Moved GameplayBehaviors out of restricted folder over to Experimental
Moved SmartObjects out of restricted folder
Moved StateTree out of restricted folder
Moved ZoneGraph out of restricted folder
Moved ZoneGraphAnnotations out of restricted folder

diff --git a/Engine/Plugins/Experimental/GameplayBehaviors/Source/GameplayBehaviorsModule/Private/AI/GameplayBehavior_BehaviorTree.cpp b/Engine/Plugins/Experimental/GameplayBehaviors/Source/GameplayBehaviorsModule/Private/AI/GameplayBehavior_BehaviorTree.cpp
new file mode 100644
--- /dev/null
+++ b/Engine/Plugins/Experimental/GameplayBehaviors/Source/GameplayBehaviorsModule/Private/AI/GameplayBehavior_BehaviorTree.cpp
@@ -0,0 +1,67 @@
+// Copyright Epic Games, Inc. All Rights Reserved.
+
+#include "AI/GameplayBehavior_BehaviorTree.h"
+#include "AI/GameplayBehaviorConfig_BehaviorTree.h"
+#include "VisualLogger/VisualLogger.h"
+#include "Blueprint/AIBlueprintHelperLibrary.h"
+#include "AIController.h"
+#include "BehaviorTree/BehaviorTreeComponent.h"
+
+//----------------------------------------------------------------------//
+// UGameplayBehavior_BehaviorTree
+//----------------------------------------------------------------------//
+UGameplayBehavior_BehaviorTree::UGameplayBehavior_BehaviorTree(const FObjectInitializer& ObjectInitializer)
+	: Super(ObjectInitializer)
+{
+	InstantiationPolicy = EGameplayBehaviorInstantiationPolicy::ConditionallyInstantiate;
+}
+
+bool UGameplayBehavior_BehaviorTree::Trigger(AActor& InAvatar, const UGameplayBehaviorConfig* Config /* = nullptr*/, AActor* SmartObjectOwner /* = nullptr*/)
+{
+	const UGameplayBehaviorConfig_BehaviorTree* BTConfig = Cast<const UGameplayBehaviorConfig_BehaviorTree>(Config);
+	if (BTConfig == nullptr || BTConfig->GetBehaviorTree() == nullptr)
+	{
+		UE_VLOG(&InAvatar, LogGameplayBehavior, Warning, TEXT("Failed to trigger behavior %s for %s due to %s being null")
+			, *GetName(), *InAvatar.GetName(), BTConfig ? TEXT("Config->BehaviorTree") : TEXT("Config"));
+		return false;
+	}
+
+	// note that the value stored in this property is unreliable if we're in the CDO
+	// If you need this to be reliable set InstantiationPolicy to Instantiate
+	AIController = UAIBlueprintHelperLibrary::GetAIController(&InAvatar);
+	if (AIController == nullptr)
+	{
+		UE_VLOG(&InAvatar, LogGameplayBehavior, Warning, TEXT("Failed to trigger behavior %s due to %s not being AI-controlled")
+			, *GetName(), *InAvatar.GetName());
+		return false;
+	}
+
+	if (BTConfig->ShouldStorePreviousBT())
+	{
+		UBehaviorTreeComponent* BTComp = Cast<UBehaviorTreeComponent>(AIController->GetBrainComponent());
+		PreviousBT = BTComp ? BTComp->GetRootTree() : nullptr;
+	}
+
+	const bool bResult = AIController->RunBehaviorTree(BTConfig->GetBehaviorTree());
+
+	UE_CVLOG(bResult == false, &InAvatar, LogGameplayBehavior, Warning, TEXT("Failed to run behavior tree %s on %s; %s")
+		, *BTConfig->GetBehaviorTree()->GetName(), *InAvatar.GetName(), *AIController->GetName());
+
+	return bResult;
+}
+
+void UGameplayBehavior_BehaviorTree::EndBehavior(AActor& InAvatar, const bool bInterrupted)
+{
+	Super::EndBehavior(InAvatar, bInterrupted);
+
+	if (PreviousBT && AIController)
+	{
+		AIController->RunBehaviorTree(PreviousBT);
+	}
+}
+
+bool UGameplayBehavior_BehaviorTree::NeedsInstance(const UGameplayBehaviorConfig* Config) const
+{
+	const UGameplayBehaviorConfig_BehaviorTree* BTConfig = Cast<const UGameplayBehaviorConfig_BehaviorTree>(Config);
+	return BTConfig && BTConfig->ShouldStorePreviousBT();
+}
diff --git a/Engine/Plugins/Experimental/GameplayBehaviors/Source/GameplayBehaviorsModule/Private/GameplayBehaviorManager.cpp b/Engine/Plugins/Experimental/GameplayBehaviors/Source/GameplayBehaviorsModule/Private/GameplayBehaviorManager.cpp
new file mode 100644
--- /dev/null
+++ b/Engine/Plugins/Experimental/GameplayBehaviors/Source/GameplayBehaviorsModule/Private/GameplayBehaviorManager.cpp
@@ -0,0 +1,109 @@
+// Copyright Epic Games, Inc. All Rights Reserved.
+
+#include "GameplayBehaviorManager.h"
+#include "GameplayBehaviorConfig.h"
+#include "DefaultManagerInstanceTracker.h"
+
+namespace
+{
+	typedef TDefaultManagerInstanceTracker<UGameplayBehaviorManager> FGameplayBehaviorsManagerInstanceTracker;
+	FGameplayBehaviorsManagerInstanceTracker InstanceTracker;
+}
+
+UGameplayBehaviorManager::FInstanceGetterSignature UGameplayBehaviorManager::InstanceGetterDelegate = UGameplayBehaviorManager::FInstanceGetterSignature::CreateRaw(&InstanceTracker, &FGameplayBehaviorsManagerInstanceTracker::GetManagerInstance);
+
+UGameplayBehaviorManager::UGameplayBehaviorManager(const FObjectInitializer& ObjectInitializer)
+	: Super(ObjectInitializer)
+{
+	bCreateIfMissing = true;
+}
+
+void UGameplayBehaviorManager::PostInitProperties()
+{
+	Super::PostInitProperties();
+
+	if (HasAnyFlags(RF_ClassDefaultObject))
+	{
+		InstanceTracker.bCreateIfMissing = bCreateIfMissing;
+	}
+}
+
+bool UGameplayBehaviorManager::StopBehavior(AActor& Avatar, TSubclassOf<UGameplayBehavior> BehaviorToStop)
+{
+	FAgentGameplayBehaviors* AgentData = AgentGameplayBehaviors.Find(&Avatar);
+
+	if (AgentData)
+	{
+		for (int32 Index = AgentData->Behaviors.Num() - 1; Index >= 0; --Index)
+		{
+			UGameplayBehavior* Beh = AgentData->Behaviors[Index];
+			// @todo make sure we're aware of this in UGameplayBehaviorManager::OnBehaviorFinished
+			if (Beh && (!BehaviorToStop || Beh->IsA(BehaviorToStop)))
+			{
+				Beh->EndBehavior(Avatar, /*bInterrupted=*/true);
+			}
+		}
+	}
+
+	return false;
+}
+
+void UGameplayBehaviorManager::OnBehaviorFinished(UGameplayBehavior& Behavior, AActor& Avatar, const bool bInterrupted)
+{
+	if (IsShuttingDown())
+	{
+		// ignore if we're shutting down
+		return;
+	}
+
+	FAgentGameplayBehaviors* AgentData = AgentGameplayBehaviors.Find(&Avatar);
+
+	if (AgentData)
+	{
+		const int32 BehaviorIndex = AgentData->Behaviors.Find(&Behavior);
+		if (BehaviorIndex != INDEX_NONE)
+		{
+			Behavior.GetOnBehaviorFinishedDelegate().RemoveAll(this);
+			AgentData->Behaviors.RemoveAtSwap(BehaviorIndex, 1, /*bAllowShrinking=*/false);
+		}
+	}
+}
+
+bool UGameplayBehaviorManager::TriggerBehavior(const UGameplayBehaviorConfig& Config, AActor& Avatar, AActor* SmartObjectOwner/* = nullptr*/)
+{
+	UWorld* World = Avatar.GetWorld();
+	UGameplayBehavior* BehaviorRun = World ? Config.GetBehavior(*World) : nullptr;
+	return BehaviorRun != nullptr
+		&& UGameplayBehaviorManager::TriggerBehavior(*BehaviorRun, Avatar, &Config, SmartObjectOwner);
+}
+
+bool UGameplayBehaviorManager::TriggerBehavior(UGameplayBehavior& Behavior, AActor& Avatar, const UGameplayBehaviorConfig* Config, AActor* SmartObjectOwner/* = nullptr*/)
+{
+	UWorld* World = Avatar.GetWorld();
+	if (World)
+	{
+		UGameplayBehaviorManager* ManagerInstance = InstanceGetterDelegate.Execute(*World);
+		return ManagerInstance != nullptr
+			&& ManagerInstance->TriggerBehaviorImpl(Behavior, Avatar, Config, SmartObjectOwner);
+	}
+	return false;
+}
+
+bool UGameplayBehaviorManager::TriggerBehaviorImpl(UGameplayBehavior& Behavior, AActor& Avatar, const UGameplayBehaviorConfig* Config, AActor* SmartObjectOwner/* = nullptr*/)
+{
+	if (Behavior.Trigger(Avatar, Config, SmartObjectOwner))
+	{
+		Behavior.GetOnBehaviorFinishedDelegate().AddUObject(this, &UGameplayBehaviorManager::OnBehaviorFinished);
+		
+		FAgentGameplayBehaviors& AgentData = AgentGameplayBehaviors.FindOrAdd(&Avatar);
+		AgentData.Behaviors.Add(&Behavior);
+
+		return true;
+	}
+	return false;
+}
+
+UGameplayBehaviorManager* UGameplayBehaviorManager::GetCurrent(UWorld* World)
+{
+	return World ? InstanceGetterDelegate.Execute(*World) : nullptr;
+}
diff --git a/Engine/Plugins/Experimental/GameplayBehaviors/Source/GameplayBehaviorsModule/Public/AI/BTTask_StopGameplayBehavior.h b/Engine/Plugins/Experimental/GameplayBehaviors/Source/GameplayBehaviorsModule/Public/AI/BTTask_StopGameplayBehavior.h
new file mode 100644
--- /dev/null
+++ b/Engine/Plugins/Experimental/GameplayBehaviors/Source/GameplayBehaviorsModule/Public/AI/BTTask_StopGameplayBehavior.h
@@ -0,0 +1,32 @@
+// Copyright Epic Games, Inc. All Rights Reserved.
+
+#pragma once
+
+#include "CoreMinimal.h"
+#include "UObject/ObjectMacros.h"
+#include "InputCoreTypes.h"
+#include "BehaviorTree/BTTaskNode.h"
+#include "BTTask_StopGameplayBehavior.generated.h"
+
+
+class UGameplayBehavior;
+/**
+*
+*/
+UCLASS()
+class GAMEPLAYBEHAVIORSMODULE_API UBTTask_StopGameplayBehavior : public UBTTaskNode
+{
+	GENERATED_BODY()
+public:
+	UBTTask_StopGameplayBehavior(const FObjectInitializer& ObjectInitializer = FObjectInitializer::Get());
+
+	virtual EBTNodeResult::Type ExecuteTask(UBehaviorTreeComponent& OwnerComp, uint8* NodeMemory) override;
+	virtual uint16 GetInstanceMemorySize() const override { return 0; }
+
+	virtual FString GetStaticDescription() const override;
+
+protected:
+	/** If None (the default) will stop any and all gameplay behaviors instigated by the agent*/
+	UPROPERTY(EditAnywhere, Category = "Node")
+	TSubclassOf<UGameplayBehavior> BehaviorToStop;
+};
diff --git a/Engine/Plugins/Experimental/GameplayBehaviors/Source/GameplayBehaviorsModule/Public/AI/GameplayBehaviorConfig_BehaviorTree.h b/Engine/Plugins/Experimental/GameplayBehaviors/Source/GameplayBehaviorsModule/Public/AI/GameplayBehaviorConfig_BehaviorTree.h
new file mode 100644
--- /dev/null
+++ b/Engine/Plugins/Experimental/GameplayBehaviors/Source/GameplayBehaviorsModule/Public/AI/GameplayBehaviorConfig_BehaviorTree.h
@@ -0,0 +1,26 @@
+// Copyright Epic Games, Inc. All Rights Reserved.
+
+#pragma once
+
+#include "GameplayBehaviorConfig.h"
+#include "BehaviorTree/BehaviorTree.h"
+#include "GameplayBehaviorConfig_BehaviorTree.generated.h"
+
+
+UCLASS()
+class GAMEPLAYBEHAVIORSMODULE_API UGameplayBehaviorConfig_BehaviorTree : public UGameplayBehaviorConfig
+{
+	GENERATED_BODY()
+public:
+	UGameplayBehaviorConfig_BehaviorTree(const FObjectInitializer& ObjectInitializer = FObjectInitializer::Get());
+
+	UBehaviorTree* GetBehaviorTree() const;
+	bool ShouldStorePreviousBT() const { return bRevertToPreviousBTOnFinish; }
+
+protected:
+	UPROPERTY(EditAnywhere, Category = SmartObject)
+	mutable TSoftObjectPtr<UBehaviorTree> BehaviorTree;
+
+	UPROPERTY(EditAnywhere, Category = SmartObject)
+	uint32 bRevertToPreviousBTOnFinish : 1;
+};
diff --git a/Engine/Plugins/Experimental/GameplayBehaviors/Source/GameplayBehaviorsModule/Public/AI/GameplayBehavior_BehaviorTree.h b/Engine/Plugins/Experimental/GameplayBehaviors/Source/GameplayBehaviorsModule/Public/AI/GameplayBehavior_BehaviorTree.h
new file mode 100644
--- /dev/null
+++ b/Engine/Plugins/Experimental/GameplayBehaviors/Source/GameplayBehaviorsModule/Public/AI/GameplayBehavior_BehaviorTree.h
@@ -0,0 +1,31 @@
+// Copyright Epic Games, Inc. All Rights Reserved.
+
+#pragma once
+
+#include "GameplayBehavior.h"
+#include "GameplayBehavior_BehaviorTree.generated.h"
+
+class UBehaviorTree;
+class AAIController;
+
+/** NOTE: this behavior works only for AIControlled pawns */
+UCLASS()
+class GAMEPLAYBEHAVIORSMODULE_API UGameplayBehavior_BehaviorTree : public UGameplayBehavior
+{
+	GENERATED_BODY()
+public:
+	UGameplayBehavior_BehaviorTree(const FObjectInitializer& ObjectInitializer = FObjectInitializer::Get());
+
+	virtual bool Trigger(AActor& InAvatar, const UGameplayBehaviorConfig* Config = nullptr, AActor* SmartObjectOwner = nullptr) override;
+	virtual void EndBehavior(AActor& InAvatar, const bool bInterrupted) override;
+
+protected:
+	virtual bool NeedsInstance(const UGameplayBehaviorConfig* Config) const override;
+
+protected:
+	UPROPERTY()
+	UBehaviorTree* PreviousBT;
+
+	UPROPERTY()
+	AAIController* AIController;
+};
diff --git a/Engine/Plugins/Experimental/GameplayBehaviors/Source/GameplayBehaviorsModule/Public/GameplayBehaviorConfig.h b/Engine/Plugins/Experimental/GameplayBehaviors/Source/GameplayBehaviorsModule/Public/GameplayBehaviorConfig.h
new file mode 100644
--- /dev/null
+++ b/Engine/Plugins/Experimental/GameplayBehaviors/Source/GameplayBehaviorsModule/Public/GameplayBehaviorConfig.h
@@ -0,0 +1,27 @@
+// Copyright Epic Games, Inc. All Rights Reserved.
+
+#pragma once
+
+#include "CoreMinimal.h"
+#include "UObject/ObjectMacros.h"
+#include "UObject/Object.h"
+#include "Templates/SubclassOf.h"
+#include "GameplayBehavior.h"
+#include "GameplayBehaviorConfig.generated.h"
+
+
+
+UCLASS(Blueprintable, BlueprintType, EditInlineNew, CollapseCategories)
+class GAMEPLAYBEHAVIORSMODULE_API UGameplayBehaviorConfig : public UObject
+{
+	GENERATED_BODY()
+public:
+	//UGameplayBehavior(const FObjectInitializer& ObjectInitializer = FObjectInitializer::Get());
+
+	/** Depending on the specific UGameplayBehavior class returns an instance or CDO of BehaviorClass. */
+	virtual UGameplayBehavior* GetBehavior(UWorld& World) const;
+
+protected:
+	UPROPERTY(EditDefaultsOnly, Category = GameplayBehavior)
+	TSubclassOf<UGameplayBehavior> BehaviorClass;
+};
diff --git a/Engine/Plugins/Experimental/GameplayBehaviors/Source/GameplayBehaviorsModule/Public/GameplayBehaviorManager.h b/Engine/Plugins/Experimental/GameplayBehaviors/Source/GameplayBehaviorsModule/Public/GameplayBehaviorManager.h
new file mode 100644
--- /dev/null
+++ b/Engine/Plugins/Experimental/GameplayBehaviors/Source/GameplayBehaviorsModule/Public/GameplayBehaviorManager.h
@@ -0,0 +1,56 @@
+// Copyright Epic Games, Inc. All Rights Reserved.
+
+#pragma once
+
+#include "AISubsystem.h"
+#include "GameplayBehaviorManager.generated.h"
+
+
+class UGameplayBehavior;
+class AActor;
+class UGameplayBehaviorConfig;
+class UWorld;
+
+USTRUCT()
+struct FAgentGameplayBehaviors
+{
+	GENERATED_BODY()
+
+	UPROPERTY()
+	TArray<UGameplayBehavior*> Behaviors;
+};
+
+UCLASS(config = Game, defaultconfig, Transient)
+class GAMEPLAYBEHAVIORSMODULE_API UGameplayBehaviorManager : public UAISubsystem//, public FSelfRegisteringExec
+{
+	GENERATED_BODY()
+public:
+	DECLARE_DELEGATE_RetVal_OneParam(UGameplayBehaviorManager*, FInstanceGetterSignature, UWorld& /*World*/);
+
+	UGameplayBehaviorManager(const FObjectInitializer& ObjectInitializer = FObjectInitializer::Get());
+	virtual void PostInitProperties() override;
+
+	bool StopBehavior(AActor& Avatar, TSubclassOf<UGameplayBehavior> BehaviorToStop);
+
+	static UGameplayBehaviorManager* GetCurrent(UWorld* World);
+	static bool TriggerBehavior(const UGameplayBehaviorConfig& Config, AActor& Avatar, AActor* SmartObjectOwner = nullptr);
+	static bool TriggerBehavior(UGameplayBehavior& Behavior, AActor& Avatar, const UGameplayBehaviorConfig* Config, AActor* SmartObjectOwner = nullptr);
+
+	bool IsShuttingDown() const { return IsPendingKillOrUnreachable(); }
+
+	UWorld* GetWorldFast() const { return Cast<UWorld>(GetOuter()); }
+
+protected:
+	void OnBehaviorFinished(UGameplayBehavior& Behavior, AActor& Avatar, const bool bInterrupted);
+
+	virtual bool TriggerBehaviorImpl(UGameplayBehavior& Behavior, AActor& Avatar, const UGameplayBehaviorConfig* Config, AActor* SmartObjectOwner = nullptr);
+
+protected:
+	UPROPERTY()
+	TMap<AActor*, FAgentGameplayBehaviors> AgentGameplayBehaviors;
+
+	UPROPERTY(config)
+	uint32 bCreateIfMissing : 1;
+
+	static FInstanceGetterSignature InstanceGetterDelegate;
+};
diff --git a/Engine/Plugins/Experimental/GameplayBehaviors/Source/GameplayBehaviorsModule/Public/GameplayBehavior_AnimationBased.h b/Engine/Plugins/Experimental/GameplayBehaviors/Source/GameplayBehaviorsModule/Public/GameplayBehavior_AnimationBased.h
new file mode 100644
--- /dev/null
+++ b/Engine/Plugins/Experimental/GameplayBehaviors/Source/GameplayBehaviorsModule/Public/GameplayBehavior_AnimationBased.h
@@ -0,0 +1,78 @@
+// Copyright Epic Games, Inc. All Rights Reserved.
+
+#pragma once
+
+#include "Engine/EngineTypes.h"
+#include "TimerManager.h"
+#include "GameplayBehavior.h"
+#include "GameplayBehavior_AnimationBased.generated.h"
+
+
+class AActor;
+class UAnimMontage;
+class UAbilitySystemComponent;
+
+
+USTRUCT()
+struct FMontagePlaybackData
+{
+	GENERATED_BODY()
+
+	UPROPERTY()
+	AActor* Avatar;
+
+	UPROPERTY()
+	UAnimMontage* AnimMontage;
+
+	UPROPERTY()
+	UAbilitySystemComponent* AbilityComponent;
+
+	float PlayRate;
+	FName SectionName;
+	FTimerHandle TimerHandle;
+	FTimerDelegate TimerDelegate;
+	uint8 bLoop : 1;
+
+	FMontagePlaybackData()
+	{
+		FMemory::Memzero(*this);
+	}
+
+	FMontagePlaybackData(AActor& InAvatar, UAnimMontage& InAnimMontage, const float InPlayRate = 1.f, const FName InSectionName = NAME_None, const bool bInLoop = false)
+		: Avatar(&InAvatar), AnimMontage(&InAnimMontage)
+		, AbilityComponent(nullptr), PlayRate(InPlayRate)
+		, SectionName(InSectionName), bLoop(bInLoop)
+	{}
+	
+	bool operator==(const FMontagePlaybackData& Other) const { return Avatar == Other.Avatar && AnimMontage == Other.AnimMontage; }
+	bool operator==(const AActor* InAvatar) const { return Avatar == InAvatar; }
+};
+
+
+/**
+ *	Note that this behavior is supporting playing only a single montage for a
+ *	given Avatar at a time. Trying to play multiple or using multiple UGameplayBehavior_AnimationBased
+ *	instances will result in requests overriding and interfering. */
+UCLASS()
+class GAMEPLAYBEHAVIORSMODULE_API UGameplayBehavior_AnimationBased : public UGameplayBehavior
+{
+	GENERATED_BODY()
+public:
+	UGameplayBehavior_AnimationBased(const FObjectInitializer& ObjectInitializer = FObjectInitializer::Get());
+
+	virtual bool Trigger(AActor& InAvatar, const UGameplayBehaviorConfig* Config = nullptr, AActor* SmartObjectOwner = nullptr) override;
+	virtual void EndBehavior(AActor& InAvatar, const bool bInterrupted) override;
+
+	bool PlayMontage(AActor& InAvatar, UAnimMontage& AnimMontage, const float InPlayRate = 1.f, const FName StartSectionName = NAME_None, const bool bLoop = false);
+
+protected:
+	UFUNCTION()
+	void OnMontageFinished(UAnimMontage* Montage, bool bInterrupted, AActor* InAvatar);
+
+	/**
+	 *	If this array ever gets more than couple elements at a time we should consider
+	 *	switching over to a TMap
+	 */ 
+	UPROPERTY()
+	mutable TArray<FMontagePlaybackData> ActivePlayback;
+};
diff --git a/Engine/Plugins/Experimental/GameplayBehaviors/Source/GameplayBehaviorsModule/Public/GameplayBehaviorsBlueprintFunctionLibrary.h b/Engine/Plugins/Experimental/GameplayBehaviors/Source/GameplayBehaviorsModule/Public/GameplayBehaviorsBlueprintFunctionLibrary.h
new file mode 100644
--- /dev/null
+++ b/Engine/Plugins/Experimental/GameplayBehaviors/Source/GameplayBehaviorsModule/Public/GameplayBehaviorsBlueprintFunctionLibrary.h
@@ -0,0 +1,22 @@
+// Copyright Epic Games, Inc. All Rights Reserved.
+
+#pragma once
+
+#include "CoreMinimal.h"
+#include "UObject/ObjectMacros.h"
+#include "Kismet/BlueprintFunctionLibrary.h"
+#include "GameplayBehavior.h"
+#include "GameplayBehaviorsBlueprintFunctionLibrary.generated.h"
+
+
+class AActor;
+
+UCLASS(meta = (ScriptName = "GameplayBehaviorsLibrary"))
+class GAMEPLAYBEHAVIORSMODULE_API UGameplayBehaviorsBlueprintFunctionLibrary : public UBlueprintFunctionLibrary
+{
+	GENERATED_BODY()
+public:
+	/** Will force-stop GameplayBehavior on given Avatar assuming the current 
+	 *	behavior is of GameplayBehaviorClass class*/
+	static bool StopGameplayBehavior(TSubclassOf<UGameplayBehavior> GameplayBehaviorClass,  AActor* Avatar);
+};
diff --git a/Engine/Plugins/Runtime/StateTree/Config/BaseStateTree.ini b/Engine/Plugins/Runtime/StateTree/Config/BaseStateTree.ini
new file mode 100644
--- /dev/null
+++ b/Engine/Plugins/Runtime/StateTree/Config/BaseStateTree.ini
@@ -0,0 +1,9 @@
+[CoreRedirects]
+
++PropertyRedirects=(OldName="StateTreeState.DefaultTransition",NewName="StateDoneTransition")
++PropertyRedirects=(OldName="BakedStateTreeState.DefaultTransitionState",NewName="StateDoneTransitionState")
++PropertyRedirects=(OldName="BakedStateTreeState.DefaultTransitionType",NewName="StateDoneTransitionType")
++PropertyRedirects=(OldName="StateTreeBindableStructDesc.ScriptStruct",NewName="StateTreeBindableStructDesc.Struct")
+
++EnumRedirects=(OldName="EStateTreePropertyCopyType",ValueChanges=(("Plain","CopyPlain"),("Complex","CopyComplex"),("Bool","CopyBool"),("Struct","CopyStruct"),("Object","CopyObject"),("Name","CopyName"),("FixedArray","CopyFixedArray")))
++EnumRedirects=(OldName="EStateTreePropertyAccessType",ValueChanges=(("Array","IndexArray")))
diff --git a/Engine/Plugins/Runtime/StateTree/Source/StateTreeEditorModule/Private/AssetTypeActions_StateTree.cpp b/Engine/Plugins/Runtime/StateTree/Source/StateTreeEditorModule/Private/AssetTypeActions_StateTree.cpp
new file mode 100644
--- /dev/null
+++ b/Engine/Plugins/Runtime/StateTree/Source/StateTreeEditorModule/Private/AssetTypeActions_StateTree.cpp
@@ -0,0 +1,36 @@
+// Copyright Epic Games, Inc. All Rights Reserved.
+
+#include "AssetTypeActions_StateTree.h"
+#include "StateTree.h"
+#include "IStateTreeEditor.h"
+#include "StateTreeEditorModule.h"
+#include "AIModule.h"
+
+#define LOCTEXT_NAMESPACE "AssetTypeActions"
+
+void FAssetTypeActions_StateTree::OpenAssetEditor(const TArray<UObject*>& InObjects, TSharedPtr<IToolkitHost> EditWithinLevelEditor)
+{
+	EToolkitMode::Type Mode = EditWithinLevelEditor.IsValid() ? EToolkitMode::WorldCentric : EToolkitMode::Standalone;
+
+	for (auto ObjIt = InObjects.CreateConstIterator(); ObjIt; ++ObjIt)
+	{
+		if (UStateTree* StateTree = Cast<UStateTree>(*ObjIt))
+		{
+			FStateTreeEditorModule& EditorModule = FModuleManager::LoadModuleChecked<FStateTreeEditorModule>("StateTreeEditorModule");
+			TSharedRef<IStateTreeEditor> NewEditor = EditorModule.CreateStateTreeEditor(EToolkitMode::Standalone, EditWithinLevelEditor, StateTree);
+		}
+	}
+}
+
+UClass* FAssetTypeActions_StateTree::GetSupportedClass() const
+{
+	return UStateTree::StaticClass();
+}
+
+uint32 FAssetTypeActions_StateTree::GetCategories()
+{ 
+	IAIModule& AIModule = FModuleManager::GetModuleChecked<IAIModule>("AIModule").Get();
+	return AIModule.GetAIAssetCategoryBit();
+}
+
+#undef LOCTEXT_NAMESPACE
diff --git a/Engine/Plugins/Runtime/StateTree/Source/StateTreeEditorModule/Private/AssetTypeActions_StateTree.h b/Engine/Plugins/Runtime/StateTree/Source/StateTreeEditorModule/Private/AssetTypeActions_StateTree.h
new file mode 100644
--- /dev/null
+++ b/Engine/Plugins/Runtime/StateTree/Source/StateTreeEditorModule/Private/AssetTypeActions_StateTree.h
@@ -0,0 +1,21 @@
+// Copyright Epic Games, Inc. All Rights Reserved.
+
+#pragma once
+
+#include "CoreMinimal.h"
+#include "AssetTypeActions_Base.h"
+
+/**
+  * Asset type actions for UStateTree.
+ */
+
+class FAssetTypeActions_StateTree : public FAssetTypeActions_Base
+{
+public:
+	// IAssetTypeActions Implementation
+	virtual FText GetName() const override { return NSLOCTEXT("AssetTypeActions", "FAssetTypeActions_StateTree", "StateTree"); }
+	virtual FColor GetTypeColor() const override { return FColor(201, 185, 29); }
+	virtual UClass* GetSupportedClass() const override;
+	virtual void OpenAssetEditor(const TArray<UObject*>& InObjects, TSharedPtr<class IToolkitHost> EditWithinLevelEditor = TSharedPtr<IToolkitHost>()) override;
+	virtual uint32 GetCategories() override;
+};
diff --git a/Engine/Plugins/Runtime/StateTree/Source/StateTreeEditorModule/Private/Customizations/StateTreeBindableItemDetails.cpp b/Engine/Plugins/Runtime/StateTree/Source/StateTreeEditorModule/Private/Customizations/StateTreeBindableItemDetails.cpp
new file mode 100644
--- /dev/null
+++ b/Engine/Plugins/Runtime/StateTree/Source/StateTreeEditorModule/Private/Customizations/StateTreeBindableItemDetails.cpp
@@ -0,0 +1,481 @@
+// Copyright Epic Games, Inc. All Rights Reserved.
+
+#include "StateTreeBindableItemDetails.h"
+#include "DetailWidgetRow.h"
+#include "DetailLayoutBuilder.h"
+#include "IPropertyUtilities.h"
+#include "IDetailPropertyRow.h"
+#include "IDetailChildrenBuilder.h"
+#include "StateTree.h"
+#include "StateTreeEditorData.h"
+#include "Widgets/DeclarativeSyntaxSupport.h"
+#include "Widgets/Input/SComboButton.h"
+#include "Widgets/Input/SEditableTextBox.h"
+#include "Widgets/Images/SImage.h"
+#include "Framework/MultiBox/MultiBoxBuilder.h"
+#include "InstancedStruct.h"
+#include "Styling/SlateIconFinder.h"
+#include "Editor.h"
+#include "Modules/ModuleManager.h"
+#include "StructViewerModule.h"
+#include "StructViewerFilter.h"
+#include "InstancedStructDetails.h"
+#include "Engine/UserDefinedStruct.h"
+#include "StateTreeBindingExtension.h"
+#include "StateTreeDelegates.h"
+#include "StateTreeTaskBase.h"
+#include "StateTreeEvaluatorBase.h"
+
+#define LOCTEXT_NAMESPACE "StateTreeEditor"
+
+// Customized version of FInstancedStructDataDetails used to hide bindable properties.
+class FBindableItemDataDetails : public FInstancedStructDataDetails
+{
+public:
+	FBindableItemDataDetails(TSharedPtr<IPropertyHandle> InStructProperty, UStateTreeEditorData* InEditorData)
+		: FInstancedStructDataDetails(InStructProperty)
+		, EditorData(InEditorData)
+	{
+		EditorPropBindings = EditorData ? EditorData->GetPropertyEditorBindings() : nullptr;
+	}
+
+	virtual void OnChildRowAdded(IDetailPropertyRow& ChildRow)
+	{
+		static const FName BindableMetaName(TEXT("Bindable"));
+
+		if (!EditorPropBindings)
+		{
+			return;
+		}
+		
+		TSharedPtr<IPropertyHandle> ChildPropHandle = ChildRow.GetPropertyHandle();
+
+		// Conditionally control visibility of the value field of bound properties.
+		if (ChildPropHandle->HasMetaData(BindableMetaName))
+		{
+			// Find property path relative to closest outer bindable struct.
+			FStateTreeEditorPropertyPath Path;
+			UE::StateTree::PropertyBinding::GetOuterStructPropertyPath(ChildPropHandle, Path);
+
+			TSharedPtr<SWidget> NameWidget;
+			TSharedPtr<SWidget> ValueWidget;
+			FDetailWidgetRow Row;
+			ChildRow.GetDefaultWidgets(NameWidget, ValueWidget, Row);
+
+			auto IsValueVisible = TAttribute<EVisibility>::Create([this, Path]() -> EVisibility
+				{
+					return EditorPropBindings->HasPropertyBinding(Path) ? EVisibility::Collapsed : EVisibility::Visible;
+				});
+
+			ChildRow
+				.CustomWidget(/*bShowChildren*/true)
+				.NameContent()
+				[
+					NameWidget.ToSharedRef()
+				]
+				.ValueContent()
+				[
+					SNew(SBox)
+					.Visibility(IsValueVisible)
+					[
+						ValueWidget.ToSharedRef()
+					]
+				];
+		}
+	}
+
+	UStateTreeEditorData* EditorData;
+	FStateTreeEditorPropertyBindings* EditorPropBindings;
+};
+
+
+////////////////////////////////////
+
+class FBindableItemStructFilter : public IStructViewerFilter
+{
+public:
+	/** The base struct for the property that classes must be a child-of. */
+	const UScriptStruct* BaseStruct = nullptr;
+
+	// A flag controlling whether we allow UserDefinedStructs
+	bool bAllowUserDefinedStructs = false;
+
+	// A flag controlling whether we allow to select the BaseStruct
+	bool bAllowBaseStruct = true;
+
+	// Schema to filter which structs are allowed.
+	const UStateTreeSchema* Schema = nullptr;
+
+	virtual bool IsStructAllowed(const FStructViewerInitializationOptions& InInitOptions, const UScriptStruct* InStruct, TSharedRef<FStructViewerFilterFuncs> InFilterFuncs) override
+	{
+		if (Schema && !Schema->IsStructAllowed(InStruct))
+		{
+			return false;
+		}
+
+		if (InStruct->IsA<UUserDefinedStruct>())
+		{
+			return bAllowUserDefinedStructs;
+		}
+
+		if (InStruct == BaseStruct)
+		{
+			return bAllowBaseStruct;
+		}
+
+		// Query the native struct to see if it has the correct parent type (if any)
+		return !BaseStruct || InStruct->IsChildOf(BaseStruct);
+	}
+
+	virtual bool IsUnloadedStructAllowed(const FStructViewerInitializationOptions& InInitOptions, const FName InStructPath, TSharedRef<FStructViewerFilterFuncs> InFilterFuncs) override
+	{
+		// User Defined Structs don't support inheritance, so only include them requested
+		return bAllowUserDefinedStructs;
+	}
+};
+
+
+////////////////////////////////////
+
+TSharedRef<IPropertyTypeCustomization> FStateTreeBindableItemDetails::MakeInstance()
+{
+	return MakeShareable(new FStateTreeBindableItemDetails);
+}
+
+void FStateTreeBindableItemDetails::CustomizeHeader(TSharedRef<class IPropertyHandle> StructPropertyHandle, class FDetailWidgetRow& HeaderRow, IPropertyTypeCustomizationUtils& StructCustomizationUtils)
+{
+	StructProperty = StructPropertyHandle;
+	PropUtils = StructCustomizationUtils.GetPropertyUtilities().Get();
+
+	TypeProperty = StructProperty->GetChildHandle(TEXT("Type"));
+
+	// Find base class from meta data.
+	static const FName BaseClassMetaName(TEXT("BaseStruct")); // TODO: move these names into one central place.
+	const FString BaseClassName = TypeProperty->GetMetaData(BaseClassMetaName);
+	BaseScriptStruct = FindObject<UScriptStruct>(ANY_PACKAGE, *BaseClassName);
+
+	UE::StateTree::Delegates::OnIdentifierChanged.AddSP(this, &FStateTreeBindableItemDetails::OnIdentifierChanged);
+
+	FindOuterObjects();
+
+	HeaderRow
+		.WholeRowContent()
+		.VAlign(VAlign_Center)
+		[
+			SNew(SHorizontalBox)
+			// Name
+			+ SHorizontalBox::Slot()
+			.FillWidth(1.5f)
+			.VAlign(VAlign_Center)
+			.Padding(FMargin(FMargin(0.0f, 0.0f, 4.0f, 0.0f)))
+			[
+				SNew(SEditableTextBox)
+				.IsEnabled(TAttribute<bool>(this, &FStateTreeBindableItemDetails::IsNameEnabled))
+				.Text(this, &FStateTreeBindableItemDetails::GetName)
+				.OnTextCommitted(this, &FStateTreeBindableItemDetails::OnNameCommitted)
+				.SelectAllTextWhenFocused(true)
+				.RevertTextOnEscape(true)
+				.Font(IDetailLayoutBuilder::GetDetailFont())
+			]
+			// Class picker
+			+ SHorizontalBox::Slot()
+			.FillWidth(1.0f)
+			.VAlign(VAlign_Center)
+			[
+				SAssignNew(ComboButton, SComboButton)
+				.OnGetMenuContent(this, &FStateTreeBindableItemDetails::GenerateStructPicker)
+				.ContentPadding(0)
+				.ButtonContent()
+				[
+					SNew(SHorizontalBox)
+					+ SHorizontalBox::Slot()
+					.AutoWidth()
+					.VAlign(VAlign_Center)
+					.Padding(0.0f, 0.0f, 4.0f, 0.0f)
+					[
+						SNew(SImage)
+						.Image(this, &FStateTreeBindableItemDetails::GetDisplayValueIcon)
+					]
+					+ SHorizontalBox::Slot()
+					.VAlign(VAlign_Center)
+					[
+						SNew(STextBlock)
+						.Text(this, &FStateTreeBindableItemDetails::GetDisplayValueString)
+						.Font(IDetailLayoutBuilder::GetDetailFont())
+					]
+				]
+			]
+			+ SHorizontalBox::Slot()
+			.AutoWidth()
+			[
+				StructPropertyHandle->CreateDefaultPropertyButtonWidgets()
+			]
+		];
+}
+
+void FStateTreeBindableItemDetails::CustomizeChildren(TSharedRef<class IPropertyHandle> StructPropertyHandle, class IDetailChildrenBuilder& StructBuilder, IPropertyTypeCustomizationUtils& StructCustomizationUtils)
+{
+	TSharedRef<FBindableItemDataDetails> DataDetails = MakeShareable(new FBindableItemDataDetails(TypeProperty, EditorData));
+	StructBuilder.AddCustomBuilder(DataDetails);
+}
+
+void FStateTreeBindableItemDetails::OnIdentifierChanged(const UStateTree& InStateTree)
+{
+	if (PropUtils && StateTree == &InStateTree)
+	{
+		PropUtils->ForceRefresh();
+	}
+}
+
+void FStateTreeBindableItemDetails::FindOuterObjects()
+{
+	EditorData = nullptr;
+	StateTree = nullptr;
+
+	if (TypeProperty)
+	{
+		TArray<UObject*> OuterObjects;
+		TypeProperty->GetOuterObjects(OuterObjects);
+		for (UObject* Outer : OuterObjects)
+		{
+			UStateTreeEditorData* OuterEditorData = Outer->GetTypedOuter<UStateTreeEditorData>();
+			UStateTree* OuterStateTree = OuterEditorData ? OuterEditorData->GetTypedOuter<UStateTree>() : nullptr;
+			if (OuterEditorData && OuterStateTree)
+			{
+				StateTree = OuterStateTree;
+				EditorData = OuterEditorData;
+				break;
+			}
+		}
+	}
+}
+
+FText FStateTreeBindableItemDetails::GetName() const
+{
+	if (TypeProperty)
+	{
+		TArray<void*> RawData;
+		TypeProperty->AccessRawData(RawData);
+		if (RawData.Num() == 1)
+		{
+			// Dig out name from the struct without knowing the type.
+			FInstancedStruct* StructPtr = static_cast<FInstancedStruct*>(RawData[0]);
+
+			// Make sure the associated property has valid data (e.g. while moving an item in array)
+			const UScriptStruct* ScriptStruct = StructPtr != nullptr ? StructPtr->GetScriptStruct() : nullptr;
+			if (ScriptStruct != nullptr)
+			{
+				if (FProperty* NameProperty = ScriptStruct->FindPropertyByName(TEXT("Name")))
+				{
+					if (NameProperty->IsA(FNameProperty::StaticClass()))
+					{
+						void* Ptr = const_cast<void*>(static_cast<const void*>(StructPtr->GetMemory()));
+						FName NameValue = *NameProperty->ContainerPtrToValuePtr<FName>(Ptr);
+						return FText::FromName(NameValue);
+					}
+				}
+			}
+			return LOCTEXT("Empty", "Empty");
+		}
+		else
+		{
+			return LOCTEXT("MultipleSelected", "Multiple selected");
+		}
+	}
+	return FText::GetEmpty();
+}
+
+void FStateTreeBindableItemDetails::OnNameCommitted(const FText& NewText, ETextCommit::Type InTextCommit)
+{
+	if (InTextCommit == ETextCommit::OnEnter || InTextCommit == ETextCommit::OnUserMovedFocus)
+	{
+		// Remove excess whitespace and prevent categories with just spaces
+		FString NewName = FText::TrimPrecedingAndTrailing(NewText).ToString();
+
+		if (TypeProperty)
+		{
+			TArray<void*> RawData;
+			TypeProperty->AccessRawData(RawData);
+			if (RawData.Num() == 1)
+			{
+				if (GEditor)
+				{
+					GEditor->BeginTransaction(LOCTEXT("SetName", "Set Name"));
+				}
+				TypeProperty->NotifyPreChange();
+					
+				// Set Name
+				FInstancedStruct* StructPtr = static_cast<FInstancedStruct*>(RawData[0]);
+				if (const UScriptStruct* ScriptStruct = StructPtr->GetScriptStruct())
+				{
+					if (FProperty* NameProperty = ScriptStruct->FindPropertyByName(TEXT("Name")))
+					{
+						if (NameProperty->IsA(FNameProperty::StaticClass()))
+						{
+							void* Ptr = const_cast<void*>(static_cast<const void*>(StructPtr->GetMemory()));
+							FName& NameValue = *NameProperty->ContainerPtrToValuePtr<FName>(Ptr);
+							NameValue = FName(NewName);
+						}
+					}
+				}
+
+				TypeProperty->NotifyPostChange(EPropertyChangeType::ValueSet);
+
+				if (StateTree)
+				{
+					UE::StateTree::Delegates::OnIdentifierChanged.Broadcast(*StateTree);
+				}
+
+				if (GEditor)
+				{
+					GEditor->EndTransaction();
+				}
+
+				TypeProperty->NotifyFinishedChangingProperties();
+			}
+		}
+	}
+}
+
+bool FStateTreeBindableItemDetails::IsNameEnabled() const
+{
+	// Can only edit if we have valid instantiated type.
+	return GetCommonScriptStruct() != nullptr;
+}
+
+const UScriptStruct* FStateTreeBindableItemDetails::GetCommonScriptStruct() const
+{
+	if (!TypeProperty)
+	{
+		return nullptr;
+	}
+	const UScriptStruct* CommonScriptStruct = nullptr;
+	bool bMultipleValues = false;
+	TArray<void*> RawData;
+	TypeProperty->AccessRawData(RawData);
+	for (void* Data : RawData)
+	{
+		if (FInstancedStruct* Struct = static_cast<FInstancedStruct*>(Data))
+		{
+			const UScriptStruct* ScriptStruct = Struct->GetScriptStruct();
+			if (!bMultipleValues && !CommonScriptStruct)
+			{
+				CommonScriptStruct = ScriptStruct;
+			}
+			else if (ScriptStruct != CommonScriptStruct)
+			{
+				CommonScriptStruct = nullptr;
+				bMultipleValues = true;
+			}
+		}
+	}
+
+	return CommonScriptStruct;
+}
+
+FText FStateTreeBindableItemDetails::GetDisplayValueString() const
+{
+	const UScriptStruct* ScriptStruct = GetCommonScriptStruct();
+	if (ScriptStruct)
+	{
+		return ScriptStruct->GetDisplayNameText();
+	}
+	return FText();
+}
+
+const FSlateBrush* FStateTreeBindableItemDetails::GetDisplayValueIcon() const
+{
+	return FSlateIconFinder::FindIconBrushForClass(UScriptStruct::StaticClass());
+}
+
+TSharedRef<SWidget> FStateTreeBindableItemDetails::GenerateStructPicker()
+{
+	static const FName ExcludeBaseStruct(TEXT("ExcludeBaseStruct")); // TODO: move these names into one central place.
+	const bool bExcludeBaseStruct = TypeProperty->HasMetaData(ExcludeBaseStruct);
+
+	TSharedRef<FBindableItemStructFilter> StructFilter = MakeShared<FBindableItemStructFilter>();
+	StructFilter->BaseStruct = BaseScriptStruct;
+	StructFilter->bAllowUserDefinedStructs = false;
+	StructFilter->bAllowBaseStruct = !bExcludeBaseStruct;
+	StructFilter->Schema = StateTree ? StateTree->GetSchema() : nullptr;
+
+	FStructViewerInitializationOptions Options;
+	Options.bShowUnloadedStructs = true;
+	Options.bShowNoneOption = true;
+	Options.StructFilter = StructFilter;
+	Options.NameTypeToDisplay = EStructViewerNameTypeToDisplay::DisplayName;
+	Options.DisplayMode = EStructViewerDisplayMode::ListView;
+	Options.bAllowViewOptions = true;
+
+	FOnStructPicked OnPicked(FOnStructPicked::CreateRaw(this, &FStateTreeBindableItemDetails::OnStructPicked));
+
+	return SNew(SBox)
+		.WidthOverride(280)
+		[
+			SNew(SVerticalBox)
+			+ SVerticalBox::Slot()
+			.AutoHeight()
+			.MaxHeight(500)
+			[
+				FModuleManager::LoadModuleChecked<FStructViewerModule>("StructViewer").CreateStructViewer(Options, OnPicked)
+			]
+		];
+}
+
+void FStateTreeBindableItemDetails::OnStructPicked(const UScriptStruct* InStruct)
+{
+	if (TypeProperty)
+	{
+		GEditor->BeginTransaction(LOCTEXT("OnStructPicked", "Set Struct"));
+
+		TypeProperty->NotifyPreChange();
+
+		TArray<void*> RawData;
+		TypeProperty->AccessRawData(RawData);
+		for (void* Data : RawData)
+		{
+			if (FInstancedStruct* Struct = static_cast<FInstancedStruct*>(Data))
+			{
+				if (InStruct)
+				{
+					Struct->InitializeAs(InStruct);
+
+					// Generate new ID and name.
+					if (InStruct->IsChildOf(FStateTreeTask2Base::StaticStruct()))
+					{
+						FStateTreeTask2Base* Task = Struct->GetMutablePtr<FStateTreeTask2Base>();
+						check(Task);
+						Task->ID = FGuid::NewGuid();
+						Task->Name = FName(InStruct->GetDisplayNameText().ToString());
+					}
+					else if (InStruct->IsChildOf(FStateTreeEvaluator2Base::StaticStruct()))
+					{
+						FStateTreeEvaluator2Base* Eval = Struct->GetMutablePtr<FStateTreeEvaluator2Base>();
+						check(Eval);
+						Eval->ID = FGuid::NewGuid();
+						Eval->Name = FName(InStruct->GetDisplayNameText().ToString());
+					}
+				}
+				else
+				{
+					Struct->Reset();
+				}
+			}
+		}
+
+		TypeProperty->NotifyPostChange(EPropertyChangeType::ValueSet);
+		GEditor->EndTransaction();
+		TypeProperty->NotifyFinishedChangingProperties();
+
+		GEditor->EndTransaction();
+	}
+
+	ComboButton->SetIsOpen(false);
+
+	if (PropUtils)
+	{
+		PropUtils->ForceRefresh();
+	}
+}
+
+#undef LOCTEXT_NAMESPACE
diff --git a/Engine/Plugins/Runtime/StateTree/Source/StateTreeEditorModule/Private/Customizations/StateTreeBindableItemDetails.h b/Engine/Plugins/Runtime/StateTree/Source/StateTreeEditorModule/Private/Customizations/StateTreeBindableItemDetails.h
new file mode 100644
--- /dev/null
+++ b/Engine/Plugins/Runtime/StateTree/Source/StateTreeEditorModule/Private/Customizations/StateTreeBindableItemDetails.h
@@ -0,0 +1,51 @@
+// Copyright Epic Games, Inc. All Rights Reserved.
+
+#pragma once
+
+#include "CoreMinimal.h"
+#include "UnrealClient.h"
+#include "IPropertyTypeCustomization.h"
+
+class IPropertyHandle;
+class UStateTree;
+class UStateTreeState;
+class UStateTreeEditorData;
+
+/**
+ * Semi-generic type customization for bindable items (Evaluators and Tasks) in StateTreeState.
+ */
+class FStateTreeBindableItemDetails : public IPropertyTypeCustomization
+{
+public:
+	/** Makes a new instance of this detail layout class for a specific detail view requesting it */
+	static TSharedRef<IPropertyTypeCustomization> MakeInstance();
+
+	/** IPropertyTypeCustomization interface */
+	virtual void CustomizeHeader(TSharedRef<IPropertyHandle> StructPropertyHandle, class FDetailWidgetRow& HeaderRow, IPropertyTypeCustomizationUtils& StructCustomizationUtils) override;
+	virtual void CustomizeChildren(TSharedRef<IPropertyHandle> StructPropertyHandle, class IDetailChildrenBuilder& StructBuilder, IPropertyTypeCustomizationUtils& StructCustomizationUtils) override;
+
+private:
+
+	FText GetName() const;
+	void OnNameCommitted(const FText& NewText, ETextCommit::Type InTextCommit);
+	bool IsNameEnabled() const;
+
+	const UScriptStruct* GetCommonScriptStruct() const;
+	FText GetDisplayValueString() const;
+	const FSlateBrush* GetDisplayValueIcon() const;
+	TSharedRef<SWidget> GenerateStructPicker();
+	void OnStructPicked(const UScriptStruct* InStruct);
+
+	void OnIdentifierChanged(const UStateTree& StateTree);
+	void FindOuterObjects();
+
+	UScriptStruct* BaseScriptStruct = nullptr;
+	TSharedPtr<class SComboButton> ComboButton;
+
+	UStateTreeEditorData* EditorData = nullptr;
+	UStateTree* StateTree = nullptr;
+
+	class IPropertyUtilities* PropUtils = nullptr;
+	TSharedPtr<IPropertyHandle> StructProperty;
+	TSharedPtr<IPropertyHandle> TypeProperty;
+};
diff --git a/Engine/Plugins/Runtime/StateTree/Source/StateTreeEditorModule/Private/Customizations/StateTreeBindingExtension.cpp b/Engine/Plugins/Runtime/StateTree/Source/StateTreeEditorModule/Private/Customizations/StateTreeBindingExtension.cpp
new file mode 100644
--- /dev/null
+++ b/Engine/Plugins/Runtime/StateTree/Source/StateTreeEditorModule/Private/Customizations/StateTreeBindingExtension.cpp
@@ -0,0 +1,307 @@
+// Copyright Epic Games, Inc. All Rights Reserved.
+
+#include "StateTreeBindingExtension.h"
+#include "DetailWidgetRow.h"
+#include "DetailLayoutBuilder.h"
+#include "IDetailChildrenBuilder.h"
+#include "EdGraphSchema_K2.h"
+#include "IPropertyAccessEditor.h"
+#include "Features/IModularFeatures.h"
+#include "StateTreeEvaluatorBase.h"
+#include "StateTree.h"
+#include "Algo/Accumulate.h"
+#include "InstancedStruct.h"
+#include "StateTreePropertyHelpers.h"
+#include "StateTreeAnyEnum.h"
+#include "StateTreePropertyBindingCompiler.h"
+
+#define LOCTEXT_NAMESPACE "StateTreeEditor"
+
+
+namespace UE { namespace StateTree { namespace PropertyBinding {
+
+UObject* FindEditorBindingsOwner(UObject* InObject)
+{
+	UObject* Result = nullptr;
+
+	for (UObject* Outer = InObject; Outer; Outer = Outer->GetOuter())
+	{
+		IStateTreeEditorPropertyBindingsOwner* BindingOwner = Cast<IStateTreeEditorPropertyBindingsOwner>(Outer);
+		if (BindingOwner)
+		{
+			Result = Outer;
+			break;
+		}
+	}
+	return Result;
+}
+
+void GetOuterStructPropertyPath(TSharedPtr<IPropertyHandle> InPropertyHandle, FStateTreeEditorPropertyPath& OutPath)
+{
+	FGuid StructID;
+	TArray<FBindingChainElement> BindingChain;
+	TSharedPtr<IPropertyHandle> CurrentPropertyHandle = InPropertyHandle;
+
+	while (CurrentPropertyHandle.IsValid())
+	{
+		if (FProperty* CurrentProperty = CurrentPropertyHandle->GetProperty())
+		{
+			// Keep track of the path.
+			BindingChain.Insert(FBindingChainElement(CurrentProperty, CurrentPropertyHandle->GetIndexInArray()), 0);
+
+			// If the owner struct of this property has a property called ID and it's FGuid, we'll use that as struct ID.
+			TSharedPtr<IPropertyHandle> OwnerHandle = CurrentPropertyHandle->GetParentHandle();
+			if (OwnerHandle.IsValid())
+			{
+				TSharedPtr<IPropertyHandle> IdHandle = OwnerHandle->GetChildHandle(TEXT("ID"), /*bRecurse*/false); // Omitting recursive, since we expect that the ID belongs to owner.
+				FGuid ID;
+				if (UE::StateTree::PropertyHelpers::GetStructValue<FGuid>(IdHandle, ID) == FPropertyAccess::Success)
+				{
+					StructID = ID;
+					break;
+				}
+			}
+
+			CurrentPropertyHandle = CurrentPropertyHandle->GetParentHandle();
+		}
+		else
+		{
+			BindingChain.Reset();
+			CurrentPropertyHandle.Reset();
+			break;
+		}
+	}
+
+	OutPath.StructID = StructID;
+
+	IPropertyAccessEditor& PropertyAccessEditor = IModularFeatures::Get().GetModularFeature<IPropertyAccessEditor>("PropertyAccessEditor");
+	PropertyAccessEditor.MakeStringPath(BindingChain, OutPath.Path);
+}
+
+FOnStateTreeBindingChanged STATETREEEDITORMODULE_API OnStateTreeBindingChanged;
+
+}}} // UE::StateTree::PropertyBinding
+
+bool FStateTreeBindingExtension::IsPropertyExtendable(const UClass* InObjectClass, const IPropertyHandle& PropertyHandle) const
+{
+	static const FName BindableMetaName(TEXT("Bindable"));
+
+	return PropertyHandle.HasMetaData(BindableMetaName);
+}
+
+
+void FStateTreeBindingExtension::ExtendWidgetRow(FDetailWidgetRow& InWidgetRow, const IDetailLayoutBuilder& InDetailBuilder, const UClass* InObjectClass, TSharedPtr<IPropertyHandle> InPropertyHandle)
+{
+	if (!IModularFeatures::Get().IsModularFeatureAvailable("PropertyAccessEditor"))
+	{
+		return;
+	}
+
+	IPropertyAccessEditor& PropertyAccessEditor = IModularFeatures::Get().GetModularFeature<IPropertyAccessEditor>("PropertyAccessEditor");
+
+	UObject* OwnerObject = nullptr;
+	FStateTreeEditorPropertyBindings* EditorBindings = nullptr;
+
+	// Array of structs we can bind to.
+	TArray<FBindingContextStruct> Context;
+	TArray<FStateTreeBindableStructDesc> AccessibleStructs;
+
+	// The struct and property where we're binding.
+	FStateTreeEditorPropertyPath TargetPath;
+
+	TArray<UObject*> OuterObjects;
+	InPropertyHandle->GetOuterObjects(OuterObjects);
+	if (OuterObjects.Num() == 1)
+	{
+		// Only allow to binding when one object is selected.
+		OwnerObject = UE::StateTree::PropertyBinding::FindEditorBindingsOwner(OuterObjects[0]);
+
+		// Figure out the structs we're editing, and property path relative to current property.
+		UE::StateTree::PropertyBinding::GetOuterStructPropertyPath(InPropertyHandle, TargetPath);
+
+		if (IStateTreeEditorPropertyBindingsOwner* BindingOwner = Cast<IStateTreeEditorPropertyBindingsOwner>(OwnerObject))
+		{
+			EditorBindings = BindingOwner->GetPropertyEditorBindings();
+			BindingOwner->GetAccessibleStructs(TargetPath.StructID, AccessibleStructs);
+			for (FStateTreeBindableStructDesc& StructDesc : AccessibleStructs)
+			{
+				const UStruct* Struct = StructDesc.Struct;
+				Context.Emplace(const_cast<UStruct*>(Struct), nullptr, FText::FromString(StructDesc.Name.ToString()));
+			}
+		}
+	}
+
+	FProperty* Property = InPropertyHandle->GetProperty();
+
+	bool bIsAnyEnum = false;
+	if (FStructProperty* StructProperty = CastField<FStructProperty>(Property))
+	{
+		bIsAnyEnum = StructProperty->Struct == FStateTreeAnyEnum::StaticStruct();
+	}
+	
+	FPropertyBindingWidgetArgs Args;
+	Args.Property = InPropertyHandle->GetProperty();
+
+	Args.OnCanBindProperty = FOnCanBindProperty::CreateLambda([EditorBindings, OwnerObject, InPropertyHandle, bIsAnyEnum](FProperty* InProperty)
+		{
+			if (!EditorBindings || !OwnerObject)
+			{
+				return false;
+			}
+
+			// Special case for binding widget calling OnCanBindProperty with Args.Property (i.e. self).
+			if (InPropertyHandle->GetProperty() == InProperty)
+			{
+				return true;
+			}
+
+			bool bCanBind = false;
+
+			// AnyEnums need special handling.
+			// It is a struct property but we want to treat it as an enum. We need to do this here, instead of 
+			// FStateTreePropertyBindingCompiler::GetPropertyCompatibility() because the treatment depends on the value too.
+			// Note: AnyEnums will need special handling before they can be used for binding.
+			if (bIsAnyEnum)
+			{
+				FStateTreeAnyEnum AnyEnum;
+				UE::StateTree::PropertyHelpers::GetStructValue(InPropertyHandle, AnyEnum);
+
+				// If the enum class is not specified, allow to bind to any enum, if the class is specified allow only that enum.
+				if (FByteProperty* ByteProperty = CastField<FByteProperty>(InProperty))
+				{
+					if (UEnum* Enum = ByteProperty->GetIntPropertyEnum())
+					{
+						bCanBind = !AnyEnum.Enum || AnyEnum.Enum == Enum;
+					}
+				}
+				else if (FEnumProperty* EnumProperty = CastField<FEnumProperty>(InProperty))
+				{
+					bCanBind = !AnyEnum.Enum || AnyEnum.Enum == EnumProperty->GetEnum();
+				}
+			}
+			else
+			{
+				// Note: We support type promotion here
+				bCanBind = FStateTreePropertyBindingCompiler::GetPropertyCompatibility(InProperty, InPropertyHandle->GetProperty()) != EPropertyAccessCompatibility::Incompatible;
+			}
+
+			return bCanBind;
+		});
+
+	Args.OnCanAcceptPropertyOrChildren = FOnCanBindProperty::CreateLambda([](FProperty* InProperty)
+		{
+			// Make only editor visible properties visible for binding.
+			return InProperty->HasAnyPropertyFlags(CPF_Edit);
+		});
+
+	Args.OnCanBindToClass = FOnCanBindToClass::CreateLambda([](UClass* InClass)
+		{
+			return true;
+		});
+
+	Args.OnAddBinding = FOnAddBinding::CreateLambda([EditorBindings, OwnerObject, TargetPath, AccessibleStructs](FName InPropertyName, const TArray<FBindingChainElement>& InBindingChain)
+		{
+			IPropertyAccessEditor& PropertyAccessEditor = IModularFeatures::Get().GetModularFeature<IPropertyAccessEditor>("PropertyAccessEditor");
+			if (EditorBindings && OwnerObject)
+			{
+				if (TargetPath.IsValid() && InBindingChain.Num() > 1)	// Assume at least: [0] struct index, [1] a property.
+				{
+					FStateTreeEditorPropertyPath SourcePath;
+					const int32 SourceStructIndex = InBindingChain[0].ArrayIndex;
+					TArray<FBindingChainElement> SourceBindingChain(InBindingChain.GetData() + 1, InBindingChain.Num() - 1);
+					check(SourceStructIndex >= 0 && SourceStructIndex < AccessibleStructs.Num());
+
+					SourcePath.StructID = AccessibleStructs[SourceStructIndex].ID;
+					PropertyAccessEditor.MakeStringPath(SourceBindingChain, SourcePath.Path);
+
+					OwnerObject->Modify();
+					EditorBindings->AddPropertyBinding(SourcePath, TargetPath);
+
+					UE::StateTree::PropertyBinding::OnStateTreeBindingChanged.Broadcast(SourcePath, TargetPath);
+				}
+			}
+		});
+
+	Args.OnRemoveBinding = FOnRemoveBinding::CreateLambda([EditorBindings, OwnerObject, TargetPath](FName InPropertyName)
+		{
+			if (EditorBindings && OwnerObject)
+			{
+				OwnerObject->Modify();
+				EditorBindings->RemovePropertyBindings(TargetPath);
+
+				FStateTreeEditorPropertyPath SourcePath; // Null path
+				UE::StateTree::PropertyBinding::OnStateTreeBindingChanged.Broadcast(SourcePath, TargetPath);
+			}
+		});
+
+	Args.OnCanRemoveBinding = FOnCanRemoveBinding::CreateLambda([EditorBindings, TargetPath](FName InPropertyName)
+		{
+			return EditorBindings && EditorBindings->HasPropertyBinding(TargetPath);
+		});
+
+	Args.CurrentBindingText = MakeAttributeLambda([EditorBindings, TargetPath, AccessibleStructs]()
+		{
+			const FText MultipleValues = LOCTEXT("MultipleValues", "Multiple Values");
+			const FText Bind = LOCTEXT("Bind", "Bind");
+			FText CurrentValue = Bind;
+
+			if (EditorBindings)
+			{
+				if (const FStateTreeEditorPropertyPath* SourcePath = EditorBindings->GetPropertyBindingSource(TargetPath))
+				{
+					FString PropertyName;
+					for (int32 i = 0; i < AccessibleStructs.Num(); i++)
+					{
+						if (AccessibleStructs[i].ID == SourcePath->StructID)
+						{
+							PropertyName = AccessibleStructs[i].Name.ToString();
+							break;
+						}
+					}
+					for (const FString& Segment : SourcePath->Path)
+					{
+						PropertyName += TEXT(".") + Segment;
+					}
+					CurrentValue = FText::FromString(PropertyName);
+				}
+			}
+			else
+			{
+				// StateTreeEditorData is not valid if there's multiple objects selected.
+				CurrentValue = MultipleValues;
+			}
+
+			return CurrentValue;
+		});
+
+	Args.CurrentBindingImage = MakeAttributeLambda([]() -> const FSlateBrush*
+		{
+			static FName PropertyIcon(TEXT("Kismet.Tabs.Variables"));
+			return FEditorStyle::GetBrush(PropertyIcon);
+		});
+
+	Args.CurrentBindingColor = MakeAttributeLambda([InPropertyHandle]() -> FLinearColor
+		{
+			const UEdGraphSchema_K2* Schema = GetDefault<UEdGraphSchema_K2>();
+
+			FEdGraphPinType PinType;
+			Schema->ConvertPropertyToPinType(InPropertyHandle->GetProperty(), PinType);
+			const FLinearColor BindingColor = Schema->GetPinTypeColor(PinType);
+
+			// TODO: Handle coloring of type promotion
+
+			return BindingColor;
+		});
+
+	Args.bAllowNewBindings = false;
+	Args.bAllowArrayElementBindings = false;
+	Args.bAllowUObjectFunctions = false;
+
+	InWidgetRow.ExtensionContent()
+	[
+		PropertyAccessEditor.MakePropertyBindingWidget(Context, Args)
+	];
+}
+
+#undef LOCTEXT_NAMESPACE
+
diff --git a/Engine/Plugins/Runtime/StateTree/Source/StateTreeEditorModule/Private/Customizations/StateTreeBindingExtension.h b/Engine/Plugins/Runtime/StateTree/Source/StateTreeEditorModule/Private/Customizations/StateTreeBindingExtension.h
new file mode 100644
--- /dev/null
+++ b/Engine/Plugins/Runtime/StateTree/Source/StateTreeEditorModule/Private/Customizations/StateTreeBindingExtension.h
@@ -0,0 +1,42 @@
+// Copyright Epic Games, Inc. All Rights Reserved.
+
+#pragma once
+
+#include "CoreMinimal.h"
+#include "Widgets/SNullWidget.h"
+#include "IDetailPropertyExtensionHandler.h"
+
+class IPropertyHandle;
+class IDetailLayoutBuilder;
+struct FStateTreeEditorPropertyPath;
+class IPropertyAccessEditor;
+struct FStateTreeEditorPropertyPath;
+
+namespace UE { namespace StateTree { namespace PropertyBinding {
+
+	/**
+	 * Get nearest Outer that implements IStateTreeEditorPropertyBindingsOwner.
+	 * @param InObject Object where to start the search.
+	 */
+	UObject* FindEditorBindingsOwner(UObject* InObject);
+
+	/**
+	 * Returns property path for a specific property. It will look for nearest outer struct that contains property named "ID",
+	 * and build property path using that struct as root.
+	 * @param InPropertyHandle Handle to the property to find path for.
+	 * @param OutPath Resulting property path.
+	 */
+	void GetOuterStructPropertyPath(TSharedPtr<IPropertyHandle> InPropertyHandle, FStateTreeEditorPropertyPath& OutPath);
+			
+	DECLARE_MULTICAST_DELEGATE_TwoParams(FOnStateTreeBindingChanged, const FStateTreeEditorPropertyPath& /*SourcePath*/, const FStateTreeEditorPropertyPath& /*TargetPath*/);
+	extern STATETREEEDITORMODULE_API FOnStateTreeBindingChanged OnStateTreeBindingChanged;
+
+} } } // UE::StateTree::PropertyBinding
+
+class FStateTreeBindingExtension : public IDetailPropertyExtensionHandler
+{
+public:
+	// IDetailPropertyExtensionHandler interface
+	virtual bool IsPropertyExtendable(const UClass* InObjectClass, const IPropertyHandle& PropertyHandle) const override;
+	virtual void ExtendWidgetRow(FDetailWidgetRow& InWidgetRow, const IDetailLayoutBuilder& InDetailBuilder, const UClass* InObjectClass, TSharedPtr<IPropertyHandle> PropertyHandle) override;
+};
diff --git a/Engine/Plugins/Runtime/StateTree/Source/StateTreeEditorModule/Private/Customizations/StateTreeConditionDetails.cpp b/Engine/Plugins/Runtime/StateTree/Source/StateTreeEditorModule/Private/Customizations/StateTreeConditionDetails.cpp
new file mode 100644
--- /dev/null
+++ b/Engine/Plugins/Runtime/StateTree/Source/StateTreeEditorModule/Private/Customizations/StateTreeConditionDetails.cpp
@@ -0,0 +1,328 @@
+// Copyright Epic Games, Inc. All Rights Reserved.
+
+#include "StateTreeConditionDetails.h"
+#include "Widgets/DeclarativeSyntaxSupport.h"
+#include "DetailWidgetRow.h"
+#include "DetailLayoutBuilder.h"
+#include "IDetailPropertyRow.h"
+#include "IDetailChildrenBuilder.h"
+#include "IPropertyUtilities.h"
+#include "StateTreeCondition.h"
+#include "Widgets/Input/SComboButton.h"
+#include "Framework/MultiBox/MultiBoxBuilder.h"
+#include "ScopedTransaction.h"
+#include "StateTreePropertyHelpers.h"
+
+#define LOCTEXT_NAMESPACE "StateTreeEditor"
+
+TSharedRef<IPropertyTypeCustomization> FStateTreeConditionDetails::MakeInstance()
+{
+	return MakeShareable(new FStateTreeConditionDetails);
+}
+
+void FStateTreeConditionDetails::CustomizeHeader(TSharedRef<IPropertyHandle> StructPropertyHandle, FDetailWidgetRow& HeaderRow, IPropertyTypeCustomizationUtils& StructCustomizationUtils)
+{
+	StructProperty = StructPropertyHandle;
+	PropUtils = StructCustomizationUtils.GetPropertyUtilities().Get();
+
+	LeftProperty = StructProperty->GetChildHandle(TEXT("Left"));
+	OperatorProperty = StructProperty->GetChildHandle(TEXT("Operator"));
+	RightProperty = StructProperty->GetChildHandle(TEXT("Right"));
+
+	if (LeftProperty)
+	{
+		LeftTypeProperty = LeftProperty->GetChildHandle(TEXT("Type"));
+		LeftBaseClassProperty = LeftProperty->GetChildHandle(TEXT("BaseClass"));
+	}
+	if (RightProperty)
+	{
+		RightTypeProperty = RightProperty->GetChildHandle(TEXT("Type"));
+		RightBaseClassProperty = RightProperty->GetChildHandle(TEXT("BaseClass"));
+		RightNameProperty = RightProperty->GetChildHandle(TEXT("Name"));
+		RightIDProperty = RightProperty->GetChildHandle(TEXT("ID"));
+		RightValueProperty = RightProperty->GetChildHandle(TEXT("Value"));
+	}
+
+	if (LeftTypeProperty)
+	{
+		LeftTypeProperty->SetOnPropertyValueChanged(FSimpleDelegate::CreateSP(this, &FStateTreeConditionDetails::OnLeftChanged));
+	}
+
+	HeaderRow
+		.WholeRowContent()
+		[
+			SNew(SHorizontalBox)
+			// Description
+			+ SHorizontalBox::Slot()
+			.FillWidth(1.0f)
+			.VAlign(VAlign_Center)
+			[
+				SNew(STextBlock)
+				.Text(this, &FStateTreeConditionDetails::GetDescription)
+				.Font(IDetailLayoutBuilder::GetDetailFont())
+			]
+			+ SHorizontalBox::Slot()
+			.AutoWidth()
+			[
+				StructPropertyHandle->CreateDefaultPropertyButtonWidgets()
+			]
+		];
+}
+
+void FStateTreeConditionDetails::CustomizeChildren(TSharedRef<IPropertyHandle> StructPropertyHandle, IDetailChildrenBuilder& StructBuilder, IPropertyTypeCustomizationUtils& StructCustomizationUtils)
+{
+	if (LeftProperty)
+	{
+		IDetailPropertyRow& Property = StructBuilder.AddProperty(LeftProperty.ToSharedRef());
+	}
+
+	if (OperatorProperty)
+	{
+		StructBuilder.AddCustomRow(FText::GetEmpty())
+			.NameContent()
+			[
+				OperatorProperty->CreatePropertyNameWidget()
+			]
+			.ValueContent()
+			.MinDesiredWidth(100.f)
+			.MaxDesiredWidth(150.f)
+			.VAlign(VAlign_Center)
+			.HAlign(HAlign_Left)
+			[
+				SNew(SComboButton)
+				.OnGetMenuContent(this, &FStateTreeConditionDetails::OnGetOperatorContent)
+				.ContentPadding(FMargin(6.0f, 2.0f))
+				.HAlign(HAlign_Center)
+				.ButtonContent()
+				[
+					SNew(STextBlock)
+					.Text(this, &FStateTreeConditionDetails::GetCurrentOperatorDesc)
+					.Font(IDetailLayoutBuilder::GetDetailFont())
+					.Justification(ETextJustify::Center)
+					.MinDesiredWidth(50.0f)
+				]
+			];
+	}
+
+	if (RightProperty)
+	{
+		IDetailPropertyRow& Property = StructBuilder.AddProperty(RightProperty.ToSharedRef());
+	}
+}
+
+void FStateTreeConditionDetails::OnLeftChanged()
+{
+	// Check that the operator can be used with the type, if not reset operator.
+	const EStateTreeVariableType LeftType = GetVariableType(LeftTypeProperty).Get(EStateTreeVariableType::Void);
+	if (OperatorProperty)
+	{
+		const EGenericAICheck Operator = GetOperator().Get(EGenericAICheck::Equal);
+		const bool bUseAllOperators = LeftType == EStateTreeVariableType::Float || LeftType == EStateTreeVariableType::Int;
+		const bool bIsSimpleOperators = Operator == EGenericAICheck::Equal || Operator == EGenericAICheck::NotEqual;
+		if (!bUseAllOperators && !bIsSimpleOperators)
+		{
+			OperatorProperty->SetValue((uint8)EGenericAICheck::Equal);
+		}
+	}
+
+	// Make sure that the Right is the same type as Left.
+	const EStateTreeVariableType RightType = GetVariableType(RightTypeProperty).Get(EStateTreeVariableType::Void);
+	UClass* LeftBaseClass = GetBaseClass(LeftBaseClassProperty).Get(nullptr);
+	UClass* RightBaseClass = GetBaseClass(RightBaseClassProperty).Get(nullptr);
+
+	if (LeftType != RightType || LeftBaseClass != RightBaseClass)
+	{
+		ResetRightVariableType(LeftType, LeftBaseClass);
+	}
+}
+
+void FStateTreeConditionDetails::OnOperatorComboChange(EGenericAICheck Operator)
+{
+	if (OperatorProperty)
+		OperatorProperty->SetValue((uint8)Operator);
+}
+
+TSharedRef<SWidget> FStateTreeConditionDetails::OnGetOperatorContent() const
+{
+	FMenuBuilder MenuBuilder(true, nullptr);
+
+	EStateTreeVariableType VarType = GetVariableType(LeftTypeProperty).Get(EStateTreeVariableType::Void);
+
+	FUIAction EqualAction(FExecuteAction::CreateSP(const_cast<FStateTreeConditionDetails*>(this), &FStateTreeConditionDetails::OnOperatorComboChange, EGenericAICheck::Equal));
+	MenuBuilder.AddMenuEntry(LOCTEXT("OperatorEqual", "=="), TAttribute<FText>(), FSlateIcon(), EqualAction);
+
+	FUIAction NotEqualAction(FExecuteAction::CreateSP(const_cast<FStateTreeConditionDetails*>(this), &FStateTreeConditionDetails::OnOperatorComboChange, EGenericAICheck::NotEqual));
+	MenuBuilder.AddMenuEntry(LOCTEXT("OperatorNotEqual", "!="), TAttribute<FText>(), FSlateIcon(), NotEqualAction);
+
+	if (VarType == EStateTreeVariableType::Float || VarType == EStateTreeVariableType::Int)
+	{
+		FUIAction LessAction(FExecuteAction::CreateSP(const_cast<FStateTreeConditionDetails*>(this), &FStateTreeConditionDetails::OnOperatorComboChange, EGenericAICheck::Less));
+		MenuBuilder.AddMenuEntry(LOCTEXT("OperatorLess", "<"), TAttribute<FText>(), FSlateIcon(), LessAction);
+
+		FUIAction LessOrEqualAction(FExecuteAction::CreateSP(const_cast<FStateTreeConditionDetails*>(this), &FStateTreeConditionDetails::OnOperatorComboChange, EGenericAICheck::LessOrEqual));
+		MenuBuilder.AddMenuEntry(LOCTEXT("OperatorLessOrEqual", "<="), TAttribute<FText>(), FSlateIcon(), LessOrEqualAction);
+
+		FUIAction GreaterAction(FExecuteAction::CreateSP(const_cast<FStateTreeConditionDetails*>(this), &FStateTreeConditionDetails::OnOperatorComboChange, EGenericAICheck::Greater));
+		MenuBuilder.AddMenuEntry(LOCTEXT("OperatorGreater", ">"), TAttribute<FText>(), FSlateIcon(), GreaterAction);
+
+		FUIAction GreaterOrEqualAction(FExecuteAction::CreateSP(const_cast<FStateTreeConditionDetails*>(this), &FStateTreeConditionDetails::OnOperatorComboChange, EGenericAICheck::GreaterOrEqual));
+		MenuBuilder.AddMenuEntry(LOCTEXT("OperatorGreaterOrEqual", ">="), TAttribute<FText>(), FSlateIcon(), GreaterOrEqualAction);
+	}
+
+	return MenuBuilder.MakeWidget();
+}
+
+FText FStateTreeConditionDetails::GetCurrentOperatorDesc() const
+{
+	const EGenericAICheck Operator = GetOperator().Get(EGenericAICheck::Equal);
+
+	switch (Operator)
+	{
+	case EGenericAICheck::Equal:
+		return LOCTEXT("OperatorEqual", "==");
+		break;
+	case EGenericAICheck::NotEqual:
+		return LOCTEXT("OperatorNotEqual", "!=");
+		break;
+	case EGenericAICheck::Greater:
+		return LOCTEXT("OperatorGreater", ">");
+		break;
+	case EGenericAICheck::GreaterOrEqual:
+		return LOCTEXT("OperatorGreaterOrEqual", ">=");
+		break;
+	case EGenericAICheck::Less:
+		return LOCTEXT("OperatorLess", "<");
+		break;
+	case EGenericAICheck::LessOrEqual:
+		return LOCTEXT("OperatorLessOrEqual", "<=");
+		break;
+	default:
+		return LOCTEXT("OperatorUnknown", "??");
+		break;
+	}
+}
+
+TOptional<EStateTreeVariableType> FStateTreeConditionDetails::GetVariableType(const TSharedPtr<IPropertyHandle>& PropertyHandle) const
+{
+	if (PropertyHandle)
+	{
+		uint8 Value;
+		if (PropertyHandle->GetValue(Value) == FPropertyAccess::Success)
+		{
+			return EStateTreeVariableType(Value);
+		}
+	}
+	return TOptional<EStateTreeVariableType>();
+
+}
+
+TOptional<UClass*> FStateTreeConditionDetails::GetBaseClass(const TSharedPtr<IPropertyHandle>& PropertyHandle) const
+{
+	if (PropertyHandle)
+	{
+		UObject* Value;
+		if (PropertyHandle->GetValue(Value) == FPropertyAccess::Success)
+		{
+			return Cast<UClass>(Value);
+		}
+	}
+	return nullptr;
+}
+
+void FStateTreeConditionDetails::ResetRightVariableType(EStateTreeVariableType NewType, UClass* NewBaseClass)
+{
+	FScopedTransaction Transaction(FText::Format(LOCTEXT("SetPropertyValue", "Set {0}"), RightTypeProperty->GetPropertyDisplayName()));
+
+	if (RightTypeProperty)
+	{
+		bool bValueTypeChanged = false;
+		uint8 OldType;
+		if (RightTypeProperty->GetValue(OldType) == FPropertyAccess::Success)
+		{
+			bValueTypeChanged = NewType != (EStateTreeVariableType)OldType;
+		}
+
+		RightTypeProperty->SetValue((uint8)NewType, EPropertyValueSetFlags::NotTransactable);
+
+		if (bValueTypeChanged)
+		{
+			switch (NewType)
+			{
+			case EStateTreeVariableType::Float:
+				UE::StateTree::PropertyHelpers::SetVariableValue<float>(RightValueProperty, 0.0f, EPropertyValueSetFlags::NotTransactable);
+				break;
+			case EStateTreeVariableType::Int:
+				UE::StateTree::PropertyHelpers::SetVariableValue<int32>(RightValueProperty, 0, EPropertyValueSetFlags::NotTransactable);
+				break;
+			case EStateTreeVariableType::Bool:
+				UE::StateTree::PropertyHelpers::SetVariableValue<bool>(RightValueProperty, false, EPropertyValueSetFlags::NotTransactable);
+				break;
+			case EStateTreeVariableType::Vector:
+				UE::StateTree::PropertyHelpers::SetVariableValue<FVector>(RightValueProperty, FVector::ZeroVector, EPropertyValueSetFlags::NotTransactable);
+				break;
+			case EStateTreeVariableType::Object:
+				// fallthrough
+			case EStateTreeVariableType::Void:
+				UE::StateTree::PropertyHelpers::SetVariableValue<int32>(RightValueProperty, 0, EPropertyValueSetFlags::NotTransactable);	// Not representation, just clear to int32.
+				break;
+			default:
+				ensureMsgf(false, TEXT("Unhandled variable type."));
+				break;
+			}
+		}
+	}
+
+	if (RightBaseClassProperty)
+	{
+		RightBaseClassProperty->SetValue(NewBaseClass, EPropertyValueSetFlags::NotTransactable);
+	}
+
+	// Reset binding
+	if (RightNameProperty)
+	{
+		RightNameProperty->SetValue(FName(), EPropertyValueSetFlags::NotTransactable);
+	}
+	if (RightIDProperty)
+	{
+		UE::StateTree::PropertyHelpers::SetStructValue<FGuid>(RightIDProperty, FGuid(), EPropertyValueSetFlags::NotTransactable);
+	}
+
+}
+
+FText FStateTreeConditionDetails::GetDescription() const
+{
+	if (StructProperty)
+	{
+		TArray<void*> RawData;
+		StructProperty->AccessRawData(RawData);
+		if (RawData.Num() == 1)
+		{
+			FStateTreeCondition* Condition = static_cast<FStateTreeCondition*>(RawData[0]);
+			if (Condition != nullptr)
+			{
+				return Condition->GetDescription();
+			}
+		}
+		else
+		{
+			return LOCTEXT("MultipleSelected", "Multiple selected");
+		}
+
+	}
+	return FText::GetEmpty();
+}
+
+TOptional<EGenericAICheck> FStateTreeConditionDetails::GetOperator() const
+{
+	if (OperatorProperty)
+	{
+		uint8 Value;
+		if (OperatorProperty->GetValue(Value) == FPropertyAccess::Success)
+		{
+			return EGenericAICheck(Value);
+		}
+	}
+	return TOptional<EGenericAICheck>();
+}
+
+#undef LOCTEXT_NAMESPACE
diff --git a/Engine/Plugins/Runtime/StateTree/Source/StateTreeEditorModule/Private/Customizations/StateTreeConditionDetails.h b/Engine/Plugins/Runtime/StateTree/Source/StateTreeEditorModule/Private/Customizations/StateTreeConditionDetails.h
new file mode 100644
--- /dev/null
+++ b/Engine/Plugins/Runtime/StateTree/Source/StateTreeEditorModule/Private/Customizations/StateTreeConditionDetails.h
@@ -0,0 +1,53 @@
+// Copyright Epic Games, Inc. All Rights Reserved.
+
+#pragma once
+
+#include "CoreMinimal.h"
+#include "UnrealClient.h"
+#include "IPropertyTypeCustomization.h"
+#include "StateTreeCondition.h"
+
+class IPropertyHandle;
+
+/**
+ * Type customization for FStateTreeCondition.
+ */
+
+class FStateTreeConditionDetails : public IPropertyTypeCustomization
+{
+public:
+	/** Makes a new instance of this detail layout class for a specific detail view requesting it */
+	static TSharedRef<IPropertyTypeCustomization> MakeInstance();
+
+	/** IPropertyTypeCustomization interface */
+	virtual void CustomizeHeader(TSharedRef<IPropertyHandle> StructPropertyHandle, class FDetailWidgetRow& HeaderRow, IPropertyTypeCustomizationUtils& StructCustomizationUtils) override;
+	virtual void CustomizeChildren(TSharedRef<IPropertyHandle> StructPropertyHandle, class IDetailChildrenBuilder& StructBuilder, IPropertyTypeCustomizationUtils& StructCustomizationUtils) override;
+
+private:
+
+	void OnLeftChanged();
+	FText GetDescription() const;
+	void OnOperatorComboChange(EGenericAICheck Operator);
+	TSharedRef<SWidget> OnGetOperatorContent() const;
+	FText GetCurrentOperatorDesc() const;
+
+	TOptional<EGenericAICheck> GetOperator() const;
+	TOptional<EStateTreeVariableType> GetVariableType(const TSharedPtr<IPropertyHandle>& PropertyHandle) const;
+	TOptional<UClass*> GetBaseClass(const TSharedPtr<IPropertyHandle>& PropertyHandle) const;
+	void ResetRightVariableType(EStateTreeVariableType NewType, UClass* NewBaseClass);
+
+	TSharedPtr<IPropertyHandle> LeftProperty;
+	TSharedPtr<IPropertyHandle> RightProperty;
+	TSharedPtr<IPropertyHandle> OperatorProperty;
+
+	TSharedPtr<IPropertyHandle> LeftTypeProperty;
+	TSharedPtr<IPropertyHandle> LeftBaseClassProperty;
+	TSharedPtr<IPropertyHandle> RightTypeProperty;
+	TSharedPtr<IPropertyHandle> RightBaseClassProperty;
+	TSharedPtr<IPropertyHandle> RightIDProperty;
+	TSharedPtr<IPropertyHandle> RightNameProperty;
+	TSharedPtr<IPropertyHandle> RightValueProperty;
+
+	class IPropertyUtilities* PropUtils;
+	TSharedPtr<IPropertyHandle> StructProperty;
+};
diff --git a/Engine/Plugins/Runtime/StateTree/Source/StateTreeEditorModule/Private/Customizations/StateTreeConditionItemDetails.cpp b/Engine/Plugins/Runtime/StateTree/Source/StateTreeEditorModule/Private/Customizations/StateTreeConditionItemDetails.cpp
new file mode 100644
--- /dev/null
+++ b/Engine/Plugins/Runtime/StateTree/Source/StateTreeEditorModule/Private/Customizations/StateTreeConditionItemDetails.cpp
@@ -0,0 +1,386 @@
+// Copyright Epic Games, Inc. All Rights Reserved.
+
+#include "StateTreeConditionItemDetails.h"
+#include "Widgets/DeclarativeSyntaxSupport.h"
+#include "DetailWidgetRow.h"
+#include "DetailLayoutBuilder.h"
+#include "IPropertyUtilities.h"
+#include "IDetailPropertyRow.h"
+#include "IDetailChildrenBuilder.h"
+#include "StateTree.h"
+#include "StateTreeEditorData.h"
+#include "StateTreeConditionBase.h"
+#include "Widgets/Text/SRichTextBlock.h"
+#include "Widgets/Input/SComboButton.h"
+#include "Widgets/Images/SImage.h"
+#include "Framework/MultiBox/MultiBoxBuilder.h"
+#include "InstancedStruct.h"
+#include "Editor.h"
+#include "StateTreeEditorPropertyBindings.h"
+#include "StateTreeBindingExtension.h"
+#include "StateTreeEditorStyle.h"
+#include "Modules/ModuleManager.h"
+#include "StructViewerModule.h"
+#include "StructViewerFilter.h"
+#include "InstancedStructDetails.h"
+#include "Engine/UserDefinedStruct.h"
+
+#define LOCTEXT_NAMESPACE "StateTreeEditor"
+
+// TODO: This has a lot in common with StateTreeBindableItemDetails, see if we could combine the efforts.
+
+// Customized version of FInstancedStructDataDetails used to hide bindable properties.
+class FConditionItemDataDetails : public FInstancedStructDataDetails
+{
+public:
+	FConditionItemDataDetails(TSharedPtr<IPropertyHandle> InStructProperty, UStateTreeEditorData* InEditorData)
+		: FInstancedStructDataDetails(InStructProperty)
+		, EditorData(InEditorData)
+	{
+		check(EditorData);
+		EditorPropBindings = EditorData->GetPropertyEditorBindings();
+	}
+
+	virtual void OnChildRowAdded(IDetailPropertyRow& ChildRow)
+	{
+		check(EditorPropBindings);
+
+		static const FName BindableMetaName(TEXT("Bindable"));
+
+		TSharedPtr<IPropertyHandle> ChildPropHandle = ChildRow.GetPropertyHandle();
+
+		// Conditionally control visibility of the value field of bound properties.
+		if (ChildPropHandle->HasMetaData(BindableMetaName))
+		{
+			// Find property path relative to closest outer bindable struct.
+			FStateTreeEditorPropertyPath Path;
+			UE::StateTree::PropertyBinding::GetOuterStructPropertyPath(ChildPropHandle, Path);
+
+			TSharedPtr<SWidget> NameWidget;
+			TSharedPtr<SWidget> ValueWidget;
+			FDetailWidgetRow Row;
+			ChildRow.GetDefaultWidgets(NameWidget, ValueWidget, Row);
+
+			auto IsValueVisible = TAttribute<EVisibility>::Create([this, Path]() -> EVisibility
+				{
+					return EditorPropBindings->HasPropertyBinding(Path) ? EVisibility::Collapsed : EVisibility::Visible;
+				});
+
+			ChildRow
+				.CustomWidget()
+				.NameContent()
+				[
+					NameWidget.ToSharedRef()
+				]
+				.ValueContent()
+				[
+					SNew(SBox)
+					.Visibility(IsValueVisible)
+					[
+						ValueWidget.ToSharedRef()
+					]
+				];
+		}
+	}
+
+	UStateTreeEditorData* EditorData;
+	FStateTreeEditorPropertyBindings* EditorPropBindings;
+};
+
+
+////////////////////////////////////
+
+class FConditionItemStructFilter : public IStructViewerFilter
+{
+public:
+	/** The base struct for the property that classes must be a child-of. */
+	const UScriptStruct* BaseStruct = nullptr;
+
+	/** A flag controlling whether we allow to select the BaseStruct */
+	bool bAllowBaseStruct = true;
+
+	virtual bool IsStructAllowed(const FStructViewerInitializationOptions& InInitOptions, const UScriptStruct* InStruct, TSharedRef<FStructViewerFilterFuncs> InFilterFuncs) override
+	{
+		if (InStruct == BaseStruct)
+		{
+			return bAllowBaseStruct;
+		}
+
+		// Query the native struct to see if it has the correct parent type (if any)
+		return !BaseStruct || InStruct->IsChildOf(BaseStruct);
+	}
+
+	virtual bool IsUnloadedStructAllowed(const FStructViewerInitializationOptions& InInitOptions, const FName InStructPath, TSharedRef<FStructViewerFilterFuncs> InFilterFuncs) override
+	{
+		// We currently assume that conditions are implemented natively.
+		return false;
+	}
+};
+
+
+////////////////////////////////////
+
+TSharedRef<IPropertyTypeCustomization> FStateTreeConditionItemDetails::MakeInstance()
+{
+	return MakeShareable(new FStateTreeConditionItemDetails);
+}
+
+FStateTreeConditionItemDetails::~FStateTreeConditionItemDetails()
+{
+	UE::StateTree::PropertyBinding::OnStateTreeBindingChanged.Remove(OnBindingChangedHandle);
+}
+
+void FStateTreeConditionItemDetails::CustomizeHeader(TSharedRef<class IPropertyHandle> StructPropertyHandle, class FDetailWidgetRow& HeaderRow, IPropertyTypeCustomizationUtils& StructCustomizationUtils)
+{
+	StructProperty = StructPropertyHandle;
+	PropUtils = StructCustomizationUtils.GetPropertyUtilities().Get();
+
+	TypeProperty = StructProperty->GetChildHandle(TEXT("Type"));
+	check(TypeProperty);
+
+	// Find base class from meta data.
+	static const FName BaseClassMetaName(TEXT("BaseStruct")); // TODO: move these names into one central place.
+	const FString BaseClassName = TypeProperty->GetMetaData(BaseClassMetaName);
+	BaseScriptStruct = FindObject<UScriptStruct>(ANY_PACKAGE, *BaseClassName);
+
+	OnBindingChangedHandle = UE::StateTree::PropertyBinding::OnStateTreeBindingChanged.AddRaw(this, &FStateTreeConditionItemDetails::OnBindingChanged);
+
+	FindOuterObjects();
+
+	HeaderRow
+		.WholeRowContent()
+		.VAlign(VAlign_Center)
+		[
+			SNew(SHorizontalBox)
+			// Description
+			+ SHorizontalBox::Slot()
+			.FillWidth(3.0f)
+			.VAlign(VAlign_Center)
+			[
+				SNew(SRichTextBlock)
+				.DecoratorStyleSet(FStateTreeEditorStyle::Get().Get())
+				.TextStyle(FStateTreeEditorStyle::Get(), "Details.Normal")
+				.Text(this, &FStateTreeConditionItemDetails::GetDescription)
+			]
+			// Class picker
+			+ SHorizontalBox::Slot()
+			.FillWidth(1.0f)
+			.VAlign(VAlign_Center)
+			[
+				SAssignNew(ComboButton, SComboButton)
+				.OnGetMenuContent(this, &FStateTreeConditionItemDetails::GenerateStructPicker)
+				.ContentPadding(0)
+				.ButtonContent()
+				[
+					SNew(SHorizontalBox)
+					+ SHorizontalBox::Slot()
+					.VAlign(VAlign_Center)
+					[
+						SNew(STextBlock)
+						.Text(this, &FStateTreeConditionItemDetails::GetDisplayValueString)
+						.Font(IDetailLayoutBuilder::GetDetailFont())
+					]
+				]
+			]
+			// Array controls
+			+ SHorizontalBox::Slot()
+			.AutoWidth()
+			[
+				StructPropertyHandle->CreateDefaultPropertyButtonWidgets()
+			]
+		];
+}
+
+void FStateTreeConditionItemDetails::CustomizeChildren(TSharedRef<class IPropertyHandle> StructPropertyHandle, class IDetailChildrenBuilder& StructBuilder, IPropertyTypeCustomizationUtils& StructCustomizationUtils)
+{
+	check(TypeProperty);
+
+	TSharedRef<FConditionItemDataDetails> DataDetails = MakeShareable(new FConditionItemDataDetails(TypeProperty, EditorData));
+	StructBuilder.AddCustomBuilder(DataDetails);
+}
+
+void FStateTreeConditionItemDetails::FindOuterObjects()
+{
+	check(TypeProperty);
+
+	EditorData = nullptr;
+	StateTree = nullptr;
+
+	TArray<UObject*> OuterObjects;
+	TypeProperty->GetOuterObjects(OuterObjects);
+	if (OuterObjects.Num() == 1 && OuterObjects[0] != nullptr)
+	{
+		EditorData = OuterObjects[0]->GetTypedOuter<UStateTreeEditorData>();
+		StateTree = EditorData ? EditorData->GetTypedOuter<UStateTree>() : nullptr;
+	}
+}
+
+void FStateTreeConditionItemDetails::OnBindingChanged(const FStateTreeEditorPropertyPath& SourcePath, const FStateTreeEditorPropertyPath& TargetPath)
+{
+	check(TypeProperty);
+
+	// Note: Does not support editing multiple objects.
+	FStateTreeConditionBase* Condition = GetSingleCondition();
+	if (!Condition)
+	{
+		return;
+	}
+
+	if (Condition->ID != TargetPath.StructID)
+	{
+		return;
+	}
+
+	TArray<UObject*> OuterObjects;
+	TypeProperty->GetOuterObjects(OuterObjects);
+	if (OuterObjects.Num() == 1 && OuterObjects[0] != nullptr)
+	{
+		UObject* OuterObject = OuterObjects[0]; // Immediate outer, i.e StateTreeState
+		UObject* OwnerObject = UE::StateTree::PropertyBinding::FindEditorBindingsOwner(OuterObject); // Outer that contains IStateTreeEditorPropertyBindingsOwner interface.
+		if (IStateTreeEditorPropertyBindingsOwner* BindingOwner = Cast<IStateTreeEditorPropertyBindingsOwner>(OwnerObject))
+		{
+			// Binding changes modify only the object that contains the bindings, we need to also set the property owner (i.e. StateTreeState) as modified.
+			OuterObject->Modify();
+			FStateTreeBindingLookup BindingLookup(BindingOwner);
+			Condition->OnBindingChanged(SourcePath, TargetPath, BindingLookup);
+		}
+	}
+}
+
+FText FStateTreeConditionItemDetails::GetDescription() const
+{
+	check(TypeProperty);
+
+	if (TypeProperty->GetNumOuterObjects() != 1)
+	{
+		return LOCTEXT("MultipleSelected", "<Details.Subdued>Multiple selected</>");
+	}
+
+	if (FStateTreeConditionBase* Condition = GetSingleCondition())
+	{
+		if (IStateTreeEditorPropertyBindingsOwner* BindingOwner = Cast<IStateTreeEditorPropertyBindingsOwner>(GetEditorBindingsOwner()))
+		{
+			FStateTreeBindingLookup BindingLookup(BindingOwner);
+			return Condition->GetDescription(BindingLookup);
+		}
+	}
+
+	return LOCTEXT("ConditionNotSet", "<Details.Subdued>Condition Not Set</>");
+}
+
+FInstancedStruct* FStateTreeConditionItemDetails::GetSingleStructPtr() const
+{
+	check(TypeProperty);
+	TArray<void*> RawData;
+	TypeProperty->AccessRawData(RawData);
+	return RawData.Num() == 1 ? static_cast<FInstancedStruct*>(RawData[0]) : nullptr;
+}
+
+FStateTreeConditionBase* FStateTreeConditionItemDetails::GetSingleCondition() const
+{
+	if (FInstancedStruct* StructPtr = GetSingleStructPtr())
+	{
+		return StructPtr->GetMutablePtr<FStateTreeConditionBase>();
+	}
+	return nullptr;
+}
+
+UObject* FStateTreeConditionItemDetails::GetEditorBindingsOwner() const
+{
+	check(TypeProperty);
+
+	TArray<UObject*> OuterObjects;
+	TypeProperty->GetOuterObjects(OuterObjects);
+	if (OuterObjects.Num() == 1)
+	{
+		return UE::StateTree::PropertyBinding::FindEditorBindingsOwner(OuterObjects[0]);
+	}
+
+	return nullptr;
+}
+
+FText FStateTreeConditionItemDetails::GetDisplayValueString() const
+{
+	if (FInstancedStruct* StructPtr = GetSingleStructPtr())
+	{
+		if (const UScriptStruct* ScriptStruct = StructPtr->GetScriptStruct())
+		{
+			return ScriptStruct->GetDisplayNameText();
+		}
+	}
+	return FText();
+}
+
+TSharedRef<SWidget> FStateTreeConditionItemDetails::GenerateStructPicker()
+{
+	static const FName ExcludeBaseStruct(TEXT("ExcludeBaseStruct")); // TODO: move these names into one central place.
+	const bool bExcludeBaseStruct = TypeProperty->HasMetaData(ExcludeBaseStruct);
+
+	TSharedRef<FConditionItemStructFilter> StructFilter = MakeShared<FConditionItemStructFilter>();
+	StructFilter->BaseStruct = BaseScriptStruct;
+	StructFilter->bAllowBaseStruct = !bExcludeBaseStruct;
+
+	FStructViewerInitializationOptions Options;
+	Options.bShowUnloadedStructs = true;
+	Options.bShowNoneOption = true;
+	Options.StructFilter = StructFilter;
+	Options.NameTypeToDisplay = EStructViewerNameTypeToDisplay::DisplayName;
+	Options.DisplayMode = EStructViewerDisplayMode::ListView;
+	Options.bAllowViewOptions = true;
+
+	FOnStructPicked OnPicked(FOnStructPicked::CreateRaw(this, &FStateTreeConditionItemDetails::OnStructPicked));
+
+	return SNew(SBox)
+		.WidthOverride(280)
+		[
+			SNew(SVerticalBox)
+			+ SVerticalBox::Slot()
+			.AutoHeight()
+			.MaxHeight(500)
+			[
+				FModuleManager::LoadModuleChecked<FStructViewerModule>("StructViewer").CreateStructViewer(Options, OnPicked)
+			]
+		];
+}
+
+void FStateTreeConditionItemDetails::OnStructPicked(const UScriptStruct* InStruct)
+{
+	check(TypeProperty);
+	check(PropUtils);
+
+	GEditor->BeginTransaction(LOCTEXT("OnStructPicked", "Set Struct"));
+
+	TypeProperty->NotifyPreChange();
+
+	if (FInstancedStruct* StructPtr = GetSingleStructPtr())
+	{
+		if (InStruct)
+		{
+			StructPtr->InitializeAs(InStruct);
+			
+			// Generate new ID.
+			if (InStruct->IsChildOf(FStateTreeConditionBase::StaticStruct()))
+			{
+				FStateTreeConditionBase* Condition = StructPtr->GetMutablePtr<FStateTreeConditionBase>();
+				check(Condition);
+				Condition->ID = FGuid::NewGuid();
+			}
+		}
+		else
+		{
+			StructPtr->Reset();
+		}
+	}
+
+	TypeProperty->NotifyPostChange(EPropertyChangeType::ValueSet);
+	GEditor->EndTransaction();
+	TypeProperty->NotifyFinishedChangingProperties();
+
+	GEditor->EndTransaction();
+
+	ComboButton->SetIsOpen(false);
+
+	PropUtils->ForceRefresh();
+}
+
+#undef LOCTEXT_NAMESPACE
diff --git a/Engine/Plugins/Runtime/StateTree/Source/StateTreeEditorModule/Private/Customizations/StateTreeConditionItemDetails.h b/Engine/Plugins/Runtime/StateTree/Source/StateTreeEditorModule/Private/Customizations/StateTreeConditionItemDetails.h
new file mode 100644
--- /dev/null
+++ b/Engine/Plugins/Runtime/StateTree/Source/StateTreeEditorModule/Private/Customizations/StateTreeConditionItemDetails.h
@@ -0,0 +1,58 @@
+// Copyright Epic Games, Inc. All Rights Reserved.
+
+#pragma once
+
+#include "CoreMinimal.h"
+#include "IPropertyTypeCustomization.h"
+
+class IPropertyHandle;
+struct FInstancedStruct;
+struct FStateTreeEditorPropertyPath;
+struct FStateTreeConditionBase;
+class UStateTreeEditorData;
+class UStateTree;
+
+/**
+ * Property customization for FStateTreeConditionItem.
+ */
+class FStateTreeConditionItemDetails : public IPropertyTypeCustomization
+{
+public:
+
+	virtual ~FStateTreeConditionItemDetails();
+
+	/** Makes a new instance of this detail layout class for a specific detail view requesting it */
+	static TSharedRef<IPropertyTypeCustomization> MakeInstance();
+
+	/** IPropertyTypeCustomization interface */
+	virtual void CustomizeHeader(TSharedRef<IPropertyHandle> StructPropertyHandle, class FDetailWidgetRow& HeaderRow, IPropertyTypeCustomizationUtils& StructCustomizationUtils) override;
+	virtual void CustomizeChildren(TSharedRef<IPropertyHandle> StructPropertyHandle, class IDetailChildrenBuilder& StructBuilder, IPropertyTypeCustomizationUtils& StructCustomizationUtils) override;
+
+private:
+
+	FText GetDescription() const;
+	void OnChildChanged();
+	void OnBindingChanged(const FStateTreeEditorPropertyPath& SourcePath, const FStateTreeEditorPropertyPath& TargetPath);
+
+	FStateTreeConditionBase* GetSingleCondition() const;
+	FInstancedStruct* GetSingleStructPtr() const;
+	UObject* GetEditorBindingsOwner() const;
+
+	FText GetDisplayValueString() const;
+	TSharedRef<SWidget> GenerateStructPicker();
+	void OnStructPicked(const UScriptStruct* InStruct);
+
+	void FindOuterObjects();
+
+	UScriptStruct* BaseScriptStruct = nullptr;
+	TSharedPtr<class SComboButton> ComboButton;
+
+	UStateTreeEditorData* EditorData = nullptr;
+	UStateTree* StateTree = nullptr;
+
+	class IPropertyUtilities* PropUtils = nullptr;
+	TSharedPtr<IPropertyHandle> StructProperty;
+	TSharedPtr<IPropertyHandle> TypeProperty;
+
+	FDelegateHandle OnBindingChangedHandle;
+};
diff --git a/Engine/Plugins/Runtime/StateTree/Source/StateTreeEditorModule/Private/Customizations/StateTreeEditorPropertyPathDetails.cpp b/Engine/Plugins/Runtime/StateTree/Source/StateTreeEditorModule/Private/Customizations/StateTreeEditorPropertyPathDetails.cpp
new file mode 100644
--- /dev/null
+++ b/Engine/Plugins/Runtime/StateTree/Source/StateTreeEditorModule/Private/Customizations/StateTreeEditorPropertyPathDetails.cpp
@@ -0,0 +1,255 @@
+// Copyright Epic Games, Inc. All Rights Reserved.
+
+#include "StateTreeEditorPropertyPathDetails.h"
+#include "Widgets/DeclarativeSyntaxSupport.h"
+#include "DetailWidgetRow.h"
+#include "DetailLayoutBuilder.h"
+#include "IPropertyUtilities.h"
+#include "IDetailPropertyRow.h"
+#include "IDetailChildrenBuilder.h"
+#include "StateTreeEditorPropertyBindings.h"
+#include "StateTreeBindingExtension.h"
+#include "Widgets/Input/SComboButton.h"
+#include "Framework/MultiBox/MultiBoxBuilder.h"
+#include "ScopedTransaction.h"
+
+#include "EdGraphSchema_K2.h"
+#include "IPropertyAccessEditor.h"
+#include "Features/IModularFeatures.h"
+#include "StateTreeEditorData.h"
+#include "Algo/Accumulate.h"
+#include "InstancedStruct.h"
+#include "StateTreePropertyHelpers.h"
+
+
+#define LOCTEXT_NAMESPACE "StateTreeEditor"
+
+TSharedRef<IPropertyTypeCustomization> FStateTreeEditorPropertyPathDetails::MakeInstance()
+{
+	return MakeShareable(new FStateTreeEditorPropertyPathDetails);
+}
+
+void FStateTreeEditorPropertyPathDetails::CustomizeHeader(TSharedRef<class IPropertyHandle> StructPropertyHandle, class FDetailWidgetRow& HeaderRow, IPropertyTypeCustomizationUtils& StructCustomizationUtils)
+{
+	StructProperty = StructPropertyHandle;
+	PropUtils = StructCustomizationUtils.GetPropertyUtilities().Get();
+
+	HeaderRow
+		.NameContent()
+		[
+			StructPropertyHandle->CreatePropertyNameWidget()
+		]
+		.ValueContent()
+		.MinDesiredWidth(250.f)
+		.VAlign(VAlign_Center)
+		[
+			CreateBindingWidget(StructProperty)
+		];
+}
+
+void FStateTreeEditorPropertyPathDetails::CustomizeChildren(TSharedRef<class IPropertyHandle> StructPropertyHandle, class IDetailChildrenBuilder& StructBuilder, IPropertyTypeCustomizationUtils& StructCustomizationUtils)
+{
+}
+
+TSharedRef<SWidget> FStateTreeEditorPropertyPathDetails::CreateBindingWidget(TSharedPtr<IPropertyHandle> InPropertyHandle)
+{
+	if (!IModularFeatures::Get().IsModularFeatureAvailable("PropertyAccessEditor"))
+	{
+		return SNullWidget::NullWidget;
+	}
+
+	IPropertyAccessEditor& PropertyAccessEditor = IModularFeatures::Get().GetModularFeature<IPropertyAccessEditor>("PropertyAccessEditor");
+
+	UObject* OwnerObject = nullptr;
+	FStateTreeEditorPropertyBindings* EditorBindings = nullptr;
+
+	// Array of structs we can bind to.
+	TArray<FBindingContextStruct> Context;
+	TArray<FStateTreeBindableStructDesc> AccessibleStructs;
+
+	// The struct and property where we're binding.
+	FStateTreeEditorPropertyPath TargetPath;
+	UStateTreeState* OuterState = nullptr;
+
+	TArray<UObject*> OuterObjects;
+	InPropertyHandle->GetOuterObjects(OuterObjects);
+	if (OuterObjects.Num() == 1)
+	{
+		// Only allow to binding when one object is selected.
+		OwnerObject = UE::StateTree::PropertyBinding::FindEditorBindingsOwner(OuterObjects[0]);
+
+		// Figure out the structs we're editing, and property path relative to current property.
+		UE::StateTree::PropertyBinding::GetOuterStructPropertyPath(InPropertyHandle, TargetPath);
+
+		if (IStateTreeEditorPropertyBindingsOwner* BindingOwner = Cast<IStateTreeEditorPropertyBindingsOwner>(OwnerObject))
+		{
+			EditorBindings = BindingOwner->GetPropertyEditorBindings();
+			BindingOwner->GetAccessibleStructs(TargetPath.StructID, AccessibleStructs);
+			for (FStateTreeBindableStructDesc& StructDesc : AccessibleStructs)
+			{
+				const UStruct* Struct = StructDesc.Struct;
+				Context.Emplace(const_cast<UStruct*>(Struct), nullptr, FText::FromString(StructDesc.Name.ToString()));
+			}
+		}
+	}
+
+	FPropertyBindingWidgetArgs Args;
+	Args.Property = nullptr; // Bind to any property type.
+
+	Args.OnCanBindProperty = FOnCanBindProperty::CreateLambda([OwnerObject](FProperty* InProperty)
+		{
+			// Special case for null, binding widget calls OnCanBindProperty with Args.Property as well, this is to catch the case that we've set Args.Property to null.
+			if (InProperty == nullptr)
+			{
+				return true;
+			}
+
+			// TODO: Make this configurable with meta tags.
+			const bool bIsNumeric = InProperty->IsA<FNumericProperty>();
+			const bool bIsEnum = InProperty->IsA<FEnumProperty>();
+
+			return OwnerObject != nullptr && (bIsNumeric || bIsEnum);
+		});
+
+	Args.OnCanBindToClass = FOnCanBindToClass::CreateLambda([](UClass* InClass)
+		{
+			return true;
+		});
+
+	Args.OnAddBinding = FOnAddBinding::CreateLambda([InPropertyHandle, OwnerObject, AccessibleStructs](FName InPropertyName, const TArray<FBindingChainElement>& InBindingChain)
+		{
+			IPropertyAccessEditor& PropertyAccessEditor = IModularFeatures::Get().GetModularFeature<IPropertyAccessEditor>("PropertyAccessEditor");
+			if (OwnerObject && InPropertyHandle)
+			{
+				if (InBindingChain.Num() > 1)	// Assume at least: [0] struct index, [1] a property.
+				{
+					const int32 SourceStructIndex = InBindingChain[0].ArrayIndex;
+					TArray<FBindingChainElement> SourceBindingChain(InBindingChain.GetData() + 1, InBindingChain.Num() - 1);
+					check(SourceStructIndex >= 0 && SourceStructIndex < AccessibleStructs.Num());
+
+					FStateTreeEditorPropertyPath NewPath;
+					NewPath.StructID = AccessibleStructs[SourceStructIndex].ID;
+					PropertyAccessEditor.MakeStringPath(SourceBindingChain, NewPath.Path);
+
+					UE::StateTree::PropertyHelpers::SetStructValue<FStateTreeEditorPropertyPath>(InPropertyHandle, NewPath, EPropertyValueSetFlags::NotTransactable);
+				}
+			}
+		});
+
+	Args.OnRemoveBinding = FOnRemoveBinding::CreateLambda([InPropertyHandle, OwnerObject](FName InPropertyName)
+		{
+			if (OwnerObject && InPropertyHandle)
+			{
+				FStateTreeEditorPropertyPath EmptyPath;
+				UE::StateTree::PropertyHelpers::SetStructValue<FStateTreeEditorPropertyPath>(InPropertyHandle, EmptyPath, EPropertyValueSetFlags::NotTransactable);
+			}
+		});
+
+	Args.OnCanRemoveBinding = FOnCanRemoveBinding::CreateLambda([](FName InPropertyName)
+		{
+			return true;
+		});
+
+	Args.CurrentBindingText = MakeAttributeLambda([InPropertyHandle, OwnerObject, AccessibleStructs]()
+		{
+			const FText MultipleValues = LOCTEXT("MultipleValues", "Multiple Values");
+			const FText Bind = LOCTEXT("Bind", "Bind");
+			FText CurrentValue = Bind;
+
+			if (OwnerObject && InPropertyHandle)
+			{
+				OwnerObject->Modify();
+
+				TArray<void*> RawData;
+				InPropertyHandle->AccessRawData(RawData);
+				if (RawData.Num() == 1 && RawData[0])
+				{
+					FStateTreeEditorPropertyPath& Path = *static_cast<FStateTreeEditorPropertyPath*>(RawData[0]);
+					if (Path.IsValid())
+					{
+						FString PropertyName;
+						for (int32 i = 0; i < AccessibleStructs.Num(); i++)
+						{
+							if (AccessibleStructs[i].ID == Path.StructID)
+							{
+								PropertyName = AccessibleStructs[i].Name.ToString();
+								break;
+							}
+						}
+						for (const FString& Segment : Path.Path)
+						{
+							PropertyName += TEXT(".") + Segment;
+						}
+						CurrentValue = FText::FromString(PropertyName);
+					}
+				}
+			}
+			else
+			{
+				// OwnerObject is not valid if there's multiple objects selected.
+				CurrentValue = MultipleValues;
+			}
+
+			return CurrentValue;
+		});
+
+	Args.CurrentBindingImage = MakeAttributeLambda([]() -> const FSlateBrush*
+		{
+			static FName PropertyIcon(TEXT("Kismet.Tabs.Variables"));
+			return FEditorStyle::GetBrush(PropertyIcon);
+		});
+
+	Args.CurrentBindingColor = MakeAttributeLambda([InPropertyHandle, AccessibleStructs]() -> FLinearColor
+		{
+			IPropertyAccessEditor& PropertyAccessEditor = IModularFeatures::Get().GetModularFeature<IPropertyAccessEditor>("PropertyAccessEditor");
+
+			// Get the path we're editing
+			FStateTreeEditorPropertyPath* SourcePath = nullptr;
+			TArray<void*> RawData;
+			InPropertyHandle->AccessRawData(RawData);
+			if (RawData.Num() == 1 && RawData[0])
+			{
+				SourcePath = static_cast<FStateTreeEditorPropertyPath*>(RawData[0]);
+			}
+
+			// Figure out source struct type.
+			const UStruct* SourceStruct = nullptr;
+			if (SourcePath && SourcePath->IsValid())
+			{
+				const FStateTreeBindableStructDesc* Desc = AccessibleStructs.FindByPredicate([SourcePath](const FStateTreeBindableStructDesc& Desc) { return Desc.ID == SourcePath->StructID; });
+				if (Desc)
+				{
+					SourceStruct = Desc->Struct;
+				}
+			}
+
+			// And resolve property from the struct and path.
+			FProperty* SourceProperty = nullptr;
+			if (SourceStruct && SourcePath)
+			{
+				int32 ArrayIndex = 0;
+				PropertyAccessEditor.ResolvePropertyAccess(SourceStruct, SourcePath->Path, SourceProperty, ArrayIndex);
+			}
+
+			FLinearColor BindingColor = FLinearColor::Gray;
+			if (SourceProperty)
+			{
+				const UEdGraphSchema_K2* Schema = GetDefault<UEdGraphSchema_K2>();
+				FEdGraphPinType PinType;
+				Schema->ConvertPropertyToPinType(SourceProperty, PinType);
+				BindingColor = Schema->GetPinTypeColor(PinType);
+			}
+
+			// TODO: Handle coloring of type promotion
+
+			return BindingColor;
+		});
+
+	Args.bAllowNewBindings = false;
+	Args.bAllowArrayElementBindings = false;
+	Args.bAllowUObjectFunctions = false;
+
+	return PropertyAccessEditor.MakePropertyBindingWidget(Context, Args);
+}
+
+#undef LOCTEXT_NAMESPACE
diff --git a/Engine/Plugins/Runtime/StateTree/Source/StateTreeEditorModule/Private/Customizations/StateTreeEditorPropertyPathDetails.h b/Engine/Plugins/Runtime/StateTree/Source/StateTreeEditorModule/Private/Customizations/StateTreeEditorPropertyPathDetails.h
new file mode 100644
--- /dev/null
+++ b/Engine/Plugins/Runtime/StateTree/Source/StateTreeEditorModule/Private/Customizations/StateTreeEditorPropertyPathDetails.h
@@ -0,0 +1,33 @@
+// Copyright Epic Games, Inc. All Rights Reserved.
+
+#pragma once
+
+#include "CoreMinimal.h"
+#include "UnrealClient.h"
+#include "IPropertyTypeCustomization.h"
+
+class IPropertyHandle;
+class IPropertyUtilities;
+class FDetailWidgetRow;
+class IDetailChildrenBuilder;
+
+/**
+ * Type customization for FStateTreeTransition.
+ */
+
+class FStateTreeEditorPropertyPathDetails : public IPropertyTypeCustomization
+{
+public:
+	/** Makes a new instance of this detail layout class for a specific detail view requesting it */
+	static TSharedRef<IPropertyTypeCustomization> MakeInstance();
+
+	/** IPropertyTypeCustomization interface */
+	virtual void CustomizeHeader(TSharedRef<IPropertyHandle> StructPropertyHandle, FDetailWidgetRow& HeaderRow, IPropertyTypeCustomizationUtils& StructCustomizationUtils) override;
+	virtual void CustomizeChildren(TSharedRef<IPropertyHandle> StructPropertyHandle, IDetailChildrenBuilder& StructBuilder, IPropertyTypeCustomizationUtils& StructCustomizationUtils) override;
+
+private:
+	TSharedRef<SWidget> CreateBindingWidget(TSharedPtr<IPropertyHandle> InPropertyHandle);
+
+	IPropertyUtilities* PropUtils;
+	TSharedPtr<IPropertyHandle> StructProperty;
+};
diff --git a/Engine/Plugins/Runtime/StateTree/Source/StateTreeEditorModule/Private/Customizations/StateTreeParameterLayoutDetails.cpp b/Engine/Plugins/Runtime/StateTree/Source/StateTreeEditorModule/Private/Customizations/StateTreeParameterLayoutDetails.cpp
new file mode 100644
--- /dev/null
+++ b/Engine/Plugins/Runtime/StateTree/Source/StateTreeEditorModule/Private/Customizations/StateTreeParameterLayoutDetails.cpp
@@ -0,0 +1,163 @@
+// Copyright Epic Games, Inc. All Rights Reserved.
+
+#include "StateTreeParameterLayoutDetails.h"
+#include "Widgets/DeclarativeSyntaxSupport.h"
+#include "Widgets/Images/SImage.h"
+#include "DetailWidgetRow.h"
+#include "DetailLayoutBuilder.h"
+#include "IPropertyUtilities.h"
+#include "IDetailPropertyRow.h"
+#include "IDetailChildrenBuilder.h"
+#include "StateTree.h"
+#include "StateTreeDelegates.h"
+#include "StateTreeState.h"
+#include "Widgets/Input/SComboButton.h"
+#include "Framework/MultiBox/MultiBoxBuilder.h"
+#include "ScopedTransaction.h"
+#include "StateTreePropertyHelpers.h"
+#include "Modules/ModuleManager.h"
+#include "ISettingsModule.h"
+
+#define LOCTEXT_NAMESPACE "StateTreeEditor"
+
+TSharedRef<IPropertyTypeCustomization> FStateTreeParameterLayoutDetails::MakeInstance()
+{
+	return MakeShareable(new FStateTreeParameterLayoutDetails);
+}
+
+void FStateTreeParameterLayoutDetails::CustomizeHeader(TSharedRef<class IPropertyHandle> StructPropertyHandle, class FDetailWidgetRow& HeaderRow, IPropertyTypeCustomizationUtils& StructCustomizationUtils)
+{
+	StructProperty = StructPropertyHandle;
+	PropUtils = StructCustomizationUtils.GetPropertyUtilities().Get();
+
+	ParametersProperty = StructProperty->GetChildHandle(TEXT("Parameters"));
+
+	UE::StateTree::Delegates::OnParametersInvalidated.AddSP(this, &FStateTreeParameterLayoutDetails::OnParametersInvalidated);
+
+	HeaderRow
+		.WholeRowContent()
+		[
+			SNew(SHorizontalBox)
+			// Name
+			+ SHorizontalBox::Slot()
+			.FillWidth(1.0f)
+			.VAlign(VAlign_Center)
+			[
+				StructPropertyHandle->CreatePropertyNameWidget()
+			]
+			// Array controls
+			+ SHorizontalBox::Slot()
+			.AutoWidth()
+			.Padding(FMargin(12.0f, 0.0f))
+			[
+				ParametersProperty->CreateDefaultPropertyButtonWidgets()
+			]
+		];
+}
+
+void FStateTreeParameterLayoutDetails::CustomizeChildren(TSharedRef<class IPropertyHandle> StructPropertyHandle, class IDetailChildrenBuilder& StructBuilder, IPropertyTypeCustomizationUtils& StructCustomizationUtils)
+{
+	// Inline variable binding array as children
+	if (ParametersProperty)
+	{
+		uint32 ChildNum = 0;
+		if (ParametersProperty->GetNumChildren(ChildNum) == FPropertyAccess::Success)
+		{
+			for (uint32 ChildIndex = 0; ChildIndex < ChildNum; ++ChildIndex)
+			{
+				TSharedPtr<IPropertyHandle> ChildHandle = ParametersProperty->GetChildHandle(ChildIndex);
+				if (ChildHandle)
+				{
+					// Set the property name to the name of the variable instead of array index.
+					TSharedPtr<IPropertyHandle> NameProperty = ChildHandle->GetChildHandle(TEXT("Name"));
+					TSharedPtr<IPropertyHandle> VariableProperty = ChildHandle->GetChildHandle(TEXT("Variable"));
+
+					if (NameProperty && VariableProperty)
+					{
+						IDetailPropertyRow& PropertyRow = StructBuilder.AddProperty(VariableProperty.ToSharedRef());
+
+						FName Name;
+						NameProperty->GetValue(Name);
+
+						TSharedPtr<SWidget> NameWidget;
+						TSharedPtr<SWidget> ValueWidget;
+						FDetailWidgetRow Row;
+						PropertyRow.GetDefaultWidgets(NameWidget, ValueWidget, Row);
+
+						PropertyRow
+							.CustomWidget()
+							.NameContent()
+							[
+								SNew(SHorizontalBox)
+								// Icon slot
+								+ SHorizontalBox::Slot()
+								.AutoWidth()
+								.VAlign(VAlign_Center)
+								[
+									SNew(SImage)
+									.Image(GetTypeIcon(VariableProperty))
+									.ColorAndOpacity(FLinearColor::White)
+								]
+								// Name
+								+ SHorizontalBox::Slot()
+								.FillWidth(1.0f)
+								.Padding(FMargin(6.0f, 2.0f))
+								[
+									SNew(STextBlock)
+									.Text(FText::FromName(Name))
+									.Font(IDetailLayoutBuilder::GetDetailFontBold())
+								]
+							]
+							.ValueContent()
+							.MinDesiredWidth(250.0f)
+							.VAlign(VAlign_Center)
+							[
+								ValueWidget.ToSharedRef()
+							];
+
+					}
+				}
+			}
+		}
+	}
+}
+
+const FSlateBrush* FStateTreeParameterLayoutDetails::GetTypeIcon(TSharedPtr<IPropertyHandle> VariableProperty) const
+{
+	EStateTreeVariableType Type = EStateTreeVariableType::Void;
+
+	TSharedPtr<IPropertyHandle> TypeProperty = VariableProperty->GetChildHandle(TEXT("Type"));
+	if (TypeProperty)
+	{
+		uint8 Value;
+		if (TypeProperty->GetValue(Value) == FPropertyAccess::Success)
+		{
+			Type = EStateTreeVariableType(Value);
+		}
+	}
+
+	return UE::StateTree::PropertyHelpers::GetTypeIcon(Type);
+}
+
+void FStateTreeParameterLayoutDetails::OnParametersInvalidated(const UStateTree& StateTree)
+{
+	// Check if the change is relevant to us.
+	bool bShouldUpdate = false;
+	TArray<UObject*> OuterObjects;
+	StructProperty->GetOuterObjects(OuterObjects);
+	for (UObject* Outer : OuterObjects)
+	{
+		if (Outer->GetTypedOuter<UStateTree>() == &StateTree)
+		{
+			bShouldUpdate = true;
+			break;
+		}
+	}
+
+	if (bShouldUpdate && PropUtils)
+	{
+		PropUtils->ForceRefresh();
+	}
+}
+
+#undef LOCTEXT_NAMESPACE
