Moved GameplayBehaviors out of restricted folder over to Experimental

diff --git a/Engine/Plugins/Runtime/StateTree/Source/StateTreeEditorModule/Private/Customizations/StateTreeParameterLayoutDetails.h b/Engine/Plugins/Runtime/StateTree/Source/StateTreeEditorModule/Private/Customizations/StateTreeParameterLayoutDetails.h
new file mode 100644
--- /dev/null
+++ b/Engine/Plugins/Runtime/StateTree/Source/StateTreeEditorModule/Private/Customizations/StateTreeParameterLayoutDetails.h
@@ -0,0 +1,36 @@
+// Copyright Epic Games, Inc. All Rights Reserved.
+
+#pragma once
+
+#include "CoreMinimal.h"
+#include "UnrealClient.h"
+#include "IPropertyTypeCustomization.h"
+
+class IPropertyHandle;
+class UStateTree;
+
+/**
+ * Type customization for FStateTreeParameterLayout.
+ */
+
+class FStateTreeParameterLayoutDetails : public IPropertyTypeCustomization
+{
+public:
+	/** Makes a new instance of this detail layout class for a specific detail view requesting it */
+	static TSharedRef<IPropertyTypeCustomization> MakeInstance();
+
+	/** IPropertyTypeCustomization interface */
+	virtual void CustomizeHeader(TSharedRef<IPropertyHandle> StructPropertyHandle, class FDetailWidgetRow& HeaderRow, IPropertyTypeCustomizationUtils& StructCustomizationUtils) override;
+	virtual void CustomizeChildren(TSharedRef<IPropertyHandle> StructPropertyHandle, class IDetailChildrenBuilder& StructBuilder, IPropertyTypeCustomizationUtils& StructCustomizationUtils) override;
+
+private:
+
+	void OnParametersInvalidated(const UStateTree& StateTree);
+
+	const FSlateBrush* GetTypeIcon(TSharedPtr<IPropertyHandle> VariableProperty) const;
+
+	TSharedPtr<IPropertyHandle> ParametersProperty;
+
+	class IPropertyUtilities* PropUtils;
+	TSharedPtr<IPropertyHandle> StructProperty;
+};
diff --git a/Engine/Plugins/Runtime/StateTree/Source/StateTreeEditorModule/Private/Customizations/StateTreeStateDetails.cpp b/Engine/Plugins/Runtime/StateTree/Source/StateTreeEditorModule/Private/Customizations/StateTreeStateDetails.cpp
new file mode 100644
--- /dev/null
+++ b/Engine/Plugins/Runtime/StateTree/Source/StateTreeEditorModule/Private/Customizations/StateTreeStateDetails.cpp
@@ -0,0 +1,134 @@
+// Copyright Epic Games, Inc. All Rights Reserved.
+
+#include "StateTreeStateDetails.h"
+#include "Misc/MessageDialog.h"
+#include "UObject/ObjectMacros.h"
+#include "UObject/Class.h"
+#include "UObject/UnrealType.h"
+#include "Layout/Visibility.h"
+#include "Misc/Attribute.h"
+#include "Widgets/DeclarativeSyntaxSupport.h"
+#include "Widgets/SWidget.h"
+#include "Widgets/Text/STextBlock.h"
+#include "Widgets/Layout/SBox.h"
+#include "Widgets/Input/SComboBox.h"
+#include "Widgets/Input/SVectorInputBox.h"
+#include "Widgets/Input/SNumericEntryBox.h"
+#include "Widgets/Input/SButton.h"
+#include "Framework/MultiBox/MultiBoxBuilder.h"
+#include "ComponentVisualizer.h"
+#include "Editor/UnrealEdEngine.h"
+#include "UnrealEdGlobals.h"
+#include "IDetailCustomNodeBuilder.h"
+#include "IDetailChildrenBuilder.h"
+#include "DetailWidgetRow.h"
+#include "DetailLayoutBuilder.h"
+#include "DetailCategoryBuilder.h"
+#include "PropertyCustomizationHelpers.h"
+#include "ScopedTransaction.h"
+#include "Modules/ModuleManager.h"
+#include "ISettingsModule.h"
+#include "Editor.h"
+#include "StateTreeEditorData.h"
+#include "StateTreeViewModel.h"
+#include "StateTree.h"
+
+
+#define LOCTEXT_NAMESPACE "StateTreeEditor"
+
+
+TSharedRef<IDetailCustomization> FStateTreeStateDetails::MakeInstance()
+{
+	return MakeShareable(new FStateTreeStateDetails);
+}
+
+void FStateTreeStateDetails::CustomizeDetails(IDetailLayoutBuilder& DetailBuilder)
+{
+	TSharedPtr<IPropertyHandle> TasksProperty = DetailBuilder.GetProperty(TEXT("Tasks"));
+	TSharedPtr<IPropertyHandle> Tasks2Property = DetailBuilder.GetProperty(TEXT("Tasks2"));
+	TSharedPtr<IPropertyHandle> EnterConditionsProperty = DetailBuilder.GetProperty(TEXT("EnterConditions"));
+	TSharedPtr<IPropertyHandle> EnterConditions2Property = DetailBuilder.GetProperty(TEXT("EnterConditions2"));
+	TSharedPtr<IPropertyHandle> EvaluatorsProperty = DetailBuilder.GetProperty(TEXT("Evaluators"));
+	TSharedPtr<IPropertyHandle> Evaluators2Property = DetailBuilder.GetProperty(TEXT("Evaluators2"));
+	TSharedPtr<IPropertyHandle> TransitionsProperty = DetailBuilder.GetProperty(TEXT("Transitions"));
+	TSharedPtr<IPropertyHandle> Transitions2Property = DetailBuilder.GetProperty(TEXT("Transitions2"));
+	TSharedPtr<IPropertyHandle> StateDoneTransitionProperty = DetailBuilder.GetProperty(TEXT("StateDoneTransition"));
+	TSharedPtr<IPropertyHandle> StateFailedTransitionProperty = DetailBuilder.GetProperty(TEXT("StateFailedTransition"));
+
+	TArray<TWeakObjectPtr<UObject>> Objects;
+	DetailBuilder.GetObjectsBeingCustomized(Objects);
+	UObject* Object = Objects.IsEmpty() ? nullptr : Objects[0].Get();
+	
+	if (IsV2(Object))
+	{
+		EnterConditionsProperty->MarkHiddenByCustomization();
+		EnterConditionsProperty = EnterConditions2Property;
+		EvaluatorsProperty->MarkHiddenByCustomization();
+		EvaluatorsProperty = Evaluators2Property;
+		TasksProperty->MarkHiddenByCustomization();
+		TasksProperty = Tasks2Property;
+		TransitionsProperty->MarkHiddenByCustomization();
+		TransitionsProperty = Transitions2Property;
+		StateDoneTransitionProperty->MarkHiddenByCustomization();
+		StateFailedTransitionProperty->MarkHiddenByCustomization();
+	}
+	else
+	{
+		EnterConditions2Property->MarkHiddenByCustomization();
+		Evaluators2Property->MarkHiddenByCustomization();
+		Tasks2Property->MarkHiddenByCustomization();
+		Transitions2Property->MarkHiddenByCustomization();
+	}
+
+	IDetailCategoryBuilder& StateCategory = DetailBuilder.EditCategory(TEXT("State"), LOCTEXT("StateDetailsState", "State"));
+	StateCategory.SetSortOrder(0);
+
+	MakeArrayCategory(DetailBuilder, "Evaluators", LOCTEXT("StateDetailsEvaluators", "Evaluators"), 1, EvaluatorsProperty);
+
+	MakeArrayCategory(DetailBuilder, "Enter Conditions", LOCTEXT("StateDetailsEnterConditions", "Enter Conditions"), 2, EnterConditionsProperty);
+
+	MakeArrayCategory(DetailBuilder, "Tasks", LOCTEXT("StateDetailsTasks", "Tasks"), 3, TasksProperty);
+
+	MakeArrayCategory(DetailBuilder, "Transitions", LOCTEXT("StateDetailsTransitions", "Transitions"), 4, TransitionsProperty);
+}
+
+void FStateTreeStateDetails::MakeArrayCategory(IDetailLayoutBuilder& DetailBuilder, FName CategoryName, const FText& DisplayName, int32 SortOrder, TSharedPtr<IPropertyHandle> PropertyHandle)
+{
+	IDetailCategoryBuilder& Category = DetailBuilder.EditCategory(CategoryName, DisplayName);
+	Category.SetSortOrder(SortOrder);
+
+	TSharedRef<SHorizontalBox> HeaderContentWidget = SNew(SHorizontalBox);
+	HeaderContentWidget->AddSlot()
+	.HAlign(HAlign_Right)
+	.VAlign(VAlign_Center)
+	[
+		PropertyHandle->CreateDefaultPropertyButtonWidgets()
+	];
+	Category.HeaderContent(HeaderContentWidget);
+
+	// Add items inline
+	TSharedRef<FDetailArrayBuilder> Builder = MakeShareable(new FDetailArrayBuilder(PropertyHandle.ToSharedRef(), /*InGenerateHeader*/ false, /*InDisplayResetToDefault*/ true, /*InDisplayElementNum*/ false));
+	Builder->OnGenerateArrayElementWidget(FOnGenerateArrayElementWidget::CreateSP(this, &FStateTreeStateDetails::GenerateArrayElementWidget));
+	Category.AddCustomBuilder(Builder, /*bForAdvanced*/ false);
+}
+
+void FStateTreeStateDetails::GenerateArrayElementWidget(TSharedRef<IPropertyHandle> PropertyHandle, int32 ArrayIndex, IDetailChildrenBuilder& ChildrenBuilder)
+{
+	ChildrenBuilder.AddProperty(PropertyHandle);
+};
+
+bool FStateTreeStateDetails::IsV2(const UObject* Object) const
+{
+	if (!Object)
+	{
+		return false;
+	}
+	if (const UStateTree* StateTree = Object->GetTypedOuter<UStateTree>())
+	{
+		return StateTree->IsV2();
+	}
+	return false;
+}
+
+
+#undef LOCTEXT_NAMESPACE
diff --git a/Engine/Plugins/Runtime/StateTree/Source/StateTreeEditorModule/Private/Customizations/StateTreeStateDetails.h b/Engine/Plugins/Runtime/StateTree/Source/StateTreeEditorModule/Private/Customizations/StateTreeStateDetails.h
new file mode 100644
--- /dev/null
+++ b/Engine/Plugins/Runtime/StateTree/Source/StateTreeEditorModule/Private/Customizations/StateTreeStateDetails.h
@@ -0,0 +1,27 @@
+// Copyright Epic Games, Inc. All Rights Reserved.
+
+#pragma once
+
+#include "CoreMinimal.h"
+#include "IDetailCustomization.h"
+#include "Input/Reply.h"
+
+class IDetailLayoutBuilder;
+class IDetailChildrenBuilder;
+class IPropertyHandle;
+
+class FStateTreeStateDetails : public IDetailCustomization
+{
+public:
+	/** Makes a new instance of this detail layout class for a specific detail view requesting it */
+	static TSharedRef<IDetailCustomization> MakeInstance();
+
+	/** IDetailCustomization interface */
+	virtual void CustomizeDetails(IDetailLayoutBuilder& DetailBuilder) override;
+
+protected:
+	void MakeArrayCategory(IDetailLayoutBuilder& DetailBuilder, FName CategoryName, const FText& DisplayName, int32 SortOrder, TSharedPtr<IPropertyHandle> PropertyHandle);
+	void GenerateArrayElementWidget(TSharedRef<IPropertyHandle> PropertyHandle, int32 ArrayIndex, IDetailChildrenBuilder& ChildrenBuilder);
+
+	bool IsV2(const UObject* Object) const;
+};
diff --git a/Engine/Plugins/Runtime/StateTree/Source/StateTreeEditorModule/Private/Customizations/StateTreeStateLinkDetails.cpp b/Engine/Plugins/Runtime/StateTree/Source/StateTreeEditorModule/Private/Customizations/StateTreeStateLinkDetails.cpp
new file mode 100644
--- /dev/null
+++ b/Engine/Plugins/Runtime/StateTree/Source/StateTreeEditorModule/Private/Customizations/StateTreeStateLinkDetails.cpp
@@ -0,0 +1,263 @@
+// Copyright Epic Games, Inc. All Rights Reserved.
+
+#include "StateTreeStateLinkDetails.h"
+#include "Widgets/DeclarativeSyntaxSupport.h"
+#include "DetailWidgetRow.h"
+#include "DetailLayoutBuilder.h"
+#include "StateTree.h"
+#include "StateTreeState.h"
+#include "StateTreeEditorData.h"
+#include "StateTreeDelegates.h"
+#include "Widgets/Input/SComboButton.h"
+#include "Framework/MultiBox/MultiBoxBuilder.h"
+#include "ScopedTransaction.h"
+#include "StateTreePropertyHelpers.h"
+
+#define LOCTEXT_NAMESPACE "StateTreeEditor"
+
+TSharedRef<IPropertyTypeCustomization> FStateTreeStateLinkDetails::MakeInstance()
+{
+	return MakeShareable(new FStateTreeStateLinkDetails);
+}
+
+void FStateTreeStateLinkDetails::CustomizeHeader(TSharedRef<class IPropertyHandle> StructPropertyHandle, class FDetailWidgetRow& HeaderRow, IPropertyTypeCustomizationUtils& StructCustomizationUtils)
+{
+	StructProperty = StructPropertyHandle;
+	PropUtils = StructCustomizationUtils.GetPropertyUtilities().Get();
+
+	NameProperty = StructProperty->GetChildHandle(TEXT("Name"));
+	IDProperty = StructProperty->GetChildHandle(TEXT("ID"));
+	TypeProperty = StructProperty->GetChildHandle(TEXT("Type"));
+
+	CacheStates();
+
+	HeaderRow
+		.NameContent()
+		[
+			StructPropertyHandle->CreatePropertyNameWidget()
+		]
+		.ValueContent()
+		.VAlign(VAlign_Center)
+		[
+			SNew(SComboButton)
+			.OnGetMenuContent(this, &FStateTreeStateLinkDetails::OnGetStateContent)
+			.ContentPadding(FMargin(6.f, 0.f))
+			.ButtonContent()
+			[
+				SNew(STextBlock)
+				.Text(this, &FStateTreeStateLinkDetails::GetCurrentStateDesc)
+				.Font(IDetailLayoutBuilder::GetDetailFont())
+			]
+		];
+
+	UE::StateTree::Delegates::OnIdentifierChanged.AddSP(this, &FStateTreeStateLinkDetails::OnIdentifierChanged);
+}
+
+void FStateTreeStateLinkDetails::CustomizeChildren(TSharedRef<class IPropertyHandle> StructPropertyHandle, class IDetailChildrenBuilder& StructBuilder, IPropertyTypeCustomizationUtils& StructCustomizationUtils)
+{
+}
+
+void FStateTreeStateLinkDetails::OnIdentifierChanged(const UStateTree& StateTree)
+{
+	CacheStates();
+}
+
+void FStateTreeStateLinkDetails::CacheStates(const UStateTreeState* State)
+{
+	if (State == nullptr)
+	{
+		return;
+	}
+
+	CachedNames.Add(State->Name);
+	CachedIDs.Add(State->ID);
+
+	for (UStateTreeState* ChildState : State->Children)
+	{
+		CacheStates(ChildState);
+	}
+}
+
+void FStateTreeStateLinkDetails::CacheStates()
+{
+	CachedNames.Reset();
+	CachedIDs.Reset();
+
+	TArray<UObject*> OuterObjects;
+	StructProperty->GetOuterObjects(OuterObjects);
+	for (int32 ObjectIdx = 0; ObjectIdx < OuterObjects.Num(); ObjectIdx++)
+	{
+		UStateTree* OuterStateTree = OuterObjects[ObjectIdx]->GetTypedOuter<UStateTree>();
+		if (OuterStateTree)
+		{
+			bIsV2 = OuterStateTree->IsV2();
+			if (UStateTreeEditorData* TreeData = Cast<UStateTreeEditorData>(OuterStateTree->EditorData))
+			{
+				for (const UStateTreeState* Routine : TreeData->Routines)
+				{
+					CacheStates(Routine);
+				}
+			}
+			break;
+		}
+	}
+
+}
+
+void FStateTreeStateLinkDetails::OnStateComboChange(int Idx)
+{
+	if (NameProperty && IDProperty)
+	{
+		FScopedTransaction Transaction(FText::Format(LOCTEXT("SetPropertyValue", "Set {0}"), StructProperty->GetPropertyDisplayName()));
+		if (Idx >= 0)
+		{
+			TypeProperty->SetValue((uint8)EStateTreeTransitionType::GotoState);
+			NameProperty->SetValue(CachedNames[Idx], EPropertyValueSetFlags::NotTransactable);
+			UE::StateTree::PropertyHelpers::SetStructValue<FGuid>(IDProperty, CachedIDs[Idx], EPropertyValueSetFlags::NotTransactable);
+		}
+		else
+		{
+			switch (Idx)
+			{
+			case ComboSucceeded:
+				TypeProperty->SetValue((uint8)EStateTreeTransitionType::Succeeded);
+				break;
+			case ComboFailed:
+				TypeProperty->SetValue((uint8)EStateTreeTransitionType::Failed);
+				break;
+			case ComboNextState:
+				TypeProperty->SetValue((uint8)EStateTreeTransitionType::NextState);
+				break;
+			case ComboSelectChildState:
+				TypeProperty->SetValue((uint8)EStateTreeTransitionType::SelectChildState);
+				break;
+			case ComboNotSet:
+			default:
+				TypeProperty->SetValue((uint8)EStateTreeTransitionType::NotSet);
+				break;
+			}
+			// Clear name and id.
+			NameProperty->SetValue(FName(), EPropertyValueSetFlags::NotTransactable);
+			UE::StateTree::PropertyHelpers::SetStructValue<FGuid>(IDProperty, FGuid(), EPropertyValueSetFlags::NotTransactable);
+		}
+	}
+}
+
+TSharedRef<SWidget> FStateTreeStateLinkDetails::OnGetStateContent() const
+{
+	FMenuBuilder MenuBuilder(true, NULL);
+
+	FUIAction NotSetItemAction(FExecuteAction::CreateSP(const_cast<FStateTreeStateLinkDetails*>(this), &FStateTreeStateLinkDetails::OnStateComboChange, ComboNotSet));
+	MenuBuilder.AddMenuEntry(LOCTEXT("TransitionNotSet", "Not Set"), LOCTEXT("TransitionNotSetTooltip", "Transition not set."), FSlateIcon(), NotSetItemAction);
+
+	MenuBuilder.AddSeparator();
+
+	FUIAction SucceededItemAction(FExecuteAction::CreateSP(const_cast<FStateTreeStateLinkDetails*>(this), &FStateTreeStateLinkDetails::OnStateComboChange, ComboSucceeded));
+	MenuBuilder.AddMenuEntry(LOCTEXT("TransitionTreeSucceeded", "Tree Succeeded"), LOCTEXT("TransitionSucceededTooltip", "Stop StateTree and mark it Succeeded."), FSlateIcon(), SucceededItemAction);
+
+	FUIAction FailedItemAction(FExecuteAction::CreateSP(const_cast<FStateTreeStateLinkDetails*>(this), &FStateTreeStateLinkDetails::OnStateComboChange, ComboFailed));
+	MenuBuilder.AddMenuEntry(LOCTEXT("TransitionTreeFailed", "Tree Failed"), LOCTEXT("TransitionFailedTooltip", "Stop StateTree and mark it Failed."), FSlateIcon(), FailedItemAction);
+
+	MenuBuilder.AddSeparator();
+
+	FUIAction NextItemAction(FExecuteAction::CreateSP(const_cast<FStateTreeStateLinkDetails*>(this), &FStateTreeStateLinkDetails::OnStateComboChange, ComboNextState));
+	MenuBuilder.AddMenuEntry(LOCTEXT("TransitionNextState", "Next State"), LOCTEXT("TransitionNextTooltip", "Goto next sibling State."), FSlateIcon(), NextItemAction);
+
+	if (!bIsV2)
+	{
+		FUIAction SelectItemAction(FExecuteAction::CreateSP(const_cast<FStateTreeStateLinkDetails*>(this), &FStateTreeStateLinkDetails::OnStateComboChange, ComboSelectChildState));
+		MenuBuilder.AddMenuEntry(LOCTEXT("TransitionSelect", "Select"), LOCTEXT("TransitionSelectTooltip", "Select a child State to run."), FSlateIcon(), SelectItemAction);
+	}
+
+	if (CachedNames.Num() > 0)
+	{
+		MenuBuilder.BeginSection(FName(), LOCTEXT("TransitionGotoState", "Goto State"));
+
+		for (int32 Idx = 0; Idx < CachedNames.Num(); Idx++)
+		{
+			FFormatNamedArguments Args;
+			Args.Add(TEXT("StateName"), FText::FromName(CachedNames[Idx]));
+
+			FUIAction ItemAction(FExecuteAction::CreateSP(const_cast<FStateTreeStateLinkDetails*>(this), &FStateTreeStateLinkDetails::OnStateComboChange, Idx));
+			MenuBuilder.AddMenuEntry(FText::FromName(CachedNames[Idx]), FText::Format(LOCTEXT("TransitionGotoStateTooltip", "Goto State {StateName}."), Args), FSlateIcon(), ItemAction);
+		}
+
+		MenuBuilder.EndSection();
+	}
+
+	return MenuBuilder.MakeWidget();
+}
+
+FText FStateTreeStateLinkDetails::GetCurrentStateDesc() const
+{
+	const EStateTreeTransitionType TransitionType = GetTransitionType().Get(EStateTreeTransitionType::Failed);
+
+	if (TransitionType == EStateTreeTransitionType::NotSet)
+	{
+		return LOCTEXT("TransitionNotSet", "Not Set");
+	}
+	else if (TransitionType == EStateTreeTransitionType::Succeeded)
+	{
+		return LOCTEXT("TransitionTreeSucceeded", "Tree Succeeded");
+	}
+	else if (TransitionType == EStateTreeTransitionType::Failed)
+	{
+		return LOCTEXT("TransitionTreeFailed", "Tree Failed");
+	}
+	else if (TransitionType == EStateTreeTransitionType::NextState)
+	{
+		return LOCTEXT("TransitionNextState", "Next State");
+	}
+	else if (TransitionType == EStateTreeTransitionType::SelectChildState)
+	{
+		return LOCTEXT("TransitionSelect", "Select");
+	}
+	else if (TransitionType == EStateTreeTransitionType::GotoState)
+	{
+		FName OldName;
+		if (NameProperty && IDProperty)
+		{
+			NameProperty->GetValue(OldName);
+
+			FGuid StateID;
+			if (UE::StateTree::PropertyHelpers::GetStructValue<FGuid>(IDProperty, StateID) == FPropertyAccess::Success)
+			{
+				if (!StateID.IsValid())
+				{
+					return LOCTEXT("None", "None");
+				}
+				else
+				{
+					for (int32 Idx = 0; Idx < CachedIDs.Num(); Idx++)
+					{
+						if (CachedIDs[Idx] == StateID)
+						{
+							return FText::FromName(CachedNames[Idx]);
+						}
+					}
+				}
+			}
+		}
+
+		FFormatNamedArguments Args;
+		Args.Add(TEXT("Identifier"), FText::FromName(OldName));
+		return FText::Format(LOCTEXT("InvalidReference", "Invalid Reference {Identifier}"), Args);
+	}
+
+	return LOCTEXT("TransitionInvalid", "Invalid");
+}
+
+TOptional<EStateTreeTransitionType> FStateTreeStateLinkDetails::GetTransitionType() const
+{
+	if (TypeProperty)
+	{
+		uint8 Value;
+		if (TypeProperty->GetValue(Value) == FPropertyAccess::Success)
+		{
+			return EStateTreeTransitionType(Value);
+		}
+	}
+	return TOptional<EStateTreeTransitionType>();
+}
+
+#undef LOCTEXT_NAMESPACE
diff --git a/Engine/Plugins/Runtime/StateTree/Source/StateTreeEditorModule/Private/Customizations/StateTreeStateLinkDetails.h b/Engine/Plugins/Runtime/StateTree/Source/StateTreeEditorModule/Private/Customizations/StateTreeStateLinkDetails.h
new file mode 100644
--- /dev/null
+++ b/Engine/Plugins/Runtime/StateTree/Source/StateTreeEditorModule/Private/Customizations/StateTreeStateLinkDetails.h
@@ -0,0 +1,56 @@
+// Copyright Epic Games, Inc. All Rights Reserved.
+
+#pragma once
+
+#include "CoreMinimal.h"
+#include "UnrealClient.h"
+#include "IPropertyTypeCustomization.h"
+#include "StateTreeTypes.h"
+
+class IPropertyHandle;
+class UStateTree;
+class UStateTreeState;
+
+/**
+ * Type customization for FStateTreeStateLink.
+ */
+
+class FStateTreeStateLinkDetails : public IPropertyTypeCustomization
+{
+public:
+	/** Makes a new instance of this detail layout class for a specific detail view requesting it */
+	static TSharedRef<IPropertyTypeCustomization> MakeInstance();
+
+	/** IPropertyTypeCustomization interface */
+	virtual void CustomizeHeader(TSharedRef<class IPropertyHandle> StructPropertyHandle, class FDetailWidgetRow& HeaderRow, IPropertyTypeCustomizationUtils& StructCustomizationUtils) override;
+	virtual void CustomizeChildren(TSharedRef<class IPropertyHandle> StructPropertyHandle, class IDetailChildrenBuilder& StructBuilder, IPropertyTypeCustomizationUtils& StructCustomizationUtils) override;
+
+private:
+
+	// Special indices for the combo selection.
+	static const int ComboNotSet = -1;
+	static const int ComboSucceeded = -2;
+	static const int ComboFailed = -3;
+	static const int ComboNextState = -4;
+	static const int ComboSelectChildState = -5;
+
+	void CacheStates();
+	void CacheStates(const UStateTreeState* State);
+	void OnIdentifierChanged(const UStateTree& StateTree);
+
+	void OnStateComboChange(int Idx);
+	TSharedRef<SWidget> OnGetStateContent() const;
+	FText GetCurrentStateDesc() const;
+	TOptional<EStateTreeTransitionType> GetTransitionType() const;
+
+	TArray<FName> CachedNames;
+	TArray<FGuid> CachedIDs;
+	bool bIsV2;
+
+	TSharedPtr<IPropertyHandle> NameProperty;
+	TSharedPtr<IPropertyHandle> IDProperty;
+	TSharedPtr<IPropertyHandle> TypeProperty;
+
+	class IPropertyUtilities* PropUtils;
+	TSharedPtr<IPropertyHandle> StructProperty;
+};
diff --git a/Engine/Plugins/Runtime/StateTree/Source/StateTreeEditorModule/Private/Customizations/StateTreeTransition2Details.cpp b/Engine/Plugins/Runtime/StateTree/Source/StateTreeEditorModule/Private/Customizations/StateTreeTransition2Details.cpp
new file mode 100644
--- /dev/null
+++ b/Engine/Plugins/Runtime/StateTree/Source/StateTreeEditorModule/Private/Customizations/StateTreeTransition2Details.cpp
@@ -0,0 +1,144 @@
+﻿// Copyright Epic Games, Inc. All Rights Reserved.
+
+#include "StateTreeTransition2Details.h"
+#include "Widgets/DeclarativeSyntaxSupport.h"
+#include "DetailWidgetRow.h"
+#include "DetailLayoutBuilder.h"
+#include "IDetailPropertyRow.h"
+#include "IDetailChildrenBuilder.h"
+#include "StateTree.h"
+#include "StateTreeState.h"
+#include "Widgets/Input/SComboButton.h"
+#include "Framework/MultiBox/MultiBoxBuilder.h"
+
+#define LOCTEXT_NAMESPACE "StateTreeEditor"
+
+TSharedRef<IPropertyTypeCustomization> FStateTreeTransition2Details::MakeInstance()
+{
+	return MakeShareable(new FStateTreeTransition2Details);
+}
+
+void FStateTreeTransition2Details::CustomizeHeader(TSharedRef<class IPropertyHandle> StructPropertyHandle, class FDetailWidgetRow& HeaderRow, IPropertyTypeCustomizationUtils& StructCustomizationUtils)
+{
+	StructProperty = StructPropertyHandle;
+	PropUtils = StructCustomizationUtils.GetPropertyUtilities().Get();
+
+	EventProperty = StructProperty->GetChildHandle(TEXT("Event"));
+	StateProperty = StructProperty->GetChildHandle(TEXT("State"));
+	GateDelayProperty = StructProperty->GetChildHandle(TEXT("GateDelay"));
+	ConditionsProperty = StructProperty->GetChildHandle(TEXT("Conditions"));
+
+	HeaderRow
+		.WholeRowContent()
+		.VAlign(VAlign_Center)
+		[
+			SNew(SHorizontalBox)
+			// Description
+			+ SHorizontalBox::Slot()
+			.FillWidth(1.0f)
+			.VAlign(VAlign_Center)
+			[
+				SNew(STextBlock)
+				.Text(this, &FStateTreeTransition2Details::GetDescription)
+				.Font(IDetailLayoutBuilder::GetDetailFont())
+			]
+			+ SHorizontalBox::Slot()
+			.AutoWidth()
+			[
+				StructPropertyHandle->CreateDefaultPropertyButtonWidgets()
+			]
+		];
+}
+
+void FStateTreeTransition2Details::CustomizeChildren(TSharedRef<class IPropertyHandle> StructPropertyHandle, class IDetailChildrenBuilder& StructBuilder, IPropertyTypeCustomizationUtils& StructCustomizationUtils)
+{
+	if (EventProperty)
+	{
+		IDetailPropertyRow& Property = StructBuilder.AddProperty(EventProperty.ToSharedRef());
+	}
+
+	if (GateDelayProperty)
+	{
+		IDetailPropertyRow& Property = StructBuilder.AddProperty(GateDelayProperty.ToSharedRef());
+	}
+
+	if (StateProperty)
+	{
+		IDetailPropertyRow& Property = StructBuilder.AddProperty(StateProperty.ToSharedRef());
+	}
+
+	if (ConditionsProperty)
+	{
+		// Show conditions always expanded, with simplified header (remove item count)
+		IDetailPropertyRow& Property = StructBuilder.AddProperty(ConditionsProperty.ToSharedRef());
+		Property.ShouldAutoExpand(true);
+
+		static const bool bShowChildren = true;
+		Property.CustomWidget(bShowChildren)
+			.NameContent()
+			.VAlign(VAlign_Center)
+			[
+				SNew(SHorizontalBox)
+				+ SHorizontalBox::Slot()
+				.AutoWidth()
+				.VAlign(VAlign_Center)
+				.Padding(FMargin(0.0f, 2.0f))
+				[
+					SNew(STextBlock)
+					.Text(ConditionsProperty->GetPropertyDisplayName())
+					.Font(IDetailLayoutBuilder::GetDetailFontBold())
+				]
+			]
+			.ValueContent()
+			.VAlign(VAlign_Center)
+			.HAlign(HAlign_Right)
+			[
+				SNew(SBox) // Empty, suppress noisy array details.
+			];
+	}
+}
+
+FText FStateTreeTransition2Details::GetDescription() const
+{
+	if (StateProperty)
+	{
+		TArray<void*> RawData;
+		StateProperty->AccessRawData(RawData);
+		if (RawData.Num() == 1)
+		{
+			FStateTreeStateLink* State = static_cast<FStateTreeStateLink*>(RawData[0]);
+			if (State != nullptr)
+			{
+				switch (State->Type)
+				{
+				case EStateTreeTransitionType::NotSet:
+					return LOCTEXT("TransitionNotSet", "None");
+					break;
+				case EStateTreeTransitionType::Succeeded:
+					return LOCTEXT("TransitionTreeSucceeded", "Tree Succeeded");
+					break;
+				case EStateTreeTransitionType::Failed:
+					return LOCTEXT("TransitionTreeFailed", "Tree Failed");
+					break;
+				case EStateTreeTransitionType::NextState:
+					return LOCTEXT("TransitionNextStart", "Next State");
+					break;
+				case EStateTreeTransitionType::GotoState:
+					{
+						FFormatNamedArguments Args;
+						Args.Add(TEXT("State"), FText::FromName(State->Name));
+						return FText::Format(LOCTEXT("TransitionActionGotoState", "Go to State {State}"), Args);
+					}
+					break;
+				}
+			}
+		}
+		else
+		{
+			return LOCTEXT("MultipleSelected", "Multiple selected");
+		}
+	}
+	return FText::GetEmpty();
+}
+
+#undef LOCTEXT_NAMESPACE
diff --git a/Engine/Plugins/Runtime/StateTree/Source/StateTreeEditorModule/Private/Customizations/StateTreeTransition2Details.h b/Engine/Plugins/Runtime/StateTree/Source/StateTreeEditorModule/Private/Customizations/StateTreeTransition2Details.h
new file mode 100644
--- /dev/null
+++ b/Engine/Plugins/Runtime/StateTree/Source/StateTreeEditorModule/Private/Customizations/StateTreeTransition2Details.h
@@ -0,0 +1,38 @@
+﻿// Copyright Epic Games, Inc. All Rights Reserved.
+
+#pragma once
+
+#include "CoreMinimal.h"
+#include "UnrealClient.h"
+#include "IPropertyTypeCustomization.h"
+
+class IPropertyHandle;
+class UStateTree;
+class UStateTreeState;
+
+/**
+ * Type customization for FStateTreeTransition2.
+ */
+
+class FStateTreeTransition2Details : public IPropertyTypeCustomization
+{
+public:
+	/** Makes a new instance of this detail layout class for a specific detail view requesting it */
+	static TSharedRef<IPropertyTypeCustomization> MakeInstance();
+
+	/** IPropertyTypeCustomization interface */
+	virtual void CustomizeHeader(TSharedRef<IPropertyHandle> StructPropertyHandle, class FDetailWidgetRow& HeaderRow, IPropertyTypeCustomizationUtils& StructCustomizationUtils) override;
+	virtual void CustomizeChildren(TSharedRef<IPropertyHandle> StructPropertyHandle, class IDetailChildrenBuilder& StructBuilder, IPropertyTypeCustomizationUtils& StructCustomizationUtils) override;
+
+private:
+
+	FText GetDescription() const;
+
+	TSharedPtr<IPropertyHandle> EventProperty;
+	TSharedPtr<IPropertyHandle> StateProperty;
+	TSharedPtr<IPropertyHandle> GateDelayProperty;
+	TSharedPtr<IPropertyHandle> ConditionsProperty;
+
+	class IPropertyUtilities* PropUtils;
+	TSharedPtr<IPropertyHandle> StructProperty;
+};
diff --git a/Engine/Plugins/Runtime/StateTree/Source/StateTreeEditorModule/Private/Customizations/StateTreeTransitionDetails.cpp b/Engine/Plugins/Runtime/StateTree/Source/StateTreeEditorModule/Private/Customizations/StateTreeTransitionDetails.cpp
new file mode 100644
--- /dev/null
+++ b/Engine/Plugins/Runtime/StateTree/Source/StateTreeEditorModule/Private/Customizations/StateTreeTransitionDetails.cpp
@@ -0,0 +1,128 @@
+// Copyright Epic Games, Inc. All Rights Reserved.
+
+#include "StateTreeTransitionDetails.h"
+#include "Widgets/DeclarativeSyntaxSupport.h"
+#include "DetailWidgetRow.h"
+#include "DetailLayoutBuilder.h"
+#include "IPropertyUtilities.h"
+#include "IDetailPropertyRow.h"
+#include "IDetailChildrenBuilder.h"
+#include "StateTree.h"
+#include "StateTreeState.h"
+#include "Widgets/Input/SComboButton.h"
+#include "Framework/MultiBox/MultiBoxBuilder.h"
+#include "ScopedTransaction.h"
+
+#define LOCTEXT_NAMESPACE "StateTreeEditor"
+
+TSharedRef<IPropertyTypeCustomization> FStateTreeTransitionDetails::MakeInstance()
+{
+	return MakeShareable(new FStateTreeTransitionDetails);
+}
+
+void FStateTreeTransitionDetails::CustomizeHeader(TSharedRef<class IPropertyHandle> StructPropertyHandle, class FDetailWidgetRow& HeaderRow, IPropertyTypeCustomizationUtils& StructCustomizationUtils)
+{
+	StructProperty = StructPropertyHandle;
+	PropUtils = StructCustomizationUtils.GetPropertyUtilities().Get();
+
+	StateProperty = StructProperty->GetChildHandle(TEXT("State"));
+	ConditionsProperty = StructProperty->GetChildHandle(TEXT("Conditions"));
+
+	HeaderRow
+		.WholeRowContent()
+		.VAlign(VAlign_Center)
+		[
+			SNew(SHorizontalBox)
+			// Description
+			+ SHorizontalBox::Slot()
+			.FillWidth(1.0f)
+			.VAlign(VAlign_Center)
+			[
+				SNew(STextBlock)
+				.Text(this, &FStateTreeTransitionDetails::GetDescription)
+				.Font(IDetailLayoutBuilder::GetDetailFont())
+			]
+			+ SHorizontalBox::Slot()
+			.AutoWidth()
+			[
+				StructPropertyHandle->CreateDefaultPropertyButtonWidgets()
+			]
+		];
+}
+
+void FStateTreeTransitionDetails::CustomizeChildren(TSharedRef<class IPropertyHandle> StructPropertyHandle, class IDetailChildrenBuilder& StructBuilder, IPropertyTypeCustomizationUtils& StructCustomizationUtils)
+{
+	if (StateProperty)
+	{
+		IDetailPropertyRow& Property = StructBuilder.AddProperty(StateProperty.ToSharedRef());
+	}
+
+	if (ConditionsProperty)
+	{
+		// Show conditions always expanded, with simplified header (remove item count)
+		IDetailPropertyRow& Property = StructBuilder.AddProperty(ConditionsProperty.ToSharedRef());
+		Property.ShouldAutoExpand(true);
+
+		static const bool bShowChildren = true;
+		Property.CustomWidget(bShowChildren)
+			.NameContent()
+			.VAlign(VAlign_Center)
+			[
+				SNew(SHorizontalBox)
+				+ SHorizontalBox::Slot()
+				.AutoWidth()
+				.VAlign(VAlign_Center)
+				.Padding(FMargin(0.0f, 2.0f))
+				[
+					SNew(STextBlock)
+					.Text(ConditionsProperty->GetPropertyDisplayName())
+					.Font(IDetailLayoutBuilder::GetDetailFontBold())
+				]
+			]
+			.ValueContent()
+			.VAlign(VAlign_Center)
+			.HAlign(HAlign_Right)
+			[
+				SNew(SBox) // Empty, suppress noisy array details.
+			];
+	}
+}
+
+
+FText FStateTreeTransitionDetails::GetDescription() const
+{
+	if (StateProperty)
+	{
+		TArray<void*> RawData;
+		StateProperty->AccessRawData(RawData);
+		if (RawData.Num() == 1)
+		{
+			FStateTreeStateLink* StateLink = static_cast<FStateTreeStateLink*>(RawData[0]);
+			if (StateLink != nullptr)
+			{
+				switch (StateLink->Type)
+				{
+				case EStateTreeTransitionType::Succeeded:
+					return LOCTEXT("TransitionSucceeded", "Succeeded");
+					break;
+				case EStateTreeTransitionType::Failed:
+					return LOCTEXT("TransitionFailed", "Failed");
+					break;
+				case EStateTreeTransitionType::SelectChildState:
+					return LOCTEXT("TransitionSelect", "Select");
+					break;
+				case EStateTreeTransitionType::GotoState:
+					return FText::Join(FText::FromString(TEXT(" ")), LOCTEXT("TransitionGoto", "Go to"), FText::FromName(StateLink->Name));
+					break;
+				}
+			}
+		}
+		else
+		{
+			return LOCTEXT("MultipleSelected", "Multiple selected");
+		}
+	}
+	return FText::GetEmpty();
+}
+
+#undef LOCTEXT_NAMESPACE
diff --git a/Engine/Plugins/Runtime/StateTree/Source/StateTreeEditorModule/Private/Customizations/StateTreeTransitionDetails.h b/Engine/Plugins/Runtime/StateTree/Source/StateTreeEditorModule/Private/Customizations/StateTreeTransitionDetails.h
new file mode 100644
--- /dev/null
+++ b/Engine/Plugins/Runtime/StateTree/Source/StateTreeEditorModule/Private/Customizations/StateTreeTransitionDetails.h
@@ -0,0 +1,37 @@
+// Copyright Epic Games, Inc. All Rights Reserved.
+
+#pragma once
+
+#include "CoreMinimal.h"
+#include "UnrealClient.h"
+#include "IPropertyTypeCustomization.h"
+
+class IPropertyHandle;
+class UStateTree;
+class UStateTreeState;
+
+/**
+ * Type customization for FStateTreeTransition.
+ */
+
+class FStateTreeTransitionDetails : public IPropertyTypeCustomization
+{
+public:
+	/** Makes a new instance of this detail layout class for a specific detail view requesting it */
+	static TSharedRef<IPropertyTypeCustomization> MakeInstance();
+
+	/** IPropertyTypeCustomization interface */
+	virtual void CustomizeHeader(TSharedRef<IPropertyHandle> StructPropertyHandle, class FDetailWidgetRow& HeaderRow, IPropertyTypeCustomizationUtils& StructCustomizationUtils) override;
+	virtual void CustomizeChildren(TSharedRef<IPropertyHandle> StructPropertyHandle, class IDetailChildrenBuilder& StructBuilder, IPropertyTypeCustomizationUtils& StructCustomizationUtils) override;
+
+private:
+
+	FText GetDescription() const;
+
+	TSharedPtr<IPropertyHandle> TypeProperty;
+	TSharedPtr<IPropertyHandle> StateProperty;
+	TSharedPtr<IPropertyHandle> ConditionsProperty;
+
+	class IPropertyUtilities* PropUtils;
+	TSharedPtr<IPropertyHandle> StructProperty;
+};
diff --git a/Engine/Plugins/Runtime/StateTree/Source/StateTreeEditorModule/Private/Customizations/StateTreeVariableDescDetails.cpp b/Engine/Plugins/Runtime/StateTree/Source/StateTreeEditorModule/Private/Customizations/StateTreeVariableDescDetails.cpp
new file mode 100644
--- /dev/null
+++ b/Engine/Plugins/Runtime/StateTree/Source/StateTreeEditorModule/Private/Customizations/StateTreeVariableDescDetails.cpp
@@ -0,0 +1,119 @@
+// Copyright Epic Games, Inc. All Rights Reserved.
+
+#include "StateTreeVariableDescDetails.h"
+#include "Widgets/DeclarativeSyntaxSupport.h"
+#include "Widgets/Images/SImage.h"
+#include "DetailWidgetRow.h"
+#include "DetailLayoutBuilder.h"
+#include "IPropertyUtilities.h"
+#include "IDetailPropertyRow.h"
+#include "IDetailChildrenBuilder.h"
+#include "StateTree.h"
+#include "StateTreeState.h"
+#include "Widgets/Input/SComboButton.h"
+#include "Framework/MultiBox/MultiBoxBuilder.h"
+#include "ScopedTransaction.h"
+#include "StateTreePropertyHelpers.h"
+
+#define LOCTEXT_NAMESPACE "StateTreeEditor"
+
+TSharedRef<IPropertyTypeCustomization> FStateTreeVariableDescDetails::MakeInstance()
+{
+	return MakeShareable(new FStateTreeVariableDescDetails);
+}
+
+void FStateTreeVariableDescDetails::CustomizeHeader(TSharedRef<class IPropertyHandle> StructPropertyHandle, class FDetailWidgetRow& HeaderRow, IPropertyTypeCustomizationUtils& StructCustomizationUtils)
+{
+	StructProperty = StructPropertyHandle;
+	PropUtils = StructCustomizationUtils.GetPropertyUtilities().Get();
+
+	NameProperty = StructProperty->GetChildHandle(TEXT("Name"));
+	TypeProperty = StructProperty->GetChildHandle(TEXT("Type"));
+	BaseClassProperty = StructProperty->GetChildHandle(TEXT("BaseClass"));
+
+	HeaderRow
+		.WholeRowContent()
+		[
+			SNew(SHorizontalBox)
+			// Icon slot
+			+ SHorizontalBox::Slot()
+			.AutoWidth()
+			.VAlign(VAlign_Center)
+			[
+				SNew(SImage)
+				.Image(this, &FStateTreeVariableDescDetails::GetIcon)
+				.ColorAndOpacity(FLinearColor::White)
+			]
+			// Description
+			+ SHorizontalBox::Slot()
+			.FillWidth(1.0f)
+			.Padding(FMargin(6.0f, 2.0f))
+			[
+				SNew(STextBlock)
+				.Text(this, &FStateTreeVariableDescDetails::GetDescription)
+				.Font(IDetailLayoutBuilder::GetDetailFontBold())
+			]
+			// Array controls
+			+ SHorizontalBox::Slot()
+			.Padding(FMargin(12.0f, 0.0f))
+			.AutoWidth()
+			[
+				StructPropertyHandle->CreateDefaultPropertyButtonWidgets()
+			]
+		];
+}
+
+void FStateTreeVariableDescDetails::CustomizeChildren(TSharedRef<class IPropertyHandle> StructPropertyHandle, class IDetailChildrenBuilder& StructBuilder, IPropertyTypeCustomizationUtils& StructCustomizationUtils)
+{
+	if (NameProperty)
+	{
+		StructBuilder.AddProperty(NameProperty.ToSharedRef());
+	}
+
+	if (TypeProperty)
+	{
+		StructBuilder.AddProperty(TypeProperty.ToSharedRef());
+	}
+
+	if (BaseClassProperty)
+	{
+		StructBuilder.AddProperty(BaseClassProperty.ToSharedRef()).Visibility(TAttribute<EVisibility>(this, &FStateTreeVariableDescDetails::IsBaseClassVisible));
+	}
+}
+
+
+EVisibility FStateTreeVariableDescDetails::IsBaseClassVisible() const
+{
+	return GetType().Get(EStateTreeVariableType::Void) == EStateTreeVariableType::Object ? EVisibility::Visible : EVisibility::Collapsed;
+}
+
+FText FStateTreeVariableDescDetails::GetDescription() const
+{
+	FName Name;
+	if (NameProperty && NameProperty->GetValue(Name) == FPropertyAccess::Success)
+	{
+		return FText::FromName(Name);
+	}
+	return FText::GetEmpty();
+}
+
+const FSlateBrush* FStateTreeVariableDescDetails::GetIcon() const
+{
+	return UE::StateTree::PropertyHelpers::GetTypeIcon(GetType().Get(EStateTreeVariableType::Void));
+}
+
+TOptional<EStateTreeVariableType> FStateTreeVariableDescDetails::GetType() const
+{
+	if (TypeProperty)
+	{
+		uint8 Value;
+		if (TypeProperty->GetValue(Value) == FPropertyAccess::Success)
+		{
+			return EStateTreeVariableType(Value);
+		}
+	}
+	return TOptional<EStateTreeVariableType>();
+}
+
+
+#undef LOCTEXT_NAMESPACE
diff --git a/Engine/Plugins/Runtime/StateTree/Source/StateTreeEditorModule/Private/Customizations/StateTreeVariableDescDetails.h b/Engine/Plugins/Runtime/StateTree/Source/StateTreeEditorModule/Private/Customizations/StateTreeVariableDescDetails.h
new file mode 100644
--- /dev/null
+++ b/Engine/Plugins/Runtime/StateTree/Source/StateTreeEditorModule/Private/Customizations/StateTreeVariableDescDetails.h
@@ -0,0 +1,40 @@
+// Copyright Epic Games, Inc. All Rights Reserved.
+
+#pragma once
+
+#include "CoreMinimal.h"
+#include "UnrealClient.h"
+#include "IPropertyTypeCustomization.h"
+#include "StateTreeVariableDesc.h"
+
+class IPropertyHandle;
+
+/**
+ * Type customization for FStateTreeVariableDesc.
+ */
+
+class FStateTreeVariableDescDetails : public IPropertyTypeCustomization
+{
+public:
+	/** Makes a new instance of this detail layout class for a specific detail view requesting it */
+	static TSharedRef<IPropertyTypeCustomization> MakeInstance();
+
+	/** IPropertyTypeCustomization interface */
+	virtual void CustomizeHeader(TSharedRef<IPropertyHandle> StructPropertyHandle, class FDetailWidgetRow& HeaderRow, IPropertyTypeCustomizationUtils& StructCustomizationUtils) override;
+	virtual void CustomizeChildren(TSharedRef<IPropertyHandle> StructPropertyHandle, class IDetailChildrenBuilder& StructBuilder, IPropertyTypeCustomizationUtils& StructCustomizationUtils) override;
+
+private:
+
+	FText GetDescription() const;
+	const FSlateBrush* GetIcon() const;
+
+	TOptional<EStateTreeVariableType> GetType() const;
+	EVisibility IsBaseClassVisible() const;
+
+	TSharedPtr<IPropertyHandle> NameProperty;
+	TSharedPtr<IPropertyHandle> TypeProperty;
+	TSharedPtr<IPropertyHandle> BaseClassProperty;
+
+	class IPropertyUtilities* PropUtils;
+	TSharedPtr<IPropertyHandle> StructProperty;
+};
diff --git a/Engine/Plugins/Runtime/StateTree/Source/StateTreeEditorModule/Private/Customizations/StateTreeVariableDetails.cpp b/Engine/Plugins/Runtime/StateTree/Source/StateTreeEditorModule/Private/Customizations/StateTreeVariableDetails.cpp
new file mode 100644
--- /dev/null
+++ b/Engine/Plugins/Runtime/StateTree/Source/StateTreeEditorModule/Private/Customizations/StateTreeVariableDetails.cpp
@@ -0,0 +1,553 @@
+// Copyright Epic Games, Inc. All Rights Reserved.
+
+#include "StateTreeVariableDetails.h"
+#include "Framework/Commands/UIAction.h"
+#include "Widgets/DeclarativeSyntaxSupport.h"
+#include "Widgets/Text/STextBlock.h"
+#include "Widgets/Input/SNumericEntryBox.h"
+#include "Widgets/Input/SCheckBox.h"
+#include "Widgets/Input/SVectorInputBox.h"
+#include "Widgets/Input/SComboButton.h"
+#include "Widgets/Input/SButton.h"
+#include "Framework/MultiBox/MultiBoxBuilder.h"
+#include "DetailWidgetRow.h"
+#include "DetailLayoutBuilder.h"
+#include "Styling/SlateIconFinder.h"
+#include "IPropertyUtilities.h"
+#include "Editor.h"
+#include "StateTreeVariableProvider.h"
+#include "StateTreeVariable.h"
+#include "StateTreePropertyHelpers.h"
+#include "ScopedTransaction.h"
+#include "StateTree.h"
+#include "StateTreeDelegates.h"
+
+#define LOCTEXT_NAMESPACE "StateTreeEditor"
+
+TSharedRef<IPropertyTypeCustomization> FStateTreeVariableDetails::MakeInstance()
+{
+	return MakeShareable( new FStateTreeVariableDetails);
+}
+
+void FStateTreeVariableDetails::CustomizeHeader( TSharedRef<IPropertyHandle> StructPropertyHandle, FDetailWidgetRow& HeaderRow, IPropertyTypeCustomizationUtils& StructCustomizationUtils )
+{
+	StructProperty = StructPropertyHandle;
+	PropUtils = StructCustomizationUtils.GetPropertyUtilities().Get();
+
+	IDProperty = StructProperty->GetChildHandle(TEXT("ID"));
+	NameProperty = StructProperty->GetChildHandle(TEXT("Name"));
+	ValueProperty = StructProperty->GetChildHandle(TEXT("Value"));
+
+	TypeProperty = StructProperty->GetChildHandle(TEXT("Type"));
+	BaseClassProperty = StructProperty->GetChildHandle(TEXT("BaseClass"));
+	BindingModeProperty = StructProperty->GetChildHandle(TEXT("BindingMode"));
+
+	const FSimpleDelegate& ForceRefresh = FSimpleDelegate::CreateSP(this, &FStateTreeVariableDetails::OnForceRefresh);
+	TypeProperty->SetOnPropertyValueChanged(ForceRefresh);
+	BaseClassProperty->SetOnPropertyValueChanged(ForceRefresh);
+	BindingModeProperty->SetOnPropertyValueChanged(ForceRefresh);
+	UE::StateTree::Delegates::OnIdentifierChanged.AddSP(this, &FStateTreeVariableDetails::OnIdentifierChanged);
+
+	CacheExposedVariables();
+
+	HeaderRow
+		.NameContent()
+		[
+			StructPropertyHandle->CreatePropertyNameWidget()
+		]
+		.ValueContent()
+		.MinDesiredWidth(250.f)
+		.VAlign(VAlign_Center)
+		[
+			SNew(SHorizontalBox)
+			// Value inputs
+			// int
+			+ SHorizontalBox::Slot()
+			[
+				SNew(SNumericEntryBox<int32>)
+				.Font(IDetailLayoutBuilder::GetDetailFont())
+				.Value(this, &FStateTreeVariableDetails::GetIntValue)
+				.OnValueCommitted(this, &FStateTreeVariableDetails::SetIntValue)
+				.Visibility(this, &FStateTreeVariableDetails::IsIntVisible)
+			]
+			// float
+			+ SHorizontalBox::Slot()
+			[
+				SNew(SNumericEntryBox<float>)
+				.Font(IDetailLayoutBuilder::GetDetailFont())
+				.Value(this, &FStateTreeVariableDetails::GetFloatValue)
+				.OnValueCommitted(this, &FStateTreeVariableDetails::SetFloatValue)
+				.Visibility(this, &FStateTreeVariableDetails::IsFloatVisible)
+			]
+			// bool
+			+ SHorizontalBox::Slot()
+			[
+				SNew(SCheckBox)
+				.IsChecked(this, &FStateTreeVariableDetails::GetBoolValue)
+				.OnCheckStateChanged(this, &FStateTreeVariableDetails::SetBoolValue)
+				.Visibility(this, &FStateTreeVariableDetails::IsBoolVisible)
+			]
+			// vector
+			+ SHorizontalBox::Slot()
+			[
+				SNew(SVectorInputBox)
+				.Font(IDetailLayoutBuilder::GetDetailFont())
+				.X(this, &FStateTreeVariableDetails::GetVectorXValue)
+				.Y(this, &FStateTreeVariableDetails::GetVectorYValue)
+				.Z(this, &FStateTreeVariableDetails::GetVectorZValue)
+				.OnXCommitted(this, &FStateTreeVariableDetails::SetVectorXValue)
+				.OnYCommitted(this, &FStateTreeVariableDetails::SetVectorYValue)
+				.OnZCommitted(this, &FStateTreeVariableDetails::SetVectorZValue)
+				.AllowSpin(false)
+				.Visibility(this, &FStateTreeVariableDetails::IsVectorVisible)
+			]
+			// Object
+			+ SHorizontalBox::Slot()
+			[
+				SNew(STextBlock)
+				.Text(LOCTEXT("Empty", "Empty"))
+				.Font(IDetailLayoutBuilder::GetDetailFont())
+				.Visibility(this, &FStateTreeVariableDetails::IsObjectOrVoidVisible)
+			]
+			// Variable
+			+ SHorizontalBox::Slot()
+			[
+				SNew(SHorizontalBox)
+				.Visibility(this, &FStateTreeVariableDetails::IsVariableVisible)
+				// Icon slot
+				+ SHorizontalBox::Slot()
+				.AutoWidth()
+				.VAlign(VAlign_Center)
+				[
+					SNew(SImage)
+					.Image(this, &FStateTreeVariableDetails::GetCurrentKeyIcon)
+					.ColorAndOpacity(FLinearColor::White)
+				]
+				// Name slot
+				+ SHorizontalBox::Slot()
+				.AutoWidth()
+				.VAlign(VAlign_Center)
+				.Padding(3, 0)
+				[
+					SNew(STextBlock)
+					.Text(this, &FStateTreeVariableDetails::GetCurrentKeyDesc)
+					.Font(IDetailLayoutBuilder::GetDetailFont())
+				]
+			]
+			// Binding button
+			+ SHorizontalBox::Slot()
+			.AutoWidth()
+			.VAlign(VAlign_Center)
+			.Padding(6, 0, 0, 0)
+			[
+				SNew(SComboButton)
+				.OnGetMenuContent(this, &FStateTreeVariableDetails::OnGetKeyContent)
+				.ContentPadding(FMargin(2.0f, 0.0f))
+				.ButtonContent()
+				[
+					SNew(STextBlock)
+					.Text(LOCTEXT("Bind", "Bind"))
+					.Font(IDetailLayoutBuilder::GetDetailFont())
+				]
+			]
+		];
+}
+
+void FStateTreeVariableDetails::CustomizeChildren( TSharedRef<IPropertyHandle> StructPropertyHandle, IDetailChildrenBuilder& StructBuilder, IPropertyTypeCustomizationUtils& StructCustomizationUtils )
+{
+}
+
+void FStateTreeVariableDetails::OnIdentifierChanged(const UStateTree& StateTree)
+{
+	OnForceRefresh();
+	if (PropUtils)
+	{
+		PropUtils->ForceRefresh();
+	}
+}
+
+void FStateTreeVariableDetails::OnForceRefresh()
+{
+	CacheExposedVariables();
+}
+
+void FStateTreeVariableDetails::CacheExposedVariables()
+{
+	KeyValues.Reset();
+
+	const TOptional<EStateTreeVariableType> TypeOpt = GetVariableType();
+	const TOptional<EStateTreeVariableBindingMode> BindingModeOpt = GetBindingMode();
+
+	if (TypeOpt.IsSet() && BindingModeOpt.IsSet())
+	{
+		const EStateTreeVariableType Type = TypeOpt.GetValue();
+		const EStateTreeVariableBindingMode BindingMode = BindingModeOpt.GetValue();
+		const TOptional<UClass*> BaseClassOpt = GetBaseClass();
+		TSubclassOf<UObject> BaseClass = BaseClassOpt.IsSet() ? BaseClassOpt.GetValue() : nullptr;
+
+		TArray<UObject*> OuterObjects;
+		StructProperty->GetOuterObjects(OuterObjects);
+		for (int32 ObjectIdx = 0; ObjectIdx < OuterObjects.Num(); ObjectIdx++)
+		{
+			// Find outer which can provide us the keys
+			IStateTreeVariableProvider* Provider = nullptr;
+			for (UObject* OuterObj = OuterObjects[ObjectIdx]; OuterObj; OuterObj = OuterObj->GetOuter())
+			{
+				Provider = Cast<IStateTreeVariableProvider>(OuterObj);
+				if (Provider)
+				{
+					break;
+				}
+			}
+
+			if (Provider)
+			{
+				FStateTreeVariableLayout Variables;
+				Provider->GetVisibleVariables(Variables);
+
+				for (const FStateTreeVariableDesc& Var : Variables.Variables)
+				{
+					if (BindingMode == EStateTreeVariableBindingMode::Any || (Var.Type == Type && Var.BaseClass == BaseClass))
+					{
+						FVariableKey Key;
+						Key.Name = Var.Name;
+						Key.ID = Var.ID;
+						Key.Type = Var.Type;
+						Key.BaseClass = Var.BaseClass;
+						KeyValues.Add(Key);
+					}
+				}
+				break;
+			}
+		}
+	}
+}
+
+TSharedRef<SWidget> FStateTreeVariableDetails::OnGetKeyContent() const
+{
+	FMenuBuilder MenuBuilder(true, NULL);
+
+	const TOptional<EStateTreeVariableBindingMode> BindingModeOpt = GetBindingMode();
+	const bool bCanBeValue = BindingModeOpt.IsSet() && BindingModeOpt.GetValue() == EStateTreeVariableBindingMode::Typed && HasValueProperty();
+
+	if (bCanBeValue)
+	{
+		FUIAction ValueItemAction(FExecuteAction::CreateSP(const_cast<FStateTreeVariableDetails*>(this), &FStateTreeVariableDetails::OnKeyComboChange, -1));
+		MenuBuilder.AddMenuEntry(LOCTEXT("LocalValue", "Local Value"), TAttribute<FText>(), FSlateIcon(), ValueItemAction);
+		MenuBuilder.AddMenuSeparator();
+	}
+	else
+	{
+		FUIAction ValueItemAction(FExecuteAction::CreateSP(const_cast<FStateTreeVariableDetails*>(this), &FStateTreeVariableDetails::OnKeyComboChange, -1));
+		MenuBuilder.AddMenuEntry(LOCTEXT("ClearItem", "Clear"), TAttribute<FText>(), FSlateIcon(), ValueItemAction);
+		MenuBuilder.AddMenuSeparator();
+	}
+
+	if (KeyValues.Num() > 0)
+	{
+		for (int32 Idx = 0; Idx < KeyValues.Num(); Idx++)
+		{
+			const FVariableKey& Key = KeyValues[Idx];
+
+			FUIAction ItemAction(FExecuteAction::CreateSP(const_cast<FStateTreeVariableDetails*>(this), &FStateTreeVariableDetails::OnKeyComboChange, Idx));
+			MenuBuilder.AddMenuEntry(FText::FromName(Key.Name), TAttribute<FText>(), FSlateIcon(), ItemAction);
+		}
+	}
+	else
+	{
+		MenuBuilder.AddMenuEntry(LOCTEXT("NoVariablesToBindTo", "No Variables to bind to."), TAttribute<FText>(), FSlateIcon(), FUIAction());
+	}
+
+	return MenuBuilder.MakeWidget();
+}
+
+bool FStateTreeVariableDetails::HasValueProperty() const
+{
+	return ValueProperty.IsValid();
+}
+
+bool FStateTreeVariableDetails::IsBound() const
+{
+	FGuid ID;
+	if (IDProperty)
+	{
+		UE::StateTree::PropertyHelpers::GetStructValue<FGuid>(IDProperty, ID);
+	}
+
+	return ID.IsValid();
+}
+
+FText FStateTreeVariableDetails::GetCurrentKeyDesc() const
+{
+	const bool bCanBeValue = GetBindingMode().Get(EStateTreeVariableBindingMode::Any) == EStateTreeVariableBindingMode::Typed && HasValueProperty();
+
+	FName Name;
+	FGuid ID;
+	if (NameProperty && IDProperty)
+	{
+		NameProperty->GetValue(Name);
+		UE::StateTree::PropertyHelpers::GetStructValue<FGuid>(IDProperty, ID);
+	}
+
+	const bool bIsBound = ID.IsValid();
+
+	// Potentially Bound
+	if (bIsBound)
+	{
+		int KeyIdx = INDEX_NONE;
+		for (int32 Idx = 0; Idx < KeyValues.Num(); Idx++)
+		{
+			const FVariableKey& Key = KeyValues[Idx];
+			if (Key.ID == ID)
+			{
+				return FText::FromName(Key.Name);
+				break;
+			}
+		}
+		FFormatNamedArguments Args;
+		Args.Add(TEXT("Identifier"), FText::FromName(Name));
+		return FText::Format(LOCTEXT("InvalidReference", "Invalid Reference {Identifier}"), Args);
+	}
+
+	return FText(); // Empty name
+}
+
+const FSlateBrush* FStateTreeVariableDetails::GetCurrentKeyIcon() const
+{
+	const bool bCanBeValue = GetBindingMode().Get(EStateTreeVariableBindingMode::Any) == EStateTreeVariableBindingMode::Typed && HasValueProperty();
+
+	FGuid ID;
+	if (NameProperty && IDProperty)
+	{
+		UE::StateTree::PropertyHelpers::GetStructValue<FGuid>(IDProperty, ID);
+	}
+
+	const bool bIsBound = ID.IsValid();
+
+	// Potentially Bound
+	if (bIsBound)
+	{
+		int KeyIdx = INDEX_NONE;
+		for (int32 Idx = 0; Idx < KeyValues.Num(); Idx++)
+		{
+			const FVariableKey& Key = KeyValues[Idx];
+			if (Key.ID == ID)
+			{
+				return UE::StateTree::PropertyHelpers::GetTypeIcon(Key.Type);
+			}
+		}
+	}
+
+	return nullptr;
+}
+
+void FStateTreeVariableDetails::OnKeyComboChange(int32 Index)
+{
+	if (Index == -1)
+	{
+		FScopedTransaction Transaction(FText::Format(LOCTEXT("SetPropertyValue", "Set {0}"), StructProperty->GetPropertyDisplayName()));
+
+		// Clear or Local Value
+		if (NameProperty && IDProperty)
+		{
+			NameProperty->SetValue(FName(), EPropertyValueSetFlags::NotTransactable);
+			UE::StateTree::PropertyHelpers::SetStructValue<FGuid>(IDProperty, FGuid(), EPropertyValueSetFlags::NotTransactable);
+		}
+
+		// Clear the associated type for Any binding, it will be updated when reference is chosen.
+		const TOptional<EStateTreeVariableBindingMode> BindingModeOpt = GetBindingMode();
+		if (BindingModeOpt.IsSet() && BindingModeOpt.GetValue() == EStateTreeVariableBindingMode::Any)
+		{
+			if (TypeProperty && BaseClassProperty)
+			{
+				// Note: Set BaseClass first. Condition details is listening to changed in property type, and base class should be set up before we hit hte callback.
+				BaseClassProperty->SetValue((UObject*)nullptr, EPropertyValueSetFlags::NotTransactable);
+				TypeProperty->SetValue((uint8)EStateTreeVariableType::Void, EPropertyValueSetFlags::NotTransactable);
+			}
+		}
+	}
+	else
+	{
+		// Set reference.
+		if (KeyValues.IsValidIndex(Index))
+		{
+			const FVariableKey& Key = KeyValues[Index];
+			if (NameProperty && IDProperty)
+			{
+				FScopedTransaction Transaction(FText::Format(LOCTEXT("SetPropertyValue", "Set {0}"), StructProperty->GetPropertyDisplayName()));
+
+				NameProperty->SetValue(Key.Name, EPropertyValueSetFlags::NotTransactable);
+				UE::StateTree::PropertyHelpers::SetStructValue<FGuid>(IDProperty, Key.ID, EPropertyValueSetFlags::NotTransactable);
+
+				const TOptional<EStateTreeVariableBindingMode> BindingModeOpt = GetBindingMode();
+				if (BindingModeOpt.IsSet() && BindingModeOpt.GetValue() == EStateTreeVariableBindingMode::Any)
+				{
+					// Note: Set BaseClass first. Condition details is listening to changed in property type, and base class should be set up before we hit hte callback.
+					// Copy base class from the variable.
+					if (BaseClassProperty)
+					{
+						BaseClassProperty->SetValue(Key.BaseClass, EPropertyValueSetFlags::NotTransactable);
+					}
+					// Copy type from the variable.
+					if (TypeProperty)
+					{
+						TypeProperty->SetValue((uint8)Key.Type, EPropertyValueSetFlags::NotTransactable);
+					}
+				}
+			}
+		}
+	}
+}
+
+EVisibility FStateTreeVariableDetails::IsIntVisible() const
+{
+	TOptional<EStateTreeVariableType> Type = GetVariableType();
+	return (!IsBound() && Type.IsSet() && Type.GetValue() == EStateTreeVariableType::Int) ? EVisibility::Visible : EVisibility::Collapsed;
+}
+
+TOptional<int32> FStateTreeVariableDetails::GetIntValue() const
+{
+	return UE::StateTree::PropertyHelpers::GetVariableValue<int32>(ValueProperty);
+}
+
+void FStateTreeVariableDetails::SetIntValue(int32 Value, ETextCommit::Type CommitType)
+{
+	UE::StateTree::PropertyHelpers::SetVariableValue<int32>(ValueProperty, Value);
+}
+
+
+EVisibility FStateTreeVariableDetails::IsFloatVisible() const
+{
+	TOptional<EStateTreeVariableType> Type = GetVariableType();
+	return (!IsBound() && Type.IsSet() && Type.GetValue() == EStateTreeVariableType::Float) ? EVisibility::Visible : EVisibility::Collapsed;
+}
+
+TOptional<float> FStateTreeVariableDetails::GetFloatValue() const
+{
+	return UE::StateTree::PropertyHelpers::GetVariableValue<float>(ValueProperty);
+}
+
+void FStateTreeVariableDetails::SetFloatValue(float Value, ETextCommit::Type CommitType)
+{
+	UE::StateTree::PropertyHelpers::SetVariableValue<float>(ValueProperty, Value);
+}
+
+
+EVisibility FStateTreeVariableDetails::IsBoolVisible() const
+{
+	TOptional<EStateTreeVariableType> Type = GetVariableType();
+	return (!IsBound() && Type.IsSet() && Type.GetValue() == EStateTreeVariableType::Bool) ? EVisibility::Visible : EVisibility::Collapsed;
+}
+
+ECheckBoxState FStateTreeVariableDetails::GetBoolValue() const
+{
+	TOptional<bool> Result = UE::StateTree::PropertyHelpers::GetVariableValue<bool>(ValueProperty);
+	if (!Result.IsSet())
+		return ECheckBoxState::Undetermined;
+	return Result.GetValue() ? ECheckBoxState::Checked : ECheckBoxState::Unchecked;
+}
+
+void FStateTreeVariableDetails::SetBoolValue(ECheckBoxState NewCheckedState)
+{
+	UE::StateTree::PropertyHelpers::SetVariableValue<bool>(ValueProperty, NewCheckedState == ECheckBoxState::Checked ? true : false);
+}
+
+
+EVisibility FStateTreeVariableDetails::IsVectorVisible() const
+{
+	TOptional<EStateTreeVariableType> Type = GetVariableType();
+	return (!IsBound() && Type.IsSet() && Type.GetValue() == EStateTreeVariableType::Vector) ? EVisibility::Visible : EVisibility::Collapsed;
+}
+
+TOptional<float> FStateTreeVariableDetails::GetVectorXValue() const
+{
+	TOptional<FVector> Result = UE::StateTree::PropertyHelpers::GetVariableValue<FVector>(ValueProperty);
+	return Result.IsSet() ? Result.GetValue().X : TOptional<float>();
+}
+
+TOptional<float> FStateTreeVariableDetails::GetVectorYValue() const
+{
+	TOptional<FVector> Result = UE::StateTree::PropertyHelpers::GetVariableValue<FVector>(ValueProperty);
+	return Result.IsSet() ? Result.GetValue().Y : TOptional<float>();
+}
+
+TOptional<float> FStateTreeVariableDetails::GetVectorZValue() const
+{
+	TOptional<FVector> Result = UE::StateTree::PropertyHelpers::GetVariableValue<FVector>(ValueProperty);
+	return Result.IsSet() ? Result.GetValue().Z : TOptional<float>();
+}
+
+void FStateTreeVariableDetails::SetVectorXValue(float X, ETextCommit::Type CommitType)
+{
+	FVector CurrentValue = UE::StateTree::PropertyHelpers::GetVariableValue<FVector>(ValueProperty).Get(FVector::ZeroVector);
+	CurrentValue.X = X;
+	UE::StateTree::PropertyHelpers::SetVariableValue<FVector>(ValueProperty, CurrentValue);
+}
+
+void FStateTreeVariableDetails::SetVectorYValue(float Y, ETextCommit::Type CommitType)
+{
+	FVector CurrentValue = UE::StateTree::PropertyHelpers::GetVariableValue<FVector>(ValueProperty).Get(FVector::ZeroVector);
+	CurrentValue.Y = Y;
+	UE::StateTree::PropertyHelpers::SetVariableValue<FVector>(ValueProperty, CurrentValue);
+}
+
+void FStateTreeVariableDetails::SetVectorZValue(float Z, ETextCommit::Type CommitType)
+{
+	FVector CurrentValue = UE::StateTree::PropertyHelpers::GetVariableValue<FVector>(ValueProperty).Get(FVector::ZeroVector);
+	CurrentValue.Z = Z;
+	UE::StateTree::PropertyHelpers::SetVariableValue<FVector>(ValueProperty, CurrentValue);
+}
+
+EVisibility FStateTreeVariableDetails::IsObjectOrVoidVisible() const
+{
+	TOptional<EStateTreeVariableType> Type = GetVariableType();
+	return (!IsBound() && Type.IsSet() && (Type.GetValue() == EStateTreeVariableType::Object || Type.GetValue() == EStateTreeVariableType::Void)) ? EVisibility::Visible : EVisibility::Collapsed;
+}
+
+EVisibility FStateTreeVariableDetails::IsVariableVisible() const
+{
+	TOptional<EStateTreeVariableBindingMode> BindingModeOpt = GetBindingMode();
+	return (IsBound() || (BindingModeOpt.IsSet() && BindingModeOpt.GetValue() == EStateTreeVariableBindingMode::Any)) ? EVisibility::Visible : EVisibility::Collapsed;
+}
+
+TOptional<EStateTreeVariableType> FStateTreeVariableDetails::GetVariableType() const
+{
+	if (TypeProperty)
+	{
+		uint8 Value;
+		if (TypeProperty->GetValue(Value) == FPropertyAccess::Success)
+		{
+			return EStateTreeVariableType(Value);
+		}
+	}
+	return TOptional<EStateTreeVariableType>();
+}
+
+TOptional<EStateTreeVariableBindingMode> FStateTreeVariableDetails::GetBindingMode() const
+{
+	if (BindingModeProperty)
+	{
+		uint8 Value;
+		if (BindingModeProperty->GetValue(Value) == FPropertyAccess::Success)
+		{
+			return EStateTreeVariableBindingMode(Value);
+		}
+	}
+	return TOptional<EStateTreeVariableBindingMode>();
+}
+
+TOptional<UClass*> FStateTreeVariableDetails::GetBaseClass() const
+{
+	if (BaseClassProperty)
+	{
+		UObject* Value;
+		if (BaseClassProperty->GetValue(Value) == FPropertyAccess::Success)
+		{
+			return Cast<UClass>(Value);
+		}
+	}
+	return nullptr;
+}
+
+
+#undef LOCTEXT_NAMESPACE
diff --git a/Engine/Plugins/Runtime/StateTree/Source/StateTreeEditorModule/Private/Customizations/StateTreeVariableDetails.h b/Engine/Plugins/Runtime/StateTree/Source/StateTreeEditorModule/Private/Customizations/StateTreeVariableDetails.h
new file mode 100644
--- /dev/null
+++ b/Engine/Plugins/Runtime/StateTree/Source/StateTreeEditorModule/Private/Customizations/StateTreeVariableDetails.h
@@ -0,0 +1,87 @@
+// Copyright Epic Games, Inc. All Rights Reserved.
+
+#pragma once
+
+#include "CoreMinimal.h"
+#include "UnrealClient.h"
+#include "Widgets/Input/SCheckBox.h"
+#include "IPropertyTypeCustomization.h"
+#include "StateTreeVariable.h"
+
+class IPropertyHandle;
+
+/**
+ * Type customization for FStateTreeVariable and FStateTreeVariable.
+ */
+
+class FStateTreeVariableDetails : public IPropertyTypeCustomization
+{
+public:
+	/** Makes a new instance of this detail layout class for a specific detail view requesting it */
+	static TSharedRef<IPropertyTypeCustomization> MakeInstance();
+
+	/** IPropertyTypeCustomization interface */
+	virtual void CustomizeHeader( TSharedRef<IPropertyHandle> StructPropertyHandle, class FDetailWidgetRow& HeaderRow, IPropertyTypeCustomizationUtils& StructCustomizationUtils ) override;
+	virtual void CustomizeChildren( TSharedRef<IPropertyHandle> StructPropertyHandle, class IDetailChildrenBuilder& StructBuilder, IPropertyTypeCustomizationUtils& StructCustomizationUtils ) override;
+
+private:
+
+	TOptional<int32> GetIntValue() const;
+	void SetIntValue(int32 Value, ETextCommit::Type CommitType);
+
+	TOptional<float> GetFloatValue() const;
+	void SetFloatValue(float Value, ETextCommit::Type CommitType);
+
+	ECheckBoxState GetBoolValue() const;
+	void SetBoolValue(ECheckBoxState NewCheckedState);
+
+	TOptional<float>  GetVectorXValue() const;
+	TOptional<float>  GetVectorYValue() const;
+	TOptional<float>  GetVectorZValue() const;
+	void SetVectorXValue(float Value, ETextCommit::Type CommitType);
+	void SetVectorYValue(float Value, ETextCommit::Type CommitType);
+	void SetVectorZValue(float Value, ETextCommit::Type CommitType);
+
+	TOptional<EStateTreeVariableType> GetVariableType() const;
+	TOptional<EStateTreeVariableBindingMode> GetBindingMode() const;
+	TOptional<UClass*> GetBaseClass() const;
+
+	bool HasValueProperty() const;
+	bool IsBound() const;
+
+	void CacheExposedVariables();
+	void OnForceRefresh();
+	void OnIdentifierChanged(const class UStateTree& StateTree);
+
+	EVisibility IsIntVisible() const;
+	EVisibility IsFloatVisible() const;
+	EVisibility IsBoolVisible() const;
+	EVisibility IsVectorVisible() const;
+	EVisibility IsObjectOrVoidVisible() const;
+	EVisibility IsVariableVisible() const;
+
+	void OnKeyComboChange(int32 Index);
+	TSharedRef<SWidget> OnGetKeyContent() const;
+	FText GetCurrentKeyDesc() const;
+	const FSlateBrush* GetCurrentKeyIcon() const;
+
+	TSharedPtr<IPropertyHandle> IDProperty;
+	TSharedPtr<IPropertyHandle> NameProperty;
+	TSharedPtr<IPropertyHandle> TypeProperty;
+	TSharedPtr<IPropertyHandle> BaseClassProperty;
+	TSharedPtr<IPropertyHandle> BindingModeProperty;
+	TSharedPtr<IPropertyHandle> ValueProperty;
+
+	struct FVariableKey
+	{
+		FName Name;
+		FGuid ID;
+		EStateTreeVariableType Type;
+		TSubclassOf<UObject> BaseClass;
+	};
+
+	TArray<FVariableKey> KeyValues;
+
+	class IPropertyUtilities* PropUtils;
+	TSharedPtr<IPropertyHandle> StructProperty;
+};
diff --git a/Engine/Plugins/Runtime/StateTree/Source/StateTreeEditorModule/Private/Customizations/StateTreeVariableLayoutDetails.cpp b/Engine/Plugins/Runtime/StateTree/Source/StateTreeEditorModule/Private/Customizations/StateTreeVariableLayoutDetails.cpp
new file mode 100644
--- /dev/null
+++ b/Engine/Plugins/Runtime/StateTree/Source/StateTreeEditorModule/Private/Customizations/StateTreeVariableLayoutDetails.cpp
@@ -0,0 +1,84 @@
+// Copyright Epic Games, Inc. All Rights Reserved.
+
+#include "StateTreeVariableLayoutDetails.h"
+#include "Widgets/DeclarativeSyntaxSupport.h"
+#include "Widgets/Images/SImage.h"
+#include "DetailWidgetRow.h"
+#include "DetailLayoutBuilder.h"
+#include "IPropertyUtilities.h"
+#include "IDetailPropertyRow.h"
+#include "IDetailChildrenBuilder.h"
+#include "Widgets/Input/SComboButton.h"
+#include "Framework/MultiBox/MultiBoxBuilder.h"
+#include "ScopedTransaction.h"
+
+#define LOCTEXT_NAMESPACE "StateTreeEditor"
+
+TSharedRef<IPropertyTypeCustomization> FStateTreeVariableLayoutDetails::MakeInstance()
+{
+	return MakeShareable(new FStateTreeVariableLayoutDetails);
+}
+
+void FStateTreeVariableLayoutDetails::CustomizeHeader(TSharedRef<class IPropertyHandle> StructPropertyHandle, class FDetailWidgetRow& HeaderRow, IPropertyTypeCustomizationUtils& StructCustomizationUtils)
+{
+	StructProperty = StructPropertyHandle;
+	PropUtils = StructCustomizationUtils.GetPropertyUtilities().Get();
+
+	VariablesProperty = StructProperty->GetChildHandle(TEXT("Variables"));
+
+	// Listen changes to the array. The list will not update automatically since we're dealing with the list items ourselves.
+	const FSimpleDelegate& ForceRefresh = FSimpleDelegate::CreateSP(this, &FStateTreeVariableLayoutDetails::OnForceRefresh);
+	VariablesProperty->SetOnPropertyValueChanged(ForceRefresh);
+
+	HeaderRow
+		.WholeRowContent()
+		[
+			SNew(SHorizontalBox)
+			// Name
+			+ SHorizontalBox::Slot()
+			.FillWidth(1.0f)
+			.VAlign(VAlign_Center)
+			[
+				StructPropertyHandle->CreatePropertyNameWidget()
+			]
+			// Array controls
+			+ SHorizontalBox::Slot()
+			.AutoWidth()
+			.Padding(FMargin(12.0f, 0.0f))
+			[
+				VariablesProperty->CreateDefaultPropertyButtonWidgets()
+			]
+		];
+}
+
+void FStateTreeVariableLayoutDetails::CustomizeChildren(TSharedRef<class IPropertyHandle> StructPropertyHandle, class IDetailChildrenBuilder& StructBuilder, IPropertyTypeCustomizationUtils& StructCustomizationUtils)
+{
+	if (VariablesProperty)
+	{
+		// Inline descriptor array as children
+		uint32 ChildNum = 0;
+		if (VariablesProperty->GetNumChildren(ChildNum) == FPropertyAccess::Success)
+		{
+			for (uint32 ChildIndex = 0; ChildIndex < ChildNum; ++ChildIndex)
+			{
+				TSharedPtr<IPropertyHandle> ChildProperty = VariablesProperty->GetChildHandle(ChildIndex);
+				if (ChildProperty)
+				{
+					StructBuilder.AddProperty(ChildProperty.ToSharedRef());
+				}
+			}
+		}
+	}
+}
+
+void FStateTreeVariableLayoutDetails::OnForceRefresh()
+{
+	if (PropUtils)
+	{
+		PropUtils->ForceRefresh();
+	}
+}
+
+
+
+#undef LOCTEXT_NAMESPACE
diff --git a/Engine/Plugins/Runtime/StateTree/Source/StateTreeEditorModule/Private/Customizations/StateTreeVariableLayoutDetails.h b/Engine/Plugins/Runtime/StateTree/Source/StateTreeEditorModule/Private/Customizations/StateTreeVariableLayoutDetails.h
new file mode 100644
--- /dev/null
+++ b/Engine/Plugins/Runtime/StateTree/Source/StateTreeEditorModule/Private/Customizations/StateTreeVariableLayoutDetails.h
@@ -0,0 +1,33 @@
+// Copyright Epic Games, Inc. All Rights Reserved.
+
+#pragma once
+
+#include "CoreMinimal.h"
+#include "UnrealClient.h"
+#include "IPropertyTypeCustomization.h"
+
+class IPropertyHandle;
+
+/**
+ * Type customization for .
+ */
+
+class FStateTreeVariableLayoutDetails : public IPropertyTypeCustomization
+{
+public:
+	/** Makes a new instance of this detail layout class for a specific detail view requesting it */
+	static TSharedRef<IPropertyTypeCustomization> MakeInstance();
+
+	/** IPropertyTypeCustomization interface */
+	virtual void CustomizeHeader(TSharedRef<IPropertyHandle> StructPropertyHandle, class FDetailWidgetRow& HeaderRow, IPropertyTypeCustomizationUtils& StructCustomizationUtils) override;
+	virtual void CustomizeChildren(TSharedRef<IPropertyHandle> StructPropertyHandle, class IDetailChildrenBuilder& StructBuilder, IPropertyTypeCustomizationUtils& StructCustomizationUtils) override;
+
+private:
+
+	void OnForceRefresh();
+
+	TSharedPtr<IPropertyHandle> VariablesProperty;
+
+	class IPropertyUtilities* PropUtils;
+	TSharedPtr<IPropertyHandle> StructProperty;
+};
diff --git a/Engine/Plugins/Runtime/StateTree/Source/StateTreeEditorModule/Private/SStateTreeView.cpp b/Engine/Plugins/Runtime/StateTree/Source/StateTreeEditorModule/Private/SStateTreeView.cpp
new file mode 100644
--- /dev/null
+++ b/Engine/Plugins/Runtime/StateTree/Source/StateTreeEditorModule/Private/SStateTreeView.cpp
@@ -0,0 +1,424 @@
+// Copyright Epic Games, Inc. All Rights Reserved.
+
+#include "SStateTreeView.h"
+
+#include "Framework/Commands/Commands.h"
+#include "Framework/Application/SlateApplication.h"
+#include "Framework/MultiBox/MultiBoxBuilder.h"
+#include "ScopedTransaction.h"
+
+#include "Templates/SharedPointer.h"
+
+#include "EditorStyleSet.h"
+#include "EditorFontGlyphs.h"
+#include "StateTreeEditorStyle.h"
+
+#include "Widgets/DeclarativeSyntaxSupport.h"
+#include "Widgets/SCompoundWidget.h"
+#include "Widgets/Input/SButton.h"
+#include "Widgets/Text/STextBlock.h"
+#include "Widgets/Layout/SBorder.h"
+#include "Widgets/Layout/SScrollBorder.h"
+#include "Widgets/Images/SImage.h"
+#include "Widgets/Text/SInlineEditableTextBlock.h"
+#include "Widgets/Layout/SScrollBox.h"
+#include "SStateTreeViewRow.h"
+
+#include "StateTreeViewModel.h"
+#include "StateTreeState.h"
+
+#define LOCTEXT_NAMESPACE "StateTreeEditor"
+
+SStateTreeView::SStateTreeView()
+	: RequestedRenameState(nullptr)
+	, bItemsDirty(false)
+	, bUpdatingSelection(false)
+{
+}
+
+SStateTreeView::~SStateTreeView()
+{
+	if (StateTreeViewModel)
+	{
+		StateTreeViewModel->GetOnAssetChanged().RemoveAll(this);
+		StateTreeViewModel->GetOnStatesRemoved().RemoveAll(this);
+		StateTreeViewModel->GetOnStatesMoved().RemoveAll(this);
+		StateTreeViewModel->GetOnStateAdded().RemoveAll(this);
+		StateTreeViewModel->GetOnStatesChanged().RemoveAll(this);
+		StateTreeViewModel->GetOnSelectionChanged().RemoveAll(this);
+	}
+}
+
+void SStateTreeView::Construct(const FArguments& InArgs, TSharedRef<FStateTreeViewModel> InStateTreeViewModel)
+{
+	StateTreeViewModel = InStateTreeViewModel;
+
+	StateTreeViewModel->GetOnAssetChanged().AddSP(this, &SStateTreeView::HandleModelAssetChanged);
+	StateTreeViewModel->GetOnStatesRemoved().AddSP(this, &SStateTreeView::HandleModelStatesRemoved);
+	StateTreeViewModel->GetOnStatesMoved().AddSP(this, &SStateTreeView::HandleModelStatesMoved);
+	StateTreeViewModel->GetOnStateAdded().AddSP(this, &SStateTreeView::HandleModelStateAdded);
+	StateTreeViewModel->GetOnStatesChanged().AddSP(this, &SStateTreeView::HandleModelStatesChanged);
+	StateTreeViewModel->GetOnSelectionChanged().AddSP(this, &SStateTreeView::HandleModelSelectionChanged);
+
+	bUpdatingSelection = false;
+
+	TreeView = SNew(STreeView<UStateTreeState*>)
+		.OnGenerateRow(this, &SStateTreeView::HandleGenerateRow)
+		.OnGetChildren(this, &SStateTreeView::HandleGetChildren)
+		.TreeItemsSource(StateTreeViewModel->GetRoutines())
+		.ItemHeight(32)
+		.OnSelectionChanged(this, &SStateTreeView::HandleTreeSelectionChanged)
+		.OnContextMenuOpening(this, &SStateTreeView::HandleContextMenuOpening);
+
+	ExternalScrollbar = SNew(SScrollBar)
+		.AlwaysShowScrollbar(true);
+
+	ChildSlot
+	[
+		SNew(SVerticalBox)
+
+		+ SVerticalBox::Slot()
+		.VAlign(VAlign_Center)
+		.AutoHeight()
+		[
+			SNew(SBorder)
+			.BorderImage(FEditorStyle::GetBrush("ToolPanel.GroupBorder"))
+			.Padding(2.0f)
+			[
+				SNew(SHorizontalBox)
+
+				// New State
+				+ SHorizontalBox::Slot()
+				.VAlign(VAlign_Center)
+				.Padding(4.0f, 2.0f)
+				.AutoWidth()
+				[
+					SNew(SButton)
+					.ButtonStyle(FEditorStyle::Get(), "FlatButton.Success")
+					.ToolTipText(LOCTEXT("AddRoutineToolTip", "Add New Routine"))
+					.OnClicked(this, &SStateTreeView::HandleAddRoutine)
+					[
+						SNew(SHorizontalBox)
+						+ SHorizontalBox::Slot()
+						.VAlign(VAlign_Center)
+						.Padding(0.0f, 0.0f, 4.0f, 0.0f)
+						.AutoWidth()
+						[
+							SNew(STextBlock)
+							.Text(FEditorFontGlyphs::File)
+							.TextStyle(FEditorStyle::Get(), "FlatButton.DefaultTextStyle")
+							.Font(FEditorStyle::Get().GetFontStyle("FontAwesome.11"))
+						]
+						+ SHorizontalBox::Slot()
+						.VAlign(VAlign_Center)
+						.AutoWidth()
+						[
+							SNew(STextBlock)
+							.Text(LOCTEXT("AddRoutine", "Add Routine"))
+							.TextStyle(FEditorStyle::Get(), "FlatButton.DefaultTextStyle")
+						]
+					]
+				]
+			]
+		]
+
+		+SVerticalBox::Slot()
+		.FillHeight(1.0f)
+		.Padding(0.0f, 6.0f, 0.0f, 0.0f)
+		[
+			SNew(SHorizontalBox)
+			+ SHorizontalBox::Slot()
+			.FillWidth(1.0f)
+			.Padding(0.0f)
+			[
+				SNew(SScrollBox)
+				.Orientation(Orient_Horizontal)
+				+ SScrollBox::Slot()
+				.HAlign(HAlign_Fill)
+				[
+					SNew(SBorder)
+					.HAlign(HAlign_Fill)
+					.BorderImage(FEditorStyle::GetBrush("ToolPanel.GroupBorder"))
+					.Padding(0.0f)
+					[
+						TreeView.ToSharedRef()
+					]
+				]
+			]
+
+			+ SHorizontalBox::Slot()
+			.AutoWidth()
+			.Padding(0.0f)
+			[
+				SNew(SBox)
+				.WidthOverride(FOptionalSize(13.0f))
+				[
+					ExternalScrollbar.ToSharedRef()
+				]
+			]
+		]
+	];
+
+	UpdateTree(true);
+}
+
+
+void SStateTreeView::SavePersistentExpandedStates()
+{
+	if (!StateTreeViewModel)
+	{
+		return;
+	}
+
+	TSet<UStateTreeState*> ExpandedStates;
+	TreeView->GetExpandedItems(ExpandedStates);
+	StateTreeViewModel->SetPersistentExpandedStates(ExpandedStates);
+}
+
+void SStateTreeView::Tick(const FGeometry& AllottedGeometry, const double InCurrentTime, const float InDeltaTime)
+{
+	if (bItemsDirty)
+	{
+		UpdateTree();
+	}
+
+	if (RequestedRenameState && !TreeView->IsPendingRefresh())
+	{
+		if (TSharedPtr<SStateTreeViewRow> Row = StaticCastSharedPtr<SStateTreeViewRow>(TreeView->WidgetFromItem(RequestedRenameState)))
+		{
+			Row->RequestRename();
+		}
+		RequestedRenameState = nullptr;
+	}
+}
+
+void SStateTreeView::UpdateTree(bool bExpandPersistent)
+{
+	if (!StateTreeViewModel)
+	{
+		return;
+	}
+
+	TSet<UStateTreeState*> ExpandedStates;
+	if (bExpandPersistent)
+	{
+		// Get expanded state from the tree data.
+		StateTreeViewModel->GetPersistentExpandedStates(ExpandedStates);
+	}
+	else
+	{
+		// Restore current expanded state.
+		TreeView->GetExpandedItems(ExpandedStates);
+	}
+
+	// Remember selection
+	TArray<UStateTreeState*> SelectedStates;
+	StateTreeViewModel->GetSelectedStates(SelectedStates);
+
+	// Regenerate items
+	TreeView->SetTreeItemsSource(StateTreeViewModel->GetRoutines());
+
+	// Restore expanded state
+	for (UStateTreeState* State : ExpandedStates)
+	{
+		TreeView->SetItemExpansion(State, true);
+	}
+
+	// Restore selected state
+	TreeView->ClearSelection();
+	TreeView->SetItemSelection(SelectedStates, true);
+
+	bItemsDirty = false;
+}
+
+
+
+void SStateTreeView::HandleModelAssetChanged()
+{
+	bItemsDirty = true;
+}
+
+void SStateTreeView::HandleModelStatesRemoved(const TSet<UStateTreeState*>& AffectedParents)
+{
+	bItemsDirty = true;
+}
+
+void SStateTreeView::HandleModelStatesMoved(const TSet<UStateTreeState*>& AffectedParents, const TSet<UStateTreeState*>& MovedStates)
+{
+	bItemsDirty = true;
+}
+
+void SStateTreeView::HandleModelStateAdded(UStateTreeState* ParentState, UStateTreeState* NewState)
+{
+	bItemsDirty = true;
+
+	HandleRenameState(NewState);
+}
+
+void SStateTreeView::HandleModelStatesChanged(const TSet<UStateTreeState*>& AffectedStates, const FPropertyChangedEvent& PropertyChangedEvent)
+{
+	bool bArraysChanged = false;
+
+	// The purpose of the rebuild below is to update the evaluator and task visualization (number of widgets change).
+	// This method is called when anything in a state changes, make sure to only rebuild when needed.
+	if (PropertyChangedEvent.MemberProperty->GetFName() == GET_MEMBER_NAME_CHECKED(UStateTreeState, Evaluators))
+	{
+		bArraysChanged = true;
+	}
+	if (PropertyChangedEvent.MemberProperty->GetFName() == GET_MEMBER_NAME_CHECKED(UStateTreeState, Evaluators2))
+	{
+		bArraysChanged = true;
+	}
+	if (PropertyChangedEvent.MemberProperty->GetFName() == GET_MEMBER_NAME_CHECKED(UStateTreeState, Tasks))
+	{
+		bArraysChanged = true;
+	}
+	if (PropertyChangedEvent.MemberProperty->GetFName() == GET_MEMBER_NAME_CHECKED(UStateTreeState, Tasks2))
+	{
+		bArraysChanged = true;
+	}
+		
+	if (bArraysChanged)
+	{
+		TreeView->RebuildList();
+	}
+}
+
+void SStateTreeView::HandleModelSelectionChanged(const TArray<UStateTreeState*>& SelectedStates)
+{
+	if (bUpdatingSelection)
+	{
+		return;
+	}
+
+	TreeView->ClearSelection();
+
+	if (SelectedStates.Num() > 0)
+	{
+		TreeView->SetItemSelection(SelectedStates, true);
+	}
+}
+
+
+TSharedRef<ITableRow> SStateTreeView::HandleGenerateRow(UStateTreeState* InState, const TSharedRef<STableViewBase>& InOwnerTableView)
+{
+	return SNew(SStateTreeViewRow, InOwnerTableView, InState, StateTreeViewModel.ToSharedRef());
+}
+
+void SStateTreeView::HandleGetChildren(UStateTreeState* InParent, TArray<UStateTreeState*>& OutChildren)
+{
+	if (InParent)
+	{
+		OutChildren.Append(InParent->Children);
+	}
+}
+
+void SStateTreeView::HandleTreeSelectionChanged(UStateTreeState* InSelectedItem, ESelectInfo::Type SelectionType)
+{
+	if (!StateTreeViewModel)
+	{
+		return;
+	}
+
+	// Do not report code based selection changes.
+	if (SelectionType == ESelectInfo::Direct)
+	{
+		return;
+	}
+
+	TArray<UStateTreeState*> SelectedItems = TreeView->GetSelectedItems();
+
+	bUpdatingSelection = true;
+	StateTreeViewModel->SetSelection(SelectedItems);
+	bUpdatingSelection = false;
+}
+
+TSharedPtr<SWidget> SStateTreeView::HandleContextMenuOpening()
+{
+	if (!StateTreeViewModel)
+	{
+		return nullptr;
+	}
+
+	TArray<UStateTreeState*> SelectedStates;
+	StateTreeViewModel->GetSelectedStates(SelectedStates);
+	UStateTreeState* FirstSelectedState = SelectedStates.Num() > 0 ? SelectedStates[0] : nullptr;
+
+	FMenuBuilder MenuBuilder(true, nullptr);
+
+	MenuBuilder.AddMenuEntry(
+		LOCTEXT("AddState", "Add State"),
+		FText(),
+		FSlateIcon(),
+		FUIAction(FExecuteAction::CreateSP(this, &SStateTreeView::HandleAddState, FirstSelectedState)));
+
+	if (FirstSelectedState)
+	{
+		MenuBuilder.AddMenuEntry(
+			LOCTEXT("AddChildState", "Add Child State"),
+			FText(),
+			FSlateIcon(),
+			FUIAction(FExecuteAction::CreateSP(this, &SStateTreeView::HandleAddChildState, FirstSelectedState)));
+
+		MenuBuilder.AddSeparator();
+
+		MenuBuilder.AddMenuEntry(
+			LOCTEXT("RenameNodeGroup", "Rename"),
+			FText(),
+			FSlateIcon(),
+			FUIAction(FExecuteAction::CreateSP(this, &SStateTreeView::HandleRenameState, FirstSelectedState)));
+
+		MenuBuilder.AddMenuEntry(
+			LOCTEXT("DeleteState", "Delete Selected"),
+			FText(),
+			FSlateIcon(),
+			FUIAction(FExecuteAction::CreateSP(this, &SStateTreeView::HandleDeleteItems)));
+	}
+
+	return MenuBuilder.MakeWidget();
+}
+
+
+FReply SStateTreeView::HandleAddRoutine()
+{
+	if (StateTreeViewModel)
+	{
+		StateTreeViewModel->AddState(nullptr);
+	}
+
+	return FReply::Handled();
+}
+
+void SStateTreeView::HandleAddState(UStateTreeState* AfterState)
+{
+	if (StateTreeViewModel)
+	{
+		StateTreeViewModel->AddState(AfterState);
+	}
+}
+
+void SStateTreeView::HandleRenameState(UStateTreeState* State)
+{
+	RequestedRenameState = State;
+}
+
+void SStateTreeView::HandleAddChildState(UStateTreeState* ParentState)
+{
+	if (StateTreeViewModel && ParentState)
+	{
+		StateTreeViewModel->AddChildState(ParentState);
+		TreeView->SetItemExpansion(ParentState, true);
+	}
+}
+
+void SStateTreeView::HandleDeleteItems()
+{
+	if (StateTreeViewModel)
+	{
+		StateTreeViewModel->RemoveSelectedStates();
+	}
+}
+
+
+
+#undef LOCTEXT_NAMESPACE
diff --git a/Engine/Plugins/Runtime/StateTree/Source/StateTreeEditorModule/Private/SStateTreeView.h b/Engine/Plugins/Runtime/StateTree/Source/StateTreeEditorModule/Private/SStateTreeView.h
new file mode 100644
--- /dev/null
+++ b/Engine/Plugins/Runtime/StateTree/Source/StateTreeEditorModule/Private/SStateTreeView.h
@@ -0,0 +1,82 @@
+// Copyright Epic Games, Inc. All Rights Reserved.
+
+#pragma once
+
+#include "CoreMinimal.h"
+#include "StateTreeViewModel.h"
+#include "Widgets/Views/STreeView.h"
+
+class UStateTreeEditorData;
+class UStateTreeState;
+
+
+class FActionTreeViewDragDrop : public FDragDropOperation
+{
+public:
+	DRAG_DROP_OPERATOR_TYPE(FActionTreeViewDragDrop, FDragDropOperation);
+
+	static TSharedRef<FActionTreeViewDragDrop> New(const UStateTreeState* InState)
+	{
+		return MakeShareable(new FActionTreeViewDragDrop(InState));
+	}
+
+	const UStateTreeState* GetDraggedState() const { return State; }
+
+private:
+	FActionTreeViewDragDrop(const UStateTreeState* InState)
+		: State(InState)
+	{
+	}
+
+	const UStateTreeState* State;
+};
+
+
+class SStateTreeView : public SCompoundWidget
+{
+public:
+	SLATE_BEGIN_ARGS(SStateTreeView) {}
+	SLATE_END_ARGS()
+
+	SStateTreeView();
+	~SStateTreeView();
+
+	void Construct(const FArguments& InArgs, TSharedRef<FStateTreeViewModel> StateTreeViewModel);
+
+	void SavePersistentExpandedStates();
+
+private:
+	virtual void Tick(const FGeometry& AllottedGeometry, const double InCurrentTime, const float InDeltaTime) override;
+
+	void UpdateTree(bool bExpandPersistent = false);
+
+	// ViewModel handlers
+	void HandleModelAssetChanged();
+	void HandleModelStatesRemoved(const TSet<UStateTreeState*>& AffectedParents);
+	void HandleModelStatesMoved(const TSet<UStateTreeState*>& AffectedParents, const TSet<UStateTreeState*>& MovedStates);
+	void HandleModelStateAdded(UStateTreeState* ParentState, UStateTreeState* NewState);
+	void HandleModelStatesChanged(const TSet<UStateTreeState*>& AffectedStates, const FPropertyChangedEvent& PropertyChangedEvent);
+	void HandleModelSelectionChanged(const TArray<UStateTreeState*>& SelectedStates);
+
+	// Treeview handlers
+	TSharedRef<ITableRow> HandleGenerateRow(UStateTreeState* InState, const TSharedRef<STableViewBase>& InOwnerTableView);
+	void HandleGetChildren(UStateTreeState* InParent, TArray<UStateTreeState*>& OutChildren);
+	void HandleTreeSelectionChanged(UStateTreeState* InSelectedItem, ESelectInfo::Type SelectionType);
+	TSharedPtr<SWidget> HandleContextMenuOpening();
+
+	// Action handlers
+	FReply HandleAddRoutine();
+	void HandleRenameState(UStateTreeState* State);
+	void HandleAddState(UStateTreeState* AfterItem);
+	void HandleAddChildState(UStateTreeState* ParentItem);
+	void HandleDeleteItems();
+
+	TSharedPtr<FStateTreeViewModel> StateTreeViewModel;
+
+	TSharedPtr<STreeView<UStateTreeState*>> TreeView;
+	TSharedPtr<SScrollBar> ExternalScrollbar;
+
+	UStateTreeState* RequestedRenameState;
+	bool bItemsDirty;
+	bool bUpdatingSelection;
+};
