Moved GameplayBehaviors out of restricted folder over to Experimental

diff --git a/Engine/Plugins/Runtime/StateTree/Source/StateTreeEditorModule/Private/SStateTreeViewRow.cpp b/Engine/Plugins/Runtime/StateTree/Source/StateTreeEditorModule/Private/SStateTreeViewRow.cpp
new file mode 100644
--- /dev/null
+++ b/Engine/Plugins/Runtime/StateTree/Source/StateTreeEditorModule/Private/SStateTreeViewRow.cpp
@@ -0,0 +1,1055 @@
+// Copyright Epic Games, Inc. All Rights Reserved.
+
+#include "SStateTreeViewRow.h"
+#include "SStateTreeView.h"
+
+#include "Framework/MultiBox/MultiBoxBuilder.h"
+#include "ScopedTransaction.h"
+#include "Templates/SharedPointer.h"
+
+#include "EditorStyleSet.h"
+#include "EditorFontGlyphs.h"
+#include "StateTreeEditorStyle.h"
+
+#include "Widgets/DeclarativeSyntaxSupport.h"
+#include "Widgets/SCompoundWidget.h"
+#include "Widgets/Input/SButton.h"
+#include "Widgets/Text/STextBlock.h"
+#include "Widgets/Layout/SBorder.h"
+#include "Widgets/Layout/SScrollBorder.h"
+#include "Widgets/Images/SImage.h"
+#include "Widgets/Text/SInlineEditableTextBlock.h"
+#include "Widgets/Layout/SScrollBox.h"
+
+#include "StateTree.h"
+#include "StateTreeState.h"
+#include "StateTreeViewModel.h"
+#include "Algo/ForEach.h"
+
+#define LOCTEXT_NAMESPACE "StateTreeEditor"
+
+namespace UE { namespace StateTreeEditor {
+
+static void AddUnique(TArray<FText>& Array, const FText& NewItem)
+{
+	if (!Array.FindByPredicate([&NewItem](const FText& Item) { return Item.IdenticalTo(NewItem); }))
+	{
+		Array.Add(NewItem);
+	}
+}
+	
+}} // UE::StateTreeEditor
+
+void SStateTreeViewRow::Construct(const FArguments& InArgs, const TSharedRef<STableViewBase>& InOwnerTableView, UStateTreeState* InState, TSharedRef<FStateTreeViewModel> InStateTreeViewModel)
+{
+	StateTreeViewModel = InStateTreeViewModel;
+	WeakState = InState;
+
+	STableRow<UStateTreeState*>::ConstructInternal(STableRow::FArguments()
+        .Padding(5.f)
+        .OnDragDetected(this, &SStateTreeViewRow::HandleDragDetected)
+        .OnCanAcceptDrop(this, &SStateTreeViewRow::HandleCanAcceptDrop)
+        .OnAcceptDrop(this, &SStateTreeViewRow::HandleAcceptDrop)
+        .Style(&FStateTreeEditorStyle::Get()->GetWidgetStyle<FTableRowStyle>("StateTree.Selection"))
+    , InOwnerTableView);
+
+	static const FLinearColor TasksBackground = FLinearColor(FColor(17, 117, 131));
+	static const FLinearColor EvaluatorsBackground = FLinearColor(FColor(48, 48, 48));
+	
+	this->ChildSlot
+    .HAlign(HAlign_Fill)
+    [
+        SNew(SHorizontalBox)
+        + SHorizontalBox::Slot()
+        .VAlign(VAlign_Fill)
+        .HAlign(HAlign_Left)
+        .AutoWidth()
+        [
+            SNew(SExpanderArrow, SharedThis(this))
+            .ShouldDrawWires(true)
+            .IndentAmount(32)
+            .BaseIndentLevel(0)
+        ]
+
+        // State Box
+        + SHorizontalBox::Slot()
+        .VAlign(VAlign_Center)
+        .Padding(FMargin(0.0f, 4.0f))
+        .AutoWidth()
+        [
+            SNew(SBox)
+            .HeightOverride(28.0f)
+            .VAlign(VAlign_Fill)
+            [
+                SNew(SBorder)
+                .BorderImage(FEditorStyle::GetBrush("WhiteBrush"))
+                .BorderBackgroundColor(this, &SStateTreeViewRow::GetTitleColor)
+                .Padding(FMargin(16.0f, 0.0f, 16.0f, 0.0f))
+                [
+                    // Conditions icon
+                    SNew(SHorizontalBox)
+                    + SHorizontalBox::Slot()
+                    .VAlign(VAlign_Center)
+                    .AutoWidth()
+                    [
+                        SNew(SBox)
+                        .Padding(FMargin(0.0f, 0.0f, 4.0f, 0.0f))
+                        .Visibility(this, &SStateTreeViewRow::GetConditionVisibility)
+                        [
+                            SNew(STextBlock)
+                            .Text(FEditorFontGlyphs::Question)
+                            .TextStyle(FStateTreeEditorStyle::Get(), "StateTree.Icon")
+						]
+					]
+
+					// Selector icon
+                    + SHorizontalBox::Slot()
+                    .VAlign(VAlign_Center)
+                    .AutoWidth()
+                    [
+                        SNew(SBox)
+                        .Padding(FMargin(0.0f, 0.0f, 4.0f, 0.0f))
+                        .Visibility(this, &SStateTreeViewRow::GetSelectorVisibility)
+                        [
+                            SNew(STextBlock)
+                            .Text(FEditorFontGlyphs::Level_Down)
+                            .TextStyle(FStateTreeEditorStyle::Get(), "StateTree.Icon")
+                        ]
+                    ]
+
+                    // State Name
+					+ SHorizontalBox::Slot()
+					.VAlign(VAlign_Center)
+					.AutoWidth()
+					[
+						SAssignNew(NameTextBlock, SInlineEditableTextBlock)
+						.Style(FStateTreeEditorStyle::Get(), "StateTree.State.TitleInlineEditableText")
+						.OnVerifyTextChanged(this, &SStateTreeViewRow::VerifyNodeTextChanged)
+						.OnTextCommitted(this, &SStateTreeViewRow::HandleNodeLabelTextCommitted)
+						.Text(this, &SStateTreeViewRow::GetStateDesc)
+						.Clipping(EWidgetClipping::ClipToBounds)
+						.IsSelected(this, &SStateTreeViewRow::IsSelected)
+					]
+				]
+			]
+		]
+
+		// Evaluators Box
+		+ SHorizontalBox::Slot()
+		.VAlign(VAlign_Center)
+		.Padding(FMargin(0.0f, 4.0f))
+		.AutoWidth()
+		[
+			SNew(SBox)
+			.HeightOverride(28.0f)
+			.VAlign(VAlign_Fill)
+			.Visibility(this, &SStateTreeViewRow::GetEvaluatorsVisibility)
+			[
+				SNew(SBorder)
+				.BorderImage(FEditorStyle::GetBrush("WhiteBrush"))
+				.BorderBackgroundColor(EvaluatorsBackground)
+				.Padding(FMargin(12.0f, 0.0f, 16.0f, 0.0f))
+				[
+					// Evaluator icon
+					SNew(SHorizontalBox)
+					+ SHorizontalBox::Slot()
+					.VAlign(VAlign_Center)
+					.Padding(FMargin(0.0f, 0.0f, 4.0f, 0.0f))
+					.AutoWidth()
+					[
+						SNew(STextBlock)
+						.Text(FEditorFontGlyphs::Crosshairs)
+						.TextStyle(FStateTreeEditorStyle::Get(), "StateTree.DetailsIcon")
+					]
+
+					// Evalutors
+					+ SHorizontalBox::Slot()
+					.VAlign(VAlign_Center)
+					.AutoWidth()
+					[
+						SNew(STextBlock)
+						.Text(this, &SStateTreeViewRow::GetEvaluatorsDesc)
+						.TextStyle(FStateTreeEditorStyle::Get(), "StateTree.Details")
+					]
+				]
+			]
+		]
+
+		// Tasks Box
+		+ SHorizontalBox::Slot()
+		.VAlign(VAlign_Center)
+		.Padding(FMargin(0.0f, 4.0f))
+		.AutoWidth()
+		[
+			SNew(SBox)
+			.HeightOverride(28.0f)
+			.VAlign(VAlign_Fill)
+			.Visibility(this, &SStateTreeViewRow::GetTasksVisibility)
+			[
+				SNew(SBorder)
+				.BorderImage(FEditorStyle::GetBrush("WhiteBrush"))
+				.BorderBackgroundColor(TasksBackground)
+				.Padding(FMargin(12.0f, 0.0f, 16.0f, 0.0f))
+				[
+					// Task icon
+					SNew(SHorizontalBox)
+					+ SHorizontalBox::Slot()
+					.VAlign(VAlign_Center)
+					.Padding(FMargin(0.0f, 0.0f, 4.0f, 0.0f))
+					.AutoWidth()
+					[
+						SNew(STextBlock)
+						.Text(FEditorFontGlyphs::Paper_Plane)
+						.TextStyle(FStateTreeEditorStyle::Get(), "StateTree.DetailsIcon")
+					]
+
+					// Tasks list
+					+ SHorizontalBox::Slot()
+					.VAlign(VAlign_Center)
+					.AutoWidth()
+					[
+						SNew(STextBlock)
+						.Text(this, &SStateTreeViewRow::GetTasksDesc)
+						.TextStyle(FStateTreeEditorStyle::Get(), "StateTree.Details")
+					]
+				]
+			]
+		]
+
+		// Completed transitions
+		+ SHorizontalBox::Slot()
+		.VAlign(VAlign_Center)
+		.AutoWidth()
+		[
+			SNew(SHorizontalBox)
+			+ SHorizontalBox::Slot()
+			.VAlign(VAlign_Center)
+			.AutoWidth()
+			.Padding(FMargin(8.0f, 0.0f, 0, 0.0f))
+			[
+				SNew(STextBlock)
+				.Text(this, &SStateTreeViewRow::GetCompletedTransitionsIcon)
+				.TextStyle(FStateTreeEditorStyle::Get(), "StateTree.Icon")
+				.Visibility(this, &SStateTreeViewRow::GetCompletedTransitionVisibility)
+			]
+			+ SHorizontalBox::Slot()
+			.VAlign(VAlign_Center)
+			.AutoWidth()
+			.Padding(FMargin(4.0f, 0, 0, 0))
+			[
+				SNew(STextBlock)
+				.Text(this, &SStateTreeViewRow::GetCompletedTransitionsDesc)
+				.TextStyle(FStateTreeEditorStyle::Get(), "StateTree.Details")
+				.Visibility(this, &SStateTreeViewRow::GetCompletedTransitionVisibility)
+			]
+		]
+
+		// Succeeded transitions
+		+ SHorizontalBox::Slot()
+		.VAlign(VAlign_Center)
+		.AutoWidth()
+		[
+			SNew(SHorizontalBox)
+			+ SHorizontalBox::Slot()
+			.VAlign(VAlign_Center)
+			.Padding(FMargin(8.0f, 0.0f, 0, 0))
+			.AutoWidth()
+			[
+				SNew(STextBlock)
+				.Text(FEditorFontGlyphs::Check_Circle)
+                .ColorAndOpacity(FLinearColor(FColor(110,143,67)))
+				.TextStyle(FStateTreeEditorStyle::Get(), "StateTree.Icon")
+				.Visibility(this, &SStateTreeViewRow::GetSucceededTransitionVisibility)
+			]
+			+ SHorizontalBox::Slot()
+			.VAlign(VAlign_Center)
+			.Padding(FMargin(4.0f, 0.0f, 0, 0))
+			.AutoWidth()
+			[
+				SNew(STextBlock)
+				.Text(this, &SStateTreeViewRow::GetSucceededTransitionIcon)
+				.TextStyle(FStateTreeEditorStyle::Get(), "StateTree.Icon")
+				.Visibility(this, &SStateTreeViewRow::GetSucceededTransitionVisibility)
+			]
+			+ SHorizontalBox::Slot()
+			.VAlign(VAlign_Center)
+			.AutoWidth()
+			.Padding(FMargin(4.0f, 0, 0, 0))
+			[
+				SNew(STextBlock)
+				.Text(this, &SStateTreeViewRow::GetSucceededTransitionDesc)
+				.TextStyle(FStateTreeEditorStyle::Get(), "StateTree.Details")
+				.Visibility(this, &SStateTreeViewRow::GetSucceededTransitionVisibility)
+			]
+		]
+
+		// Failed transitions
+		+ SHorizontalBox::Slot()
+		.VAlign(VAlign_Center)
+		.AutoWidth()
+		[
+			SNew(SHorizontalBox)
+			+ SHorizontalBox::Slot()
+			.VAlign(VAlign_Center)
+			.Padding(FMargin(8.0f, 0.0f, 0, 0))
+			.AutoWidth()
+			[
+				SNew(STextBlock)
+				.Text(FEditorFontGlyphs::Times_Circle)
+                .ColorAndOpacity(FLinearColor(FColor(187,77,42)))
+				.TextStyle(FStateTreeEditorStyle::Get(), "StateTree.Icon")
+				.Visibility(this, &SStateTreeViewRow::GetFailedTransitionVisibility)
+			]
+			+ SHorizontalBox::Slot()
+			.VAlign(VAlign_Center)
+			.Padding(FMargin(4.0f, 0.0f, 0, 0))
+			.AutoWidth()
+			[
+				SNew(STextBlock)
+				.Text(this, &SStateTreeViewRow::GetFailedTransitionIcon)
+				.TextStyle(FStateTreeEditorStyle::Get(), "StateTree.Icon")
+				.Visibility(this, &SStateTreeViewRow::GetFailedTransitionVisibility)
+			]
+			+ SHorizontalBox::Slot()
+			.VAlign(VAlign_Center)
+			.AutoWidth()
+			.Padding(FMargin(4.0f, 0, 0, 0))
+			[
+				SNew(STextBlock)
+				.Text(this, &SStateTreeViewRow::GetFailedTransitionDesc)
+				.TextStyle(FStateTreeEditorStyle::Get(), "StateTree.Details")
+				.Visibility(this, &SStateTreeViewRow::GetFailedTransitionVisibility)
+			]
+		]
+		// Transitions
+		+ SHorizontalBox::Slot()
+		.VAlign(VAlign_Center)
+		.AutoWidth()
+		[
+			SNew(SHorizontalBox)
+			+ SHorizontalBox::Slot()
+			.VAlign(VAlign_Center)
+			.Padding(FMargin(8.0f, 0.0f, 0, 0))
+			.AutoWidth()
+			[
+				SNew(STextBlock)
+				.Text(FEditorFontGlyphs::Question_Circle)
+                .ColorAndOpacity(FLinearColor(FColor(31,151,167)))
+				.TextStyle(FStateTreeEditorStyle::Get(), "StateTree.Icon")
+				.Visibility(this, &SStateTreeViewRow::GetConditionalTransitionsVisibility)
+			]
+            + SHorizontalBox::Slot()
+            .VAlign(VAlign_Center)
+            .Padding(FMargin(4.0f, 0.0f, 0, 0))
+            .AutoWidth()
+            [
+                SNew(STextBlock)
+                .Text(FEditorFontGlyphs::Long_Arrow_Right)
+                .TextStyle(FStateTreeEditorStyle::Get(), "StateTree.Icon")
+                .Visibility(this, &SStateTreeViewRow::GetConditionalTransitionsVisibility)
+            ]
+            + SHorizontalBox::Slot()
+			.VAlign(VAlign_Center)
+			.AutoWidth()
+			.Padding(FMargin(4.0f, 0, 0, 0))
+			[
+				SNew(STextBlock)
+				.Text(this, &SStateTreeViewRow::GetConditionalTransitionsDesc)
+				.TextStyle(FStateTreeEditorStyle::Get(), "StateTree.Details")
+				.Visibility(this, &SStateTreeViewRow::GetConditionalTransitionsVisibility)
+			]
+		]
+	];
+}
+
+void SStateTreeViewRow::RequestRename()
+{
+	if (NameTextBlock)
+	{
+		NameTextBlock->EnterEditingMode();
+	}
+}
+
+FSlateColor SStateTreeViewRow::GetTitleColor() const
+{
+	if (const UStateTreeState* State = WeakState.Get())
+	{
+		if (StateTreeViewModel && StateTreeViewModel->IsSelected(State))
+		{
+			return FLinearColor(FColor(236, 134, 39));
+		}
+		if (IsRoutine())
+		{
+			return FLinearColor(FColor(17, 117, 131));
+		}
+	}
+
+	return FLinearColor(FColor(31, 151, 167));
+}
+
+FText SStateTreeViewRow::GetStateDesc() const
+{
+	if (const UStateTreeState* State = WeakState.Get())
+	{
+		return FText::FromName(State->Name);
+	}
+	return FText::FromName(FName());
+}
+
+EVisibility SStateTreeViewRow::GetConditionVisibility() const
+{
+	if (const UStateTreeState* State = WeakState.Get())
+	{
+		if (IsV2())
+		{
+			return State->EnterConditions2.Num() > 0 ? EVisibility::Visible : EVisibility::Collapsed;
+		}
+		else
+		{
+			return State->EnterConditions.Num() > 0 ? EVisibility::Visible : EVisibility::Collapsed;
+		}
+	}
+	return EVisibility::Collapsed;
+}
+
+EVisibility SStateTreeViewRow::GetSelectorVisibility() const
+{
+	if (const UStateTreeState* State = WeakState.Get())
+	{
+		return State->Children.Num() > 0 ? EVisibility::Visible : EVisibility::Collapsed;
+	}
+	return EVisibility::Collapsed;
+}
+
+EVisibility SStateTreeViewRow::GetEvaluatorsVisibility() const
+{
+	if (const UStateTreeState* State = WeakState.Get())
+	{
+		int32 ValidCount = 0;
+		if (IsV2())
+		{
+			for (int32 i = 0; i < State->Evaluators2.Num(); i++)
+			{
+				if (const FStateTreeEvaluator2Base* Eval = State->Evaluators2[i].Type.GetPtr<FStateTreeEvaluator2Base>())
+				{
+					ValidCount++;
+				}
+			}
+		}
+		else
+		{
+			for (int32 i = 0; i < State->Evaluators.Num(); i++)
+			{
+				if (UStateTreeEvaluatorBase* Eval = State->Evaluators[i])
+				{
+					ValidCount++;
+				}
+			}
+		}
+		return ValidCount > 0 ? EVisibility::Visible : EVisibility::Collapsed;
+	}
+	return EVisibility::Collapsed;
+}
+
+FText SStateTreeViewRow::GetEvaluatorsDesc() const
+{
+	const UStateTreeState* State = WeakState.Get();
+	if (!State)
+	{
+		return FText::GetEmpty();
+	}
+
+	TArray<FText> Names;
+
+	if (IsV2())
+	{
+		for (int32 i = 0; i < State->Evaluators2.Num(); i++)
+		{
+			if (const FStateTreeEvaluator2Base* Eval = State->Evaluators2[i].Type.GetPtr<FStateTreeEvaluator2Base>())
+			{
+				Names.Add(FText::FromName(Eval->Name));
+			}
+		}
+	}
+	else
+	{
+		for (int32 i = 0; i < State->Evaluators.Num(); i++)
+		{
+			if (UStateTreeEvaluatorBase* Eval = State->Evaluators[i])
+			{
+				Names.Add(FText::FromName(Eval->GetName()));
+			}
+		}
+	}
+
+	return FText::Join((FText::FromString(TEXT(", "))), Names);
+}
+
+EVisibility SStateTreeViewRow::GetTasksVisibility() const
+{
+	if (const UStateTreeState* State = WeakState.Get())
+	{
+		int32 ValidCount = 0;
+		if (IsV2())
+		{
+			for (int32 i = 0; i < State->Tasks2.Num(); i++)
+			{
+				if (const FStateTreeTask2Base* Task = State->Tasks2[i].Type.GetPtr<FStateTreeTask2Base>())
+				{
+					ValidCount++;
+				}
+			}
+		}
+		else
+		{
+			for (int32 i = 0; i < State->Tasks.Num(); i++)
+			{
+				if (UStateTreeTaskBase* Task = State->Tasks[i])
+				{
+					ValidCount++;
+				}
+			}
+		}
+		return ValidCount > 0 ? EVisibility::Visible : EVisibility::Collapsed;
+	}
+	return EVisibility::Collapsed;
+}
+
+FText SStateTreeViewRow::GetTasksDesc() const
+{
+	const UStateTreeState* State = WeakState.Get();
+	if (!State)
+	{
+		return FText::GetEmpty();
+	}
+
+	TArray<FText> Names;
+
+	if (IsV2())
+	{
+		for (int32 i = 0; i < State->Tasks2.Num(); i++)
+		{
+			if (const FStateTreeTask2Base* Task = State->Tasks2[i].Type.GetPtr<FStateTreeTask2Base>())
+			{
+				Names.Add(FText::FromName(Task->Name));
+			}
+		}
+	}
+	else
+	{
+		for (int32 i = 0; i < State->Tasks.Num(); i++)
+		{
+			if (UStateTreeTaskBase* Task = State->Tasks[i])
+			{
+				Names.Add(FText::FromName(Task->Name));
+			}
+		}
+	}
+
+	if (IsV2())
+	{
+		return FText::Join((FText::FromString(TEXT(" & "))), Names);
+	}
+
+	return FText::Join((FText::FromString(TEXT(", "))), Names);
+}
+
+bool SStateTreeViewRow::HasParentTransitionForEvent(const UStateTreeState& State, const EStateTreeTransitionEvent Event) const
+{
+	EStateTreeTransitionEvent CombinedEvents = EStateTreeTransitionEvent::None;
+	for (const UStateTreeState* ParentState = State.Parent; ParentState != nullptr; ParentState = ParentState->Parent)
+	{
+		for (const FStateTreeTransition2& Transition : ParentState->Transitions2)
+		{
+			CombinedEvents |= Transition.Event;
+		}
+	}
+	return EnumHasAllFlags(CombinedEvents, Event);
+}
+
+FText SStateTreeViewRow::GetTransitionsDesc(const UStateTreeState& State, const EStateTreeTransitionEvent Event) const
+{
+	TArray<FText> DescItems;
+	for (const FStateTreeTransition2& Transition : State.Transitions2)
+	{
+		if (Transition.Event == Event)
+		{
+			switch (Transition.State.Type)
+			{
+			case EStateTreeTransitionType::NotSet:
+				DescItems.Add(LOCTEXT("TransitionNoneStyled", "[None]"));
+				break;
+			case EStateTreeTransitionType::Succeeded:
+				DescItems.Add(LOCTEXT("TransitionTreeSucceededStyled", "[Succeeded]"));
+				break;
+			case EStateTreeTransitionType::Failed:
+				DescItems.Add(LOCTEXT("TransitionTreeFailedStyled", "[Failed]"));
+				break;
+			case EStateTreeTransitionType::NextState:
+				DescItems.Add(LOCTEXT("TransitionNextStateStyled", "[Next]"));
+				break;
+			case EStateTreeTransitionType::GotoState:
+				DescItems.Add(FText::FromName(Transition.State.Name));
+				break;
+			default:
+				ensureMsgf(false, TEXT("Unhandled transition type."));
+				break;
+			}
+		}
+	}
+
+	if (State.Children.Num() == 0
+		&& DescItems.Num() == 0
+		&& Event != EStateTreeTransitionEvent::OnCondition)
+	{
+		if (HasParentTransitionForEvent(State, Event))
+		{
+			DescItems.Add(LOCTEXT("TransitionActionHandleInParentStyled", "[Parent]"));
+		}
+		else
+		{
+			DescItems.Add(LOCTEXT("TransitionActionMissingTransition", "Missing Transition"));
+		}
+	}
+	
+	return FText::Join(FText::FromString(TEXT(", ")), DescItems);
+}
+
+FText SStateTreeViewRow::GetTransitionsIcon(const UStateTreeState& State, const EStateTreeTransitionEvent Event) const
+{
+	enum EIconType
+	{
+		IconNone = 0,
+        IconRightArrow =	1 << 0,
+        IconDownArrow =		1 << 1,
+        IconLevelUp =		1 << 2,
+        IconWarning =		1 << 3,
+    };
+	uint8 IconType = IconNone;
+	
+	for (const FStateTreeTransition2& Transition : State.Transitions2)
+	{
+		// The icons here depict "transition direction", not the type specifically.
+		if (Transition.Event == Event)
+		{
+			switch (Transition.State.Type)
+			{
+			case EStateTreeTransitionType::NotSet:
+				IconType |= IconRightArrow;
+				break;
+			case EStateTreeTransitionType::Succeeded:
+				IconType |= IconRightArrow;
+				break;
+			case EStateTreeTransitionType::Failed:
+				IconType |= IconRightArrow;
+				break;
+			case EStateTreeTransitionType::NextState:
+				IconType |= IconDownArrow;
+				break;
+			case EStateTreeTransitionType::GotoState:
+				IconType |= IconRightArrow;
+				break;
+			default:
+				ensureMsgf(false, TEXT("Unhandled transition type."));
+				break;
+			}
+		}
+	}
+
+	if (FMath::CountBits(static_cast<uint64>(IconType)) > 1)
+	{
+		// Prune down to just one icon.
+		IconType = IconRightArrow;
+	}
+	
+	if (State.Children.Num() == 0
+        && IconType == IconNone
+        && Event != EStateTreeTransitionEvent::OnCondition)
+	{
+		if (HasParentTransitionForEvent(State, Event))
+		{
+			IconType = IconLevelUp;
+		}
+		else
+		{
+			IconType = IconWarning;
+		}
+	}
+
+	switch (IconType)
+	{
+		case IconRightArrow:
+			return FEditorFontGlyphs::Long_Arrow_Right;
+		case IconDownArrow:
+			return FEditorFontGlyphs::Long_Arrow_Down;
+		case IconLevelUp:
+			return FEditorFontGlyphs::Level_Up;
+		case IconWarning:
+			return FEditorFontGlyphs::Exclamation_Triangle;
+		default:
+			return FText::GetEmpty();
+	}
+
+	return FText::GetEmpty();
+}
+
+EVisibility SStateTreeViewRow::GetTransitionsVisibility(const UStateTreeState& State, const EStateTreeTransitionEvent Event) const
+{
+	TStaticArray<int32, 5> Counts;
+	Algo::ForEach(Counts, [](int32& Count) { Count = 0; }); 
+
+	for (const FStateTreeTransition2& Transition : State.Transitions2)
+	{
+		Counts[static_cast<uint8>(Transition.Event)]++;
+	}
+
+	if (State.Children.Num() == 0
+	    && Counts[static_cast<uint8>(Event)] == 0
+	    && Event != EStateTreeTransitionEvent::OnCondition)
+	{
+		// Find the missing transition type, note: Completed = Succeeded|Failed.
+		EStateTreeTransitionEvent HandledEvents = EStateTreeTransitionEvent::None;
+		HandledEvents |= Counts[static_cast<uint8>(EStateTreeTransitionEvent::OnCompleted)] > 0 ? EStateTreeTransitionEvent::OnCompleted : EStateTreeTransitionEvent::None;
+		HandledEvents |= Counts[static_cast<uint8>(EStateTreeTransitionEvent::OnSucceeded)] > 0 ? EStateTreeTransitionEvent::OnSucceeded : EStateTreeTransitionEvent::None;
+		HandledEvents |= Counts[static_cast<uint8>(EStateTreeTransitionEvent::OnFailed)] > 0 ? EStateTreeTransitionEvent::OnFailed : EStateTreeTransitionEvent::None;
+		const EStateTreeTransitionEvent MissingEvent = HandledEvents ^ EStateTreeTransitionEvent::OnCompleted;
+		
+		return MissingEvent == Event ? EVisibility::Visible : EVisibility::Collapsed;
+	}
+	
+	return Counts[static_cast<uint8>(Event)] > 0 ? EVisibility::Visible : EVisibility::Collapsed;
+}
+
+EVisibility SStateTreeViewRow::GetCompletedTransitionVisibility() const
+{
+	if (const UStateTreeState* State = WeakState.Get())
+	{
+		if (IsV2())
+		{
+			return GetTransitionsVisibility(*State, EStateTreeTransitionEvent::OnCompleted);
+		}
+	}
+	return EVisibility::Visible;
+}
+
+FText SStateTreeViewRow::GetCompletedTransitionsDesc() const
+{
+	if (const UStateTreeState* State = WeakState.Get())
+	{
+		if (IsV2())
+		{
+			return GetTransitionsDesc(*State, EStateTreeTransitionEvent::OnCompleted);
+		}
+		else
+		{
+			switch (State->StateDoneTransition.Type)
+			{
+			case EStateTreeTransitionType::NotSet:
+				return FText::GetEmpty();
+			case EStateTreeTransitionType::Succeeded:
+				return LOCTEXT("TransitionSucceeded", "Succeeded");
+			case EStateTreeTransitionType::Failed:
+				return LOCTEXT("TransitionFailed", "Failed");
+			case EStateTreeTransitionType::NextState:
+				return LOCTEXT("TransitionNext", "Next");
+			case EStateTreeTransitionType::SelectChildState:
+				return LOCTEXT("TransitionSelect", "Select");
+			case EStateTreeTransitionType::GotoState:
+				return FText::FromName(State->StateDoneTransition.Name);
+			default:
+				ensureMsgf(false, TEXT("Unhandled transition type."));
+				break;
+			}
+		}
+	}
+	return LOCTEXT("Invalid", "Invalid");
+}
+
+FText SStateTreeViewRow::GetCompletedTransitionsIcon() const
+{
+	if (const UStateTreeState* State = WeakState.Get())
+	{
+		if (IsV2())
+		{
+			return GetTransitionsIcon(*State, EStateTreeTransitionEvent::OnCompleted);
+		}
+		else
+		{
+			switch (State->StateDoneTransition.Type)
+			{
+			case EStateTreeTransitionType::NotSet:
+				return FEditorFontGlyphs::Times;
+			case EStateTreeTransitionType::Succeeded:
+				return FEditorFontGlyphs::Check;
+			case EStateTreeTransitionType::Failed:
+				return FEditorFontGlyphs::Times;
+			case EStateTreeTransitionType::NextState:
+				return FEditorFontGlyphs::Long_Arrow_Down;
+			case EStateTreeTransitionType::SelectChildState:
+				return FEditorFontGlyphs::Level_Down;
+			case EStateTreeTransitionType::GotoState:
+				return FEditorFontGlyphs::Long_Arrow_Right;
+			default:
+				ensureMsgf(false, TEXT("Unhandled transition type."));
+				break;
+			}
+		}
+	}
+	return FText::GetEmpty();
+}
+
+EVisibility SStateTreeViewRow::GetSucceededTransitionVisibility() const
+{
+	if (const UStateTreeState* State = WeakState.Get())
+	{
+		if (IsV2())
+		{
+			return GetTransitionsVisibility(*State, EStateTreeTransitionEvent::OnSucceeded);
+		}
+	}
+	return EVisibility::Collapsed;
+}
+
+FText SStateTreeViewRow::GetSucceededTransitionDesc() const
+{
+	if (const UStateTreeState* State = WeakState.Get())
+	{
+		if (IsV2())
+		{
+			return GetTransitionsDesc(*State, EStateTreeTransitionEvent::OnSucceeded);
+		}
+	}
+	return FText::GetEmpty();
+}
+
+FText SStateTreeViewRow::GetSucceededTransitionIcon() const
+{
+	if (const UStateTreeState* State = WeakState.Get())
+	{
+		if (IsV2())
+		{
+			return GetTransitionsIcon(*State, EStateTreeTransitionEvent::OnSucceeded);
+		}
+	}
+	return FText::GetEmpty();
+}
+
+EVisibility SStateTreeViewRow::GetFailedTransitionVisibility() const
+{
+	if (const UStateTreeState* State = WeakState.Get())
+	{
+		if (IsV2())
+		{
+			return GetTransitionsVisibility(*State, EStateTreeTransitionEvent::OnFailed);
+		}
+		else
+		{
+			return State->StateFailedTransition.Type != EStateTreeTransitionType::NotSet ? EVisibility::Visible : EVisibility::Collapsed;
+		}
+	}
+	return EVisibility::Collapsed;
+}
+
+FText SStateTreeViewRow::GetFailedTransitionDesc() const
+{
+	if (const UStateTreeState* State = WeakState.Get())
+	{
+		if (IsV2())
+		{
+			return GetTransitionsDesc(*State, EStateTreeTransitionEvent::OnFailed);
+		}
+		else
+		{
+			switch (State->StateFailedTransition.Type)
+			{
+			case EStateTreeTransitionType::NotSet:
+				return FText::GetEmpty();
+			case EStateTreeTransitionType::Succeeded:
+				return LOCTEXT("TransitionSucceeded", "Succeeded");
+			case EStateTreeTransitionType::Failed:
+				return LOCTEXT("TransitionFailed", "Failed");
+			case EStateTreeTransitionType::NextState:
+				return LOCTEXT("TransitionNext", "Next");
+			case EStateTreeTransitionType::SelectChildState:
+				return LOCTEXT("TransitionSelect", "Select");
+			case EStateTreeTransitionType::GotoState:
+				return FText::FromName(State->StateFailedTransition.Name);
+			default:
+				ensureMsgf(false, TEXT("Unhandled transition type."));
+				break;
+			}
+		}
+	}
+	return LOCTEXT("Invalid", "Invalid");
+}
+
+FText SStateTreeViewRow::GetFailedTransitionIcon() const
+{
+	if (const UStateTreeState* State = WeakState.Get())
+	{
+		if (IsV2())
+		{
+			return GetTransitionsIcon(*State, EStateTreeTransitionEvent::OnFailed);
+		}
+		else
+		{
+			return FEditorFontGlyphs::Ban;
+		}
+	}
+	return FEditorFontGlyphs::Ban;
+}
+
+EVisibility SStateTreeViewRow::GetConditionalTransitionsVisibility() const
+{
+	if (const UStateTreeState* State = WeakState.Get())
+	{
+		if (IsV2())
+		{
+			return GetTransitionsVisibility(*State, EStateTreeTransitionEvent::OnCondition);
+		}
+		else
+		{
+			int32 ValidCount = 0;
+			for (const FStateTreeTransition& Transition : State->Transitions)
+			{
+				ValidCount++;
+			}
+			return ValidCount > 0 ? EVisibility::Visible : EVisibility::Collapsed;
+		}
+	}
+	return EVisibility::Collapsed;
+}
+
+FText SStateTreeViewRow::GetConditionalTransitionsDesc() const
+{
+	TArray<FText> DescItems;
+	if (const UStateTreeState* State = WeakState.Get())
+	{
+		if (IsV2())
+		{
+			return GetTransitionsDesc(*State, EStateTreeTransitionEvent::OnCondition);
+		}
+		else
+		{
+			for (const FStateTreeTransition& Transition : State->Transitions)
+			{
+				switch (Transition.State.Type)
+				{
+				case EStateTreeTransitionType::Succeeded:
+					DescItems.Add(LOCTEXT("TransitionSucceeded", "Succeeded"));
+					break;
+				case EStateTreeTransitionType::Failed:
+					DescItems.Add(LOCTEXT("TransitionFailed", "Failed"));
+					break;
+				case EStateTreeTransitionType::SelectChildState:
+					DescItems.Add(LOCTEXT("TransitionSelect", "Select"));
+					break;
+				case EStateTreeTransitionType::GotoState:
+					DescItems.Add(FText::FromName(Transition.State.Name));
+					break;
+				default:
+					ensureMsgf(false, TEXT("Unhandled transition type."));
+					break;
+				}
+			}
+		}
+	}
+	return FText::Join(FText::FromString(TEXT(", ")), DescItems);
+}
+
+bool SStateTreeViewRow::IsRoutine() const
+{
+	// Routines can be identified by not having parent state.
+	const UStateTreeState* State = WeakState.Get();
+	return State ? State->Parent == nullptr : false;
+}
+
+bool SStateTreeViewRow::IsSelected() const
+{
+	if (const UStateTreeState* State = WeakState.Get())
+	{
+		if (StateTreeViewModel)
+		{
+			return StateTreeViewModel->IsSelected(State);
+		}
+	}
+	return false;
+}
+
+
+bool SStateTreeViewRow::VerifyNodeTextChanged(const FText& NewLabel, FText& OutErrorMessage)
+{
+	return !NewLabel.IsEmptyOrWhitespace();
+}
+
+void SStateTreeViewRow::HandleNodeLabelTextCommitted(const FText& NewLabel, ETextCommit::Type CommitType)
+{
+	if (StateTreeViewModel)
+	{
+		if (UStateTreeState* State = WeakState.Get())
+		{
+			StateTreeViewModel->RenameState(State, FName(*FText::TrimPrecedingAndTrailing(NewLabel).ToString()));
+		}
+	}
+}
+
+
+FReply SStateTreeViewRow::HandleDragDetected(const FGeometry&, const FPointerEvent&)
+{
+	return FReply::Handled().BeginDragDrop(FActionTreeViewDragDrop::New(WeakState.Get()));
+}
+
+TOptional<EItemDropZone> SStateTreeViewRow::HandleCanAcceptDrop(const FDragDropEvent& DragDropEvent, EItemDropZone DropZone, UStateTreeState* TargetState)
+{
+	TSharedPtr<FActionTreeViewDragDrop> DragDropOperation = DragDropEvent.GetOperationAs<FActionTreeViewDragDrop>();
+	if (DragDropOperation.IsValid())
+	{
+		// Cannot drop on selection or child of selection.
+		if (StateTreeViewModel && StateTreeViewModel->IsChildOfSelection(TargetState))
+		{
+			return TOptional<EItemDropZone>();
+		}
+
+		return DropZone;
+	}
+
+	return TOptional<EItemDropZone>();
+}
+
+FReply SStateTreeViewRow::HandleAcceptDrop(const FDragDropEvent& DragDropEvent, EItemDropZone DropZone, UStateTreeState* TargetState)
+{
+	TSharedPtr<FActionTreeViewDragDrop> DragDropOperation = DragDropEvent.GetOperationAs<FActionTreeViewDragDrop>();
+	if (DragDropOperation.IsValid())
+	{
+		if (StateTreeViewModel)
+		{
+			if (DropZone == EItemDropZone::AboveItem)
+			{
+				StateTreeViewModel->MoveSelectedStatesBefore(TargetState);
+			}
+			else if (DropZone == EItemDropZone::BelowItem)
+			{
+				StateTreeViewModel->MoveSelectedStatesAfter(TargetState);
+			}
+			else
+			{
+				StateTreeViewModel->MoveSelectedStatesInto(TargetState);
+			}
+
+			return FReply::Handled();
+		}
+	}
+
+	return FReply::Unhandled();
+}
+
+bool SStateTreeViewRow::IsV2() const
+{
+	if (const UStateTreeState* State = WeakState.Get())
+	{
+		return State->IsV2();
+	}
+	return false;	
+}
+
+#undef LOCTEXT_NAMESPACE
diff --git a/Engine/Plugins/Runtime/StateTree/Source/StateTreeEditorModule/Private/SStateTreeViewRow.h b/Engine/Plugins/Runtime/StateTree/Source/StateTreeEditorModule/Private/SStateTreeViewRow.h
new file mode 100644
--- /dev/null
+++ b/Engine/Plugins/Runtime/StateTree/Source/StateTreeEditorModule/Private/SStateTreeViewRow.h
@@ -0,0 +1,76 @@
+// Copyright Epic Games, Inc. All Rights Reserved.
+
+#pragma once
+
+#include "CoreMinimal.h"
+#include "StateTreeState.h"
+#include "Widgets/Views/STableRow.h"
+
+class UStateTreeEditorData;
+class SStateTreeView;
+class SInlineEditableTextBlock;
+class FStateTreeViewModel;
+
+class SStateTreeViewRow : public STableRow<UStateTreeState*>
+{
+public:
+
+	SLATE_BEGIN_ARGS(SStateTreeViewRow)
+	{
+	}
+	SLATE_END_ARGS()
+
+	void Construct(const FArguments& InArgs, const TSharedRef<STableViewBase>& InOwnerTableView, UStateTreeState* InState, TSharedRef<FStateTreeViewModel> InStateTreeViewModel);
+	void RequestRename();
+
+private:
+
+	FSlateColor GetTitleColor() const;
+	FText GetStateDesc() const;
+
+	EVisibility GetConditionVisibility() const;
+	EVisibility GetSelectorVisibility() const;
+
+	EVisibility GetEvaluatorsVisibility() const;
+	FText GetEvaluatorsDesc() const;
+
+	EVisibility GetTasksVisibility() const;
+	FText GetTasksDesc() const;
+
+	EVisibility GetCompletedTransitionVisibility() const;
+	FText GetCompletedTransitionsDesc() const;
+	FText GetCompletedTransitionsIcon() const;
+
+	EVisibility GetSucceededTransitionVisibility() const;
+	FText GetSucceededTransitionDesc() const;
+	FText GetSucceededTransitionIcon() const;
+
+	EVisibility GetFailedTransitionVisibility() const;
+	FText GetFailedTransitionDesc() const;
+	FText GetFailedTransitionIcon() const;
+
+	EVisibility GetConditionalTransitionsVisibility() const;
+	FText GetConditionalTransitionsDesc() const;
+
+	FText GetTransitionsDesc(const UStateTreeState& State, const EStateTreeTransitionEvent Event) const;
+	FText GetTransitionsIcon(const UStateTreeState& State, const EStateTreeTransitionEvent Event) const;
+	EVisibility GetTransitionsVisibility(const UStateTreeState& State, const EStateTreeTransitionEvent Event) const;
+
+	bool HasParentTransitionForEvent(const UStateTreeState& State, const EStateTreeTransitionEvent Event) const;
+
+	bool IsRoutine() const;
+	bool IsSelected() const;
+
+	bool VerifyNodeTextChanged(const FText& NewLabel, FText& OutErrorMessage);
+	void HandleNodeLabelTextCommitted(const FText& NewLabel, ETextCommit::Type CommitType);
+
+	FReply HandleDragDetected(const FGeometry&, const FPointerEvent&);
+	TOptional<EItemDropZone> HandleCanAcceptDrop(const FDragDropEvent& DragDropEvent, EItemDropZone DropZone, UStateTreeState* TargetState);
+	FReply HandleAcceptDrop(const FDragDropEvent& DragDropEvent, EItemDropZone DropZone, UStateTreeState* TargetState);
+
+	bool IsV2() const;
+	
+	TSharedPtr<FStateTreeViewModel> StateTreeViewModel;
+	TWeakObjectPtr<UStateTreeState> WeakState;
+	TSharedPtr<SInlineEditableTextBlock> NameTextBlock;
+};
diff --git a/Engine/Plugins/Runtime/StateTree/Source/StateTreeEditorModule/Private/StateTreeBaker.cpp b/Engine/Plugins/Runtime/StateTree/Source/StateTreeEditorModule/Private/StateTreeBaker.cpp
new file mode 100644
--- /dev/null
+++ b/Engine/Plugins/Runtime/StateTree/Source/StateTreeEditorModule/Private/StateTreeBaker.cpp
@@ -0,0 +1,788 @@
+// Copyright Epic Games, Inc. All Rights Reserved.
+
+#include "StateTreeBaker.h"
+#include "StateTree.h"
+#include "StateTreeEditorData.h"
+#include "StateTreeTypes.h"
+#include "StateTreeCondition.h"
+#include "Conditions/StateTreeCondition_Common.h"
+#include "StateTreeEvaluatorBase.h"
+#include "StateTreeTaskBase.h"
+#include "StateTreeVariable.h"
+#include "StateTreeVariableDesc.h"
+#include "StateTreeVariableLayout.h"
+#include "StateTreeParameter.h"
+#include "StateTreeParameterLayout.h"
+#include "StateTreeState.h"
+#include "StateTreeExecutionContext.h"
+#include "CoreMinimal.h"
+#include "StateTreePropertyBindingCompiler.h"
+
+bool FStateTreeBaker::Bake(UStateTree& InStateTree)
+{
+	if (InStateTree.IsV2())
+	{
+		return Bake2(InStateTree);
+	}
+
+	StateTree = &InStateTree;
+	TreeData = Cast<UStateTreeEditorData>(StateTree->EditorData);
+	if (!TreeData)
+	{
+		return false;
+	}
+
+	// Cleanup existing state
+	StateTree->ResetBaked();
+
+	DefineParameters();
+
+	if (!CreateStates())
+	{
+		StateTree->ResetBaked();
+		return false;
+	}
+
+	StateTree->Variables.CalculateOffsetsAndMemoryUsage();
+	StateTree->Constants.ConstantBaseOffset = StateTree->Variables.GetMemoryUsage();
+
+	ResolveParameters();
+
+	if (!CreateStateTransitions())
+	{
+		StateTree->ResetBaked();
+		return false;
+	}
+
+	// Create and initialize tasks
+	for (const FSourceTask& Source : SourceTasks)
+	{
+		UStateTreeTaskBase* Task = Source.Task;
+		const UStateTreeState* SourceState = Source.State;
+		if (!Task->ResolveVariables(StateTree->Variables, StateTree->Constants, StateTree))
+		{
+			UE_LOG(LogStateTree, Error, TEXT("%s: %s:%s, Failed to resolve task %s."), ANSI_TO_TCHAR(__FUNCTION__), *GetNameSafe(StateTree), *SourceState->Name.ToString(), *Task->GetName());
+			StateTree->ResetBaked();
+			return false;
+		}
+		StateTree->Tasks.Add(Cast<UStateTreeTaskBase>(StaticDuplicateObject(Task, StateTree)));
+	}
+
+	// Create and initialize evaluators
+	for (const FSourceEvaluator& Source : SourceEvaluators)
+	{
+		UStateTreeEvaluatorBase* Eval = Source.Eval;
+		const UStateTreeState* SourceState = Source.State;
+		if (!Eval->ResolveVariables(StateTree->Variables, StateTree->Constants, StateTree))
+		{
+			UE_LOG(LogStateTree, Error, TEXT("%s: %s:%s, Failed to instantiate evaluator %s."), ANSI_TO_TCHAR(__FUNCTION__), *GetNameSafe(StateTree), *SourceState->Name.ToString(), *Eval->GetName().ToString());
+			StateTree->ResetBaked();
+			return false;
+		}
+		StateTree->Evaluators.Add(Cast<UStateTreeEvaluatorBase>(StaticDuplicateObject(Eval, StateTree)));
+	}
+
+	return true;
+}
+
+void FStateTreeBaker::DefineParameters() const
+{
+	const FStateTreeVariableLayout& InputParameterLayout = StateTree->GetInputParameterLayout();
+	for (const FStateTreeVariableDesc& Desc : InputParameterLayout.Variables)
+	{
+		StateTree->Variables.DefineVariable(Desc);
+	}
+}
+
+void FStateTreeBaker::ResolveParameters() const
+{
+	// Find where to write parameter values
+	StateTree->Parameters.Reset();
+
+	const FStateTreeVariableLayout& InputParameterLayout = StateTree->GetInputParameterLayout();
+	for (const FStateTreeVariableDesc& Desc : InputParameterLayout.Variables)
+	{
+		FStateTreeParameter& Param = StateTree->Parameters.Parameters.AddDefaulted_GetRef();
+		Param.Name = Desc.Name;
+		Param.Variable.Type = Desc.Type;
+		Param.Variable.Handle = StateTree->Variables.GetVariableHandle(Desc.ID);
+	}
+}
+
+bool FStateTreeBaker::CreateStates()
+{
+	for (UStateTreeState* Routine : TreeData->Routines)
+	{
+		if (Routine)
+		{
+			if (!CreateStateRecursive(*Routine, FStateTreeHandle::Invalid))
+			{
+				return false;
+			}
+		}
+	}
+	return true;
+}
+
+bool FStateTreeBaker::CreateStateRecursive(UStateTreeState& State, const FStateTreeHandle Parent)
+{
+	const int32 Idx = StateTree->States.AddDefaulted();
+	FBakedStateTreeState& BakedState = StateTree->States[Idx];
+	BakedState.Name = State.Name;
+	BakedState.Parent = Parent;
+
+	SourceStates.Add(&State);
+	IDToState.Add(State.ID, Idx);
+
+	// Collect tasks
+	BakedState.TasksBegin = uint16(SourceTasks.Num());
+	for (UStateTreeTaskBase* Task : State.Tasks)
+	{
+		if (Task)
+		{
+			SourceTasks.Add(FSourceTask(Task, &State));
+		}
+	}
+	BakedState.TasksNum = uint8(uint16(SourceTasks.Num()) - BakedState.TasksBegin);
+
+	// Collect evaluators and define variables
+	for (UStateTreeEvaluatorBase* Evaluator : State.Evaluators)
+	{
+		if (Evaluator)
+		{
+			Evaluator->DefineOutputVariables(StateTree->Variables);
+			SourceEvaluators.Add(FSourceEvaluator(Evaluator, &State));
+		}
+	}
+
+	// Child states
+	BakedState.ChildrenBegin = uint16(StateTree->States.Num());
+	for (UStateTreeState* Child : State.Children)
+	{
+		if (Child)
+		{
+			if (!CreateStateRecursive(*Child, FStateTreeHandle((uint16)Idx)))
+			{
+				return false;
+			}
+		}
+	}
+	StateTree->States[Idx].ChildrenEnd = uint16(StateTree->States.Num()); // Cannot use BakedState here, it may be invalid due to array resize.
+
+	return true;
+}
+
+bool FStateTreeBaker::ResolveTransitionState(const UStateTreeState& SourceState, const TCHAR* ContextStr, const FStateTreeStateLink& Link, FStateTreeHandle& OutTransitionHandle) const 
+{
+	if (Link.Type == EStateTreeTransitionType::GotoState)
+	{
+		OutTransitionHandle = GetStateHandle(Link.ID);
+		if (!OutTransitionHandle.IsValid())
+		{
+			UE_LOG(LogStateTree, Error, TEXT("%s: '%s':%s: Failed to resolve %s to %s."), ANSI_TO_TCHAR(__FUNCTION__), *GetNameSafe(StateTree), *SourceState.Name.ToString(), ContextStr, *Link.Name.ToString());
+			return false;
+		}
+	}
+	else if (Link.Type == EStateTreeTransitionType::NextState)
+	{
+		// Find next state.
+		const UStateTreeState* NextState = SourceState.GetNextSiblingState();
+		if (NextState == nullptr)
+		{
+			UE_LOG(LogStateTree, Error, TEXT("%s: '%s':%s: Failed to resolve default transition, there's no next State after %s."), ANSI_TO_TCHAR(__FUNCTION__), *GetNameSafe(StateTree), *SourceState.Name.ToString(), ContextStr, *SourceState.Name.ToString());
+			return false;
+		}
+		OutTransitionHandle = GetStateHandle(NextState->ID);
+		if (!OutTransitionHandle.IsValid())
+		{
+			UE_LOG(LogStateTree, Error, TEXT("%s: '%s':%s: Failed to resolve default transition next state %s."), ANSI_TO_TCHAR(__FUNCTION__), *GetNameSafe(StateTree), *SourceState.Name.ToString(), ContextStr, *Link.Name.ToString());
+			return false;
+		}
+	}
+	
+	return true;
+}
+
+bool FStateTreeBaker::CreateStateTransitions()
+{
+	for (int32 i = 0; i < StateTree->States.Num(); i++)
+	{
+		FBakedStateTreeState& BakedState = StateTree->States[i];
+		UStateTreeState* SourceState = SourceStates[i];
+		check(SourceState);
+
+		// Resolve default transition
+		BakedState.StateDoneTransitionType = SourceState->StateDoneTransition.Type;
+		BakedState.StateDoneTransitionState = FStateTreeHandle::Invalid;
+		if (!ResolveTransitionState(*SourceState, TEXT("State Done transition"), SourceState->StateDoneTransition, BakedState.StateDoneTransitionState))
+		{
+			return false;
+		}
+		// Default transition must be set.
+		if (BakedState.StateDoneTransitionType == EStateTreeTransitionType::NotSet)
+		{
+			UE_LOG(LogStateTree, Error, TEXT("%s: '%s':%s: State Done transition must be set."), ANSI_TO_TCHAR(__FUNCTION__), *GetNameSafe(StateTree), *SourceState->Name.ToString());
+			return false;
+		}
+
+		// Resolve failed transition
+		BakedState.StateFailedTransitionType = SourceState->StateFailedTransition.Type;
+		BakedState.StateFailedTransitionState = FStateTreeHandle::Invalid;
+		if (!ResolveTransitionState(*SourceState, TEXT("State Failed transition"), SourceState->StateFailedTransition, BakedState.StateFailedTransitionState))
+		{
+			return false;
+		}
+
+		// Enter conditions.
+		BakedState.EnterConditionsBegin = uint16(StateTree->Conditions.Num());
+		for (FStateTreeCondition& Condition : SourceState->EnterConditions)
+		{
+			if (!CreateCondition(Condition))
+			{
+				UE_LOG(LogStateTree, Error, TEXT("%s: '%s':%s: Failed to create state enter condition."), ANSI_TO_TCHAR(__FUNCTION__), *GetNameSafe(StateTree), *SourceState->Name.ToString());
+				return false;
+			}
+		}
+		BakedState.EnterConditionsNum = uint8(uint16(StateTree->Conditions.Num()) - BakedState.EnterConditionsBegin);
+
+		// Transitions
+		BakedState.TransitionsBegin = uint16(StateTree->Transitions.Num());
+		for (FStateTreeTransition& Transition : SourceState->Transitions)
+		{
+			FBakedStateTransition& BakedTransition = StateTree->Transitions.AddDefaulted_GetRef();
+			BakedTransition.Type = Transition.State.Type;
+			BakedTransition.State = FStateTreeHandle::Invalid;
+			if (!ResolveTransitionState(*SourceState, TEXT("transition"), Transition.State, BakedTransition.State))
+			{
+				return false;
+			}
+			// Note: Unset transition is allowed here. It can be used to mask a transition at parent.
+
+			BakedTransition.ConditionsBegin = uint16(StateTree->Conditions.Num());
+			for (FStateTreeCondition& Condition : Transition.Conditions)
+			{
+				if (!CreateCondition(Condition))
+				{
+					UE_LOG(LogStateTree, Error, TEXT("%s: '%s':%s: Failed to create transition condition to %s."), ANSI_TO_TCHAR(__FUNCTION__), *GetNameSafe(StateTree), *SourceState->Name.ToString(), *Transition.State.Name.ToString());
+					return false;
+				}
+			}
+			BakedTransition.ConditionsNum = uint8(uint16(StateTree->Conditions.Num()) - BakedTransition.ConditionsBegin);
+		}
+		BakedState.TransitionsNum = uint8(uint16(StateTree->Transitions.Num()) - BakedState.TransitionsBegin);
+	}
+
+	return true;
+}
+
+
+bool FStateTreeBaker::CreateCondition(const FStateTreeCondition& Condition)
+{
+	FStateTreeCondition& BakedCondition = StateTree->Conditions.Add_GetRef(Condition);
+
+	if (BakedCondition.Left.Type != BakedCondition.Right.Type)
+	{
+		UE_LOG(LogStateTree, Error, TEXT("%s: '%s': Left and right types do not match."), ANSI_TO_TCHAR(__FUNCTION__), *GetNameSafe(StateTree));
+		return false;
+	}
+
+	if (!BakedCondition.Left.IsBound())
+	{
+		UE_LOG(LogStateTree, Error, TEXT("%s: '%s': Left hand variable must be bound."), ANSI_TO_TCHAR(__FUNCTION__), *GetNameSafe(StateTree));
+		return false;
+	}
+
+	BakedCondition.Left.ResolveHandle(StateTree->Variables, StateTree->Constants);
+	if (!BakedCondition.Left.Handle.IsValid())
+	{
+		UE_LOG(LogStateTree, Error, TEXT("%s: '%s': Cannot resolve left hand variable %s."), ANSI_TO_TCHAR(__FUNCTION__), *GetNameSafe(StateTree), *BakedCondition.Left.Name.ToString());
+		return false;
+	}
+
+	BakedCondition.Right.ResolveHandle(StateTree->Variables, StateTree->Constants);
+	if (!BakedCondition.Right.Handle.IsValid())
+	{
+		UE_LOG(LogStateTree, Error, TEXT("%s: '%s': Cannot resolve right hand variable %s."), ANSI_TO_TCHAR(__FUNCTION__), *GetNameSafe(StateTree), *BakedCondition.Right.Name.ToString());
+		return false;
+	}
+
+	return true;
+}
+
+FStateTreeHandle FStateTreeBaker::GetStateHandle(const FGuid& StateID) const
+{
+	const int32* Idx = IDToState.Find(StateID);
+	if (Idx == nullptr)
+	{
+		return FStateTreeHandle::Invalid;
+	}
+
+	return FStateTreeHandle(uint16(*Idx));
+}
+
+
+bool FStateTreeBaker::Bake2(UStateTree& InStateTree)
+{
+	StateTree = &InStateTree;
+	TreeData = Cast<UStateTreeEditorData>(StateTree->EditorData);
+	if (!TreeData)
+	{
+		return false;
+	}
+
+	// Cleanup existing state
+	StateTree->ResetBaked();
+
+	BindingsCompiler.Init(StateTree->PropertyBindings);
+
+	if (!CreateStates2())
+	{
+		StateTree->ResetBaked();
+		return false;
+	}
+
+	if (!CreateStateTransitions2())
+	{
+		StateTree->ResetBaked();
+		return false;
+	}
+
+	BindingsCompiler.Finalize();
+
+	StateTree->InitRuntimeStorage();
+
+	// TODO: This is just for testing. It should be called just before use.
+	StateTree->ResolvePropertyPaths();
+
+	return true;
+}
+
+bool FStateTreeBaker::CreateStates2()
+{
+	// Create item for the runtime execution state
+	StateTree->RuntimeStorageItems.Add(FInstancedStruct::Make<FStateTreeExecutionState>());
+
+	for (UStateTreeState* Routine : TreeData->Routines)
+	{
+		if (Routine)
+		{
+			if (!CreateStateRecursive2(*Routine, FStateTreeHandle::Invalid))
+			{
+				return false;
+			}
+		}
+	}
+	return true;
+}
+
+bool FStateTreeBaker::CreateStateTransitions2()
+{
+	for (int32 i = 0; i < StateTree->States.Num(); i++)
+	{
+		FBakedStateTreeState& BakedState = StateTree->States[i];
+		UStateTreeState* SourceState = SourceStates[i];
+		check(SourceState);
+
+		// Enter conditions.
+		BakedState.EnterConditionsBegin = uint16(StateTree->Conditions2.Num());
+		for (FStateTreeConditionItem& ConditionItem : SourceState->EnterConditions2)
+		{
+			if (!CreateCondition2(ConditionItem))
+			{
+				UE_LOG(LogStateTree, Error, TEXT("%s: '%s':%s: Failed to create state enter condition."), ANSI_TO_TCHAR(__FUNCTION__), *GetNameSafe(StateTree), *SourceState->Name.ToString());
+				return false;
+			}
+		}
+		BakedState.EnterConditionsNum = uint8(uint16(StateTree->Conditions2.Num()) - BakedState.EnterConditionsBegin);
+
+		// Transitions
+		BakedState.TransitionsBegin = uint16(StateTree->Transitions.Num());
+		for (FStateTreeTransition2& Transition : SourceState->Transitions2)
+		{
+			FBakedStateTransition& BakedTransition = StateTree->Transitions.AddDefaulted_GetRef();
+			BakedTransition.Event = Transition.Event;
+			BakedTransition.Type = Transition.State.Type;
+			BakedTransition.GateDelay = (uint8)FMath::Clamp(FMath::CeilToInt(Transition.GateDelay * 10.0f), 0, 255);
+			BakedTransition.State = FStateTreeHandle::Invalid;
+			if (!ResolveTransitionState(*SourceState, TEXT("transition"), Transition.State, BakedTransition.State))
+			{
+				return false;
+			}
+			// Note: Unset transition is allowed here. It can be used to mask a transition at parent.
+
+			BakedTransition.ConditionsBegin = uint16(StateTree->Conditions2.Num());
+			for (FStateTreeConditionItem& ConditionItem : Transition.Conditions)
+			{
+				if (!CreateCondition2(ConditionItem))
+				{
+					UE_LOG(LogStateTree, Error, TEXT("%s: '%s':%s: Failed to create transition condition to %s."), ANSI_TO_TCHAR(__FUNCTION__), *GetNameSafe(StateTree), *SourceState->Name.ToString(), *Transition.State.Name.ToString());
+					return false;
+				}
+			}
+			BakedTransition.ConditionsNum = uint8(uint16(StateTree->Conditions2.Num()) - BakedTransition.ConditionsBegin);
+		}
+		BakedState.TransitionsNum = uint8(uint16(StateTree->Transitions.Num()) - BakedState.TransitionsBegin);
+	}
+
+	return true;
+}
+
+bool FStateTreeBaker::CreateCondition2(const FStateTreeConditionItem& CondItem)
+{
+	if (!CondItem.Type.IsValid())
+	{
+		// Empty line in conditions array, just silently ignore.
+		return true;
+	}
+
+	// Copy the condition
+	FInstancedStruct& CondPtr = StateTree->Conditions2.AddDefaulted_GetRef();
+	CondPtr = CondItem.Type;
+	check(CondPtr.IsValid());
+
+	FStateTreeConditionBase& Cond = CondPtr.GetMutable<FStateTreeConditionBase>();
+
+	// Create binding source struct descriptor. Note: not exposing the struct for reading.
+	FStateTreeBindableStructDesc StructDesc;
+	StructDesc.Struct = CondPtr.GetScriptStruct();
+	StructDesc.Name = CondPtr.GetScriptStruct()->GetFName();
+	StructDesc.ID = Cond.ID;
+
+	// Check that the bindings for this struct are still all valid.
+	TArray<FStateTreeEditorPropertyBinding> Bindings;
+	if (!GetAndValidateBindings(StructDesc, Bindings))
+	{
+		return false;
+	}
+
+	// Compile batch copy for this struct, we pass in all the bindings, the compiler will pick up the ones for the target structs.
+	int32 BatchIndex = INDEX_NONE;
+	if (!BindingsCompiler.CompileBatch(StructDesc, Bindings, BatchIndex))
+	{
+		return false;
+	}
+	check(BatchIndex < int32(MAX_uint16));
+	Cond.BindingsBatch = BatchIndex == INDEX_NONE ? FStateTreeHandle::Invalid : FStateTreeHandle(uint16(BatchIndex));
+
+	return true;
+}
+
+bool FStateTreeBaker::IsPropertyAnyEnum(const FStateTreeBindableStructDesc& Struct, FStateTreeEditorPropertyPath Path) const
+{
+	bool bIsAnyEnum = false;
+	TArray<FStateTreePropertySegment> Segments;
+	FProperty* LeafProperty = nullptr;
+	int32 LeafArrayIndex = INDEX_NONE;
+	const bool bResolved = FStateTreePropertyBindingCompiler::ResolvePropertyPath(Struct, Path, Segments, LeafProperty, LeafArrayIndex);
+	if (bResolved && LeafProperty)
+	{
+		if (FProperty* OwnerProperty = LeafProperty->GetOwnerProperty())
+		{
+			if (const FStructProperty* OwnerStructProperty = CastField<FStructProperty>(OwnerProperty))
+			{
+				bIsAnyEnum = OwnerStructProperty->Struct == FStateTreeAnyEnum::StaticStruct();
+			}
+		}
+	}
+	return bIsAnyEnum;
+}
+
+bool FStateTreeBaker::GetAndValidateBindings(const FStateTreeBindableStructDesc& TargetStruct, TArray<FStateTreeEditorPropertyBinding>& OutBindings) const
+{
+	OutBindings.Reset();
+	
+	for (const FStateTreeEditorPropertyBinding& Binding : TreeData->EditorBindings.GetBindings())
+	{
+		if (Binding.TargetPath.StructID != TargetStruct.ID)
+		{
+			continue;
+		}
+
+		// Source must be one of the source structs we have discovered in the tree.
+		const FGuid SourceStructID = Binding.SourcePath.StructID;
+		const int32 SourceStructIdx = BindingsCompiler.GetSourceStructIndexByID(SourceStructID);
+		if (SourceStructIdx == INDEX_NONE)
+		{
+			UE_LOG(LogStateTree, Error, TEXT("%s: '%s': Failed to find container for path %s in %s."), ANSI_TO_TCHAR(__FUNCTION__), *GetNameSafe(StateTree), *Binding.SourcePath.ToString(), *TargetStruct.Name.ToString());
+			return false;
+		}
+		const FStateTreeBindableStructDesc& SourceStruct = BindingsCompiler.GetSourceStructDesc(SourceStructIdx);
+
+		// Source must be accessible by the target struct.
+		TArray<FStateTreeBindableStructDesc> AccessibleStructs;
+		TreeData->GetAccessibleStructs(Binding.TargetPath.StructID, AccessibleStructs);
+		const bool SourceAccessible = AccessibleStructs.ContainsByPredicate([SourceStructID](const FStateTreeBindableStructDesc& Structs)
+			{
+				return (Structs.ID == SourceStructID);
+			});
+		if (!SourceAccessible)
+		{
+			UE_LOG(LogStateTree, Error, TEXT("%s: '%s': %s:%s is not accessible to %s:%s."), ANSI_TO_TCHAR(__FUNCTION__), *GetNameSafe(StateTree),
+				*SourceStruct.Name.ToString(), *Binding.SourcePath.ToString(), *TargetStruct.Name.ToString(), *Binding.TargetPath.ToString());
+			return false;
+		}
+
+		// Special case fo AnyEnum. StateTreeBindingExtension allows AnyEnums to bind to other enum types.
+		// The actual copy will be done via potential type promotion copy, into the value property inside the AnyEnum.
+		// We amend the paths here to point to the 'Value' property.
+		const bool bSourceIsAnyEnum = IsPropertyAnyEnum(SourceStruct, Binding.SourcePath);
+		const bool bTargetIsAnyEnum = IsPropertyAnyEnum(TargetStruct, Binding.TargetPath);
+		if (bSourceIsAnyEnum || bTargetIsAnyEnum)
+		{
+			FStateTreeEditorPropertyBinding ModifiedBinding(Binding);
+			if (bSourceIsAnyEnum)
+			{
+				ModifiedBinding.SourcePath.Path.Add(GET_MEMBER_NAME_STRING_CHECKED(FStateTreeAnyEnum, Value));
+			}
+			if (bTargetIsAnyEnum)
+			{
+				ModifiedBinding.TargetPath.Path.Add(GET_MEMBER_NAME_STRING_CHECKED(FStateTreeAnyEnum, Value));
+			}
+			OutBindings.Add(ModifiedBinding);
+		}
+		else
+		{
+			OutBindings.Add(Binding);
+		}
+	}
+
+	return true;
+}
+
+bool FStateTreeBaker::CreateExternalItemHandles(FStructView Item)
+{
+	check(StateTree);
+	
+	const UScriptStruct* ScriptStruct = Item.GetScriptStruct();
+	check(ScriptStruct);
+
+	static const FName BaseStructMetaName(TEXT("BaseStruct"));
+	static const FName BaseClassMetaName(TEXT("BaseClass"));
+	static const FName OptionalMetaName(TEXT("Optional"));
+
+	//
+	// Iterate over the main level properties of the struct and look for properties of type FStateTreeExternalItemHandle.
+	// The meta data of the property defines what type will bound to it. Examples:
+	//
+	// Handle to WorldSubsystem:
+	//		UPROPERTY(meta=(BaseClass="MassStateTreeSubsystem"))
+	//		FStateTreeExternalItemHandle MassStateTreeSubSystemHandle;
+	//
+	// Handle to a struct (fragment), optional, meaning it can be null:
+	//		UPROPERTY(meta=(BaseStruct="DataFragment_SmartObjectUser", Optional))
+	//		FStateTreeExternalItemHandle SmartObjectUserHandle;
+
+    for (TPropertyValueIterator<const FStructProperty> It(ScriptStruct, Item.GetMutableMemory(), EPropertyValueIteratorFlags::NoRecursion); It; ++It)
+    {
+    	const FStructProperty* StructProperty = It->Key;
+    	check(StructProperty);
+    	if (StructProperty->Struct == FStateTreeExternalItemHandle::StaticStruct())
+    	{
+    		// Find Class or ScriptStruct.
+    		const UStruct* Struct = nullptr;
+    		if (StructProperty->HasMetaData(BaseStructMetaName))
+    		{
+	    		const FString& BaseStructName = StructProperty->GetMetaData(BaseStructMetaName);
+				Struct = FindObject<UScriptStruct>(ANY_PACKAGE, *BaseStructName);
+    			if (Struct == nullptr)
+    			{
+    				UE_LOG(LogStateTree, Error, TEXT("%s: Struct '%s' does not exists."), ANSI_TO_TCHAR(__FUNCTION__), *BaseStructName);
+    				return false;
+    			}
+    		}
+    		else if (StructProperty->HasMetaData(BaseClassMetaName))
+    		{
+    			const FString& BaseClassName = StructProperty->GetMetaData(BaseClassMetaName);
+				Struct = FindObject<UClass>(ANY_PACKAGE, *BaseClassName);
+    			if (Struct == nullptr)
+    			{
+    				UE_LOG(LogStateTree, Error, TEXT("%s: Class '%s' does not exists."), ANSI_TO_TCHAR(__FUNCTION__), *BaseClassName);
+    				return false;
+    			}
+    		}
+    		else
+    		{
+    			UE_LOG(LogStateTree, Error, TEXT("%s: FStateTreeExternalItemHandle must have 'BaseStruct' or 'BaseClass' set."), ANSI_TO_TCHAR(__FUNCTION__));
+    			return false;
+    		}
+
+    		// Parse optional
+    		const bool bOptional = StructProperty->HasMetaData(OptionalMetaName);
+
+    		// The struct/class must be accepted by the schema.
+    		if (StateTree->Schema && !StateTree->Schema->IsExternalItemAllowed(*Struct))
+    		{
+    			UE_LOG(LogStateTree, Error, TEXT("%s: Schema %s does not allow item type %s."), ANSI_TO_TCHAR(__FUNCTION__), *GetNameSafe(StateTree->Schema), *GetNameSafe(Struct));
+    			return false;
+    		}
+
+    		// Check if similar struct already exists, if not add new.
+    		FStateTreeExternalItemDesc* ExternalItem = StateTree->ExternalItems.FindByPredicate([Struct](const FStateTreeExternalItemDesc& Item) { return Item.Struct == Struct; });
+    		if (ExternalItem == nullptr)
+    		{
+    			ExternalItem = &StateTree->ExternalItems.Emplace_GetRef(Struct, bOptional);
+    		
+    			// The external struct pointer is stored in the same array as property binding sources.
+    			// This is done also as anticipation to allow to bind to external items.
+    			FStateTreeBindableStructDesc StructDesc;
+    			StructDesc.Struct = Struct;
+    			StructDesc.Name = FName(Struct->GetName() + TEXT("_External"));
+    			StructDesc.ID = FGuid(); // Empty GUID, this item cannot be bound to.
+
+    			const int32 StructIndex = BindingsCompiler.AddSourceStruct(StructDesc);
+
+    			check(StructIndex <= int32(MAX_uint16));
+    			ExternalItem->Handle.SetIndex(uint16(StructIndex));
+			}
+    		else
+    		{
+    			// If same type is requested as required, clear optional flag. 
+    			if (!bOptional)
+    			{
+    				ExternalItem->bOptional = false;
+    			}
+    		}
+
+    		check(ExternalItem);
+
+    		// Update the handle on Eval/struct
+    		FStateTreeExternalItemHandle& Handle = *static_cast<FStateTreeExternalItemHandle*>(const_cast<void*>(It.Value()));
+    		Handle = ExternalItem->Handle;
+    	}
+    }
+
+	return true;
+}
+
+bool FStateTreeBaker::CreateStateRecursive2(UStateTreeState& State, const FStateTreeHandle Parent)
+{
+	const int32 StateIdx = StateTree->States.AddDefaulted();
+	FBakedStateTreeState& BakedState = StateTree->States[StateIdx];
+	BakedState.Name = State.Name;
+	BakedState.Parent = Parent;
+
+	SourceStates.Add(&State);
+	IDToState.Add(State.ID, StateIdx);
+
+	// Collect evaluators
+	check(StateTree->RuntimeStorageItems.Num() <= int32(MAX_uint16));
+	BakedState.EvaluatorsBegin = uint16(StateTree->RuntimeStorageItems.Num());
+
+	for (FStateTreeEvaluatorItem& EvaluatorItem : State.Evaluators2)
+	{
+		if (EvaluatorItem.Type.IsValid())
+		{
+			// Copy the evaluator
+			FInstancedStruct& EvalPtr = StateTree->RuntimeStorageItems.AddDefaulted_GetRef();
+			EvalPtr = EvaluatorItem.Type;
+			check(EvalPtr.IsValid());
+			FStateTreeEvaluator2Base& Eval = EvalPtr.GetMutable<FStateTreeEvaluator2Base>();
+
+			// Create binding source struct descriptor.
+			FStateTreeBindableStructDesc StructDesc;
+			StructDesc.Struct = EvalPtr.GetScriptStruct();
+			StructDesc.Name = Eval.Name;
+			StructDesc.ID = Eval.ID;
+
+			// Mark the struct as binding source.
+			const int32 SourceStructIndex = BindingsCompiler.AddSourceStruct(StructDesc);
+
+			// Check that the bindings for this struct are still all valid.
+			TArray<FStateTreeEditorPropertyBinding> Bindings;
+			if (!GetAndValidateBindings(StructDesc, Bindings))
+			{
+				return false;
+			}
+
+			// Compile batch copy for this struct, we pass in all the bindings, the compiler will pick up the ones for the target structs.
+			int32 BatchIndex = INDEX_NONE;
+			if (!BindingsCompiler.CompileBatch(StructDesc, Bindings, BatchIndex))
+			{
+				return false;
+			}
+			check(BatchIndex < int32(MAX_uint16));
+			Eval.BindingsBatch = BatchIndex == INDEX_NONE ? FStateTreeHandle::Invalid : FStateTreeHandle(uint16(BatchIndex));
+			check(SourceStructIndex <= int32(MAX_uint16));
+			Eval.SourceStructIndex = uint16(SourceStructIndex);
+
+			if (!CreateExternalItemHandles(EvalPtr))
+			{
+				UE_LOG(LogStateTree, Error, TEXT("%s: '%s': %s failed to create external item handles."), ANSI_TO_TCHAR(__FUNCTION__), *GetNameSafe(StateTree), *Eval.Name.ToString());
+				return false;
+			}
+		}
+	}
+	const int32 EvaluatorsNum = StateTree->RuntimeStorageItems.Num() - int32(BakedState.EvaluatorsBegin);
+	check(EvaluatorsNum <= int32(MAX_uint8));
+	BakedState.EvaluatorsNum = uint8(EvaluatorsNum);
+
+	// Collect tasks
+	check(StateTree->RuntimeStorageItems.Num() <= int32(MAX_uint16));
+	BakedState.TasksBegin = uint16(StateTree->RuntimeStorageItems.Num());
+
+	for (FStateTreeTaskItem& TaskItem : State.Tasks2)
+	{
+		if (TaskItem.Type.IsValid())
+		{
+			// Copy the task
+			FInstancedStruct& TaskPtr = StateTree->RuntimeStorageItems.AddDefaulted_GetRef();
+			TaskPtr = TaskItem.Type;
+			check(TaskPtr.IsValid());
+			FStateTreeTask2Base& Task = TaskPtr.GetMutable<FStateTreeTask2Base>();
+
+			// Create binding source struct descriptor.
+			FStateTreeBindableStructDesc StructDesc;
+			StructDesc.Struct = TaskPtr.GetScriptStruct();
+			StructDesc.Name = Task.Name;
+			StructDesc.ID = Task.ID;
+
+			// Mark the struct as binding source.
+			const int32 SourceStructIndex = BindingsCompiler.AddSourceStruct(StructDesc);
+
+			// Check that the bindings for this struct are still all valid.
+			TArray<FStateTreeEditorPropertyBinding> Bindings;
+			if (!GetAndValidateBindings(StructDesc, Bindings))
+			{
+				return false;
+			}
+
+			// Compile batch copy for this struct, we pass in all the bindings, the compiler will pick up the ones for the target structs.
+			int32 BatchIndex = INDEX_NONE;
+			if (!BindingsCompiler.CompileBatch(StructDesc, Bindings, BatchIndex))
+			{
+				return false;
+			}
+			check(BatchIndex < int32(MAX_uint16));
+			Task.BindingsBatch = BatchIndex == INDEX_NONE ? FStateTreeHandle::Invalid : FStateTreeHandle(uint16(BatchIndex));
+			check(SourceStructIndex <= int32(MAX_uint16));
+			Task.SourceStructIndex = uint16(SourceStructIndex);
+
+			if (!CreateExternalItemHandles(TaskPtr))
+			{
+				UE_LOG(LogStateTree, Error, TEXT("%s: '%s': %s failed to create external item handles."), ANSI_TO_TCHAR(__FUNCTION__), *GetNameSafe(StateTree), *Task.Name.ToString());
+				return false;
+			}
+		}
+	}
+	const int32 TasksNum = StateTree->RuntimeStorageItems.Num() - int32(BakedState.TasksBegin);
+	check(TasksNum <= int32(MAX_uint8));
+	BakedState.TasksNum = uint8(TasksNum);
+
+	// Child states
+	check(StateTree->States.Num() <= int32(MAX_uint16));
+	BakedState.ChildrenBegin = uint16(StateTree->States.Num());
+	for (UStateTreeState* Child : State.Children)
+	{
+		if (Child)
+		{
+			if (!CreateStateRecursive2(*Child, FStateTreeHandle((uint16)StateIdx)))
+			{
+				return false;
+			}
+		}
+	}
+	check(StateTree->States.Num() <= int32(MAX_uint16));
+	StateTree->States[StateIdx].ChildrenEnd = uint16(StateTree->States.Num()); // Cannot use BakedState here, it may be invalid due to array resize.
+
+	return true;
+}
diff --git a/Engine/Plugins/Runtime/StateTree/Source/StateTreeEditorModule/Private/StateTreeEditor.cpp b/Engine/Plugins/Runtime/StateTree/Source/StateTreeEditorModule/Private/StateTreeEditor.cpp
new file mode 100644
--- /dev/null
+++ b/Engine/Plugins/Runtime/StateTree/Source/StateTreeEditorModule/Private/StateTreeEditor.cpp
@@ -0,0 +1,643 @@
+// Copyright Epic Games, Inc. All Rights Reserved.
+
+#include "StateTreeEditor.h"
+#include "Modules/ModuleManager.h"
+#include "EditorStyleSet.h"
+#include "SStateTreeView.h"
+#include "PropertyEditorModule.h"
+#include "StateTreeEditorModule.h"
+#include "StateTree.h"
+#include "StateTreeEditorData.h"
+#include "StateTreeDelegates.h"
+#include "StateTreeState.h"
+#include "StateTreeEvaluatorBase.h"
+#include "StateTreeTaskBase.h"
+#include "StateTreeBaker.h"
+#include "IDetailsView.h"
+#include "Widgets/Docking/SDockTab.h"
+#include "Customizations/StateTreeBindingExtension.h"
+#include "Logging/MessageLog.h"
+
+
+#define LOCTEXT_NAMESPACE "StateTreeEditor"
+
+const FName StateTreeEditorAppName(TEXT("StateTreeEditorApp"));
+
+const FName FStateTreeEditor::SelectionDetailsTabId(TEXT("StateTreeEditor_SelectionDetails"));
+const FName FStateTreeEditor::AssetDetailsTabId(TEXT("StateTreeEditor_AssetDetails"));
+const FName FStateTreeEditor::StateTreeViewTabId(TEXT("StateTreeEditor_StateTreeView"));
+
+void FStateTreeEditor::PostUndo(bool bSuccess)
+{
+}
+
+void FStateTreeEditor::PostRedo(bool bSuccess)
+{
+}
+
+void FStateTreeEditor::AddReferencedObjects(FReferenceCollector& Collector)
+{
+	if (StateTree != nullptr)
+	{
+		Collector.AddReferencedObject(StateTree);
+	}
+}
+
+void FStateTreeEditor::RegisterTabSpawners(const TSharedRef<class FTabManager>& InTabManager)
+{
+	WorkspaceMenuCategory = InTabManager->AddLocalWorkspaceMenuCategory(LOCTEXT("WorkspaceMenu_StateTreeEditor", "StateTree Editor"));
+	auto WorkspaceMenuCategoryRef = WorkspaceMenuCategory.ToSharedRef();
+
+	FAssetEditorToolkit::RegisterTabSpawners(InTabManager);
+
+	InTabManager->RegisterTabSpawner(SelectionDetailsTabId, FOnSpawnTab::CreateSP(this, &FStateTreeEditor::SpawnTab_SelectionDetails) )
+		.SetDisplayName( NSLOCTEXT("StateTreeEditor", "SelectionDetailsTab", "Details" ) )
+		.SetGroup(WorkspaceMenuCategoryRef)
+		.SetIcon(FSlateIcon(FEditorStyle::GetStyleSetName(), "LevelEditor.Tabs.Details"));
+
+	InTabManager->RegisterTabSpawner(AssetDetailsTabId, FOnSpawnTab::CreateSP(this, &FStateTreeEditor::SpawnTab_AssetDetails))
+		.SetDisplayName(NSLOCTEXT("StateTreeEditor", "AssetDetailsTab", "Asset Details"))
+		.SetGroup(WorkspaceMenuCategoryRef)
+		.SetIcon(FSlateIcon(FEditorStyle::GetStyleSetName(), "LevelEditor.Tabs.Details"));
+
+	InTabManager->RegisterTabSpawner(StateTreeViewTabId, FOnSpawnTab::CreateSP(this, &FStateTreeEditor::SpawnTab_StateTreeView))
+		.SetDisplayName(NSLOCTEXT("StateTreeEditor", "StateTreeViewTab", "StateTree"))
+		.SetGroup(WorkspaceMenuCategoryRef)
+		.SetIcon(FSlateIcon(FEditorStyle::GetStyleSetName(), "LevelEditor.Tabs.Outliner"));
+}
+
+void FStateTreeEditor::UnregisterTabSpawners(const TSharedRef<class FTabManager>& InTabManager)
+{
+	FAssetEditorToolkit::UnregisterTabSpawners(InTabManager);
+
+	InTabManager->UnregisterTabSpawner(SelectionDetailsTabId);
+	InTabManager->UnregisterTabSpawner(AssetDetailsTabId);
+	InTabManager->UnregisterTabSpawner(StateTreeViewTabId);
+}
+
+void FStateTreeEditor::InitEditor( const EToolkitMode::Type Mode, const TSharedPtr< class IToolkitHost >& InitToolkitHost, UStateTree* InStateTree)
+{
+	StateTree = InStateTree;
+	check(StateTree != NULL);
+
+	UStateTreeEditorData* EditorData = Cast<UStateTreeEditorData>(StateTree->EditorData);
+	if (EditorData == NULL)
+	{
+		EditorData = NewObject<UStateTreeEditorData>(StateTree, FName(), RF_Transactional);
+		StateTree->EditorData = EditorData;
+	}
+
+	StateTreeViewModel = MakeShareable(new FStateTreeViewModel());
+	StateTreeViewModel->Init(EditorData);
+
+	StateTreeViewModel->GetOnAssetChanged().AddSP(this, &FStateTreeEditor::HandleModelAssetChanged);
+	StateTreeViewModel->GetOnSelectionChanged().AddSP(this, &FStateTreeEditor::HandleModelSelectionChanged);
+
+
+	TSharedRef<FTabManager::FLayout> StandaloneDefaultLayout = FTabManager::NewLayout("Standalone_StateTree_Layout_v2")
+	->AddArea
+	(
+		FTabManager::NewPrimaryArea() ->SetOrientation(Orient_Vertical)
+		->Split
+		(
+			FTabManager::NewSplitter() ->SetOrientation(Orient_Horizontal)
+			->Split
+			(
+				FTabManager::NewStack()
+				->SetSizeCoefficient(0.2f)
+				->AddTab(AssetDetailsTabId, ETabState::OpenedTab)
+				->SetForegroundTab(AssetDetailsTabId)
+			)
+			->Split
+			(
+				FTabManager::NewStack()
+				->SetSizeCoefficient(0.5f)
+				->AddTab(StateTreeViewTabId, ETabState::OpenedTab)
+			)
+			->Split
+			(
+				FTabManager::NewStack()
+				->SetSizeCoefficient(0.3f)
+				->AddTab(SelectionDetailsTabId, ETabState::OpenedTab)
+				->SetForegroundTab(SelectionDetailsTabId)
+			)
+		)
+	);
+
+
+	const bool bCreateDefaultStandaloneMenu = true;
+	const bool bCreateDefaultToolbar = true;
+	FAssetEditorToolkit::InitAssetEditor(Mode, InitToolkitHost, StateTreeEditorAppName, StandaloneDefaultLayout, bCreateDefaultStandaloneMenu, bCreateDefaultToolbar, StateTree);
+	
+	FStateTreeEditorModule& StateTreeEditorModule = FModuleManager::LoadModuleChecked<FStateTreeEditorModule>("StateTreeEditorModule");
+	AddMenuExtender(StateTreeEditorModule.GetMenuExtensibilityManager()->GetAllExtenders(GetToolkitCommands(), GetEditingObjects()));
+
+	RegenerateMenusAndToolbars();
+
+	// TODO: should only do validation here to not edit asset each time it is opened.
+	UpdateAsset();
+
+	UE::StateTree::Delegates::OnIdentifierChanged.AddSP(this, &FStateTreeEditor::OnIdentifierChanged);
+	UE::StateTree::Delegates::OnParameterLayoutChanged.AddSP(this, &FStateTreeEditor::OnParameterLayoutChanged);
+	UE::StateTree::Delegates::OnSchemaChanged.AddSP(this, &FStateTreeEditor::OnSchemaChanged);
+}
+
+FName FStateTreeEditor::GetToolkitFName() const
+{
+	return FName("StateTreeEditor");
+}
+
+FText FStateTreeEditor::GetBaseToolkitName() const
+{
+	return NSLOCTEXT("StateTreeEditor", "AppLabel", "State Tree");
+}
+
+FString FStateTreeEditor::GetWorldCentricTabPrefix() const
+{
+	return NSLOCTEXT("StateTreeEditor", "WorldCentricTabPrefix", "State Tree").ToString();
+}
+
+
+FLinearColor FStateTreeEditor::GetWorldCentricTabColorScale() const
+{
+	return FLinearColor( 0.0f, 0.0f, 0.2f, 0.5f );
+}
+
+TSharedRef<SDockTab> FStateTreeEditor::SpawnTab_StateTreeView(const FSpawnTabArgs& Args)
+{
+	check(Args.GetTabId() == StateTreeViewTabId);
+
+	return SNew(SDockTab)
+		.Label(NSLOCTEXT("StateTreeEditor", "StateTreeViewTab", "State Tree"))
+		.TabColorScale(GetTabColorScale())
+		[
+			SAssignNew(StateTreeView, SStateTreeView, StateTreeViewModel.ToSharedRef())
+		];
+}
+
+
+TSharedRef<SDockTab> FStateTreeEditor::SpawnTab_SelectionDetails(const FSpawnTabArgs& Args)
+{
+	check(Args.GetTabId() == SelectionDetailsTabId);
+
+	FPropertyEditorModule& PropertyEditorModule = FModuleManager::GetModuleChecked<FPropertyEditorModule>("PropertyEditor");
+	FDetailsViewArgs DetailsViewArgs;
+	DetailsViewArgs.NameAreaSettings = FDetailsViewArgs::HideNameArea;
+
+	SelectionDetailsView = PropertyEditorModule.CreateDetailView(DetailsViewArgs);
+	SelectionDetailsView->SetObject(nullptr);
+	SelectionDetailsView->OnFinishedChangingProperties().AddSP(this, &FStateTreeEditor::OnSelectionFinishedChangingProperties);
+
+	SelectionDetailsView->SetExtensionHandler(MakeShared<FStateTreeBindingExtension>());
+
+	TSharedRef<SDockTab> SpawnedTab = SNew(SDockTab)
+		.Label(NSLOCTEXT("StateTreeEditor", "SelectionDetailsTab", "Details"))
+		[
+			SelectionDetailsView.ToSharedRef()
+		];
+
+	return SpawnedTab;
+}
+
+TSharedRef<SDockTab> FStateTreeEditor::SpawnTab_AssetDetails(const FSpawnTabArgs& Args)
+{
+	check(Args.GetTabId() == AssetDetailsTabId);
+
+	FPropertyEditorModule& PropertyEditorModule = FModuleManager::GetModuleChecked<FPropertyEditorModule>("PropertyEditor");
+	FDetailsViewArgs DetailsViewArgs;
+	DetailsViewArgs.NameAreaSettings = FDetailsViewArgs::HideNameArea;
+
+	AssetDetailsView = PropertyEditorModule.CreateDetailView(DetailsViewArgs);
+	AssetDetailsView->SetObject(StateTree);
+	AssetDetailsView->OnFinishedChangingProperties().AddSP(this, &FStateTreeEditor::OnAssetFinishedChangingProperties);
+
+	TSharedRef<SDockTab> SpawnedTab = SNew(SDockTab)
+		.Label(NSLOCTEXT("StateTreeEditor", "AssetDetailsTab", "StateTree"))
+		[
+			AssetDetailsView.ToSharedRef()
+		];
+
+	return SpawnedTab;
+}
+
+void FStateTreeEditor::HandleModelAssetChanged()
+{
+	UpdateAsset();
+}
+
+void FStateTreeEditor::HandleModelSelectionChanged(const TArray<UStateTreeState*>& SelectedStates)
+{
+	if (SelectionDetailsView)
+	{
+		TArray<UObject*> Selected;
+		for (UStateTreeState* State : SelectedStates)
+		{
+			Selected.Add(State);
+		}
+		SelectionDetailsView->SetObjects(Selected);
+	}
+}
+
+void FStateTreeEditor::SaveAsset_Execute()
+{
+	// Remember the treview expnasion state
+	if (StateTreeView)
+	{
+		StateTreeView->SavePersistentExpandedStates();
+	}
+
+	UpdateAsset();
+
+	// save it
+	FAssetEditorToolkit::SaveAsset_Execute();
+}
+
+void FStateTreeEditor::OnIdentifierChanged(const UStateTree& InStateTree)
+{
+	if (StateTree == &InStateTree)
+	{
+		UpdateAsset();
+	}
+}
+
+void FStateTreeEditor::OnParameterLayoutChanged(const UStateTree& InStateTree)
+{
+	// Update when other StateTree assets change.
+	if (StateTree != &InStateTree)
+	{
+		// Update parameter sets. For assets that are not open, this will be done on load.
+		UpdateAsset();
+	}
+}
+
+void FStateTreeEditor::OnSchemaChanged(const UStateTree& InStateTree)
+{
+	if (StateTree == &InStateTree)
+	{
+		if (StateTreeViewModel)
+		{
+			StateTreeViewModel->NotifyAssetChangedExternally();
+		}
+	}
+}
+
+void FStateTreeEditor::OnAssetFinishedChangingProperties(const FPropertyChangedEvent& PropertyChangedEvent)
+{
+	// Make sure nodes get updates when properties are changed.
+	if (StateTreeViewModel)
+	{
+		StateTreeViewModel->NotifyAssetChangedExternally();
+	}
+}
+
+void FStateTreeEditor::OnSelectionFinishedChangingProperties(const FPropertyChangedEvent& PropertyChangedEvent)
+{
+	// Make sure nodes get updates when properties are changed.
+	if (StateTreeViewModel)
+	{
+		const TArray<TWeakObjectPtr<UObject>>& SelectedObjects = SelectionDetailsView->GetSelectedObjects();
+		TSet<UStateTreeState*> ChangedStates;
+		for (const TWeakObjectPtr<UObject>& WeakObject : SelectedObjects)
+		{
+			if (UObject* Object = WeakObject.Get())
+			{
+				if (UStateTreeState* State = Cast<UStateTreeState>(Object))
+				{
+					ChangedStates.Add(State);
+				}
+			}
+		}
+		if (ChangedStates.Num() > 0)
+		{
+			StateTreeViewModel->NotifyStatesChangedExternally(ChangedStates, PropertyChangedEvent);
+		}
+	}
+}
+
+namespace FStateTreeEditorHelpers
+{
+	bool FixChangedStateLinkName(FStateTreeStateLink& StateLink, const TMap<FGuid, FName>& IDToName)
+	{
+		if (StateLink.ID.IsValid())
+		{
+			const FName* Name = IDToName.Find(StateLink.ID);
+			if (Name == nullptr)
+			{
+				// Missing link, we'll show these in the UI
+				return false;
+			}
+			if (StateLink.Name != *Name)
+			{
+				// Name changed, fix!
+				StateLink.Name = *Name;
+				return true;
+			}
+		}
+		return false;
+	}
+
+	bool FixChangedVariableName(FStateTreeVariable& VarBinding, const TMap<FGuid, FName>& IDToName)
+	{
+		if (VarBinding.ID.IsValid())
+		{
+			const FName* Name = IDToName.Find(VarBinding.ID);
+			if (Name == nullptr)
+			{
+				// Missing link, we'll show these in the UI
+				return false;
+			}
+			if (VarBinding.Name != *Name)
+			{
+				// Name changed, fix!
+				VarBinding.Name = *Name;
+				return true;
+			}
+		}
+		return false;
+	}
+
+	bool FixChangedVariableName(FStateTreeCondition& Condition, const TMap<FGuid, FName>& IDToName)
+	{
+		bool Result = false;
+		Result |= FixChangedVariableName(Condition.Left, IDToName);
+		Result |= FixChangedVariableName(Condition.Right, IDToName);
+		return Result;
+	}
+
+	bool FixChangedVariableName(TArray<FStateTreeCondition>& Conditions, const TMap<FGuid, FName>& IDToName)
+	{
+		bool Result = false;
+		for (FStateTreeCondition& Cond : Conditions)
+		{
+			Result |= FixChangedVariableName(Cond, IDToName);
+		}
+		return Result;
+	}
+
+	void ValidateLinkedStates(UStateTree& StateTree)
+	{
+		UStateTreeEditorData* TreeData = Cast<UStateTreeEditorData>(StateTree.EditorData);
+		if (!TreeData)
+		{
+			return;
+		}
+
+		// Make sure all state links are valid and update the names if needed.
+
+		// Create ID to state name map.
+		TMap<FGuid, FName> IDToName;
+		TArray<UStateTreeState*> Stack;
+
+		for (UStateTreeState* Routine : TreeData->Routines)
+		{
+			if (Routine)
+			{
+				Stack.Reset();
+				Stack.Add(Routine);
+				while (Stack.Num() > 0)
+				{
+					UStateTreeState* CurState = Stack.Pop();
+					IDToName.Add(CurState->ID, CurState->Name);
+					for (UStateTreeState* ChildState : CurState->Children)
+					{
+						if (ChildState)
+						{
+							Stack.Append(CurState->Children);
+						}
+					}
+				}
+			}
+		}
+
+		// Fix changed names.
+		for (UStateTreeState* Routine : TreeData->Routines)
+		{
+			if (Routine)
+			{
+				Stack.Reset();
+				Stack.Add(Routine);
+				while (Stack.Num() > 0)
+				{
+					UStateTreeState* CurState = Stack.Pop();
+
+					FixChangedStateLinkName(CurState->StateDoneTransition, IDToName);
+					FixChangedStateLinkName(CurState->StateFailedTransition, IDToName);
+
+					for (FStateTreeTransition& Transition : CurState->Transitions)
+					{
+						FixChangedStateLinkName(Transition.State, IDToName);
+					}
+
+					for (UStateTreeState* ChildState : CurState->Children)
+					{
+						if (ChildState)
+						{
+							Stack.Append(CurState->Children);
+						}
+					}
+				}
+			}
+		}
+
+	}
+
+	void ValidateLinkedVariables(UStateTree& StateTree)
+	{
+		UStateTreeEditorData* TreeData = Cast<UStateTreeEditorData>(StateTree.EditorData);
+		if (!TreeData)
+		{
+			return;
+		}
+
+		// Make sure all state links are valid and update the names if needed.
+		TArray<UStateTreeState*> Stack;
+
+		// Create ID to state name map.
+		TMap<FGuid, FName> IDToName;
+
+		for (UStateTreeState* Routine : TreeData->Routines)
+		{
+			if (Routine)
+			{
+				Stack.Reset();
+				Stack.Add(Routine);
+				while (Stack.Num() > 0)
+				{
+					UStateTreeState* CurState = Stack.Pop();
+					for (const UStateTreeEvaluatorBase* Evaluator : CurState->Evaluators)
+					{
+						if (Evaluator)
+						{
+							IDToName.Add(Evaluator->GetID(), Evaluator->GetName());
+						}
+					}
+
+					for (UStateTreeState* ChildState : CurState->Children)
+					{
+						if (ChildState)
+						{
+							Stack.Append(CurState->Children);
+						}
+					}
+				}
+			}
+		}
+
+		for (UStateTreeState* Routine : TreeData->Routines)
+		{
+			if (Routine)
+			{
+				// Fix changed names.
+				Stack.Reset();
+				Stack.Add(Routine);
+				while (Stack.Num() > 0)
+				{
+					UStateTreeState* CurState = Stack.Pop();
+
+					FixChangedVariableName(CurState->EnterConditions, IDToName);
+
+					for (FStateTreeTransition& Transition : CurState->Transitions)
+					{
+						FixChangedVariableName(Transition.Conditions, IDToName);
+					}
+
+					for (UStateTreeState* ChildState : CurState->Children)
+					{
+						if (ChildState)
+						{
+							Stack.Append(CurState->Children);
+						}
+					}
+				}
+			}
+		}
+	}
+
+	void ValidateEvaluatorsAndTracks(UStateTree& StateTree)
+	{
+		UStateTreeEditorData* TreeData = Cast<UStateTreeEditorData>(StateTree.EditorData);
+		if (!TreeData)
+		{
+			return;
+		}
+
+		TArray<UStateTreeState*> Stack;
+		bool bChanged = false;
+
+		for (UStateTreeState* Routine : TreeData->Routines)
+		{
+			if (Routine)
+			{
+				// Make sure all parameterset are valid and update the layout if needed.
+				Stack.Reset();
+				Stack.Add(Routine);
+				while (Stack.Num() > 0)
+				{
+					UStateTreeState* CurState = Stack.Pop();
+					for (UStateTreeEvaluatorBase* Evaluator : CurState->Evaluators)
+					{
+						if (Evaluator)
+						{
+							bChanged |= Evaluator->ValidateParameterLayout();
+						}
+					}
+
+					for (UStateTreeTaskBase* Task : CurState->Tasks)
+					{
+						if (Task)
+						{
+							bChanged |= Task->ValidateParameterLayout();
+						}
+					}
+
+					for (UStateTreeState* ChildState : CurState->Children)
+					{
+						if (ChildState)
+						{
+							Stack.Append(CurState->Children);
+						}
+					}
+				}
+			}
+		}
+
+		if (bChanged)
+		{
+			UE::StateTree::Delegates::OnParametersInvalidated.Broadcast(StateTree);
+		}
+	}
+
+	void UpdateParents(UStateTree& StateTree)
+	{
+		UStateTreeEditorData* TreeData = Cast<UStateTreeEditorData>(StateTree.EditorData);
+		if (!TreeData)
+		{
+			return;
+		}
+
+		TArray<UStateTreeState*> Stack;
+
+		for (UStateTreeState* Routine : TreeData->Routines)
+		{
+			if (Routine)
+			{
+				Routine->Parent = nullptr;
+				Stack.Reset();
+				Stack.Add(Routine);
+				while (Stack.Num() > 0)
+				{
+					UStateTreeState* CurState = Stack.Pop();
+					for (UStateTreeState* ChildState : CurState->Children)
+					{
+						if (ChildState)
+						{
+							ChildState->Parent = CurState;
+							Stack.Append(CurState->Children);
+						}
+					}
+				}
+			}
+		}
+	}
+
+	void RemoveUnusedBindings(UStateTree& StateTree)
+	{
+		UStateTreeEditorData* TreeData = Cast<UStateTreeEditorData>(StateTree.EditorData);
+		if (!TreeData)
+		{
+			return;
+		}
+
+		TMap<FGuid, const UScriptStruct*> AllStructIDs;
+		TreeData->GetAllStructIDs(AllStructIDs);
+		TreeData->GetPropertyEditorBindings()->RemoveUnusedBindings(AllStructIDs);
+	}
+}
+
+void FStateTreeEditor::UpdateAsset()
+{
+	if (!StateTree)
+	{
+		return;
+	}
+
+	FStateTreeEditorHelpers::UpdateParents(*StateTree);
+	FStateTreeEditorHelpers::RemoveUnusedBindings(*StateTree);
+	FStateTreeEditorHelpers::ValidateLinkedStates(*StateTree);
+	FStateTreeEditorHelpers::ValidateLinkedVariables(*StateTree);
+	FStateTreeEditorHelpers::ValidateEvaluatorsAndTracks(*StateTree);
+
+	FStateTreeBaker Baker;
+	if (!Baker.Bake(*StateTree))
+	{
+		FMessageLog EditorErrors("StateTree");
+		EditorErrors.Error(FText::Format(LOCTEXT("StateTreeBakerFailed", "Failed to bake state tree asset {0}, see the log for details"), FText::FromString(StateTree->GetName())));
+		EditorErrors.Notify(FText::Format(LOCTEXT("StateTreeBakerFailed", "Failed to bake state tree asset {0}, see the log for details"), FText::FromString(StateTree->GetName())));
+
+		// Make sure not to leave stale data on failed bake.
+		StateTree->ResetBaked();
+	}
+}
+
+
+#undef LOCTEXT_NAMESPACE
diff --git a/Engine/Plugins/Runtime/StateTree/Source/StateTreeEditorModule/Private/StateTreeEditor.h b/Engine/Plugins/Runtime/StateTree/Source/StateTreeEditorModule/Private/StateTreeEditor.h
new file mode 100644
--- /dev/null
+++ b/Engine/Plugins/Runtime/StateTree/Source/StateTreeEditorModule/Private/StateTreeEditor.h
@@ -0,0 +1,85 @@
+// Copyright Epic Games, Inc. All Rights Reserved.
+
+#pragma once
+
+#include "CoreMinimal.h"
+#include "Toolkits/IToolkitHost.h"
+#include "EditorUndoClient.h"
+#include "Toolkits/AssetEditorToolkit.h"
+#include "IStateTreeEditor.h"
+#include "StateTreeViewModel.h"
+
+class IDetailsView;
+class UStateTree;
+class FStateTreeViewModel;
+
+class FStateTreeEditor : public IStateTreeEditor, public FSelfRegisteringEditorUndoClient, public FGCObject
+{
+public:
+
+	virtual void RegisterTabSpawners(const TSharedRef<class FTabManager>& TabManager) override;
+	virtual void UnregisterTabSpawners(const TSharedRef<class FTabManager>& TabManager) override;
+
+	void InitEditor(const EToolkitMode::Type Mode, const TSharedPtr< class IToolkitHost >& InitToolkitHost, UStateTree* StateTree);
+
+	//~ Begin IToolkit Interface
+	virtual FName GetToolkitFName() const override;
+	virtual FText GetBaseToolkitName() const override;
+	virtual FString GetWorldCentricTabPrefix() const override;
+	virtual FLinearColor GetWorldCentricTabColorScale() const override;
+	//~ End IToolkit Interface
+
+	void OnAssetFinishedChangingProperties(const FPropertyChangedEvent& PropertyChangedEvent);
+	void OnSelectionFinishedChangingProperties(const FPropertyChangedEvent& PropertyChangedEvent);
+
+	/** FEditorUndoClient interface */
+	virtual void PostUndo(bool bSuccess) override;
+	virtual void PostRedo(bool bSuccess) override;
+
+	/** FGCObject interface */
+	virtual void AddReferencedObjects(FReferenceCollector& Collector) override;
+	virtual FString GetReferencerName() const override
+	{
+		return TEXT("FStateTreeEditor");
+	}
+
+protected:
+	/** Called when "Save" is clicked for this asset */
+	virtual void SaveAsset_Execute() override;
+
+private:
+
+	void HandleModelSelectionChanged(const TArray<UStateTreeState*>& SelectedStates);
+
+	/** Spawns the tab with the update graph inside */
+	TSharedRef<SDockTab> SpawnTab_StateTreeView(const FSpawnTabArgs& Args);
+	TSharedRef<SDockTab> SpawnTab_SelectionDetails(const FSpawnTabArgs& Args);
+	TSharedRef<SDockTab> SpawnTab_AssetDetails(const FSpawnTabArgs& Args);
+
+private:
+
+	void UpdateAsset();
+
+	void HandleModelAssetChanged();
+	void OnIdentifierChanged(const UStateTree& StateTree);
+	void OnParameterLayoutChanged(const UStateTree& StateTree);
+	void OnSchemaChanged(const UStateTree& StateTree);
+
+	/* State Tree being edited */
+	UStateTree* StateTree;
+
+	/** Selection Property View */
+	TSharedPtr<class IDetailsView> SelectionDetailsView;
+
+	/** Asset Property View */
+	TSharedPtr<class IDetailsView> AssetDetailsView;
+
+	/** Tree View */
+	TSharedPtr<class SStateTreeView> StateTreeView;
+
+	TSharedPtr<FStateTreeViewModel> StateTreeViewModel;
+
+	static const FName StateTreeViewTabId;
+	static const FName SelectionDetailsTabId;
+	static const FName AssetDetailsTabId;
+};
diff --git a/Engine/Plugins/Runtime/StateTree/Source/StateTreeEditorModule/Private/StateTreeEditorData.cpp b/Engine/Plugins/Runtime/StateTree/Source/StateTreeEditorModule/Private/StateTreeEditorData.cpp
new file mode 100644
--- /dev/null
+++ b/Engine/Plugins/Runtime/StateTree/Source/StateTreeEditorModule/Private/StateTreeEditorData.cpp
@@ -0,0 +1,205 @@
+// Copyright Epic Games, Inc. All Rights Reserved.
+
+#include "StateTreeEditorData.h"
+#include "StateTreeConditionBase.h"
+#include "CoreMinimal.h"
+
+void UStateTreeEditorData::GetAccessibleStructs(const FGuid TargetStructID, TArray<FStateTreeBindableStructDesc>& OutStructDescs) const
+{
+	// Find the states that are updated before the current state.
+	TSet<const UStateTreeState*> ValidStates;
+	const UStateTreeState* CurrentState = GetStateByStructID(TargetStructID);
+	for (const UStateTreeState* State = CurrentState; State != nullptr; State = State->Parent)
+	{
+		ValidStates.Add(State);
+	}
+
+	TArray<FStateTreeBindableStructDesc> EvalDescs;
+	TArray<FStateTreeBindableStructDesc> TaskDescs;
+
+	if (ValidStates.Num() > 0)
+	{
+		VisitHierarchy([&ValidStates, &OutStructDescs, &EvalDescs, &TaskDescs, TargetStructID, CurrentState](const UStateTreeState& State, const FGuid& ID, const FName& Name, const UScriptStruct* Struct) -> bool
+			{
+				if (ValidStates.Contains(&State))
+				{
+					if (ID == TargetStructID)
+					{
+						OutStructDescs.Append(EvalDescs);
+
+						// Only tasks can see other tasks too.
+						if (Struct->IsChildOf(FStateTreeTask2Base::StaticStruct()))
+						{
+							OutStructDescs.Append(TaskDescs);
+						}
+						
+						return false; // Stop visit
+					}
+
+					if (Struct->IsChildOf(FStateTreeEvaluator2Base::StaticStruct()))
+					{
+						// All evaluators up to the target struct are accessible.
+						FStateTreeBindableStructDesc& Desc = EvalDescs.AddDefaulted_GetRef();
+						Desc.Struct = Struct;
+						Desc.Name = Name;
+						Desc.ID = ID;
+					}
+					else if (Struct->IsChildOf(FStateTreeTask2Base::StaticStruct()))
+					{
+						// All tasks up to the target struct are accessible (for other tasks).
+						FStateTreeBindableStructDesc& Desc = TaskDescs.AddDefaulted_GetRef();
+						Desc.Struct = Struct;
+						Desc.Name = Name;
+						Desc.ID = ID;
+					}
+				}
+				return true; // Continue
+			});
+	}
+}
+
+bool UStateTreeEditorData::GetStructByID(const FGuid StructID, FStateTreeBindableStructDesc& OutStructDesc) const
+{
+	bool bResult = false;
+
+	VisitHierarchy([&bResult, &OutStructDesc, StructID](const UStateTreeState& State, const FGuid& ID, const FName& Name, const UScriptStruct* Struct) -> bool
+		{
+			if (ID == StructID)
+			{
+				OutStructDesc.Struct = Struct;
+				OutStructDesc.Name = Name;
+				OutStructDesc.ID = ID;
+				bResult = true;
+				return false; // Stop visit
+			}
+			return true; // Continue
+		});
+
+	return bResult;
+}
+
+const UStateTreeState* UStateTreeEditorData::GetStateByStructID(const FGuid TargetStructID) const
+{
+	const UStateTreeState* Result = nullptr;
+
+	VisitHierarchy([&Result, TargetStructID](const UStateTreeState& State, const FGuid& ID, const FName& Name, const UScriptStruct* Struct) -> bool
+		{
+			if (ID == TargetStructID)
+			{
+				Result = &State;
+				return false; // Stop visit
+			}
+			return true; // Continue
+		});
+
+	return Result;
+}
+
+void UStateTreeEditorData::GetAllStructIDs(TMap<FGuid, const UScriptStruct*>& AllStructs) const
+{
+	AllStructs.Reset();
+
+	VisitHierarchy([&AllStructs](const UStateTreeState& State, const FGuid& ID, const FName& Name, const UScriptStruct* Struct) -> bool
+		{
+			AllStructs.Add(ID, Struct);
+			return true; // Continue
+		});
+}
+
+void UStateTreeEditorData::VisitHierarchy(TFunctionRef<bool(const UStateTreeState& State, const FGuid& ID, const FName& Name, const UScriptStruct* Struct)> InFunc) const
+{
+	TArray<const UStateTreeState*> Stack;
+	bool bContinue = true;
+
+	for (const UStateTreeState* Routine : Routines)
+	{
+		if (!Routine)
+		{
+			continue;
+		}
+
+		Stack.Add(Routine);
+
+		while (!Stack.IsEmpty() && bContinue)
+		{
+			const UStateTreeState* State = Stack[0];
+			check(State);
+
+			Stack.RemoveAt(0);
+
+			// Evaluators
+			for (const FStateTreeEvaluatorItem& Item : State->Evaluators2)
+			{
+				if (const FStateTreeEvaluator2Base* Evaluator = Item.Type.GetPtr<FStateTreeEvaluator2Base>())
+				{
+					if (!InFunc(*State, Evaluator->ID, Evaluator->Name, Item.Type.GetScriptStruct()))
+					{
+						bContinue = false;
+						break;
+					}
+				}
+			}
+			if (bContinue)
+			{
+				// Enter conditions
+				for (const FStateTreeConditionItem& Item : State->EnterConditions2)
+				{
+					if (const FStateTreeConditionBase* Cond = Item.Type.GetPtr<FStateTreeConditionBase>())
+					{
+						if (!InFunc(*State, Cond->ID, Item.Type.GetScriptStruct()->GetFName(), Item.Type.GetScriptStruct()))
+						{
+							bContinue = false;
+							break;
+						}
+					}
+				}
+			}
+			if (bContinue)
+			{
+				// Tasks
+				for (const FStateTreeTaskItem& Item : State->Tasks2)
+				{
+					if (const FStateTreeTask2Base* Task = Item.Type.GetPtr<FStateTreeTask2Base>())
+					{
+						if (!InFunc(*State, Task->ID, Task->Name, Item.Type.GetScriptStruct()))
+						{
+							bContinue = false;
+							break;
+						}
+					}
+				}
+			}
+			if (bContinue)
+			{
+				// Transitions
+				for (const FStateTreeTransition2& Transition : State->Transitions2)
+				{
+					for (const FStateTreeConditionItem& Item : Transition.Conditions)
+					{
+						if (const FStateTreeConditionBase* Cond = Item.Type.GetPtr<FStateTreeConditionBase>())
+						{
+							if (!InFunc(*State, Cond->ID, Item.Type.GetScriptStruct()->GetFName(), Item.Type.GetScriptStruct()))
+							{
+								bContinue = false;
+								break;
+							}
+						}
+					}
+				}
+			}
+			if (bContinue)
+			{
+				// Children
+				for (const UStateTreeState* ChildState : State->Children)
+				{
+					Stack.Add(ChildState);
+				}
+			}
+		}
+		if (!bContinue)
+		{
+			break;
+		}
+	}
+}
+
diff --git a/Engine/Plugins/Runtime/StateTree/Source/StateTreeEditorModule/Private/StateTreeEditorModule.cpp b/Engine/Plugins/Runtime/StateTree/Source/StateTreeEditorModule/Private/StateTreeEditorModule.cpp
new file mode 100644
--- /dev/null
+++ b/Engine/Plugins/Runtime/StateTree/Source/StateTreeEditorModule/Private/StateTreeEditorModule.cpp
@@ -0,0 +1,124 @@
+// Copyright Epic Games, Inc. All Rights Reserved.
+
+#include "StateTreeEditorModule.h"
+#include "Modules/ModuleManager.h"
+#include "Developer/AssetTools/Public/IAssetTools.h"
+#include "Developer/AssetTools/Public/AssetToolsModule.h"
+#include "AssetTypeActions_StateTree.h"
+#include "Customizations/StateTreeVariableDescDetails.h"
+#include "Customizations/StateTreeVariableLayoutDetails.h"
+#include "Customizations/StateTreeVariableDetails.h"
+#include "Customizations/StateTreeParameterLayoutDetails.h"
+#include "Customizations/StateTreeConditionDetails.h"
+#include "Customizations/StateTreeTransitionDetails.h"
+#include "Customizations/StateTreeStateLinkDetails.h"
+#include "Customizations/StateTreeTransition2Details.h"
+#include "Customizations/StateTreeStateDetails.h"
+#include "Customizations/StateTreeBindableItemDetails.h"
+#include "Customizations/StateTreeEditorPropertyPathDetails.h"
+#include "Customizations/StateTreeConditionItemDetails.h"
+#include "Customizations/StateTreeAnyEnumDetails.h"
+#include "AIGraphTypes.h" // Class cache
+#include "StateTreeEditor.h"
+#include "StateTree.h"
+#include "StateTreeState.h"
+#include "StateTreeTaskBase.h"
+#include "StateTreeEditorStyle.h"
+
+
+#define LOCTEXT_NAMESPACE "StateTreeEditor"
+
+IMPLEMENT_MODULE(FStateTreeEditorModule, StateTreeEditorModule)
+
+void FStateTreeEditorModule::StartupModule()
+{
+	MenuExtensibilityManager = MakeShareable(new FExtensibilityManager);
+	ToolBarExtensibilityManager = MakeShareable(new FExtensibilityManager);
+
+	FStateTreeEditorStyle::Initialize();
+
+	// Register asset types
+	IAssetTools& AssetTools = FModuleManager::LoadModuleChecked<FAssetToolsModule>("AssetTools").Get();
+
+	TSharedPtr<FAssetTypeActions_StateTree> StateTreeAssetTypeAction = MakeShareable(new FAssetTypeActions_StateTree);
+	ItemDataAssetTypeActions.Add(StateTreeAssetTypeAction);
+	AssetTools.RegisterAssetTypeActions(StateTreeAssetTypeAction.ToSharedRef());
+
+	// Register the details customizer
+	FPropertyEditorModule& PropertyModule = FModuleManager::LoadModuleChecked<FPropertyEditorModule>("PropertyEditor");
+	PropertyModule.RegisterCustomPropertyTypeLayout("StateTreeVariableDesc", FOnGetPropertyTypeCustomizationInstance::CreateStatic(&FStateTreeVariableDescDetails::MakeInstance));
+	PropertyModule.RegisterCustomPropertyTypeLayout("StateTreeVariableLayout", FOnGetPropertyTypeCustomizationInstance::CreateStatic(&FStateTreeVariableLayoutDetails::MakeInstance));
+	PropertyModule.RegisterCustomPropertyTypeLayout("StateTreeVariable", FOnGetPropertyTypeCustomizationInstance::CreateStatic(&FStateTreeVariableDetails::MakeInstance));
+	PropertyModule.RegisterCustomPropertyTypeLayout("StateTreeParameterLayout", FOnGetPropertyTypeCustomizationInstance::CreateStatic(&FStateTreeParameterLayoutDetails::MakeInstance));
+	PropertyModule.RegisterCustomPropertyTypeLayout("StateTreeCondition", FOnGetPropertyTypeCustomizationInstance::CreateStatic(&FStateTreeConditionDetails::MakeInstance));
+	PropertyModule.RegisterCustomPropertyTypeLayout("StateTreeTransition", FOnGetPropertyTypeCustomizationInstance::CreateStatic(&FStateTreeTransitionDetails::MakeInstance));
+	PropertyModule.RegisterCustomPropertyTypeLayout("StateTreeStateLink", FOnGetPropertyTypeCustomizationInstance::CreateStatic(&FStateTreeStateLinkDetails::MakeInstance));
+	PropertyModule.RegisterCustomPropertyTypeLayout("StateTreeTransition2", FOnGetPropertyTypeCustomizationInstance::CreateStatic(&FStateTreeTransition2Details::MakeInstance));
+	PropertyModule.RegisterCustomPropertyTypeLayout("StateTreeEvaluatorItem", FOnGetPropertyTypeCustomizationInstance::CreateStatic(&FStateTreeBindableItemDetails::MakeInstance));
+	PropertyModule.RegisterCustomPropertyTypeLayout("StateTreeTaskItem", FOnGetPropertyTypeCustomizationInstance::CreateStatic(&FStateTreeBindableItemDetails::MakeInstance));
+	PropertyModule.RegisterCustomPropertyTypeLayout("StateTreeEditorPropertyPath", FOnGetPropertyTypeCustomizationInstance::CreateStatic(&FStateTreeEditorPropertyPathDetails::MakeInstance));
+	PropertyModule.RegisterCustomPropertyTypeLayout("StateTreeConditionItem", FOnGetPropertyTypeCustomizationInstance::CreateStatic(&FStateTreeConditionItemDetails::MakeInstance));
+	PropertyModule.RegisterCustomPropertyTypeLayout("StateTreeAnyEnum", FOnGetPropertyTypeCustomizationInstance::CreateStatic(&FStateTreeAnyEnumDetails::MakeInstance));
+	PropertyModule.RegisterCustomClassLayout("StateTreeState", FOnGetDetailCustomizationInstance::CreateStatic(&FStateTreeStateDetails::MakeInstance));
+
+	PropertyModule.NotifyCustomizationModuleChanged();
+}
+
+void FStateTreeEditorModule::ShutdownModule()
+{
+	TaskClassCache.Reset();
+	MenuExtensibilityManager.Reset();
+	ToolBarExtensibilityManager.Reset();
+
+	FStateTreeEditorStyle::Shutdown();
+
+	// Unregister the data asset type actions
+	if (FModuleManager::Get().IsModuleLoaded("AssetTools"))
+	{
+		IAssetTools& AssetToolsModule = FModuleManager::GetModuleChecked<FAssetToolsModule>("AssetTools").Get();
+		for (int i = 0; i < ItemDataAssetTypeActions.Num(); i++)
+		{
+			if (ItemDataAssetTypeActions[i].IsValid())
+			{
+				AssetToolsModule.UnregisterAssetTypeActions(ItemDataAssetTypeActions[i].ToSharedRef());
+			}
+		}
+	}
+	ItemDataAssetTypeActions.Empty();
+
+	// Unregister the details customization
+	if (FModuleManager::Get().IsModuleLoaded("PropertyEditor"))
+	{
+		FPropertyEditorModule& PropertyModule = FModuleManager::LoadModuleChecked<FPropertyEditorModule>("PropertyEditor");
+		PropertyModule.UnregisterCustomPropertyTypeLayout("StateTreeVariableDesc");
+		PropertyModule.UnregisterCustomPropertyTypeLayout("StateTreeVariableLayout");
+		PropertyModule.UnregisterCustomPropertyTypeLayout("StateTreeVariable");
+		PropertyModule.UnregisterCustomPropertyTypeLayout("StateTreeParameterLayout");
+		PropertyModule.UnregisterCustomPropertyTypeLayout("StateTreeCondition");
+		PropertyModule.UnregisterCustomPropertyTypeLayout("StateTreeTransition");
+		PropertyModule.UnregisterCustomPropertyTypeLayout("StateTreeTransition2");
+		PropertyModule.UnregisterCustomPropertyTypeLayout("StateTreeStateLink");
+		PropertyModule.UnregisterCustomPropertyTypeLayout("StateTreeEvaluatorItem");
+		PropertyModule.UnregisterCustomPropertyTypeLayout("StateTreeTaskItem");
+		PropertyModule.UnregisterCustomPropertyTypeLayout("StateTreeEditorPropertyPath");
+		PropertyModule.UnregisterCustomPropertyTypeLayout("StateTreeConditionItem");
+		PropertyModule.UnregisterCustomPropertyTypeLayout("StateTreeAnyEnum");
+		PropertyModule.NotifyCustomizationModuleChanged();
+	}
+
+}
+
+TSharedRef<IStateTreeEditor> FStateTreeEditorModule::CreateStateTreeEditor(const EToolkitMode::Type Mode, const TSharedPtr< class IToolkitHost >& InitToolkitHost, UStateTree* StateTree)
+{
+	if (!TaskClassCache.IsValid())
+	{
+		TaskClassCache = MakeShareable(new FGraphNodeClassHelper(UStateTreeTaskBase::StaticClass()));
+		TaskClassCache->UpdateAvailableBlueprintClasses();
+	}
+
+	TSharedRef<FStateTreeEditor> NewEditor(new FStateTreeEditor());
+	NewEditor->InitEditor(Mode, InitToolkitHost, StateTree);
+	return NewEditor;
+}
+
+#undef LOCTEXT_NAMESPACE
