Moved GameplayBehaviors out of restricted folder over to Experimental

diff --git a/Engine/Plugins/Runtime/StateTree/Source/StateTreeEditorModule/Private/StateTreeEditorPropertyBindings.cpp b/Engine/Plugins/Runtime/StateTree/Source/StateTreeEditorModule/Private/StateTreeEditorPropertyBindings.cpp
new file mode 100644
--- /dev/null
+++ b/Engine/Plugins/Runtime/StateTree/Source/StateTreeEditorModule/Private/StateTreeEditorPropertyBindings.cpp
@@ -0,0 +1,144 @@
+// Copyright Epic Games, Inc. All Rights Reserved.
+
+#include "StateTreeEditorPropertyBindings.h"
+#include "PropertyPathHelpers.h"
+#include "StateTreePropertyBindingCompiler.h"
+
+
+UStateTreeEditorPropertyBindingsOwner::UStateTreeEditorPropertyBindingsOwner(const FObjectInitializer& ObjectInitializer)
+	: Super(ObjectInitializer)
+{
+}
+
+//////////////////////////////////////////////////////////////////////////
+
+void FStateTreeEditorPropertyBindings::AddPropertyBinding(const FStateTreeEditorPropertyPath& SourcePath, const FStateTreeEditorPropertyPath& TargetPath)
+{
+	RemovePropertyBindings(TargetPath);
+	PropertyBindings.Add(FStateTreeEditorPropertyBinding(SourcePath, TargetPath));
+}
+
+void FStateTreeEditorPropertyBindings::RemovePropertyBindings(const FStateTreeEditorPropertyPath& TargetPath)
+{
+	PropertyBindings.RemoveAll([TargetPath](const FStateTreeEditorPropertyBinding& Binding)
+		{
+			return Binding.TargetPath == TargetPath;
+		});
+}
+
+bool FStateTreeEditorPropertyBindings::HasPropertyBinding(const FStateTreeEditorPropertyPath& TargetPath) const
+{
+	return PropertyBindings.ContainsByPredicate([TargetPath](const FStateTreeEditorPropertyBinding& Binding)
+		{
+			return Binding.TargetPath == TargetPath;
+		});
+}
+
+const FStateTreeEditorPropertyPath* FStateTreeEditorPropertyBindings::GetPropertyBindingSource(const FStateTreeEditorPropertyPath& TargetPath) const
+{
+	const FStateTreeEditorPropertyBinding* Binding = PropertyBindings.FindByPredicate([TargetPath](const FStateTreeEditorPropertyBinding& Binding)
+		{
+			return Binding.TargetPath == TargetPath;
+		});
+	return Binding ? &Binding->SourcePath : nullptr;
+}
+
+void FStateTreeEditorPropertyBindings::GetPropertyBindingsFor(const FGuid StructID, TArray<FStateTreeEditorPropertyBinding>& OutBindings) const
+{
+	OutBindings = PropertyBindings.FilterByPredicate([StructID](const FStateTreeEditorPropertyBinding& Binding)
+		{
+			return Binding.TargetPath.StructID == StructID && Binding.IsValid();
+		});
+}
+
+void FStateTreeEditorPropertyBindings::RemoveUnusedBindings(const TMap<FGuid, const UScriptStruct*>& ValidStructs)
+{
+	PropertyBindings.RemoveAll([ValidStructs](const FStateTreeEditorPropertyBinding& Binding)
+		{
+			// Remove binding if it's target struct has been removed
+			if (!ValidStructs.Contains(Binding.TargetPath.StructID))
+			{
+				return true;
+			}
+
+			// Remove binding if the target property cannot be resolved, likely renamed property.
+			// TODO: Try to use core redirect to find new name.
+			const UScriptStruct* Struct = ValidStructs.FindChecked(Binding.TargetPath.StructID);
+
+			FStateTreeBindableStructDesc StructDesc;
+			StructDesc.Struct = Struct;
+			StructDesc.ID = Binding.TargetPath.StructID;
+		
+			TArray<FStateTreePropertySegment> Segments;
+			FProperty* LeafProperty = nullptr;
+			int32 LeafArrayIndex = INDEX_NONE;
+			if (!FStateTreePropertyBindingCompiler::ResolvePropertyPath(StructDesc, Binding.TargetPath, Segments, LeafProperty, LeafArrayIndex))
+			{
+				return true;
+			}
+
+			return false;
+		});
+}
+
+//////////////////////////////////////////////////////////////////////////
+
+FStateTreeBindingLookup::FStateTreeBindingLookup(IStateTreeEditorPropertyBindingsOwner* InBindingOwner)
+	: BindingOwner(InBindingOwner)
+{
+}
+
+const FStateTreeEditorPropertyPath* FStateTreeBindingLookup::GetPropertyBindingSource(const FStateTreeEditorPropertyPath& InTargetPath) const
+{
+	check(BindingOwner);
+	FStateTreeEditorPropertyBindings* EditorBindings = BindingOwner->GetPropertyEditorBindings();
+	check(EditorBindings);
+
+	return EditorBindings->GetPropertyBindingSource(InTargetPath);
+}
+
+FText FStateTreeBindingLookup::GetPropertyPathDisplayName(const FStateTreeEditorPropertyPath& InPath) const
+{
+	check(BindingOwner);
+	FStateTreeEditorPropertyBindings* EditorBindings = BindingOwner->GetPropertyEditorBindings();
+	check(EditorBindings);
+
+	FString Result;
+
+	FStateTreeBindableStructDesc Struct;
+	if (BindingOwner->GetStructByID(InPath.StructID, Struct))
+	{
+		Result = Struct.Name.ToString();
+	}
+
+	for (const FString& Segment : InPath.Path)
+	{
+		Result += TEXT(".") + Segment;
+	}
+
+	return FText::FromString(Result);
+}
+
+const FProperty* FStateTreeBindingLookup::GetPropertyPathLeafProperty(const FStateTreeEditorPropertyPath& InPath) const
+{
+	check(BindingOwner);
+	FStateTreeEditorPropertyBindings* EditorBindings = BindingOwner->GetPropertyEditorBindings();
+	check(EditorBindings);
+
+	FProperty* Result = nullptr;
+	FStateTreeBindableStructDesc Struct;
+	if (BindingOwner->GetStructByID(InPath.StructID, Struct))
+	{
+		// TODO: Could use inline allocator here, since there are usually only few segments. Needs API change in ResolvePropertyPath().
+		TArray<FStateTreePropertySegment> Segments;
+		FProperty* LeafProperty = nullptr;
+		int32 LeafArrayIndex = INDEX_NONE;
+		if (FStateTreePropertyBindingCompiler::ResolvePropertyPath(Struct, InPath, Segments, LeafProperty, LeafArrayIndex))
+		{
+			Result = LeafProperty;
+		}
+	}
+
+	return Result;
+}
+
diff --git a/Engine/Plugins/Runtime/StateTree/Source/StateTreeEditorModule/Private/StateTreeEditorStyle.cpp b/Engine/Plugins/Runtime/StateTree/Source/StateTreeEditorModule/Private/StateTreeEditorStyle.cpp
new file mode 100644
--- /dev/null
+++ b/Engine/Plugins/Runtime/StateTree/Source/StateTreeEditorModule/Private/StateTreeEditorStyle.cpp
@@ -0,0 +1,124 @@
+// Copyright Epic Games, Inc. All Rights Reserved.
+
+#include "StateTreeEditorStyle.h"
+#include "Styling/SlateStyleRegistry.h"
+#include "Styling/CoreStyle.h"
+#include "Styling/SlateTypes.h"
+#include "EditorStyleSet.h"
+#include "Interfaces/IPluginManager.h"
+
+#define IMAGE_BRUSH( RelativePath, ... ) FSlateImageBrush( StyleSet->RootToContentDir( RelativePath, TEXT( ".png" ) ), __VA_ARGS__ )
+#define IMAGE_PLUGIN_BRUSH( RelativePath, ... ) FSlateImageBrush( FMeshEditorStyle::InContent( RelativePath, ".png" ), __VA_ARGS__ )
+#define BOX_BRUSH( RelativePath, ... ) FSlateBoxBrush( StyleSet->RootToContentDir( RelativePath, TEXT( ".png" ) ), __VA_ARGS__ )
+#define TTF_CORE_FONT( RelativePath, ... ) FSlateFontInfo( StyleSet->RootToCoreContentDir( RelativePath, TEXT( ".ttf" ) ), __VA_ARGS__ )
+
+#define DEFAULT_FONT(...) FCoreStyle::GetDefaultFontStyle(__VA_ARGS__)
+
+TSharedPtr<FSlateStyleSet> FStateTreeEditorStyle::StyleSet = nullptr;
+
+FString FStateTreeEditorStyle::InContent(const FString& RelativePath, const ANSICHAR* Extension)
+{
+	static FString ContentDir = IPluginManager::Get().FindPlugin(TEXT("StateTreeEditorModule"))->GetContentDir() / TEXT("Slate");
+	return (ContentDir / RelativePath) + Extension;
+}
+
+void FStateTreeEditorStyle::Initialize()
+{
+	if (StyleSet.IsValid())
+	{
+		return;
+	}
+
+	const FVector2D Icon8x8(8.0f, 8.0f);
+
+	StyleSet = MakeShared<FSlateStyleSet>(GetStyleSetName());
+
+	StyleSet->SetContentRoot(FPaths::EngineContentDir() / TEXT("Editor/Slate"));
+	StyleSet->SetCoreContentRoot(FPaths::EngineContentDir() / TEXT("Slate"));
+
+	const FScrollBarStyle ScrollBar = FEditorStyle::GetWidgetStyle<FScrollBarStyle>("ScrollBar");
+	const FTextBlockStyle& NormalText = FEditorStyle::Get().GetWidgetStyle<FTextBlockStyle>("NormalText");
+
+	// State
+	{
+		FTextBlockStyle StateIcon = FTextBlockStyle(NormalText)
+			.SetFont(FEditorStyle::Get().GetFontStyle("FontAwesome.12"))
+			.SetColorAndOpacity(FLinearColor(230.0f / 255.0f, 230.0f / 255.0f, 230.0f / 255.0f, 0.5f));
+		StyleSet->Set("StateTree.Icon", StateIcon);
+
+		FTextBlockStyle StateDetailsIcon = FTextBlockStyle(NormalText)
+			.SetFont(FEditorStyle::Get().GetFontStyle("FontAwesome.10"))
+		    .SetColorAndOpacity(FLinearColor(230.0f / 255.0f, 230.0f / 255.0f, 230.0f / 255.0f, 0.5f));
+		StyleSet->Set("StateTree.DetailsIcon", StateDetailsIcon);
+
+		FTextBlockStyle StateTitle = FTextBlockStyle(NormalText)
+			.SetFont(TTF_CORE_FONT("Fonts/Roboto-Bold", 12))
+			.SetColorAndOpacity(FLinearColor(230.0f / 255.0f, 230.0f / 255.0f, 230.0f / 255.0f, 0.9f));
+		StyleSet->Set("StateTree.State.Title", StateTitle);
+
+		FEditableTextBoxStyle StateTitleEditableText = FEditableTextBoxStyle()
+			.SetFont(TTF_CORE_FONT("Fonts/Roboto-Bold", 10))
+			.SetBackgroundImageNormal(BOX_BRUSH("Common/TextBox", FMargin(4.0f / 16.0f)))
+			.SetBackgroundImageHovered(BOX_BRUSH("Common/TextBox_Hovered", FMargin(4.0f / 16.0f)))
+			.SetBackgroundImageFocused(BOX_BRUSH("Common/TextBox_Hovered", FMargin(4.0f / 16.0f)))
+			.SetBackgroundImageReadOnly(BOX_BRUSH("Common/TextBox_ReadOnly", FMargin(4.0f / 16.0f)))
+			.SetScrollBarStyle(ScrollBar);
+		StyleSet->Set("StateTree.State.TitleEditableText", StateTitleEditableText);
+
+		StyleSet->Set("StateTree.State.TitleInlineEditableText", FInlineEditableTextBlockStyle()
+			.SetTextStyle(StateTitle)
+			.SetEditableTextBoxStyle(StateTitleEditableText));
+	}
+
+	// Details
+	{
+		FTextBlockStyle StateTitle = FTextBlockStyle(NormalText)
+			.SetFont(TTF_CORE_FONT("Fonts/Roboto-Regular", 10))
+			.SetColorAndOpacity(FLinearColor(230.0f / 255.0f, 230.0f / 255.0f, 230.0f / 255.0f, 0.75f));
+		StyleSet->Set("StateTree.Details", StateTitle);
+	}
+
+	// Details rich text
+	{
+		StyleSet->Set("Details.Normal", FTextBlockStyle(NormalText)
+			.SetFont(FEditorStyle::GetFontStyle(TEXT("PropertyWindow.NormalFont"))));
+
+		StyleSet->Set("Details.Bold", FTextBlockStyle(NormalText)
+			.SetFont(FEditorStyle::GetFontStyle(TEXT("PropertyWindow.BoldFont"))));
+
+		StyleSet->Set("Details.Subdued", FTextBlockStyle(NormalText)
+			.SetColorAndOpacity(FSlateColor::UseSubduedForeground())
+			.SetFont(FEditorStyle::GetFontStyle(TEXT("PropertyWindow.NormalFont"))));
+	}
+
+	const FLinearColor SelectionColor = FColor(0, 0, 0, 32);
+	const FTableRowStyle& NormalTableRowStyle = FEditorStyle::Get().GetWidgetStyle<FTableRowStyle>("TableView.Row");
+	StyleSet->Set("StateTree.Selection",
+		FTableRowStyle(NormalTableRowStyle)
+		.SetActiveBrush(IMAGE_BRUSH("Common/Selection", Icon8x8, SelectionColor))
+		.SetActiveHoveredBrush(IMAGE_BRUSH("Common/Selection", Icon8x8, SelectionColor))
+		.SetInactiveBrush(IMAGE_BRUSH("Common/Selection", Icon8x8, SelectionColor))
+		.SetInactiveHoveredBrush(IMAGE_BRUSH("Common/Selection", Icon8x8, SelectionColor))
+		.SetSelectorFocusedBrush(IMAGE_BRUSH("Common/Selection", Icon8x8, SelectionColor))
+	);
+
+	FSlateStyleRegistry::RegisterSlateStyle(*StyleSet.Get());
+}
+
+
+void FStateTreeEditorStyle::Shutdown()
+{
+	if (StyleSet.IsValid())
+	{
+		FSlateStyleRegistry::UnRegisterSlateStyle(*StyleSet.Get());
+		ensure(StyleSet.IsUnique());
+		StyleSet.Reset();
+	}
+}
+
+
+FName FStateTreeEditorStyle::GetStyleSetName()
+{
+	static FName StyleName("StateTreeEditorStyle");
+	return StyleName;
+}
diff --git a/Engine/Plugins/Runtime/StateTree/Source/StateTreeEditorModule/Private/StateTreeFactory.h b/Engine/Plugins/Runtime/StateTree/Source/StateTreeEditorModule/Private/StateTreeFactory.h
new file mode 100644
--- /dev/null
+++ b/Engine/Plugins/Runtime/StateTree/Source/StateTreeEditorModule/Private/StateTreeFactory.h
@@ -0,0 +1,23 @@
+// Copyright Epic Games, Inc. All Rights Reserved.
+
+#pragma once
+
+#include "CoreMinimal.h"
+#include "UObject/ObjectMacros.h"
+#include "Factories/Factory.h"
+#include "StateTreeFactory.generated.h"
+
+/**
+ * Factory for UStateTree
+ */
+
+UCLASS()
+class STATETREEEDITORMODULE_API UStateTreeFactory : public UFactory
+{
+	GENERATED_UCLASS_BODY()
+
+	// UFactory interface
+	virtual uint32 GetMenuCategories() const override;
+	virtual UObject* FactoryCreateNew(UClass* Class, UObject* InParent, FName Name, EObjectFlags Flags, UObject* Context, FFeedbackContext* Warn) override;
+	// End of UFactory interface
+};
diff --git a/Engine/Plugins/Runtime/StateTree/Source/StateTreeEditorModule/Private/StateTreePropertyBindingCompiler.cpp b/Engine/Plugins/Runtime/StateTree/Source/StateTreeEditorModule/Private/StateTreePropertyBindingCompiler.cpp
new file mode 100644
--- /dev/null
+++ b/Engine/Plugins/Runtime/StateTree/Source/StateTreeEditorModule/Private/StateTreePropertyBindingCompiler.cpp
@@ -0,0 +1,511 @@
+// Copyright Epic Games, Inc. All Rights Reserved.
+#include "StateTreePropertyBindingCompiler.h"
+#include "IPropertyAccessEditor.h"
+#include "PropertyPathHelpers.h"
+#include "StateTreeTypes.h"
+#include "StateTreePropertyBindings.h"
+
+bool FStateTreePropertyBindingCompiler::Init(FStateTreePropertyBindings& InPropertyBindings)
+{
+	PropertyBindings = &InPropertyBindings;
+	PropertyBindings->Reset();
+	SourceStructs.Reset();
+	return true;
+}
+
+bool FStateTreePropertyBindingCompiler::CompileBatch(const FStateTreeBindableStructDesc& TargetStruct, TConstArrayView<FStateTreeEditorPropertyBinding> EditorPropertyBindings, int32& OutBatchIndex)
+{
+	check(PropertyBindings);
+	OutBatchIndex = INDEX_NONE;
+
+	StoreSourceStructs();
+	
+	int32 BindingsBegin = PropertyBindings->PropertyBindings.Num();
+
+	for (const FStateTreeEditorPropertyBinding& EditorBinding : EditorPropertyBindings)
+	{
+		if (EditorBinding.TargetPath.StructID != TargetStruct.ID)
+		{
+			continue;
+		}
+		// Source must be in the source array
+		const FGuid SourceStructID = EditorBinding.SourcePath.StructID;
+		const int32 SourceStructIdx = SourceStructs.IndexOfByPredicate([SourceStructID](const FStateTreeBindableStructDesc& Struct)
+			{
+				return (Struct.ID == SourceStructID);
+			});
+		if (SourceStructIdx == INDEX_NONE)
+		{
+			UE_LOG(LogStateTree, Error, TEXT("%s: Could not find a binding source for '%s'."), ANSI_TO_TCHAR(__FUNCTION__), *TargetStruct.Name.ToString());
+			return false;
+		}
+		const FStateTreeBindableStructDesc& SourceStruct = SourceStructs[SourceStructIdx];
+
+		FStateTreePropertyBinding& NewBinding = PropertyBindings->PropertyBindings.AddDefaulted_GetRef();
+
+		// Resolve paths
+		FResolvedPathResult SourceResult;
+		if (!ResolvePropertyPath(SourceStruct, EditorBinding.SourcePath, SourceResult))
+		{
+			UE_LOG(LogStateTree, Error, TEXT("%s: Could not resolve path to '%s:%s'."), ANSI_TO_TCHAR(__FUNCTION__), *SourceStruct.Name.ToString(), *EditorBinding.SourcePath.ToString());
+			return false;
+		}
+
+		// Destination container is set to 0, it is assumed to be passed in when doing the batch copy.
+		FResolvedPathResult TargetResult;
+		if (!ResolvePropertyPath(TargetStruct, EditorBinding.TargetPath, TargetResult))
+		{
+			UE_LOG(LogStateTree, Error, TEXT("%s: Could not resolve path to '%s:%s'."), ANSI_TO_TCHAR(__FUNCTION__), *TargetStruct.Name.ToString(), *EditorBinding.TargetPath.ToString());
+			return false;
+		}
+
+		NewBinding.SourcePathIndex = SourceResult.PathIndex;
+		NewBinding.TargetPathIndex = TargetResult.PathIndex;
+		NewBinding.SourceStructIndex = SourceStructIdx;
+		NewBinding.CopyType = GetCopyType(SourceResult.LeafProperty, SourceResult.LeafArrayIndex, TargetResult.LeafProperty, TargetResult.LeafArrayIndex);
+	}
+
+	const int32 BindingsEnd = PropertyBindings->PropertyBindings.Num();
+	if (BindingsBegin != BindingsEnd)
+	{
+		FStateTreePropCopyBatch& Batch = PropertyBindings->CopyBatches.AddDefaulted_GetRef();
+		Batch.TargetStruct = TargetStruct;
+		Batch.BindingsBegin = BindingsBegin;
+		Batch.BindingsEnd = BindingsEnd;
+		OutBatchIndex = PropertyBindings->CopyBatches.Num() - 1;
+	}
+
+	return true;
+}
+
+void FStateTreePropertyBindingCompiler::Finalize()
+{
+	StoreSourceStructs();
+}
+
+int32 FStateTreePropertyBindingCompiler::AddSourceStruct(const FStateTreeBindableStructDesc& SourceStruct)
+{
+	SourceStructs.Add(SourceStruct);
+	return SourceStructs.Num() - 1;
+}
+
+int32 FStateTreePropertyBindingCompiler::GetSourceStructIndexByID(const FGuid& ID) const
+{
+	return SourceStructs.IndexOfByPredicate([ID](const FStateTreeBindableStructDesc& Structs) { return (Structs.ID == ID); });
+}
+
+EStateTreePropertyCopyType FStateTreePropertyBindingCompiler::GetCopyType(FProperty* SourceProperty, int32 SourceArrayIndex, FProperty* TargetProperty, int32 TargetArrayIndex)
+{
+	if (FArrayProperty* SourceArrayProperty = CastField<FArrayProperty>(SourceProperty))
+	{
+		// use the array's inner property if we are not trying to copy the whole array
+		if (SourceArrayIndex != INDEX_NONE)
+		{
+			SourceProperty = SourceArrayProperty->Inner;
+		}
+	}
+
+	if (FArrayProperty* TargetArrayProperty = CastField<FArrayProperty>(TargetProperty))
+	{
+		// use the array's inner property if we are not trying to copy the whole array
+		if (TargetArrayIndex != INDEX_NONE)
+		{
+			TargetProperty = TargetArrayProperty->Inner;
+		}
+	}
+
+	EPropertyAccessCompatibility Compatibility = GetPropertyCompatibility(SourceProperty, TargetProperty);
+
+	// Extract underlying types for enums
+	if (const FEnumProperty* EnumPropertyA = CastField<const FEnumProperty>(SourceProperty))
+	{
+		SourceProperty = EnumPropertyA->GetUnderlyingProperty();
+	}
+
+	if (const FEnumProperty* EnumPropertyB = CastField<const FEnumProperty>(TargetProperty))
+	{
+		TargetProperty = EnumPropertyB->GetUnderlyingProperty();
+	}
+
+	if (Compatibility == EPropertyAccessCompatibility::Compatible)
+	{
+		if (CastField<FNameProperty>(TargetProperty))
+		{
+			return EStateTreePropertyCopyType::CopyName;
+		}
+		else if (CastField<FBoolProperty>(TargetProperty))
+		{
+			return EStateTreePropertyCopyType::CopyBool;
+		}
+		else if (CastField<FStructProperty>(TargetProperty))
+		{
+			return EStateTreePropertyCopyType::CopyStruct;
+		}
+		else if (CastField<FObjectPropertyBase>(TargetProperty))
+		{
+			return EStateTreePropertyCopyType::CopyObject;
+		}
+		else if (CastField<FArrayProperty>(TargetProperty) && TargetProperty->HasAnyPropertyFlags(CPF_EditFixedSize))
+		{
+			// only apply array copying rules if the destination array is fixed size, otherwise it will be 'complex'
+			return EStateTreePropertyCopyType::CopyFixedArray;
+		}
+		else if (TargetProperty->PropertyFlags & CPF_IsPlainOldData)
+		{
+			return EStateTreePropertyCopyType::CopyPlain;
+		}
+		else
+		{
+			return EStateTreePropertyCopyType::CopyComplex;
+		}
+	}
+	else if (Compatibility == EPropertyAccessCompatibility::Promotable)
+	{
+		if (SourceProperty->IsA<FBoolProperty>())
+		{
+			if (TargetProperty->IsA<FByteProperty>())
+			{
+				return EStateTreePropertyCopyType::PromoteBoolToByte;
+			}
+			else if (TargetProperty->IsA<FIntProperty>())
+			{
+				return EStateTreePropertyCopyType::PromoteBoolToInt32;
+			}
+			else if (TargetProperty->IsA<FUInt32Property>())
+			{
+				return EStateTreePropertyCopyType::PromoteBoolToUInt32;
+			}
+			else if (TargetProperty->IsA<FInt64Property>())
+			{
+				return EStateTreePropertyCopyType::PromoteBoolToInt64;
+			}
+			else if (TargetProperty->IsA<FFloatProperty>())
+			{
+				return EStateTreePropertyCopyType::PromoteBoolToFloat;
+			}
+			else if (TargetProperty->IsA<FDoubleProperty>())
+			{
+				return EStateTreePropertyCopyType::PromoteBoolToDouble;
+			}
+		}
+		else if (SourceProperty->IsA<FByteProperty>())
+		{
+			if (TargetProperty->IsA<FIntProperty>())
+			{
+				return EStateTreePropertyCopyType::PromoteByteToInt32;
+			}
+			else if (TargetProperty->IsA<FUInt32Property>())
+			{
+				return EStateTreePropertyCopyType::PromoteByteToUInt32;
+			}
+			else if (TargetProperty->IsA<FInt64Property>())
+			{
+				return EStateTreePropertyCopyType::PromoteByteToInt64;
+			}
+			else if (TargetProperty->IsA<FFloatProperty>())
+			{
+				return EStateTreePropertyCopyType::PromoteByteToFloat;
+			}
+			else if (TargetProperty->IsA<FDoubleProperty>())
+			{
+				return EStateTreePropertyCopyType::PromoteByteToDouble;
+			}
+		}
+		else if (SourceProperty->IsA<FIntProperty>())
+		{
+			if (TargetProperty->IsA<FInt64Property>())
+			{
+				return EStateTreePropertyCopyType::PromoteInt32ToInt64;
+			}
+			else if (TargetProperty->IsA<FFloatProperty>())
+			{
+				return EStateTreePropertyCopyType::PromoteInt32ToFloat;
+			}
+			else if (TargetProperty->IsA<FDoubleProperty>())
+			{
+				return EStateTreePropertyCopyType::PromoteInt32ToDouble;
+			}
+		}
+		else if (SourceProperty->IsA<FFloatProperty>())
+		{
+			if (TargetProperty->IsA<FDoubleProperty>())
+			{
+				return EStateTreePropertyCopyType::PromoteFloatToDouble;
+			}
+		}
+		else if (SourceProperty->IsA<FDoubleProperty>())
+		{
+			if (TargetProperty->IsA<FFloatProperty>())
+			{
+				return EStateTreePropertyCopyType::DemoteDoubleToFloat;
+			}
+		}
+		else if (SourceProperty->IsA<FUInt32Property>())
+		{
+			if (TargetProperty->IsA<FInt64Property>())
+			{
+				return EStateTreePropertyCopyType::PromoteUInt32ToInt64;
+			}
+			else if (TargetProperty->IsA<FFloatProperty>())
+			{
+				return EStateTreePropertyCopyType::PromoteUInt32ToFloat;
+			}
+		}
+	}
+
+	ensureMsgf(false, TEXT("Couldnt determine property copy type (%s -> %s)"), *SourceProperty->GetNameCPP(), *TargetProperty->GetNameCPP());
+
+	return EStateTreePropertyCopyType::None;
+}
+
+bool FStateTreePropertyBindingCompiler::ResolvePropertyPath(const FStateTreeBindableStructDesc& InStructDesc, const FStateTreeEditorPropertyPath& InPath, FResolvedPathResult& OutResult)
+{
+	const int32 PathIndex = PropertyBindings->PropertyPaths.Num();
+	FStateTreePropertyPath& Path = PropertyBindings->PropertyPaths.AddDefaulted_GetRef();
+	Path.SegmentsBegin = PropertyBindings->PropertySegments.Num();
+
+	const bool bResult = ResolvePropertyPath(InStructDesc, InPath, PropertyBindings->PropertySegments, OutResult.LeafProperty, OutResult.LeafArrayIndex, /*bLogErrors*/true);
+
+	Path.SegmentsEnd = PropertyBindings->PropertySegments.Num();
+	OutResult.PathIndex = PathIndex;
+
+	return bResult;
+}
+
+bool FStateTreePropertyBindingCompiler::ResolvePropertyPath(const FStateTreeBindableStructDesc& InStructDesc, const FStateTreeEditorPropertyPath& InPath,
+															TArray<FStateTreePropertySegment>& OutSegments, FProperty*& OutLeafProperty, int32& OutLeafArrayIndex,
+															const bool bLogErrors)
+{
+	if (!InPath.IsValid())
+	{
+		UE_CLOG(bLogErrors, LogStateTree, Error, TEXT("%s: Invalid path '%s:%s'."), ANSI_TO_TCHAR(__FUNCTION__), *InStructDesc.Name.ToString(), *InPath.ToString());
+		return false;
+	}
+
+	const UStruct* CurrentStruct = InStructDesc.Struct;
+	FProperty* LeafProperty = nullptr;
+	int32 LeafArrayIndex = INDEX_NONE;
+	bool bResult = InPath.Path.Num() > 0;
+
+	for (int32 SegmentIndex = 0; SegmentIndex < InPath.Path.Num(); SegmentIndex++)
+	{
+		const FString& SegmentString = InPath.Path[SegmentIndex];
+		const TCHAR* PropertyNamePtr = nullptr;
+		int32 PropertyNameLength = 0;
+		int32 ArrayIndex = INDEX_NONE;
+		PropertyPathHelpers::FindFieldNameAndArrayIndex(SegmentString.Len(), *SegmentString, PropertyNameLength, &PropertyNamePtr, ArrayIndex);
+		ensure(PropertyNamePtr != nullptr);
+		FString PropertyNameString(PropertyNameLength, PropertyNamePtr);
+		FName PropertyName = FName(*PropertyNameString, FNAME_Find);
+
+		const bool bFinalSegment = SegmentIndex == (InPath.Path.Num() - 1);
+
+		if (CurrentStruct == nullptr)
+		{
+			UE_CLOG(bLogErrors, LogStateTree, Error, TEXT("%s: Malformed path '%s:%s'."), ANSI_TO_TCHAR(__FUNCTION__), *InStructDesc.Name.ToString(), *InPath.ToString(SegmentIndex, TEXT("<"), TEXT(">")));
+			bResult = false;
+			break;
+		}
+
+		FProperty* Property = CurrentStruct->FindPropertyByName(PropertyName);
+		if (Property == nullptr)
+		{
+			// TODO: use core redirects to fix up the name.
+			UE_CLOG(bLogErrors, LogStateTree, Error, TEXT("%s: Malformed path '%s:%s', could not find property '%s%s.%s'."),
+				ANSI_TO_TCHAR(__FUNCTION__), *InStructDesc.Name.ToString(), *InPath.ToString(SegmentIndex, TEXT("<"), TEXT(">")),
+				CurrentStruct->GetPrefixCPP(), *CurrentStruct->GetName(), *PropertyName.ToString());
+			bResult = false;
+			break;
+		}
+
+		FStateTreePropertySegment& Segment = OutSegments.AddDefaulted_GetRef();
+		Segment.Name = PropertyName;
+		Segment.ArrayIndex = ArrayIndex;
+
+		// Check to see if it is an array access first
+		FArrayProperty* ArrayProperty = CastField<FArrayProperty>(Property);
+		if (ArrayProperty != nullptr && ArrayIndex != INDEX_NONE)
+		{
+			// It is an array, now check to see if this is an array of structures
+			if (FStructProperty* ArrayOfStructsProperty = CastField<FStructProperty>(ArrayProperty->Inner))
+			{
+				Segment.Type = EStateTreePropertyAccessType::IndexArray;
+				CurrentStruct = ArrayOfStructsProperty->Struct;
+			}
+			// if it's not an array of structs, maybe it's an array of objects
+			else if (FObjectPropertyBase* ArrayOfObjectsProperty = CastField<FObjectPropertyBase>(ArrayProperty->Inner))
+			{
+				Segment.Type = EStateTreePropertyAccessType::IndexArray;
+				CurrentStruct = ArrayOfObjectsProperty->PropertyClass;
+
+				if (!bFinalSegment)
+				{
+					// Object arrays need an object dereference adding if non-leaf
+					FStateTreePropertySegment& ExtraSegment = OutSegments.AddDefaulted_GetRef();
+					ExtraSegment.ArrayIndex = 0;
+					ExtraSegment.Type = EStateTreePropertyAccessType::Object;
+					FProperty* InnerProperty = ArrayProperty->Inner;
+					if (FObjectProperty* ObjectProperty = CastField<FObjectProperty>(InnerProperty))
+					{
+						ExtraSegment.Type = EStateTreePropertyAccessType::Object;
+					}
+					else if (FWeakObjectProperty* WeakObjectProperty = CastField<FWeakObjectProperty>(InnerProperty))
+					{
+						ExtraSegment.Type = EStateTreePropertyAccessType::WeakObject;
+					}
+					else if (FSoftObjectProperty* SoftObjectProperty = CastField<FSoftObjectProperty>(InnerProperty))
+					{
+						ExtraSegment.Type = EStateTreePropertyAccessType::SoftObject;
+					}
+				}
+			}
+			else
+			{
+				Segment.Type = EStateTreePropertyAccessType::IndexArray;
+				Segment.ArrayIndex = ArrayIndex;
+				CurrentStruct = nullptr;
+			}
+		}
+		// Leaf segments all get treated the same, plain, array, struct or object. Copy type is figured out separately.
+		else if (bFinalSegment)
+		{
+			Segment.Type = EStateTreePropertyAccessType::Offset;
+			CurrentStruct = nullptr;
+		}
+		// Check to see if this is a simple structure (eg. not an array of structures)
+		else if (FStructProperty* StructProperty = CastField<FStructProperty>(Property))
+		{
+			Segment.Type = EStateTreePropertyAccessType::Offset;
+			CurrentStruct = StructProperty->Struct;
+		}
+		// Check to see if this is a simple object (eg. not an array of objects)
+		else if (FObjectProperty* ObjectProperty = CastField<FObjectProperty>(Property))
+		{
+			Segment.Type = EStateTreePropertyAccessType::Object;
+			CurrentStruct = ObjectProperty->PropertyClass;
+		}
+		// Check to see if this is a simple weak object property (eg. not an array of weak objects).
+		else if (FWeakObjectProperty* WeakObjectProperty = CastField<FWeakObjectProperty>(Property))
+		{
+			Segment.Type = EStateTreePropertyAccessType::WeakObject;
+			CurrentStruct = WeakObjectProperty->PropertyClass;
+		}
+		// Check to see if this is a simple soft object property (eg. not an array of soft objects).
+		else if (FSoftObjectProperty* SoftObjectProperty = CastField<FSoftObjectProperty>(Property))
+		{
+			Segment.Type = EStateTreePropertyAccessType::SoftObject;
+			CurrentStruct = SoftObjectProperty->PropertyClass;
+		}
+		else
+		{
+			UE_CLOG(bLogErrors, LogStateTree, Error, TEXT("%s: Unsupported segment %s in path '%s:%s'."),
+				ANSI_TO_TCHAR(__FUNCTION__), *Property->GetCPPType(), *InStructDesc.Name.ToString(), *InPath.ToString(SegmentIndex, TEXT("<"), TEXT(">")));
+			bResult = false;
+			break;
+		}
+
+		if (bFinalSegment)
+		{
+			LeafProperty = Property;
+			LeafArrayIndex = ArrayIndex;
+		}
+	}
+
+	if (!bResult)
+	{
+		return false;
+	}
+
+	OutLeafProperty = LeafProperty;
+	OutLeafArrayIndex = LeafArrayIndex;
+
+	return true;
+}
+
+EPropertyAccessCompatibility FStateTreePropertyBindingCompiler::GetPropertyCompatibility(const FProperty* InPropertyA, const FProperty* InPropertyB)
+{
+	if (InPropertyA == InPropertyB)
+	{
+		return EPropertyAccessCompatibility::Compatible;
+	}
+
+	if (InPropertyA == nullptr || InPropertyB == nullptr)
+	{
+		return EPropertyAccessCompatibility::Incompatible;
+	}
+
+	// Special case for object properties
+	if (InPropertyA->IsA<FObjectPropertyBase>() && InPropertyB->IsA<FObjectPropertyBase>())
+	{
+		return EPropertyAccessCompatibility::Compatible;
+	}
+
+	// Extract underlying types for enums
+	if (const FEnumProperty* EnumPropertyA = CastField<const FEnumProperty>(InPropertyA))
+	{
+		InPropertyA = EnumPropertyA->GetUnderlyingProperty();
+	}
+
+	if (const FEnumProperty* EnumPropertyB = CastField<const FEnumProperty>(InPropertyB))
+	{
+		InPropertyB = EnumPropertyB->GetUnderlyingProperty();
+	}
+
+	if (InPropertyA->SameType(InPropertyB))
+	{
+		return EPropertyAccessCompatibility::Compatible;
+	}
+	else
+	{
+		// Not directly compatible, check for promotions
+		if (InPropertyA->IsA<FBoolProperty>())
+		{
+			if (InPropertyB->IsA<FByteProperty>() || InPropertyB->IsA<FIntProperty>() || InPropertyB->IsA<FUInt32Property>() || InPropertyB->IsA<FInt64Property>() || InPropertyB->IsA<FFloatProperty>())
+			{
+				return EPropertyAccessCompatibility::Promotable;
+			}
+		}
+		else if (InPropertyA->IsA<FByteProperty>())
+		{
+			if (InPropertyB->IsA<FIntProperty>() || InPropertyB->IsA<FUInt32Property>() || InPropertyB->IsA<FInt64Property>() || InPropertyB->IsA<FFloatProperty>())
+			{
+				return EPropertyAccessCompatibility::Promotable;
+			}
+		}
+		else if (InPropertyA->IsA<FIntProperty>())
+		{
+			if (InPropertyB->IsA<FInt64Property>() || InPropertyB->IsA<FFloatProperty>())
+			{
+				return EPropertyAccessCompatibility::Promotable;
+			}
+		}
+		else if (InPropertyA->IsA<FUInt32Property>())
+		{
+			if (InPropertyB->IsA<FInt64Property>() || InPropertyB->IsA<FFloatProperty>())
+			{
+				return EPropertyAccessCompatibility::Promotable;
+			}
+		}
+	}
+
+	return EPropertyAccessCompatibility::Incompatible;
+}
+
+void FStateTreePropertyBindingCompiler::StoreSourceStructs()
+{
+	// Check that existing structs are compatible
+	check(PropertyBindings->SourceStructs.Num() <= SourceStructs.Num());
+	for (int32 i = 0; i < PropertyBindings->SourceStructs.Num(); i++)
+	{
+		check(PropertyBindings->SourceStructs[i] == SourceStructs[i]);
+	}
+
+	// Add new
+	if (SourceStructs.Num() > PropertyBindings->SourceStructs.Num())
+	{
+		for (int32 i = PropertyBindings->SourceStructs.Num(); i < SourceStructs.Num(); i++)
+		{
+			PropertyBindings->SourceStructs.Add(SourceStructs[i]);
+		}
+	}
+}
diff --git a/Engine/Plugins/Runtime/StateTree/Source/StateTreeEditorModule/Private/StateTreePropertyHelpers.cpp b/Engine/Plugins/Runtime/StateTree/Source/StateTreeEditorModule/Private/StateTreePropertyHelpers.cpp
new file mode 100644
--- /dev/null
+++ b/Engine/Plugins/Runtime/StateTree/Source/StateTreeEditorModule/Private/StateTreePropertyHelpers.cpp
@@ -0,0 +1,35 @@
+// Copyright Epic Games, Inc. All Rights Reserved.
+
+#include "StateTreePropertyHelpers.h"
+#include "CoreMinimal.h"
+#include "Misc/Guid.h"
+#include "PropertyEditorModule.h"
+#include "PropertyHandle.h"
+#include "ScopedTransaction.h"
+
+#define LOCTEXT_NAMESPACE "StateTreeEditor"
+
+namespace UE { namespace StateTree { namespace PropertyHelpers {
+
+const FSlateBrush* GetTypeIcon(EStateTreeVariableType Type)
+{
+	switch (Type)
+	{
+	case EStateTreeVariableType::Bool:
+		return FEditorStyle::GetBrush(TEXT("ClassIcon.BlackboardKeyType_Bool"));
+	case EStateTreeVariableType::Float:
+		return FEditorStyle::GetBrush(TEXT("ClassIcon.BlackboardKeyType_Float"));
+	case EStateTreeVariableType::Int:
+		return FEditorStyle::GetBrush(TEXT("ClassIcon.BlackboardKeyType_Int"));
+	case EStateTreeVariableType::Vector:
+		return FEditorStyle::GetBrush(TEXT("ClassIcon.BlackboardKeyType_Vector"));
+	case EStateTreeVariableType::Object:
+		return FEditorStyle::GetBrush(TEXT("ClassIcon.BlackboardKeyType_Object"));
+	default:
+		return FEditorStyle::GetBrush(TEXT("NoBrush"));
+	}
+}
+
+} } } // UE::StateTree::PropertyHelpers
+
+#undef LOCTEXT_NAMESPACE
diff --git a/Engine/Plugins/Runtime/StateTree/Source/StateTreeEditorModule/Private/StateTreeState.cpp b/Engine/Plugins/Runtime/StateTree/Source/StateTreeEditorModule/Private/StateTreeState.cpp
new file mode 100644
--- /dev/null
+++ b/Engine/Plugins/Runtime/StateTree/Source/StateTreeEditorModule/Private/StateTreeState.cpp
@@ -0,0 +1,275 @@
+// Copyright Epic Games, Inc. All Rights Reserved.
+
+#include "StateTreeState.h"
+#include "StateTreeEditorData.h"
+#include "StateTree.h"
+#include "StateTreeConditionBase.h"
+#include "StateTreeDelegates.h"
+#include "CoreMinimal.h"
+#include "UObject/Field.h"
+
+void FStateTreeStateLink::Set(const EStateTreeTransitionType InType, const UStateTreeState* InState)
+{
+	Type = InType;
+	if (Type == EStateTreeTransitionType::GotoState)
+	{
+		check(InState);
+		Name = InState->Name;
+		ID = InState->ID;
+	}
+}
+
+
+FStateTreeTransition2::FStateTreeTransition2(const EStateTreeTransitionEvent InEvent, const EStateTreeTransitionType InType, const UStateTreeState* InState)
+{
+	Event = InEvent;
+	State.Set(InType, InState);
+}
+
+
+UStateTreeState::UStateTreeState(const FObjectInitializer& ObjectInitializer)
+	: Super(ObjectInitializer)
+	, ID(FGuid::NewGuid())
+	, Parent(nullptr)
+{
+}
+
+#if WITH_EDITOR
+void UStateTreeState::PostEditChangeChainProperty(FPropertyChangedChainEvent& PropertyChangedEvent)
+{
+	Super::PostEditChangeChainProperty(PropertyChangedEvent);
+
+	FProperty* Property = PropertyChangedEvent.Property;
+	FProperty* MemberProperty = nullptr;
+	if (PropertyChangedEvent.PropertyChain.GetActiveMemberNode())
+	{
+		MemberProperty = PropertyChangedEvent.PropertyChain.GetActiveMemberNode()->GetValue();
+	}
+
+	if (Property)
+	{
+		if (Property->GetOwnerClass() == UStateTreeState::StaticClass() &&
+			Property->GetFName() == GET_MEMBER_NAME_CHECKED(UStateTreeState, Name))
+		{
+			UStateTree* StateTree = GetTypedOuter<UStateTree>();
+			if (ensure(StateTree))
+			{
+				UE::StateTree::Delegates::OnIdentifierChanged.Broadcast(*StateTree);
+			}
+		}
+
+		if (PropertyChangedEvent.ChangeType == EPropertyChangeType::Duplicate && MemberProperty)
+		{
+			// Ensure unique ID on duplicated items.
+			if (MemberProperty->GetFName() == GET_MEMBER_NAME_CHECKED(UStateTreeState, Evaluators))
+			{
+				const int32 ArrayIndex = PropertyChangedEvent.GetArrayIndex(MemberProperty->GetFName().ToString());
+				if (Evaluators.IsValidIndex(ArrayIndex))
+				{
+					if (UStateTreeEvaluatorBase* Eval = Evaluators[ArrayIndex])
+					{
+						Eval->SetNewUniqueID();
+						Eval->SetName(FName(Eval->GetName().ToString() + TEXT(" Duplicate")));
+					}
+				}
+			}
+			if (MemberProperty->GetFName() == GET_MEMBER_NAME_CHECKED(UStateTreeState, Tasks))
+			{
+				const int32 ArrayIndex = PropertyChangedEvent.GetArrayIndex(MemberProperty->GetFName().ToString());
+				if (Tasks.IsValidIndex(ArrayIndex))
+				{
+					if (UStateTreeTaskBase* Task = Tasks[ArrayIndex])
+					{
+						Task->SetNewUniqueID();
+						Task->Name = FName(Task->Name.ToString() + TEXT(" Duplicate"));
+					}
+				}
+			}
+		}
+
+
+		if (IsV2())
+		{
+			if (MemberProperty)
+			{
+				// Ensure unique ID on duplicated items.
+				if (PropertyChangedEvent.ChangeType == EPropertyChangeType::Duplicate)
+				{
+					if (MemberProperty->GetFName() == GET_MEMBER_NAME_CHECKED(UStateTreeState, Evaluators2))
+					{
+						const int32 ArrayIndex = PropertyChangedEvent.GetArrayIndex(MemberProperty->GetFName().ToString());
+						if (Evaluators2.IsValidIndex(ArrayIndex))
+						{
+							if (FStateTreeEvaluator2Base* Eval = Evaluators2[ArrayIndex].Type.GetMutablePtr<FStateTreeEvaluator2Base>())
+							{
+								Eval->ID = FGuid::NewGuid();
+								Eval->Name = FName(Eval->Name.ToString() + TEXT(" Duplicate"));
+							}
+						}
+					}
+					if (MemberProperty->GetFName() == GET_MEMBER_NAME_CHECKED(UStateTreeState, Tasks2))
+					{
+						const int32 ArrayIndex = PropertyChangedEvent.GetArrayIndex(MemberProperty->GetFName().ToString());
+						if (Tasks2.IsValidIndex(ArrayIndex))
+						{
+							if (FStateTreeTask2Base* Task = Tasks2[ArrayIndex].Type.GetMutablePtr<FStateTreeTask2Base>())
+							{
+								Task->ID = FGuid::NewGuid();
+								Task->Name = FName(Task->Name.ToString() + TEXT(" Duplicate"));
+							}
+						}
+					}
+					if (MemberProperty->GetFName() == GET_MEMBER_NAME_CHECKED(UStateTreeState, EnterConditions2))
+					{
+						const int32 ArrayIndex = PropertyChangedEvent.GetArrayIndex(MemberProperty->GetFName().ToString());
+						if (Tasks2.IsValidIndex(ArrayIndex))
+						{
+							if (FStateTreeConditionBase* Cond = EnterConditions2[ArrayIndex].Type.GetMutablePtr<FStateTreeConditionBase>())
+							{
+								Cond->ID = FGuid::NewGuid();
+							}
+						}
+					}
+				}
+			}
+		}
+	}
+}
+
+void UStateTreeState::PostLoad()
+{
+	Super::PostLoad();
+	
+	UStateTree* StateTree = GetTypedOuter<UStateTree>();
+	if (StateTree->IsV2())
+	{
+		// Convert existing default transitions.
+		Modify();
+
+		if (StateFailedTransition.Type != EStateTreeTransitionType::NotSet
+			&& StateFailedTransition.Type != EStateTreeTransitionType::SelectChildState)
+		{
+			if (StateFailedTransition.IsValid())
+			{
+				FStateTreeTransition2& Transition = Transitions2.AddDefaulted_GetRef();
+				Transition.Event = EStateTreeTransitionEvent::OnFailed;
+				Transition.State = StateFailedTransition;
+			}
+			// Reset
+			StateFailedTransition = FStateTreeStateLink(EStateTreeTransitionType::NotSet);
+		}
+
+		if (StateDoneTransition.Type != EStateTreeTransitionType::NotSet
+			&& StateDoneTransition.Type != EStateTreeTransitionType::SelectChildState) // Select child state is inferred automatically in V2
+		{
+			if (StateDoneTransition.IsValid())
+			{
+				FStateTreeTransition2& Transition = Transitions2.AddDefaulted_GetRef();
+				Transition.Event = EStateTreeTransitionEvent::OnCompleted;
+				Transition.State = StateDoneTransition;
+			}
+			// Reset
+			StateDoneTransition = FStateTreeStateLink(EStateTreeTransitionType::NotSet);
+		}
+	}
+}
+
+#endif
+
+void UStateTreeState::GetVisibleVariables(FStateTreeVariableLayout& Variables) const
+{
+	// Add input parameters
+	UStateTree* StateTree = GetTypedOuter<UStateTree>();
+	if (ensure(StateTree))
+	{
+		const FStateTreeVariableLayout& ParameterLayout = StateTree->GetInputParameterLayout();
+		for (const FStateTreeVariableDesc& Var : ParameterLayout.Variables)
+		{
+			Variables.DefineVariable(Var);
+		}
+	}
+
+	// Collect states variable descs up to the highest parent
+	for (const UStateTreeState* State = this; State; State = State->Parent)
+	{
+		for (const UStateTreeEvaluatorBase* Evaluator : State->Evaluators)
+		{
+			if (Evaluator)
+			{
+				Evaluator->DefineOutputVariables(Variables);
+			}
+		}
+	}
+}
+
+const UStateTreeTaskBase* UStateTreeState::GetTaskByID(FGuid InID) const
+{
+	const int32 Index = Tasks.IndexOfByPredicate([InID](const UStateTreeTaskBase* Task) -> bool { return Task->ID == InID; });
+	return Index != INDEX_NONE ? Tasks[Index] : nullptr;
+}
+
+UStateTreeTaskBase* UStateTreeState::GetTaskByID(FGuid InID)
+{
+	const int32 Index = Tasks.IndexOfByPredicate([InID](const UStateTreeTaskBase* Task) -> bool { return Task->ID == InID; });
+	return Index != INDEX_NONE ? Tasks[Index] : nullptr;
+}
+
+const FStateTreeTask2Base* UStateTreeState::GetTask2ByID(FGuid InID) const
+{
+	for (const FStateTreeTaskItem& TaskItem : Tasks2)
+	{
+		if (const FStateTreeTask2Base* Task = TaskItem.Type.GetPtr<FStateTreeTask2Base>())
+		{
+			if (Task->ID == InID)
+			{
+				return Task;
+			}
+		}
+	}
+	return nullptr;
+}
+
+FStateTreeTask2Base* UStateTreeState::GetTask2ByID(FGuid InID)
+{
+	for (FStateTreeTaskItem& TaskItem : Tasks2)
+	{
+		if (FStateTreeTask2Base* Task = TaskItem.Type.GetMutablePtr<FStateTreeTask2Base>())
+		{
+			if (Task->ID == InID)
+			{
+				return Task;
+			}
+		}
+	}
+	return nullptr;
+}
+
+UStateTreeState* UStateTreeState::GetNextSiblingState() const
+{
+	if (!Parent)
+	{
+		return nullptr;
+	}
+	for (int32 ChildIdx = 0; ChildIdx < Parent->Children.Num(); ChildIdx++)
+	{
+		if (Parent->Children[ChildIdx] == this)
+		{
+			const int NextIdx = ChildIdx + 1;
+			if (NextIdx < Parent->Children.Num())
+			{
+				return Parent->Children[NextIdx];
+			}
+			break;
+		}
+	}
+	return nullptr;
+}
+
+bool UStateTreeState::IsV2() const
+{
+	if (const UStateTree* StateTree = GetTypedOuter<UStateTree>())
+	{
+		return StateTree->IsV2();
+	}
+	return false;
+}
\ No newline at end of file
diff --git a/Engine/Plugins/Runtime/StateTree/Source/StateTreeEditorModule/Private/StateTreeVariableProvider.cpp b/Engine/Plugins/Runtime/StateTree/Source/StateTreeEditorModule/Private/StateTreeVariableProvider.cpp
new file mode 100644
--- /dev/null
+++ b/Engine/Plugins/Runtime/StateTree/Source/StateTreeEditorModule/Private/StateTreeVariableProvider.cpp
@@ -0,0 +1,13 @@
+// Copyright Epic Games, Inc. All Rights Reserved.
+
+#include "StateTreeVariableProvider.h"
+
+
+//----------------------------------------------------------------------//
+// UStateTreeVariableProvider
+//----------------------------------------------------------------------//
+
+UStateTreeVariableProvider::UStateTreeVariableProvider(const FObjectInitializer& ObjectInitializer)
+	: Super(ObjectInitializer)
+{
+}
diff --git a/Engine/Plugins/Runtime/StateTree/Source/StateTreeEditorModule/Private/StateTreeViewModel.cpp b/Engine/Plugins/Runtime/StateTree/Source/StateTreeEditorModule/Private/StateTreeViewModel.cpp
new file mode 100644
--- /dev/null
+++ b/Engine/Plugins/Runtime/StateTree/Source/StateTreeEditorModule/Private/StateTreeViewModel.cpp
@@ -0,0 +1,534 @@
+// Copyright Epic Games, Inc. All Rights Reserved.
+
+#include "StateTreeViewModel.h"
+#include "Templates/SharedPointer.h"
+#include "StateTree.h"
+#include "StateTreeEditorData.h"
+#include "StateTreeState.h"
+#include "StateTreeTaskBase.h"
+#include "StateTreeDelegates.h"
+#include "Editor.h"
+#include "EditorSupportDelegates.h"
+#include "ScopedTransaction.h"
+#include "Modules/ModuleManager.h"
+#include "AIGraphTypes.h" // Class cache
+#include "StateTreeEditorModule.h"
+#include "Framework/MultiBox/MultiBoxBuilder.h"
+
+
+#define LOCTEXT_NAMESPACE "StateTreeEditor"
+
+
+namespace FStateTreeViewUtilities
+{
+	// Remove state from the array. Recurses into children if no match found.
+	bool RemoveRecursive(TArray<UStateTreeState*>& Array, UStateTreeState* StateToRemove)
+	{
+		if (!StateToRemove)
+		{
+			return false;
+		}
+
+		int32 ItemIndex = Array.Find(StateToRemove);
+		if (ItemIndex != INDEX_NONE)
+		{
+			Array.RemoveAt(ItemIndex);
+			StateToRemove->Parent = nullptr;
+			return true;
+		}
+
+		// Did not successfully remove an item. Try all the children.
+		for (int32 i = 0; i < Array.Num(); ++i)
+		{
+			if (RemoveRecursive(Array[i]->Children, StateToRemove))
+			{
+				return true;
+			}
+		}
+
+		return false;
+	}
+
+	// Insert state in the array relative to target state. Recurses into children if no match found.
+	bool InsertRecursive(UStateTreeState* ParentState, TArray<UStateTreeState*>& ParentArray, UStateTreeState* TargetState, UStateTreeState* StateToInsert, int32 RelativeLocation)
+	{
+		if (!StateToInsert || !TargetState)
+		{
+			return false;
+		}
+
+		const int32 TargetIndex = ParentArray.Find(TargetState);
+		if (TargetIndex != INDEX_NONE)
+		{
+			if (RelativeLocation == -1)
+			{
+				ParentArray.Insert(StateToInsert, TargetIndex);
+				StateToInsert->Parent = ParentState;
+			}
+			else if (RelativeLocation == 1)
+			{
+				ParentArray.Insert(StateToInsert, TargetIndex + 1);
+				StateToInsert->Parent = ParentState;
+			}
+			else
+			{
+				ensure(RelativeLocation == 0);
+				ParentArray[TargetIndex]->Children.Insert(StateToInsert, 0);
+				StateToInsert->Parent = ParentArray[TargetIndex];
+			}
+			return true;
+		}
+
+		for (int32 i = 0; i < ParentArray.Num(); ++i)
+		{
+			if (InsertRecursive(ParentArray[i], ParentArray[i]->Children, TargetState, StateToInsert, RelativeLocation))
+			{
+				return true;
+			}
+		}
+
+		return false;
+	}
+
+	// Removes states from the array which are children of any other state.
+	void RemoveContainedChildren(TArray<UStateTreeState*>& States)
+	{
+		TSet<UStateTreeState*> UniqueStates;
+		for (UStateTreeState* State : States)
+		{
+			UniqueStates.Add(State);
+		}
+
+		for (int32 i = 0; i < States.Num(); )
+		{
+			UStateTreeState* State = States[i];
+
+			// Walk up the parent state sand if the current state
+			// exists in any of them, remove it.
+			UStateTreeState* StateParent = State->Parent;
+			bool bShouldRemove = false;
+			while (StateParent)
+			{
+				if (UniqueStates.Contains(StateParent))
+				{
+					bShouldRemove = true;
+					break;
+				}
+				StateParent = StateParent->Parent;
+			}
+
+			if (bShouldRemove)
+			{
+				States.RemoveAt(i);
+			}
+			else
+			{
+				i++;
+			}
+		}
+	}
+
+	// Returns true if the state is child of parent state.
+	bool IsChildOf(const UStateTreeState* ParentState, const UStateTreeState* State)
+	{
+		for (const UStateTreeState* Child : ParentState->Children)
+		{
+			if (Child == State)
+			{
+				return true;
+			}
+			if (IsChildOf(Child, State))
+			{
+				return true;
+			}
+		}
+		return false;
+	}
+
+};
+
+
+FStateTreeViewModel::FStateTreeViewModel()
+	: TreeData(nullptr)
+{
+}
+
+FStateTreeViewModel::~FStateTreeViewModel()
+{
+	GEditor->UnregisterForUndo(this);
+
+	UE::StateTree::Delegates::OnIdentifierChanged.RemoveAll(this);
+	UE::StateTree::Delegates::OnParameterLayoutChanged.RemoveAll(this);
+}
+
+void FStateTreeViewModel::Init(UStateTreeEditorData* InTreeData)
+{
+	TreeData = InTreeData;
+
+	GEditor->RegisterForUndo(this);
+
+	UE::StateTree::Delegates::OnIdentifierChanged.AddSP(this, &FStateTreeViewModel::HandleIdentifierChanged);
+}
+
+void FStateTreeViewModel::HandleIdentifierChanged(const UStateTree& StateTree)
+{
+	const UStateTree* OuterStateTree = TreeData ? Cast<UStateTree>(TreeData->GetOuter()) : nullptr;
+	if (OuterStateTree == &StateTree)
+	{
+		OnAssetChanged.Broadcast();
+	}
+}
+
+void FStateTreeViewModel::NotifyAssetChangedExternally()
+{
+	OnAssetChanged.Broadcast();
+}
+
+void FStateTreeViewModel::NotifyStatesChangedExternally(const TSet<UStateTreeState*>& ChangedStates, const FPropertyChangedEvent& PropertyChangedEvent)
+{
+	OnStatesChanged.Broadcast(ChangedStates, PropertyChangedEvent);
+}
+
+TArray<UStateTreeState*>* FStateTreeViewModel::GetRoutines()
+{
+	return TreeData ? &TreeData->Routines : nullptr;
+}
+
+void FStateTreeViewModel::PostUndo(bool bSuccess)
+{
+	// TODO: see if we can narrow this down.
+	OnAssetChanged.Broadcast();
+}
+
+void FStateTreeViewModel::PostRedo(bool bSuccess)
+{
+	OnAssetChanged.Broadcast();
+}
+
+void FStateTreeViewModel::ClearSelection()
+{
+	SelectedStates.Reset();
+
+	TArray<UStateTreeState*> SelectedStatesArr;
+	OnSelectionChanged.Broadcast(SelectedStatesArr);
+}
+
+void FStateTreeViewModel::SetSelection(UStateTreeState* SelectedState)
+{
+	SelectedStates.Reset();
+
+	SelectedStates.Add(SelectedState);
+
+	TArray<UStateTreeState*> SelectedStatesArr;
+	SelectedStatesArr.Add(SelectedState);
+	OnSelectionChanged.Broadcast(SelectedStatesArr);
+}
+
+void FStateTreeViewModel::SetSelection(const TArray<UStateTreeState*>& InSelectedStates)
+{
+	SelectedStates.Reset();
+
+	for (UStateTreeState* State : InSelectedStates)
+	{
+		if (State)
+		{
+			SelectedStates.Add(State);
+		}
+	}
+
+	TArray<FGuid> SelectedTaskIDArr;
+	OnSelectionChanged.Broadcast(InSelectedStates);
+}
+
+bool FStateTreeViewModel::IsSelected(const UStateTreeState* State) const
+{
+	return SelectedStates.Contains(State);
+}
+
+bool FStateTreeViewModel::IsChildOfSelection(const UStateTreeState* State) const
+{
+	for (const FWeakObjectPtr& WeakSelectedState : SelectedStates)
+	{
+		if (const UStateTreeState* SelectedState = Cast<UStateTreeState>(WeakSelectedState.Get()))
+		{
+			if (SelectedState == State)
+			{
+				return true;
+			}
+			else
+			{
+				if (FStateTreeViewUtilities::IsChildOf(SelectedState, State))
+				{
+					return true;
+				}
+			}
+		}
+	}
+	return false;
+}
+
+void FStateTreeViewModel::GetSelectedStates(TArray<UStateTreeState*>& OutSelectedStates)
+{
+	OutSelectedStates.Reset();
+	for (FWeakObjectPtr& WeakState : SelectedStates)
+	{
+		if (UStateTreeState* State = Cast<UStateTreeState>(WeakState.Get()))
+		{
+			OutSelectedStates.Add(State);
+		}
+	}
+}
+
+void FStateTreeViewModel::GetPersistentExpandedStates(TSet<UStateTreeState*>& OutExpandedStates)
+{
+	OutExpandedStates.Reset();
+	if (!TreeData)
+	{
+		return;
+	}
+	for (UStateTreeState* Routine : TreeData->Routines)
+	{
+		GetExpandedStatesRecursive(Routine, OutExpandedStates);
+	}
+}
+
+void FStateTreeViewModel::GetExpandedStatesRecursive(UStateTreeState* State, TSet<UStateTreeState*>& OutExpandedStates)
+{
+	if (State->bExpanded)
+	{
+		OutExpandedStates.Add(State);
+	}
+	for (UStateTreeState* Child : State->Children)
+	{
+		GetExpandedStatesRecursive(Child, OutExpandedStates);
+	}
+}
+
+void FStateTreeViewModel::SetPersistentExpandedStates(TSet<UStateTreeState*>& InExpandedStates)
+{
+	if (!TreeData)
+	{
+		return;
+	}
+
+	TreeData->Modify();
+
+	for (UStateTreeState* State : InExpandedStates)
+	{
+		if (State)
+		{
+			State->bExpanded = true;
+		}
+	}
+}
+
+
+void FStateTreeViewModel::AddState(UStateTreeState* AfterState)
+{
+	if (!TreeData)
+	{
+		return;
+	}
+
+	const FScopedTransaction Transaction(LOCTEXT("AddStateTransaction", "Add State"));
+
+	UStateTreeState* NewState = NewObject<UStateTreeState>(TreeData, FName(), RF_Transactional);
+	UStateTreeState* ParentState = nullptr;
+
+	if (AfterState)
+	{
+		ParentState = AfterState->Parent;
+		if (ParentState)
+		{
+			ParentState->Modify();
+			NewState->Parent = ParentState;
+		}
+		else
+		{
+			TreeData->Modify();
+			NewState->Parent = nullptr;
+		}
+
+		TArray<UStateTreeState*>& ArrayToAddTo = ParentState ? ParentState->Children : TreeData->Routines;
+		FStateTreeViewUtilities::InsertRecursive(ParentState, ArrayToAddTo, AfterState, NewState, 1); // Insert after
+	}
+	else
+	{
+		TreeData->Modify();
+		NewState->Parent = nullptr;
+		TreeData->Routines.Add(NewState);
+	}
+
+	OnStateAdded.Broadcast(ParentState, NewState);
+}
+
+void FStateTreeViewModel::AddChildState(UStateTreeState* ParentState)
+{
+	if (!TreeData || !ParentState)
+	{
+		return;
+	}
+
+	const FScopedTransaction Transaction(LOCTEXT("AddChildStateTransaction", "Add Child State"));
+
+	TreeData->Modify();
+
+	UStateTreeState* NewState = NewObject<UStateTreeState>(TreeData, FName(), RF_Transactional);
+
+	// Make sure the parent state is set to select when it has children.
+/*	if (ParentState->StateDoneTransition.Type != EStateTreeTransitionType::SelectChildState)
+	{
+		ParentState->StateDoneTransition.Type = EStateTreeTransitionType::SelectChildState;
+	}*/
+	
+	ParentState->Children.Add(NewState);
+	NewState->Parent = ParentState;
+
+	OnStateAdded.Broadcast(ParentState, NewState);
+}
+
+void FStateTreeViewModel::RenameState(UStateTreeState* State, FName NewName)
+{
+	if (!State)
+	{
+		return;
+	}
+
+	const FScopedTransaction Transaction(LOCTEXT("RenameTransaction", "Rename"));
+	State->Modify();
+	State->Name = NewName;
+
+	TSet<UStateTreeState*> AffectedStates;
+	AffectedStates.Add(State);
+
+	FProperty* NameProperty = FindFProperty<FProperty>(UStateTreeState::StaticClass(), GET_MEMBER_NAME_CHECKED(UStateTreeState, Name));
+	FPropertyChangedEvent PropertyChangedEvent(NameProperty, EPropertyChangeType::ValueSet);
+	OnStatesChanged.Broadcast(AffectedStates, PropertyChangedEvent);
+}
+
+void FStateTreeViewModel::RemoveSelectedStates()
+{
+	if (!TreeData)
+	{
+		return;
+	}
+
+	TArray<UStateTreeState*> States;
+	GetSelectedStates(States);
+
+	// Remove items whose parent also exists in the selection.
+	FStateTreeViewUtilities::RemoveContainedChildren(States);
+
+	if (States.Num() > 0)
+	{
+		const FScopedTransaction Transaction(LOCTEXT("DeleteStateTransaction", "Delete State"));
+
+		TSet<UStateTreeState*> AffectedParents;
+
+		for (UStateTreeState* State : States)
+		{
+			if (State)
+			{
+				if (UStateTreeState* ParentState = State->Parent)
+				{
+					AffectedParents.Add(ParentState);
+					ParentState->Modify();
+					FStateTreeViewUtilities::RemoveRecursive(ParentState->Children, State);
+				}
+				else
+				{
+					AffectedParents.Add(nullptr);
+					TreeData->Modify();
+					FStateTreeViewUtilities::RemoveRecursive(TreeData->Routines, State);
+				}
+			}
+		}
+
+		OnStatesRemoved.Broadcast(AffectedParents);
+
+		ClearSelection();
+	}
+}
+
+void FStateTreeViewModel::MoveSelectedStatesBefore(UStateTreeState* TargetState)
+{
+	MoveSelectedStates(TargetState, -1);
+}
+
+void FStateTreeViewModel::MoveSelectedStatesAfter(UStateTreeState* TargetState)
+{
+	MoveSelectedStates(TargetState, 1);
+}
+
+void FStateTreeViewModel::MoveSelectedStatesInto(UStateTreeState* TargetState)
+{
+	MoveSelectedStates(TargetState, 0);
+}
+
+void FStateTreeViewModel::MoveSelectedStates(UStateTreeState* TargetState, int32 RelativeLocation)
+{
+	if (!TreeData || !TargetState)
+	{
+		return;
+	}
+
+	TArray<UStateTreeState*> States;
+	GetSelectedStates(States);
+
+	UStateTreeState* TargetParent = TargetState->Parent;
+
+	// Remove child items whose parent also exists in the selection.
+	FStateTreeViewUtilities::RemoveContainedChildren(States);
+	FStateTreeViewUtilities::RemoveRecursive(States, TargetState);
+
+	if (States.Num() > 0 && TargetState)
+	{
+		const FScopedTransaction Transaction(LOCTEXT("MoveTransaction", "Move"));
+
+		TSet<UStateTreeState*> AffectedParents;
+		TSet<UStateTreeState*> AffectedStates;
+
+		AffectedParents.Add(TargetParent);
+		for (int32 i = States.Num() - 1; i >= 0; i--)
+		{
+			if (UStateTreeState* State = States[i])
+			{
+				AffectedParents.Add(State->Parent);
+			}
+		}
+
+		for (UStateTreeState* Parent : AffectedParents)
+		{
+			if (Parent)
+			{
+				Parent->Modify();
+			}
+			else
+			{
+				TreeData->Modify();
+			}
+		}
+
+		// Add in reverse order to keep the original order.
+		for (int32 i = States.Num() - 1; i >= 0; i--)
+		{
+			if (UStateTreeState* SelectedState = States[i])
+			{
+				UStateTreeState* SelectedParent = SelectedState->Parent;
+				AffectedStates.Add(SelectedState);
+
+				TArray<UStateTreeState*>& ArrayToRemoveFrom = SelectedParent ? SelectedParent->Children : TreeData->Routines;
+				TArray<UStateTreeState*>& ArrayToMoveTo = TargetParent ? TargetParent->Children : TreeData->Routines;
+				FStateTreeViewUtilities::RemoveRecursive(ArrayToRemoveFrom, SelectedState);
+				FStateTreeViewUtilities::InsertRecursive(TargetParent, ArrayToMoveTo, TargetState, SelectedState, RelativeLocation);
+			}
+		}
+
+		OnStatesMoved.Broadcast(AffectedParents, AffectedStates);
+
+		SetSelection(States);
+	}
+}
+
+
+#undef LOCTEXT_NAMESPACE
diff --git a/Engine/Plugins/Runtime/StateTree/Source/StateTreeEditorModule/Private/StateTreeViewModel.h b/Engine/Plugins/Runtime/StateTree/Source/StateTreeEditorModule/Private/StateTreeViewModel.h
new file mode 100644
--- /dev/null
+++ b/Engine/Plugins/Runtime/StateTree/Source/StateTreeEditorModule/Private/StateTreeViewModel.h
@@ -0,0 +1,100 @@
+// Copyright Epic Games, Inc. All Rights Reserved.
+
+#pragma once
+
+#include "CoreMinimal.h"
+#include "Templates/Casts.h"
+#include "UObject/WeakObjectPtr.h"
+#include "EditorUndoClient.h"
+
+class UStateTreeEditorData;
+class UStateTreeState;
+class UStateTreeTaskBase;
+class UStateTree;
+class FMenuBuilder;
+
+/**
+ * ModelView for editing StateTreeEditorData.
+ */
+class FStateTreeViewModel : public FEditorUndoClient, public TSharedFromThis<FStateTreeViewModel>
+{
+public:
+
+	DECLARE_MULTICAST_DELEGATE(FOnAssetChanged);
+	DECLARE_MULTICAST_DELEGATE_TwoParams(FOnStatesChanged, const TSet<UStateTreeState*>& /*AffectedStates*/, const FPropertyChangedEvent& /*PropertyChangedEvent*/);
+	DECLARE_MULTICAST_DELEGATE_TwoParams(FOnStateAdded, UStateTreeState* /*ParentState*/, UStateTreeState* /*NewState*/);
+	DECLARE_MULTICAST_DELEGATE_OneParam(FOnStatesRemoved, const TSet<UStateTreeState*>& /*AffectedParents*/);
+	DECLARE_MULTICAST_DELEGATE_TwoParams(FOnStatesMoved, const TSet<UStateTreeState*>& /*AffectedParents*/, const TSet<UStateTreeState*>& /*MovedStates*/);
+	DECLARE_MULTICAST_DELEGATE_OneParam(FOnSelectionChanged, const TArray<UStateTreeState*>& /*SelectedStates*/);
+
+	FStateTreeViewModel();
+	~FStateTreeViewModel();
+
+	void Init(UStateTreeEditorData* InTreeData);
+
+	//~ FEditorUndoClient
+	void PostUndo(bool bSuccess);
+	void PostRedo(bool bSuccess);
+
+	// Selection handling.
+	void ClearSelection();
+	void SetSelection(UStateTreeState* Selected);
+	void SetSelection(const TArray<UStateTreeState*>& InSelection);
+	bool IsSelected(const UStateTreeState* State) const;
+	bool IsChildOfSelection(const UStateTreeState* State) const;
+	void GetSelectedStates(TArray<UStateTreeState*>& OutSelectedStates);
+
+	// Returns array of routines to edit.
+	TArray<UStateTreeState*>* GetRoutines();
+
+	// Gets and sets StateTree view expansion state store in the asset.
+	void GetPersistentExpandedStates(TSet<UStateTreeState*>& OutExpandedStates);
+	void SetPersistentExpandedStates(TSet<UStateTreeState*>& InExpandedStates);
+
+	// State manupulation commands
+	void AddState(UStateTreeState* AfterState);
+	void AddChildState(UStateTreeState* ParentState);
+	void RenameState(UStateTreeState* State, FName NewName);
+	void RemoveSelectedStates();
+	void MoveSelectedStatesBefore(UStateTreeState* TargetState);
+	void MoveSelectedStatesAfter(UStateTreeState* TargetState);
+	void MoveSelectedStatesInto(UStateTreeState* TargetState);
+
+	// Force to update the view externally.
+	void NotifyAssetChangedExternally();
+	void NotifyStatesChangedExternally(const TSet<UStateTreeState*>& ChangedStates, const FPropertyChangedEvent& PropertyChangedEvent);
+
+	// Called when the whole asset is updated (i.e. undo/redo).
+	FOnAssetChanged& GetOnAssetChanged() { return OnAssetChanged; }
+	
+	// Called when States are changed (i.e. change name or properties).
+	FOnStatesChanged& GetOnStatesChanged() { return OnStatesChanged; }
+	
+	// Called each time a state is added.
+	FOnStateAdded& GetOnStateAdded() { return OnStateAdded; }
+
+	// Called each time a states are removed.
+	FOnStatesRemoved& GetOnStatesRemoved() { return OnStatesRemoved; }
+
+	// Called each time a state is removed.
+	FOnStatesMoved& GetOnStatesMoved() { return OnStatesMoved; }
+
+	// Called each time the selection changes.
+	FOnSelectionChanged& GetOnSelectionChanged() { return OnSelectionChanged; }
+
+protected:
+	void GetExpandedStatesRecursive(UStateTreeState* State, TSet<UStateTreeState*>& ExpandedStates);
+	void MoveSelectedStates(UStateTreeState* TargetState, int32 RelativeLocation);
+
+	void HandleIdentifierChanged(const UStateTree& StateTree);
+	void HandleParameterLayoutChanged(const UStateTree& StateTree);
+
+	UStateTreeEditorData* TreeData;	// TODO: is this safe? should we used TWeakObjectPtr here istead?
+	TSet<FWeakObjectPtr> SelectedStates;
+	FOnAssetChanged OnAssetChanged;
+	FOnStatesChanged OnStatesChanged;
+	FOnStateAdded OnStateAdded;
+	FOnStatesRemoved OnStatesRemoved;
+	FOnStatesMoved OnStatesMoved;
+	FOnSelectionChanged OnSelectionChanged;
+};
diff --git a/Engine/Plugins/Runtime/StateTree/Source/StateTreeEditorModule/Public/IStateTreeEditor.h b/Engine/Plugins/Runtime/StateTree/Source/StateTreeEditorModule/Public/IStateTreeEditor.h
new file mode 100644
--- /dev/null
+++ b/Engine/Plugins/Runtime/StateTree/Source/StateTreeEditorModule/Public/IStateTreeEditor.h
@@ -0,0 +1,16 @@
+// Copyright Epic Games, Inc. All Rights Reserved.
+
+#pragma once
+
+#include "CoreMinimal.h"
+#include "Toolkits/AssetEditorToolkit.h"
+
+
+/** StateTree Editor public interface */
+class STATETREEEDITORMODULE_API IStateTreeEditor : public FAssetEditorToolkit
+{
+public:
+
+};
+
+
diff --git a/Engine/Plugins/Runtime/StateTree/Source/StateTreeEditorModule/Public/StateTreeBaker.h b/Engine/Plugins/Runtime/StateTree/Source/StateTreeEditorModule/Public/StateTreeBaker.h
new file mode 100644
--- /dev/null
+++ b/Engine/Plugins/Runtime/StateTree/Source/StateTreeEditorModule/Public/StateTreeBaker.h
@@ -0,0 +1,75 @@
+// Copyright Epic Games, Inc. All Rights Reserved.
+
+#pragma once
+
+#include "CoreMinimal.h"
+#include "StateTreeTypes.h"
+#include "StateTreePropertyBindingCompiler.h"
+
+class UStateTree;
+class UStateTreeState;
+class UStateTreeEditorData;
+struct FStateTreeCondition;
+class UStateTreeEvaluatorBase;
+class UStateTreeTaskBase;
+struct FStateTreeConditionItem;
+struct FStateTreeStateLink;
+
+/**
+ * Helper class to convert StateTree editor representation into a compact baked data.
+ * Holds data needed during baking.
+ */
+struct STATETREEEDITORMODULE_API FStateTreeBaker
+{
+public:
+	bool Bake(UStateTree& InStateTree);
+	bool Bake2(UStateTree& InStateTree);
+
+private:
+	enum class EBindingContainerType : uint8
+	{
+		Evaluator,
+		Task,
+	};
+
+	// Helper struct used when collecting evaluators, keeping the state reference so that we can have better errors.
+	struct FSourceEvaluator
+	{
+		FSourceEvaluator(UStateTreeEvaluatorBase* InEval, UStateTreeState* InState) : Eval(InEval), State(InState) {}
+		UStateTreeEvaluatorBase* Eval = nullptr;
+		UStateTreeState* State = nullptr;
+	};
+
+	// Helper struct used when collecting tracks, keeping the state reference so that we can have better errors.
+	struct FSourceTask
+	{
+		FSourceTask(UStateTreeTaskBase* InTask, UStateTreeState* InState) : Task(InTask), State(InState) {}
+		UStateTreeTaskBase* Task = nullptr;
+		UStateTreeState* State = nullptr;
+	};
+
+	void DefineParameters() const;
+	void ResolveParameters() const;
+	bool CreateStates();
+	bool CreateStateRecursive(UStateTreeState& State, const FStateTreeHandle Parent);
+	bool ResolveTransitionState(const UStateTreeState& SourceState, const TCHAR* ContextStr, const FStateTreeStateLink& Link, FStateTreeHandle& OutTransitionHandle) const;
+	bool CreateStateTransitions();
+	bool CreateCondition(const FStateTreeCondition& Cond);
+	FStateTreeHandle GetStateHandle(const FGuid& StateID) const;
+
+	bool CreateStates2();
+	bool CreateStateRecursive2(UStateTreeState& State, const FStateTreeHandle Parent);
+	bool CreateStateTransitions2();
+	bool CreateCondition2(const FStateTreeConditionItem& CondItem);
+	bool GetAndValidateBindings(const FStateTreeBindableStructDesc& TargetStruct, TArray<FStateTreeEditorPropertyBinding>& OutBindings) const;
+	bool IsPropertyAnyEnum(const FStateTreeBindableStructDesc& Struct, FStateTreeEditorPropertyPath Path) const;
+	bool CreateExternalItemHandles(FStructView Item);
+
+	UStateTree* StateTree = nullptr;
+	UStateTreeEditorData* TreeData = nullptr;
+	TMap<FGuid, int32> IDToState;
+	TArray<UStateTreeState*> SourceStates;
+	TArray<FSourceEvaluator> SourceEvaluators;
+	TArray<FSourceTask> SourceTasks;
+	FStateTreePropertyBindingCompiler BindingsCompiler;
+};
diff --git a/Engine/Plugins/Runtime/StateTree/Source/StateTreeEditorModule/Public/StateTreeEditorData.h b/Engine/Plugins/Runtime/StateTree/Source/StateTreeEditorModule/Public/StateTreeEditorData.h
new file mode 100644
--- /dev/null
+++ b/Engine/Plugins/Runtime/StateTree/Source/StateTreeEditorModule/Public/StateTreeEditorData.h
@@ -0,0 +1,66 @@
+// Copyright Epic Games, Inc. All Rights Reserved.
+
+#pragma once
+
+#include "StateTreeState.h"
+#include "StateTreeEditorPropertyBindings.h"
+#include "StateTreeEditorData.generated.h"
+
+/**
+ * Edit time data for StateTree asset. This data gets baked into runtime format before being used by the StateTreeInstance.
+ */
+UCLASS(BlueprintType, EditInlineNew, CollapseCategories)
+class STATETREEEDITORMODULE_API UStateTreeEditorData : public UObject, public IStateTreeEditorPropertyBindingsOwner
+{
+	GENERATED_BODY()
+public:
+	// IStateTreeEditorPropertyBindingsOwner
+	virtual void GetAccessibleStructs(const FGuid TargetStructID, TArray<FStateTreeBindableStructDesc>& OutStructDescs) const override;
+	virtual bool GetStructByID(const FGuid StructID, FStateTreeBindableStructDesc& OutStructDesc) const override;
+	virtual FStateTreeEditorPropertyBindings* GetPropertyEditorBindings() override { return &EditorBindings; }
+	// ~IStateTreeEditorPropertyBindingsOwner
+
+	/** Returns parent state of a struct, or nullptr if not found. */
+	const UStateTreeState* GetStateByStructID(const FGuid TargetStructID) const;
+
+	/** Gets the IDs of all bindable structs in the StateTree. */
+	void GetAllStructIDs(TMap<FGuid, const UScriptStruct*>& AllStructs) const;
+
+	/**
+	 * Iterates over all structs that are related to binding
+	 * @param InFunc function called at each item, should return true if visiting is continued or false to stop.
+	 */
+	void VisitHierarchy(TFunctionRef<bool(const UStateTreeState& State, const FGuid& ID, const FName& Name, const UScriptStruct* Struct)> InFunc) const;
+
+	// StateTree Builder API
+
+	/**
+	 * Adds new Subtree with specified name.
+	 * @return Pointer to the new Subtree.
+	 */
+	UStateTreeState& AddSubTree(const FName Name)
+	{
+		UStateTreeState* SubTreeState = NewObject<UStateTreeState>(this);
+		check(SubTreeState);
+		SubTreeState->Name = Name;
+		Routines.Add(SubTreeState);
+		return *SubTreeState;
+	}
+
+	/**
+	 * Adds property binding between two structs.
+	 */
+	void AddPropertyBinding(const FStateTreeEditorPropertyPath& SourcePath, const FStateTreeEditorPropertyPath& TargetPath)
+	{
+		EditorBindings.AddPropertyBinding(SourcePath, TargetPath);
+	}
+	// ~StateTree Builder API
+
+	
+	UPROPERTY()
+	FStateTreeEditorPropertyBindings EditorBindings;
+
+	// Top level States aka Routines.
+	UPROPERTY()
+	TArray<UStateTreeState*> Routines;
+};
diff --git a/Engine/Plugins/Runtime/StateTree/Source/StateTreeEditorModule/Public/StateTreeEditorModule.h b/Engine/Plugins/Runtime/StateTree/Source/StateTreeEditorModule/Public/StateTreeEditorModule.h
new file mode 100644
--- /dev/null
+++ b/Engine/Plugins/Runtime/StateTree/Source/StateTreeEditorModule/Public/StateTreeEditorModule.h
@@ -0,0 +1,40 @@
+// Copyright Epic Games, Inc. All Rights Reserved.
+
+#pragma  once
+
+#include "CoreMinimal.h"
+#include "Modules/ModuleManager.h"
+#include "AssetTypeCategories.h"
+#include "Toolkits/IToolkitHost.h"
+#include "Toolkits/AssetEditorToolkit.h"
+
+class UStateTree;
+class IStateTreeEditor;
+class FAssetTypeActions_Base;
+
+/**
+* The public interface to this module
+*/
+class STATETREEEDITORMODULE_API FStateTreeEditorModule : public IModuleInterface, public IHasMenuExtensibility, public IHasToolBarExtensibility
+{
+public:
+	// Begin IModuleInterface
+	virtual void StartupModule() override;
+	virtual void ShutdownModule() override;
+
+	/** Creates an instance of StateTree editor. Only virtual so that it can be called across the DLL boundary. */
+	virtual TSharedRef<IStateTreeEditor> CreateStateTreeEditor(const EToolkitMode::Type Mode, const TSharedPtr< class IToolkitHost >& InitToolkitHost, UStateTree* StateTree);
+
+	virtual TSharedPtr<FExtensibilityManager> GetMenuExtensibilityManager() override { return MenuExtensibilityManager; }
+	virtual TSharedPtr<FExtensibilityManager> GetToolBarExtensibilityManager() override { return ToolBarExtensibilityManager; }
+
+	TSharedPtr<struct FGraphNodeClassHelper> GetTaskClassCache() { return TaskClassCache; }
+
+protected:
+	TSharedPtr<struct FGraphNodeClassHelper> TaskClassCache;
+
+	TSharedPtr<FExtensibilityManager> MenuExtensibilityManager;
+	TSharedPtr<FExtensibilityManager> ToolBarExtensibilityManager;
+
+	TArray<TSharedPtr<FAssetTypeActions_Base>> ItemDataAssetTypeActions;
+};
diff --git a/Engine/Plugins/Runtime/StateTree/Source/StateTreeEditorModule/Public/StateTreeEditorPropertyBindings.h b/Engine/Plugins/Runtime/StateTree/Source/StateTreeEditorModule/Public/StateTreeEditorPropertyBindings.h
new file mode 100644
--- /dev/null
+++ b/Engine/Plugins/Runtime/StateTree/Source/StateTreeEditorModule/Public/StateTreeEditorPropertyBindings.h
@@ -0,0 +1,131 @@
+// Copyright Epic Games, Inc. All Rights Reserved.
+
+#pragma once
+
+#include "CoreMinimal.h"
+#include "UObject/Interface.h"
+#include "Misc/Guid.h"
+#include "StateTreePropertyBindings.h"
+#include "StateTreeEditorPropertyBindings.generated.h"
+
+/**
+ * Editor representation of a property binding in StateTree
+ */
+USTRUCT()
+struct STATETREEEDITORMODULE_API FStateTreeEditorPropertyBinding
+{
+	GENERATED_BODY()
+
+	FStateTreeEditorPropertyBinding() = default;
+	FStateTreeEditorPropertyBinding(const FStateTreeEditorPropertyPath& InSourcePath, const FStateTreeEditorPropertyPath& InTargetPath) : SourcePath(InSourcePath), TargetPath(InTargetPath) {}
+
+	bool IsValid() const { return SourcePath.IsValid() && TargetPath.IsValid(); }
+
+	/** Source property path of the binding */
+	UPROPERTY()
+	FStateTreeEditorPropertyPath SourcePath;
+
+	/** Target property path of the binding */
+	UPROPERTY()
+	FStateTreeEditorPropertyPath TargetPath;
+};
+
+
+/**
+ * Editor representation of a all property bindings in a StateTree
+ */
+USTRUCT()
+struct STATETREEEDITORMODULE_API FStateTreeEditorPropertyBindings
+{
+	GENERATED_BODY()
+
+	/**
+	 * Adds binding between source and destination paths. Removes any bindings to TargetPath before adding the new one.
+	 * @param SourcePath Binding source property path.
+	 * @param TargetPath Binding target property path.
+	 */
+	void AddPropertyBinding(const FStateTreeEditorPropertyPath& SourcePath, const FStateTreeEditorPropertyPath& TargetPath);
+	
+	/**
+	 * Removes all bindings to target path.
+	 * @param TargetPath Target property path.
+	 */ 
+	void RemovePropertyBindings(const FStateTreeEditorPropertyPath& TargetPath);
+	
+	/**
+	 * @param TargetPath Target property path.
+	 * @return True of the target path has any bindings.
+	 */
+	bool HasPropertyBinding(const FStateTreeEditorPropertyPath& TargetPath) const;
+	
+	/**
+	 * @return Source path for given target path, or null if binding does not exists.
+	 */
+	const FStateTreeEditorPropertyPath* GetPropertyBindingSource(const FStateTreeEditorPropertyPath& TargetPath) const;
+	
+	/**
+	 * Returns all bindings for a specified structs based in struct ID.
+	 * @param StructID ID of the struct to find bindings for.
+	 * @param OutBindings Bindings for specified struct.
+	 */
+	void GetPropertyBindingsFor(const FGuid StructID, TArray<FStateTreeEditorPropertyBinding>& OutBindings) const;
+	
+	/**
+	 * Removes bindings which do not point to valid structs IDs.
+	 * @param ValidStructsIDs Set of struct IDs that are currently valid.
+	 */
+	void RemoveUnusedBindings(const TMap<FGuid, const UScriptStruct*>& ValidStructs);
+
+	/** @return array view to all bindings. */
+	TConstArrayView<FStateTreeEditorPropertyBinding> GetBindings() const { return PropertyBindings; }
+
+protected:
+
+	UPROPERTY()
+	TArray<FStateTreeEditorPropertyBinding> PropertyBindings;
+};
+
+
+UINTERFACE(Blueprintable, MinimalAPI, meta = (CannotImplementInterfaceInBlueprint))
+class UStateTreeEditorPropertyBindingsOwner : public UInterface
+{
+	GENERATED_UINTERFACE_BODY()
+};
+
+class STATETREEEDITORMODULE_API IStateTreeEditorPropertyBindingsOwner
+{
+	GENERATED_IINTERFACE_BODY()
+
+	/**
+	 * Returns structs within the owner that are visible for target struct.
+	 * @param TargetStructID Target struct ID
+	 * @param OutStructDescs Result descriptors of the visible structs.
+	 */
+	virtual void GetAccessibleStructs(const FGuid TargetStructID, TArray<FStateTreeBindableStructDesc>& OutStructDescs) const PURE_VIRTUAL(IStateTreeEditorPropertyBindingsOwner::GetAccessibleStructs, return; );
+
+	/**
+	 * Returns struct descriptor based on struct ID.
+	 * @param StructID Target struct ID
+	 * @param OutStructDesc Result descriptor.
+	 * @return True if struct found.
+	 */
+	virtual bool GetStructByID(const FGuid StructID, FStateTreeBindableStructDesc& OutStructDesc) const PURE_VIRTUAL(IStateTreeEditorPropertyBindingsOwner::GetStructByID, return false; );
+
+	/** @return Pointer to editor property bindings. */
+	virtual FStateTreeEditorPropertyBindings* GetPropertyEditorBindings() PURE_VIRTUAL(IStateTreeEditorPropertyBindingsOwner::GetPropertyEditorBindings, return nullptr; );
+};
+
+// TODO: We should merge this with IStateTreeEditorPropertyBindingsOwner and FStateTreeEditorPropertyBindings.
+// Currently FStateTreeEditorPropertyBindings is meant to be used as a member for just to store things,
+// IStateTreeEditorPropertyBindingsOwner is meant return model specific stuff,
+// and IStateTreeBindingLookup is used in non-editor code and it cannot be in FStateTreeEditorPropertyBindings because bindings don't know about the owner.
+struct STATETREEEDITORMODULE_API FStateTreeBindingLookup : public IStateTreeBindingLookup
+{
+	FStateTreeBindingLookup(IStateTreeEditorPropertyBindingsOwner* InBindingOwner);
+
+	virtual const FStateTreeEditorPropertyPath* GetPropertyBindingSource(const FStateTreeEditorPropertyPath& InTargetPath) const override;
+	virtual FText GetPropertyPathDisplayName(const FStateTreeEditorPropertyPath& InTargetPath) const override;
+	virtual const FProperty* GetPropertyPathLeafProperty(const FStateTreeEditorPropertyPath& InPath) const override;
+
+	IStateTreeEditorPropertyBindingsOwner* BindingOwner = nullptr;
+};
diff --git a/Engine/Plugins/Runtime/StateTree/Source/StateTreeEditorModule/Public/StateTreeEditorStyle.h b/Engine/Plugins/Runtime/StateTree/Source/StateTreeEditorModule/Public/StateTreeEditorStyle.h
new file mode 100644
--- /dev/null
+++ b/Engine/Plugins/Runtime/StateTree/Source/StateTreeEditorModule/Public/StateTreeEditorStyle.h
@@ -0,0 +1,27 @@
+// Copyright Epic Games, Inc. All Rights Reserved.
+
+#pragma once
+
+#include "Styling/SlateStyle.h"
+
+class STATETREEEDITORMODULE_API FStateTreeEditorStyle
+{
+public:
+
+	static void Initialize();
+	static void Shutdown();
+	static TSharedPtr<ISlateStyle> Get() { return StyleSet; }
+	static FName GetStyleSetName();
+
+	static const FSlateBrush* GetBrush(FName PropertyName, const ANSICHAR* Specifier = NULL)
+	{
+		return StyleSet->GetBrush(PropertyName, Specifier);
+	}
+
+	static FString InContent(const FString& RelativePath, const ANSICHAR* Extension);
+
+
+private:
+
+	static TSharedPtr<FSlateStyleSet> StyleSet;
+};
diff --git a/Engine/Plugins/Runtime/StateTree/Source/StateTreeEditorModule/Public/StateTreePropertyBindingCompiler.h b/Engine/Plugins/Runtime/StateTree/Source/StateTreeEditorModule/Public/StateTreePropertyBindingCompiler.h
new file mode 100644
--- /dev/null
+++ b/Engine/Plugins/Runtime/StateTree/Source/StateTreeEditorModule/Public/StateTreePropertyBindingCompiler.h
@@ -0,0 +1,98 @@
+// Copyright Epic Games, Inc. All Rights Reserved.
+
+#pragma once
+
+#include "CoreMinimal.h"
+#include "Misc/Guid.h"
+#include "InstancedStruct.h"
+#include "UObject/Interface.h"
+#include "IPropertyAccessEditor.h"
+#include "StateTreeEditorPropertyBindings.h"
+#include "StateTreePropertyBindings.h"
+#include "StateTreePropertyBindingCompiler.generated.h"
+
+/**
+ * Helper class to compile editor representation of property bindings into runtime representation.
+ * TODO: Better error reporting, something that can be shown in the UI.
+ */
+USTRUCT()
+struct STATETREEEDITORMODULE_API FStateTreePropertyBindingCompiler
+{
+	GENERATED_BODY()
+
+	/**
+	  * Initializes the compiler to compile copies to specified Property Bindings.
+	  * @param PropertyBindings - Reference to the Property Bindings where all the batches will be stored.
+	  * @return true on success.
+	  */
+	bool Init(FStateTreePropertyBindings& InPropertyBindings);
+
+	/**
+	  * Compiles a batch of property copies.
+	  * @param TargetStruct - Description of the structs which contains the target properties.
+	  * @param PropertyBindings - Array of bindings to compile, all bindings that point to TargetStructs will be added to the batch.
+	  * @param OutBatchIndex - Resulting batch index, if index is INDEX_NONE, no bindings were found and no batch was generated.
+	  * @return True on success, false on failure.
+	 */
+	bool CompileBatch(const FStateTreeBindableStructDesc& TargetStruct, TConstArrayView<FStateTreeEditorPropertyBinding> EditorPropertyBindings, int32& OutBatchIndex);
+
+	/** Finalizes compilation, should be called once all batches are compiled. */
+	void Finalize();
+
+	/**
+	  * Adds source struct. When compiling a batch, the bindings can be between any of the defined source structs, and the target struct.
+	  * Source structs can be added between calls to Compilebatch().
+	  * @param SourceStruct - Description of the source struct to add.
+	  * @return Source struct index.
+	  */
+	int32 AddSourceStruct(const FStateTreeBindableStructDesc& SourceStruct);
+
+	/** @return Index of a source struct by specified ID, or INDEX_NONE if not found. */
+	int32 GetSourceStructIndexByID(const FGuid& ID) const;
+
+	/** @return Reference to a source struct based on ID. */
+	const FStateTreeBindableStructDesc& GetSourceStructDesc(const int32 Index) const
+	{
+		return SourceStructs[Index];
+	}
+
+	/**
+	 * Resolves a string based property path in specified struct into segments of property names and access types.
+	 * @param InStructDesc Description of the struct in which the property path is valid.
+	 * @param InPath The property path in string format.
+	 * @param OutSegments The resolved property access path as segments.
+	 * @param OutLeafProperty The leaf property of the resolved path.
+	 * @param OutLeafArrayIndex The left array index (or INDEX_NONE if not applicable) of the resolved path.
+	 * @param bLogErrors Set to true to log errors during resolve into the console.
+	 * @return True of the property was solved successfully.
+	 */
+	static bool ResolvePropertyPath(const FStateTreeBindableStructDesc& InStructDesc, const FStateTreeEditorPropertyPath& InPath,
+									TArray<FStateTreePropertySegment>& OutSegments, FProperty*& OutLeafProperty, int32& OutLeafArrayIndex, const bool bLogErrors = false);
+
+	/**
+	 * Checks if two property types can are compatible for copying.
+	 * @param InPropertyA First property used in the check.
+	 * @param InPropertyB Second property used in the check.
+	 * @return Incompatible if the properties cannot be copied, Compatible if they are trivially copyable, or Promotable if numeric values can be promoted to another numeric type.
+	 */
+	static EPropertyAccessCompatibility GetPropertyCompatibility(const FProperty* InPropertyA, const FProperty* InPropertyB);
+	
+protected:
+
+	void StoreSourceStructs();
+
+	UPROPERTY()
+	TArray<FStateTreeBindableStructDesc> SourceStructs;
+
+	FStateTreePropertyBindings* PropertyBindings = nullptr;
+
+	struct FResolvedPathResult
+	{
+		int32 PathIndex = INDEX_NONE;
+		FProperty* LeafProperty = nullptr;
+		int32 LeafArrayIndex = INDEX_NONE;
+	};
+
+	EStateTreePropertyCopyType GetCopyType(FProperty* SourceProperty, int32 SourceArrayIndex, FProperty* TargetProperty, int32 TargetArrayIndex);
+	bool ResolvePropertyPath(const FStateTreeBindableStructDesc& InStructDesc, const FStateTreeEditorPropertyPath& InPath, FResolvedPathResult& OutResult);
+};
diff --git a/Engine/Plugins/Runtime/StateTree/Source/StateTreeEditorModule/Public/StateTreeState.h b/Engine/Plugins/Runtime/StateTree/Source/StateTreeEditorModule/Public/StateTreeState.h
new file mode 100644
--- /dev/null
+++ b/Engine/Plugins/Runtime/StateTree/Source/StateTreeEditorModule/Public/StateTreeState.h
@@ -0,0 +1,267 @@
+// Copyright Epic Games, Inc. All Rights Reserved.
+
+#pragma once
+
+#include "StateTreeCondition.h"
+#include "StateTreeEvaluatorBase.h"
+#include "StateTreeTaskBase.h"
+#include "Misc/Guid.h"
+#include "StateTreeVariableProvider.h"
+#include "InstancedStruct.h"
+#include "StateTreeState.generated.h"
+
+class UStateTreeState;
+
+/**
+ * Editor representation of a link to another state in StateTree
+ */
+USTRUCT()
+struct STATETREEEDITORMODULE_API FStateTreeStateLink
+{
+	GENERATED_BODY()
+public:
+	FStateTreeStateLink() : Type(EStateTreeTransitionType::GotoState) {}
+	FStateTreeStateLink(EStateTreeTransitionType InType) : Type(InType) {}
+
+	void Set(const EStateTreeTransitionType InType, const UStateTreeState* InState = nullptr);
+	
+	bool IsValid() const { return ID.IsValid(); }
+
+	UPROPERTY(EditDefaultsOnly, Category = Link)
+	FName Name;
+	
+	UPROPERTY(EditDefaultsOnly, Category = Link)
+	FGuid ID;
+
+	UPROPERTY(EditDefaultsOnly, Category = Transition)
+	EStateTreeTransitionType Type;
+};
+
+/**
+ * Editor representation of a transition in StateTree
+ */
+USTRUCT()
+struct STATETREEEDITORMODULE_API FStateTreeTransition
+{
+	GENERATED_BODY()
+public:
+	UPROPERTY(EditDefaultsOnly, Category = Transition)
+	FStateTreeStateLink State;
+
+	UPROPERTY(EditDefaultsOnly, Category = Transition)
+	TArray<FStateTreeCondition> Conditions;
+};
+
+/**
+ * Helper struct for Evaluator details customization (i.e. show summary when collapsed)
+ */
+USTRUCT()
+struct STATETREEEDITORMODULE_API FStateTreeEvaluatorItem
+{
+	GENERATED_BODY()
+
+	UPROPERTY(EditDefaultsOnly, Category = Evaluator, meta = (BaseStruct = "StateTreeEvaluator2Base", ExcludeBaseStruct))
+	FInstancedStruct Type;
+};
+
+/**
+ * Helper struct for Task details customization (i.e. show summary when collapsed)
+ */
+USTRUCT()
+struct STATETREEEDITORMODULE_API FStateTreeTaskItem
+{
+	GENERATED_BODY()
+
+	UPROPERTY(EditDefaultsOnly, Category = Task, meta = (BaseStruct = "StateTreeTask2Base", ExcludeBaseStruct))
+	FInstancedStruct Type;
+};
+
+
+/**
+ * Helper struct for Condition details customization (i.e. show summary when collapsed)
+ */
+USTRUCT()
+struct STATETREEEDITORMODULE_API FStateTreeConditionItem
+{
+	GENERATED_BODY()
+
+	UPROPERTY(EditAnywhere, Category = Condition, meta = (BaseStruct = "StateTreeConditionBase", ExcludeBaseStruct))
+	FInstancedStruct Type;
+};
+
+/**
+ * Editor representation of a transition in StateTree
+ */
+USTRUCT()
+struct STATETREEEDITORMODULE_API FStateTreeTransition2
+{
+	GENERATED_BODY()
+
+	FStateTreeTransition2() = default;
+	FStateTreeTransition2(const EStateTreeTransitionEvent InEvent, const EStateTreeTransitionType InType, const UStateTreeState* InState = nullptr);
+
+	template<typename T, typename... TArgs>
+	T& AddCondition(TArgs&&... InArgs)
+	{
+		FStateTreeConditionItem& CondItem = Conditions.AddDefaulted_GetRef();
+		CondItem.Type.InitializeAs<T>(Forward<TArgs>(InArgs)...);
+		T& Condition = CondItem.Type.GetMutable<T>();
+		Condition.ID = FGuid::NewGuid();
+		return Condition;
+	}
+
+	UPROPERTY(EditDefaultsOnly, Category = Transition)
+	EStateTreeTransitionEvent Event = EStateTreeTransitionEvent::OnCondition;
+
+	UPROPERTY(EditDefaultsOnly, Category = Transition)
+	FStateTreeStateLink State;
+
+	// Gate delay in seconds.
+	UPROPERTY(EditDefaultsOnly, Category = Transition, meta = (UIMin = "0", ClampMin = "0", UIMax = "25", ClampMax = "25"))
+	float GateDelay = 0.0f;
+	
+	UPROPERTY(EditDefaultsOnly, Category = Transition)
+	TArray<FStateTreeConditionItem> Conditions;
+};
+
+
+/**
+ * Editor representation of a state in StateTree
+ */
+UCLASS(BlueprintType, EditInlineNew, CollapseCategories)
+class STATETREEEDITORMODULE_API UStateTreeState : public UObject, public IStateTreeVariableProvider
+{
+	GENERATED_BODY()
+
+public:
+	UStateTreeState(const FObjectInitializer& ObjectInitializer);
+
+#if WITH_EDITOR
+	virtual void PostEditChangeChainProperty(FPropertyChangedChainEvent& PropertyChangedEvent) override;
+	virtual void PostLoad() override;
+#endif
+
+	void GetVisibleVariables(FStateTreeVariableLayout& Variables) const override;
+
+	const UStateTreeTaskBase* GetTaskByID(FGuid ID) const;
+	UStateTreeTaskBase* GetTaskByID(FGuid ID);
+
+	const FStateTreeTask2Base* GetTask2ByID(FGuid ID) const;
+	FStateTreeTask2Base* GetTask2ByID(FGuid ID);
+
+	UStateTreeState* GetNextSiblingState() const;
+
+	// StateTree Builder API
+	
+	/** Adds child state with specified name. */
+	UStateTreeState& AddChildState(const FName ChildName)
+	{
+		UStateTreeState* ChildState = NewObject<UStateTreeState>(this);
+		check(ChildState);
+		ChildState->Name = ChildName;
+		ChildState->Parent = this;
+		Children.Add(ChildState);
+		return *ChildState;
+	}
+
+	/**
+	 * Adds enter condition of specified type.
+	 * @return reference to the new condition. 
+	 */
+	template<typename T, typename... TArgs>
+	T& AddEnterCondition(TArgs&&... InArgs)
+	{
+		FStateTreeConditionItem& CondItem = EnterConditions2.AddDefaulted_GetRef();
+		CondItem.Type.InitializeAs<T>(Forward<TArgs>(InArgs)...);
+		return CondItem.Type.GetMutable<T>();
+	}
+
+	/**
+	 * Adds Task of specified type.
+	 * @return reference to the new Task. 
+	 */
+	template<typename T, typename... TArgs>
+	T& AddTask(TArgs&&... InArgs)
+	{
+		FStateTreeTaskItem& TaskItem = Tasks2.AddDefaulted_GetRef();
+		TaskItem.Type.InitializeAs<T>(Forward<TArgs>(InArgs)...);
+		T& Task = TaskItem.Type.GetMutable<T>();
+		Task.ID = FGuid::NewGuid();
+		return Task;
+	}
+
+	/**
+	 * Adds Evaluator of specified type.
+	 * @return reference to the new Evaluator. 
+	 */
+	template<typename T, typename... TArgs>
+    T& AddEvaluator(TArgs&&... InArgs)
+	{
+		FStateTreeEvaluatorItem& EvalItem = Evaluators2.AddDefaulted_GetRef();
+		EvalItem.Type.InitializeAs<T>(Forward<TArgs>(InArgs)...);
+		T& Eval = EvalItem.Type.GetMutable<T>();
+		Eval.ID = FGuid::NewGuid();
+		return Eval;
+	}
+
+	/**
+	 * Adds Transition.
+	 * @return reference to the new Transition. 
+	 */
+	FStateTreeTransition2& AddTransition(const EStateTreeTransitionEvent InEvent, const EStateTreeTransitionType InType, const UStateTreeState* InState = nullptr)
+	{
+		return Transitions2.Emplace_GetRef(InEvent, InType, InState);
+	}
+
+	// ~StateTree Builder API
+
+	bool IsV2() const;
+	
+	UPROPERTY(EditDefaultsOnly, Category = State)
+	FName Name;
+
+	UPROPERTY()
+	FGuid ID;
+
+	UPROPERTY(EditDefaultsOnly, Category = "State", Meta = (ToolTip = "State to transition to when the State execution is done. If no State Failed transition is set, this transition will be used."))
+	FStateTreeStateLink StateDoneTransition;
+
+	UPROPERTY(EditDefaultsOnly, Category = "State", Meta = (ToolTip = "State to transition to if the State execution fails."))
+	FStateTreeStateLink StateFailedTransition = FStateTreeStateLink(EStateTreeTransitionType::NotSet);
+
+	UPROPERTY(EditDefaultsOnly, Category = "Enter Conditions")
+	TArray<FStateTreeCondition> EnterConditions;
+
+	UPROPERTY(EditDefaultsOnly, Category = "Evaluators", Instanced)
+	TArray<UStateTreeEvaluatorBase*> Evaluators;
+
+	UPROPERTY(EditDefaultsOnly, Category = "Tasks", Instanced)
+	TArray<UStateTreeTaskBase*> Tasks;
+
+	UPROPERTY(EditDefaultsOnly, Category = "Transitions")
+	TArray<FStateTreeTransition> Transitions;
+
+// STATETREE_V2
+	UPROPERTY(EditDefaultsOnly, Category = "Enter Conditions")
+	TArray<FStateTreeConditionItem> EnterConditions2;
+
+	UPROPERTY(EditDefaultsOnly, Category = "Evaluators")
+	TArray<FStateTreeEvaluatorItem> Evaluators2;
+
+	UPROPERTY(EditDefaultsOnly, Category = "Tasks")
+	TArray<FStateTreeTaskItem> Tasks2;
+
+	UPROPERTY(EditDefaultsOnly, Category = "Transitions")
+	TArray<FStateTreeTransition2> Transitions2;
+// ~STATETREE_V2
+
+	UPROPERTY()
+	TArray<UStateTreeState*> Children;
+
+	UPROPERTY()
+	bool bExpanded;
+
+	UPROPERTY()
+	UStateTreeState* Parent;
+
+};
diff --git a/Engine/Plugins/Runtime/StateTree/Source/StateTreeEditorModule/Public/StateTreeVariableProvider.h b/Engine/Plugins/Runtime/StateTree/Source/StateTreeEditorModule/Public/StateTreeVariableProvider.h
new file mode 100644
--- /dev/null
+++ b/Engine/Plugins/Runtime/StateTree/Source/StateTreeEditorModule/Public/StateTreeVariableProvider.h
@@ -0,0 +1,29 @@
+// Copyright Epic Games, Inc. All Rights Reserved.
+
+#pragma once
+
+#include "CoreMinimal.h"
+#include "UObject/ObjectMacros.h"
+#include "UObject/Interface.h"
+#include "StateTreeVariableDesc.h"
+#include "StateTreeVariableProvider.generated.h"
+
+struct FStateTreeVariableLayout;
+
+/**
+ * Variable provider is used by the UI (see StateTreeVariableBindingDetails) to get variables that are visible for a StateTreeVariableBinding property.
+ * The UI uses outer objects to find the nearest provider.
+*/
+UINTERFACE(Blueprintable, MinimalAPI, meta = (CannotImplementInterfaceInBlueprint))
+class UStateTreeVariableProvider : public UInterface
+{
+	GENERATED_UINTERFACE_BODY()
+};
+
+class STATETREEEDITORMODULE_API IStateTreeVariableProvider
+{
+	GENERATED_IINTERFACE_BODY()
+
+	virtual void GetVisibleVariables(FStateTreeVariableLayout& Variables) const PURE_VIRTUAL(IStateTreeVariableProvider::GetVisibleVariableBlocks, return; );
+};
+
diff --git a/Engine/Plugins/Runtime/StateTree/Source/StateTreeModule/Private/Conditions/StateTreeCondition_Common.cpp b/Engine/Plugins/Runtime/StateTree/Source/StateTreeModule/Private/Conditions/StateTreeCondition_Common.cpp
new file mode 100644
--- /dev/null
+++ b/Engine/Plugins/Runtime/StateTree/Source/StateTreeModule/Private/Conditions/StateTreeCondition_Common.cpp
@@ -0,0 +1,359 @@
+// Copyright Epic Games, Inc. All Rights Reserved.
+
+#include "Conditions/StateTreeCondition_Common.h"
+#include "StateTreePropertyBindings.h"
+
+#define LOCTEXT_NAMESPACE "StateTreeEditor"
+
+namespace UE { namespace StateTree { namespace Conditions {
+
+#if WITH_EDITOR
+FText GetOperatorText(const EGenericAICheck Operator)
+{
+	switch (Operator)
+	{
+	case EGenericAICheck::Equal:
+		return FText::FromString(TEXT("=="));
+		break;
+	case EGenericAICheck::NotEqual:
+		return FText::FromString(TEXT("!="));
+		break;
+	case EGenericAICheck::Less:
+		return FText::FromString(TEXT("&lt;"));
+		break;
+	case EGenericAICheck::LessOrEqual:
+		return FText::FromString(TEXT("&lt;="));
+		break;
+	case EGenericAICheck::Greater:
+		return FText::FromString(TEXT("&gt;"));
+		break;
+	case EGenericAICheck::GreaterOrEqual:
+		return FText::FromString(TEXT("&gt;="));
+		break;
+	default:
+		return FText::FromString(TEXT("??"));
+		break;
+	}
+	return FText::GetEmpty();
+}
+#endif // WITH_EDITOR
+
+template<typename T>
+bool CompareNumbers(const T Left, const T Right, const EGenericAICheck Operator)
+{
+	switch (Operator)
+	{
+	case EGenericAICheck::Equal:
+		return Left == Right;
+		break;
+	case EGenericAICheck::NotEqual:
+		return Left != Right;
+		break;
+	case EGenericAICheck::Less:
+		return Left < Right;
+		break;
+	case EGenericAICheck::LessOrEqual:
+		return Left <= Right;
+		break;
+	case EGenericAICheck::Greater:
+		return Left > Right;
+		break;
+	case EGenericAICheck::GreaterOrEqual:
+		return Left >= Right;
+		break;
+	default:
+		ensureMsgf(false, TEXT("Unhandled operator %d"), Operator);
+		return false;
+		break;
+	}
+	return false;
+}
+
+}}} // UE::StateTree::Conditions
+
+#if WITH_EDITOR
+FText FStateTreeCondition_CompareInt::GetDescription(const IStateTreeBindingLookup& BindingLookup) const
+{
+	const FStateTreeEditorPropertyPath LeftPath(ID, GET_MEMBER_NAME_STRING_CHECKED(FStateTreeCondition_CompareInt, Left));
+	const FStateTreeEditorPropertyPath RightPath(ID, GET_MEMBER_NAME_STRING_CHECKED(FStateTreeCondition_CompareInt, Right));
+
+	FText InvertText;
+	if (bInvert)
+	{
+		InvertText = LOCTEXT("Not", "Not");
+	}
+
+	FText LeftText;
+	if (const FStateTreeEditorPropertyPath* Binding = BindingLookup.GetPropertyBindingSource(LeftPath))
+	{
+		LeftText = BindingLookup.GetPropertyPathDisplayName(*Binding);
+	}
+	else
+	{
+		LeftText = FText::AsNumber(Left);
+	}
+
+	FText OperatorText = UE::StateTree::Conditions::GetOperatorText(Operator);
+
+	FText RightText;
+	if (const FStateTreeEditorPropertyPath* Binding = BindingLookup.GetPropertyBindingSource(RightPath))
+	{
+		RightText = BindingLookup.GetPropertyPathDisplayName(*Binding);
+	}
+	else
+	{
+		RightText = FText::AsNumber(Right);
+	}
+
+	return FText::Format(LOCTEXT("CompareIntDesc", "{0} <Details.Bold>{1}</> {2} <Details.Bold>{3}</>"), InvertText, LeftText, OperatorText, RightText);
+}
+#endif
+
+bool FStateTreeCondition_CompareInt::TestCondition() const
+{
+	const bool bResult = UE::StateTree::Conditions::CompareNumbers<int32>(Left, Right, Operator);
+	return bResult ^ bInvert;
+}
+
+
+#if WITH_EDITOR
+FText FStateTreeCondition_CompareFloat::GetDescription(const IStateTreeBindingLookup& BindingLookup) const
+{
+	const FStateTreeEditorPropertyPath LeftPath(ID, GET_MEMBER_NAME_STRING_CHECKED(FStateTreeCondition_CompareFloat, Left));
+	const FStateTreeEditorPropertyPath RightPath(ID, GET_MEMBER_NAME_STRING_CHECKED(FStateTreeCondition_CompareFloat, Right));
+
+	FText InvertText;
+	if (bInvert)
+	{
+		InvertText = LOCTEXT("Not", "Not");
+	}
+
+	FText LeftText;
+	if (const FStateTreeEditorPropertyPath* Binding = BindingLookup.GetPropertyBindingSource(LeftPath))
+	{
+		LeftText = BindingLookup.GetPropertyPathDisplayName(*Binding);
+	}
+	else
+	{
+		LeftText = FText::AsNumber(Left);
+	}
+
+	FText OperatorText = UE::StateTree::Conditions::GetOperatorText(Operator);
+
+	FText RightText;
+	if (const FStateTreeEditorPropertyPath* Binding = BindingLookup.GetPropertyBindingSource(RightPath))
+	{
+		RightText = BindingLookup.GetPropertyPathDisplayName(*Binding);
+	}
+	else
+	{
+		RightText = FText::AsNumber(Right);
+	}
+
+	return FText::Format(LOCTEXT("CompareFloatDesc", "{0} <Details.Bold>{1}</> {2} <Details.Bold>{3}</>"), InvertText, LeftText, OperatorText, RightText);
+}
+#endif
+
+bool FStateTreeCondition_CompareFloat::TestCondition() const
+{
+	const bool bResult = UE::StateTree::Conditions::CompareNumbers<float>(Left, Right, Operator);
+	return bResult ^ bInvert;
+}
+
+
+#if WITH_EDITOR
+FText FStateTreeCondition_CompareBool::GetDescription(const IStateTreeBindingLookup& BindingLookup) const
+{
+	const FStateTreeEditorPropertyPath LeftPath(ID, GET_MEMBER_NAME_STRING_CHECKED(FStateTreeCondition_CompareBool, bLeft));
+	const FStateTreeEditorPropertyPath RightPath(ID, GET_MEMBER_NAME_STRING_CHECKED(FStateTreeCondition_CompareBool, bRight));
+
+	FText InvertText;
+	if (bInvert)
+	{
+		InvertText = LOCTEXT("Not", "Not");
+	}
+
+	FText LeftText;
+	if (const FStateTreeEditorPropertyPath* Binding = BindingLookup.GetPropertyBindingSource(LeftPath))
+	{
+		LeftText = BindingLookup.GetPropertyPathDisplayName(*Binding);
+	}
+	else
+	{
+		LeftText = FText::FromString(LexToSanitizedString(bLeft));
+	}
+
+	FText RightText;
+	if (const FStateTreeEditorPropertyPath* Binding = BindingLookup.GetPropertyBindingSource(RightPath))
+	{
+		RightText = BindingLookup.GetPropertyPathDisplayName(*Binding);
+	}
+	else
+	{
+		RightText = FText::FromString(LexToSanitizedString(bRight));
+	}
+
+	return FText::Format(LOCTEXT("CompareBoolDesc", "{0} <Details.Bold>{1}</> is <Details.Bold>{2}</>"), InvertText, LeftText, RightText);
+}
+#endif
+
+bool FStateTreeCondition_CompareBool::TestCondition() const
+{
+	return (bLeft == bRight) ^ bInvert;
+}
+
+
+#if WITH_EDITOR
+FText FStateTreeCondition_CompareEnum::GetDescription(const IStateTreeBindingLookup& BindingLookup) const
+{
+	const FStateTreeEditorPropertyPath LeftPath(ID, GET_MEMBER_NAME_STRING_CHECKED(FStateTreeCondition_CompareEnum, Left));
+	const FStateTreeEditorPropertyPath RightPath(ID, GET_MEMBER_NAME_STRING_CHECKED(FStateTreeCondition_CompareEnum, Right));
+
+	FText InvertText;
+	if (bInvert)
+	{
+		InvertText = LOCTEXT("Not", "Not");
+	}
+
+	FText LeftText;
+	if (const FStateTreeEditorPropertyPath* Binding = BindingLookup.GetPropertyBindingSource(LeftPath))
+	{
+		LeftText = BindingLookup.GetPropertyPathDisplayName(*Binding);
+	}
+	else
+	{
+		if (Left.Enum)
+		{
+			LeftText = Left.Enum->GetDisplayNameTextByValue(int64(Left.Value));
+		}
+		else
+		{
+			LeftText = LOCTEXT("Invalid", "Invalid");
+		}
+	}
+
+	FText RightText;
+	if (const FStateTreeEditorPropertyPath* Binding = BindingLookup.GetPropertyBindingSource(RightPath))
+	{
+		RightText = BindingLookup.GetPropertyPathDisplayName(*Binding);
+	}
+	else
+	{
+		if (Right.Enum)
+		{
+			RightText = Right.Enum->GetDisplayNameTextByValue(int64(Right.Value));
+		}
+		else
+		{
+			RightText = LOCTEXT("Invalid", "Invalid");
+		}
+	}
+
+	return FText::Format(LOCTEXT("EnumEqualsDesc", "{0} <Details.Bold>{1}</> is <Details.Bold>{2}</>"), InvertText, LeftText, RightText);
+}
+
+void FStateTreeCondition_CompareEnum::OnBindingChanged(const FStateTreeEditorPropertyPath& SourcePath, const FStateTreeEditorPropertyPath& TargetPath, const IStateTreeBindingLookup& BindingLookup)
+{
+	if (!TargetPath.IsValid())
+	{
+		return;
+	}
+
+	// Left has changed, update enums from the leaf property.
+	if (TargetPath.Path.Last() == GET_MEMBER_NAME_STRING_CHECKED(FStateTreeCondition_CompareEnum, Left))
+	{
+		if (const FProperty* LeafProperty = BindingLookup.GetPropertyPathLeafProperty(SourcePath))
+		{
+			// Handle both old stype namespace enums and new class enum properties.
+			UEnum* NewEnum = nullptr;
+			if (const FByteProperty* ByteProperty = CastField<FByteProperty>(LeafProperty))
+			{
+				NewEnum = ByteProperty->GetIntPropertyEnum();
+			}
+			else if (const FEnumProperty* EnumProperty = CastField<FEnumProperty>(LeafProperty))
+			{
+				NewEnum = EnumProperty->GetEnum();
+			}
+
+			if (Left.Enum != NewEnum)
+			{
+				Left.Initialize(NewEnum);
+			}
+		}
+		else
+		{
+			Left.Initialize(nullptr);
+		}
+
+		if (Right.Enum != Left.Enum)
+		{
+			Right.Initialize(Left.Enum);
+		}
+	}
+}
+
+#endif
+bool FStateTreeCondition_CompareEnum::TestCondition() const
+{
+	return (Left == Right) ^ bInvert;
+}
+
+
+#if WITH_EDITOR
+FText FStateTreeCondition_CompareDistance::GetDescription(const IStateTreeBindingLookup& BindingLookup) const
+{
+	const FStateTreeEditorPropertyPath SourcePath(ID, GET_MEMBER_NAME_STRING_CHECKED(FStateTreeCondition_CompareDistance, Source));
+	const FStateTreeEditorPropertyPath TargetPath(ID, GET_MEMBER_NAME_STRING_CHECKED(FStateTreeCondition_CompareDistance, Target));
+	const FStateTreeEditorPropertyPath DistancePath(ID, GET_MEMBER_NAME_STRING_CHECKED(FStateTreeCondition_CompareDistance, Distance));
+
+	FText InvertText;
+	if (bInvert)
+	{
+		InvertText = LOCTEXT("Not", "Not");
+	}
+
+	FText SourceText;
+	if (const FStateTreeEditorPropertyPath* Binding = BindingLookup.GetPropertyBindingSource(SourcePath))
+	{
+		SourceText = BindingLookup.GetPropertyPathDisplayName(*Binding);
+	}
+	else
+	{
+		SourceText = Source.ToCompactText();
+	}
+
+	FText TargetText;
+	if (const FStateTreeEditorPropertyPath* Binding = BindingLookup.GetPropertyBindingSource(TargetPath))
+	{
+		TargetText = BindingLookup.GetPropertyPathDisplayName(*Binding);
+	}
+	else
+	{
+		TargetText = Target.ToCompactText();
+	}
+
+	FText OperatorText = UE::StateTree::Conditions::GetOperatorText(Operator);
+
+	FText DistanceText;
+	if (const FStateTreeEditorPropertyPath* DistanceBinding = BindingLookup.GetPropertyBindingSource(DistancePath))
+	{
+		DistanceText = BindingLookup.GetPropertyPathDisplayName(*DistanceBinding);
+	}
+	else
+	{
+		DistanceText = FText::AsNumber(Distance);
+	}
+
+	return FText::Format(LOCTEXT("CompareDistanceDesc", "{0} Distance <Details.Subdued>from</> <Details.Bold>{1}</> <Details.Subdued>to</> <Details.Bold>{2}</> {3} <Details.Bold>{4}</>"), InvertText, SourceText, TargetText, OperatorText, DistanceText);
+}
+#endif
+
+bool FStateTreeCondition_CompareDistance::TestCondition() const
+{
+	const float Left = FVector::DistSquared(Source, Target);
+	const float Right = FMath::Square(Distance);
+	const bool bResult = UE::StateTree::Conditions::CompareNumbers<float>(Left, Right, Operator);
+	return bResult ^ bInvert;
+}
+
+#undef LOCTEXT_NAMESPACE
