Moved GameplayBehaviors out of restricted folder over to Experimental

diff --git a/Engine/Plugins/Runtime/StateTree/Source/StateTreeModule/Private/Evaluators/StateTreeEvaluator_Test.cpp b/Engine/Plugins/Runtime/StateTree/Source/StateTreeModule/Private/Evaluators/StateTreeEvaluator_Test.cpp
new file mode 100644
--- /dev/null
+++ b/Engine/Plugins/Runtime/StateTree/Source/StateTreeModule/Private/Evaluators/StateTreeEvaluator_Test.cpp
@@ -0,0 +1,135 @@
+// Copyright Epic Games, Inc. All Rights Reserved.
+
+#include "Evaluators/StateTreeEvaluator_Test.h"
+#include "CoreMinimal.h"
+#include "GameFramework/Actor.h"
+#include "StateTreeInstance.h"
+#include "StateTreeVariableLayout.h"
+#include "StateTreeConstantStorage.h"
+#include "Engine/Engine.h"
+
+UStateTreeEvaluator_Test::UStateTreeEvaluator_Test(const FObjectInitializer& ObjectInitializer)
+	:Super(ObjectInitializer)
+	, Timer(0.0f)
+	, SomeFloat(EStateTreeVariableBindingMode::Typed, EStateTreeVariableType::Float)
+	, BoolVar(EStateTreeVariableBindingMode::Definition, FName(TEXT("Bool")), EStateTreeVariableType::Bool)
+	, FloatVar(EStateTreeVariableBindingMode::Definition, FName(TEXT("Float")), EStateTreeVariableType::Float)
+	, IntVar(EStateTreeVariableBindingMode::Definition, FName(TEXT("Int")), EStateTreeVariableType::Int)
+	, VectorVar(EStateTreeVariableBindingMode::Definition, FName(TEXT("Vector")), EStateTreeVariableType::Vector)
+	, Object1Var(EStateTreeVariableBindingMode::Definition, FName(TEXT("Object1")), EStateTreeVariableType::Object, AActor::StaticClass())
+	, Object2Var(EStateTreeVariableBindingMode::Definition, FName(TEXT("Object2")), EStateTreeVariableType::Object, AActor::StaticClass())
+	, ObjectRVar(EStateTreeVariableBindingMode::Definition, FName(TEXT("ObjectR")), EStateTreeVariableType::Object, UStateTreeEvaluatorBase::StaticClass())
+{
+}
+
+
+bool UStateTreeEvaluator_Test::Initialize(FStateTreeInstance& StateTreeInstance)
+{
+	return true;
+}
+
+void UStateTreeEvaluator_Test::Activate(FStateTreeInstance& StateTreeInstance)
+{
+	// Clear variables
+	StateTreeInstance.SetValueBool(BoolVar.Handle, false);
+	StateTreeInstance.SetValueFloat(FloatVar.Handle, 0.0f);
+	StateTreeInstance.SetValueInt(IntVar.Handle, 0);
+	StateTreeInstance.SetValueVector(VectorVar.Handle, FVector(1.0f, 20.0f, 300.0f));
+	StateTreeInstance.SetValueObject(Object1Var.Handle, nullptr);
+	StateTreeInstance.SetValueObject(Object2Var.Handle, nullptr);
+	StateTreeInstance.SetValueObject(ObjectRVar.Handle, nullptr);
+
+	Timer = 0.0f;
+}
+
+void UStateTreeEvaluator_Test::Deactivate(FStateTreeInstance& StateTreeInstance)
+{
+}
+
+void UStateTreeEvaluator_Test::Tick(FStateTreeInstance& StateTreeInstance, const float DeltaTime)
+{
+	Timer += DeltaTime;
+
+	const float FloatVal = StateTreeInstance.GetValueFloat(SomeFloat.Handle, 0.0f);
+
+	StateTreeInstance.SetValueFloat(FloatVar.Handle, Timer);
+	StateTreeInstance.SetValueInt(IntVar.Handle, int32(Timer / 4.0f));
+	StateTreeInstance.SetValueBool(BoolVar.Handle, FMath::Frac(Timer / 3.0f) > (2.5f / 3.0f));
+}
+
+#if WITH_GAMEPLAY_DEBUGGER
+void UStateTreeEvaluator_Test::AppendDebugInfoString(FString& DebugString, const FStateTreeInstance& StateTreeInstance) const
+{
+	Super::AppendDebugInfoString(DebugString, StateTreeInstance);
+}
+#endif // WITH_GAMEPLAY_DEBUGGER
+
+
+#if WITH_EDITOR
+void UStateTreeEvaluator_Test::DefineOutputVariables(FStateTreeVariableLayout& Variables) const
+{
+	Variables.DefineVariable(Name, BoolVar);
+	Variables.DefineVariable(Name, FloatVar);
+	Variables.DefineVariable(Name, IntVar);
+	Variables.DefineVariable(Name, VectorVar);
+	Variables.DefineVariable(Name, Object1Var);
+	Variables.DefineVariable(Name, Object2Var);
+	Variables.DefineVariable(Name, ObjectRVar);
+}
+
+bool UStateTreeEvaluator_Test::ResolveVariables(const FStateTreeVariableLayout& Variables, FStateTreeConstantStorage& Constants, UObject* Outer)
+{
+	SomeFloat.ResolveHandle(Variables, Constants);
+
+	BoolVar.ResolveHandle(Variables, Constants);
+	FloatVar.ResolveHandle(Variables, Constants);
+	IntVar.ResolveHandle(Variables, Constants);
+	VectorVar.ResolveHandle(Variables, Constants);
+	Object1Var.ResolveHandle(Variables, Constants);
+	Object2Var.ResolveHandle(Variables, Constants);
+	ObjectRVar.ResolveHandle(Variables, Constants);
+
+	return true;
+}
+#endif
+
+
+FStateTreeEvaluator2_Test::FStateTreeEvaluator2_Test()
+{
+}
+
+FStateTreeEvaluator2_Test::~FStateTreeEvaluator2_Test()
+{
+}
+
+void FStateTreeEvaluator2_Test::EnterState(FStateTreeExecutionContext& Context, const EStateTreeStateChangeType ChangeType, const FStateTreeTransitionResult& Transition)
+{
+	Foo = 0.0f;
+	IntVal = 0;
+	bBoolVal = false;
+	OutResult = &DummyResult;
+	DummyResult.Count = 0;
+}
+
+void FStateTreeEvaluator2_Test::ExitState(FStateTreeExecutionContext& Context, const EStateTreeStateChangeType ChangeType, const FStateTreeTransitionResult& Transition)
+{
+	OutResult = &DummyResult;
+}
+
+void FStateTreeEvaluator2_Test::Evaluate(FStateTreeExecutionContext& Context, const EStateTreeEvaluationType EvalType, const float DeltaTime)
+{
+
+	Foo += DeltaTime;
+	IntVal = int32(Foo / 4.0f);
+	bBoolVal = FMath::Frac(Foo / 3.0f) > (2.5f / 3.0f);
+
+	if (FDummyStateTreeResult* OtherDummy = InResult.GetMutablePtr<FDummyStateTreeResult>())
+	{
+		OtherDummy->Count = IntVal;
+	}
+	
+	OutResult = &DummyResult;
+
+	// Keeping this here until we get Gameplay Debugger to work again.
+//	GEngine->AddOnScreenDebugMessage(INDEX_NONE, 1.0f, FColor::Orange, *FString::Printf(TEXT("[%s] Foo=%f IntVal=%d bBoolVal=%s\n"), *Name.ToString(), Foo, IntVal, *LexToString(bBoolVal)));
+}
diff --git a/Engine/Plugins/Runtime/StateTree/Source/StateTreeModule/Private/StateTree.cpp b/Engine/Plugins/Runtime/StateTree/Source/StateTreeModule/Private/StateTree.cpp
new file mode 100644
--- /dev/null
+++ b/Engine/Plugins/Runtime/StateTree/Source/StateTreeModule/Private/StateTree.cpp
@@ -0,0 +1,190 @@
+// Copyright Epic Games, Inc. All Rights Reserved.
+
+#include "StateTree.h"
+#include "StateTreeEvaluatorBase.h"
+#include "CoreMinimal.h"
+#include "StateTreeDelegates.h"
+
+UStateTree::UStateTree(const FObjectInitializer& ObjectInitializer)
+	: Super(ObjectInitializer)
+{
+}
+
+bool UStateTree::IsValidStateTree() const
+{
+	// Valid tree must have at least one state.
+	return States.Num() > 0;
+}
+
+
+#if WITH_EDITOR
+void UStateTree::ResetBaked()
+{
+	Variables.Reset();
+	Constants.Reset();
+	Tasks.Reset();
+	Evaluators.Reset();
+	States.Reset();
+	Conditions.Reset();
+	Conditions2.Reset();
+	Transitions.Reset();
+	Parameters.Reset();
+
+	RuntimeStorageItems.Reset();
+	RuntimeStorageStruct = nullptr;
+	RuntimeStorageOffsets.Reset();
+	RuntimeStorageDefaultValue.Reset();
+	ExternalItems.Reset();
+	PropertyBindings.Reset();
+}
+
+void UStateTree::PostEditChangeProperty(struct FPropertyChangedEvent& PropertyChangedEvent)
+{
+	if (PropertyChangedEvent.MemberProperty && PropertyChangedEvent.Property)
+	{
+		if (PropertyChangedEvent.MemberProperty->GetFName() == GET_MEMBER_NAME_CHECKED(UStateTree, Parameters))
+		{
+			UE::StateTree::Delegates::OnParameterLayoutChanged.Broadcast(*this);
+		}
+
+		if (PropertyChangedEvent.MemberProperty->GetFName() == GET_MEMBER_NAME_CHECKED(UStateTree, Schema))
+		{
+			UE::StateTree::Delegates::OnParameterLayoutChanged.Broadcast(*this);
+		}
+
+		if (PropertyChangedEvent.Property->GetOwnerStruct() == FStateTreeVariableDesc::StaticStruct() &&
+			PropertyChangedEvent.Property->GetFName() == GET_MEMBER_NAME_CHECKED(FStateTreeVariableDesc, Name))
+		{
+			UE::StateTree::Delegates::OnIdentifierChanged.Broadcast(*this);
+		}
+	}
+}
+
+void UStateTree::GetAssetRegistryTags(TArray<FAssetRegistryTag>& OutTags) const
+{
+	static const FName SchemaTag(TEXT("Schema"));
+
+	const FString SchemaClassName = Schema ? Schema->GetClass()->GetName() : TEXT("");
+	OutTags.Add(FAssetRegistryTag(SchemaTag, SchemaClassName, FAssetRegistryTag::TT_Alphabetical));
+
+	Super::GetAssetRegistryTags(OutTags);
+}
+#endif // WITH_EDITOR
+
+void UStateTree::PostLoad()
+{
+	Super::PostLoad();
+	PropertyBindings.ResolvePaths();
+
+#if WITH_EDITOR
+	InitRuntimeStorage();
+#else
+	// Item offsets still need to be calculated in non editor target since the struct sizes might be different.
+	checkf(RuntimeStorageOffsets.Num() == 0, TEXT("RuntimeStorageOffsets is transient and should only be computed once."));
+
+	for (TFieldIterator<FProperty> PropertyIt(RuntimeStorageStruct); PropertyIt; ++PropertyIt)
+	{
+		FStructProperty* StructProperty = CastField<FStructProperty>(*PropertyIt);
+		checkf(StructProperty, TEXT("RuntimeStorageStruct is expected to only contain Struct properties"));
+		RuntimeStorageOffsets.Emplace(StructProperty->Struct, StructProperty->GetOffset_ForInternal());
+	}
+#endif // WITH_EDITOR
+}
+
+void UStateTree::BeginDestroy()
+{
+	Super::BeginDestroy();
+	
+	// Destroy the runtime storage before the UScriptStruct it uses gets removed.
+	RuntimeStorageDefaultValue.Reset();
+}
+
+void UStateTree::ResolvePropertyPaths()
+{
+	// TODO: find better hook when to call this. Currently this gets called from StateTreeComponent, it should be called once PostLoad() or when entering SIE/PIE in editor.
+	PropertyBindings.ResolvePaths();
+}
+
+void UStateTree::InitRuntimeStorage()
+{
+	RuntimeStorageStruct = nullptr;
+	RuntimeStorageOffsets.Reset();
+	RuntimeStorageDefaultValue.Reset();
+
+	// Check that the items are valid before trying to create the type.
+	// The structs can become invalid i.e. because of missing type.
+	bool bValid = true;
+	for (const FInstancedStruct& Item : RuntimeStorageItems)
+	{
+		if (!Item.IsValid())
+		{
+			bValid = false;
+			break;
+		}
+	}
+	if (!bValid)
+	{
+		return;
+	}
+
+	const FString StructName = GetName() + TEXT("_RuntimeStorage");
+
+	// Remove existing struct of same name.
+	UScriptStruct* OldStruct = FindObject<UScriptStruct>(this, *StructName);
+	if (OldStruct)
+	{
+		const FString OldStructName = MakeUniqueObjectName(OldStruct->GetOuter(), OldStruct->GetClass(), *FString::Printf(TEXT("%s_TRASH"), *OldStruct->GetName())).ToString();
+		OldStruct->SetFlags(RF_NewerVersionExists);
+		OldStruct->ClearFlags(RF_Public | RF_Standalone);
+		OldStruct->SetStructTrashed(/*bIsTrash*/true);
+		OldStruct->Rename(*OldStructName, nullptr, REN_DontCreateRedirectors|REN_ForceNoResetLoaders);
+	}
+
+	UScriptStruct* NewStruct = NewObject<UScriptStruct>(this, *StructName, RF_Public);
+
+	if (Schema)
+	{
+		NewStruct->SetSuperStruct(Schema->GetStorageSuperStruct());
+	}
+
+	// Append all evaluators and tasks.
+	// Since properties are stored in linked list, add in reverse order so that we retain the correct order in output.
+	TArray<FStructProperty*> NewProperties;
+	NewProperties.SetNumZeroed(RuntimeStorageItems.Num());
+	for (int32 ItemIndex = RuntimeStorageItems.Num() - 1; ItemIndex >= 0; ItemIndex--)
+	{
+		const FInstancedStruct& ItemPtr = RuntimeStorageItems[ItemIndex];
+		UScriptStruct* ItemStruct = const_cast<UScriptStruct*>(ItemPtr.GetScriptStruct());
+		FName PropName(FString::Printf(TEXT("%s%d"), *ItemStruct->GetName(), NewProperties.Num()));
+		FStructProperty* NewStructProperty = new FStructProperty(NewStruct, PropName, RF_Public);
+		NewStructProperty->Struct = ItemStruct;
+
+		NewStruct->AddCppProperty(NewStructProperty);
+
+		NewProperties[ItemIndex] = NewStructProperty;
+	}
+
+	// Finalize the struct
+	NewStruct->Bind();
+	NewStruct->StaticLink(/*bRelinkExistingProperties*/true);
+
+	// Store item offsets for fast struct view creation at runtime.
+	TArray<FStateTreeRuntimeStorageItemOffset> NewItemOffsets;
+	for (FStructProperty* StructProperty : NewProperties)
+	{
+		NewItemOffsets.Emplace(StructProperty->Struct, StructProperty->GetOffset_ForInternal());
+	}
+	check(RuntimeStorageItems.Num() == NewItemOffsets.Num());
+
+	// Instantiate default value
+	RuntimeStorageDefaultValue.InitializeAs(NewStruct);
+	for (int i = 0; i < NewItemOffsets.Num(); i++)
+	{
+		const FStateTreeRuntimeStorageItemOffset& Item = NewItemOffsets[i];
+		uint8* Dest = RuntimeStorageDefaultValue.GetMutableMemory() + Item.Offset;
+		Item.Struct->CopyScriptStruct(Dest, RuntimeStorageItems[i].GetMemory());
+	}
+
+	RuntimeStorageOffsets = NewItemOffsets;
+	RuntimeStorageStruct = NewStruct;
+}
diff --git a/Engine/Plugins/Runtime/StateTree/Source/StateTreeModule/Private/StateTreeCondition.cpp b/Engine/Plugins/Runtime/StateTree/Source/StateTreeModule/Private/StateTreeCondition.cpp
new file mode 100644
--- /dev/null
+++ b/Engine/Plugins/Runtime/StateTree/Source/StateTreeModule/Private/StateTreeCondition.cpp
@@ -0,0 +1,58 @@
+// Copyright Epic Games, Inc. All Rights Reserved.
+
+#include "StateTreeCondition.h"
+#include "StateTreeVariableDesc.h"
+#include "StateTreeVariable.h"
+
+#define LOCTEXT_NAMESPACE "StateTreeEditor"
+
+FStateTreeCondition::FStateTreeCondition()
+	: Left(EStateTreeVariableBindingMode::Any, FName(), EStateTreeVariableType::Void)
+	, Right(EStateTreeVariableBindingMode::Typed, FName(), EStateTreeVariableType::Void)
+	, Operator(EGenericAICheck::MAX)
+{
+}
+
+#if WITH_EDITOR
+FText FStateTreeCondition::GetDescription() const
+{
+	if (!Left.IsBound())
+	{
+		return FText(LOCTEXT("ConditionNotSet", "Condition Not Set"));
+	}
+
+	FFormatNamedArguments Args;
+	Args.Add(TEXT("Left"), Left.GetDescription());
+
+	switch (Operator)
+	{
+	case EGenericAICheck::Equal:
+		Args.Add(TEXT("Operator"), LOCTEXT("OperatorEqual", "=="));
+		break;
+	case EGenericAICheck::NotEqual:
+		Args.Add(TEXT("Operator"), LOCTEXT("OperatorNotEqual", "!="));
+		break;
+	case EGenericAICheck::Less:
+		Args.Add(TEXT("Operator"), LOCTEXT("OperatorLess", "<"));
+		break;
+	case EGenericAICheck::LessOrEqual:
+		Args.Add(TEXT("Operator"), LOCTEXT("OperatorLessOrEqual", "<="));
+		break;
+	case EGenericAICheck::Greater:
+		Args.Add(TEXT("Operator"), LOCTEXT("OperatorGreater", ">"));
+		break;
+	case EGenericAICheck::GreaterOrEqual:
+		Args.Add(TEXT("Operator"), LOCTEXT("OperatorGreaterOrEqual", ">="));
+		break;
+	default:
+		Args.Add(TEXT("Operator"), LOCTEXT("OperatorUnknown", "??"));
+		break;
+	}
+
+	Args.Add(TEXT("Right"), Right.GetDescription());
+
+	return FText::Format(LOCTEXT("Condition", "{Left} {Operator} {Right}"), Args);
+}
+#endif
+
+#undef LOCTEXT_NAMESPACE
diff --git a/Engine/Plugins/Runtime/StateTree/Source/StateTreeModule/Private/StateTreeDelegates.cpp b/Engine/Plugins/Runtime/StateTree/Source/StateTreeModule/Private/StateTreeDelegates.cpp
new file mode 100644
--- /dev/null
+++ b/Engine/Plugins/Runtime/StateTree/Source/StateTreeModule/Private/StateTreeDelegates.cpp
@@ -0,0 +1,16 @@
+// Copyright Epic Games, Inc. All Rights Reserved.
+
+#include "StateTreeDelegates.h"
+#include "CoreMinimal.h"
+
+namespace UE::StateTree::Delegates
+{
+	
+#if WITH_EDITOR
+FOnIdentifierChanged OnIdentifierChanged;
+FOnParametersInvalidated OnParametersInvalidated;
+FOnParameterLayoutChanged OnParameterLayoutChanged;
+FOnSchemaChanged OnSchemaChanged;
+#endif
+	
+}; // UE::StateTree::Delegates
diff --git a/Engine/Plugins/Runtime/StateTree/Source/StateTreeModule/Private/StateTreeEvaluatorBase.cpp b/Engine/Plugins/Runtime/StateTree/Source/StateTreeModule/Private/StateTreeEvaluatorBase.cpp
new file mode 100644
--- /dev/null
+++ b/Engine/Plugins/Runtime/StateTree/Source/StateTreeModule/Private/StateTreeEvaluatorBase.cpp
@@ -0,0 +1,46 @@
+// Copyright Epic Games, Inc. All Rights Reserved.
+
+#include "StateTreeEvaluatorBase.h"
+#include "CoreMinimal.h"
+#include "StateTree.h"
+#include "StateTreeDelegates.h"
+
+// Base class for all evaluators
+UStateTreeEvaluatorBase::UStateTreeEvaluatorBase(const FObjectInitializer& ObjectInitializer)
+	: Super(ObjectInitializer)
+{
+#if WITH_EDITORONLY_DATA
+	ID = FGuid::NewGuid();
+#endif // WITH_EDITORONLY_DATA
+}
+
+#if WITH_GAMEPLAY_DEBUGGER
+void UStateTreeEvaluatorBase::AppendDebugInfoString(FString& DebugString, const FStateTreeInstance& StateTreeInstance) const
+{
+	DebugString += FString::Printf(TEXT("[%s]\n"), *Name.ToString());
+}
+#endif // WITH_GAMEPLAY_DEBUGGER
+
+#if WITH_EDITOR
+void UStateTreeEvaluatorBase::PostEditChangeProperty(struct FPropertyChangedEvent& PropertyChangedEvent)
+{
+	if (PropertyChangedEvent.Property)
+	{
+		if (PropertyChangedEvent.Property->GetFName() == GET_MEMBER_NAME_CHECKED(UStateTreeEvaluatorBase, Name))
+		{
+			UStateTree* StateTree = GetTypedOuter<UStateTree>();
+			if (ensure(StateTree))
+			{
+				UE::StateTree::Delegates::OnIdentifierChanged.Broadcast(*StateTree);
+			}
+		}
+	}
+}
+#endif
+
+#if WITH_GAMEPLAY_DEBUGGER
+void FStateTreeEvaluator2Base::AppendDebugInfoString(FString& DebugString, const FStateTreeExecutionContext& Context) const
+{
+	DebugString += FString::Printf(TEXT("[%s]\n"), *Name.ToString());
+}
+#endif // WITH_GAMEPLAY_DEBUGGER
diff --git a/Engine/Plugins/Runtime/StateTree/Source/StateTreeModule/Private/StateTreeExecutionContext.cpp b/Engine/Plugins/Runtime/StateTree/Source/StateTreeModule/Private/StateTreeExecutionContext.cpp
new file mode 100644
--- /dev/null
+++ b/Engine/Plugins/Runtime/StateTree/Source/StateTreeModule/Private/StateTreeExecutionContext.cpp
@@ -0,0 +1,1058 @@
+// Copyright Epic Games, Inc. All Rights Reserved.
+
+#include "StateTreeExecutionContext.h"
+#include "StateTree.h"
+#include "StateTreeTaskBase.h"
+#include "StateTreeEvaluatorBase.h"
+#include "StateTreeConditionBase.h"
+#include "CoreMinimal.h"
+#include "Engine/World.h"
+#include "VisualLogger/VisualLogger.h"
+#include "ProfilingDebugging/CsvProfiler.h"
+#include "Algo/Reverse.h"
+#include "Logging/LogScopedVerbosityOverride.h"
+
+#define STATETREE_LOG(Verbosity, Format, ...) UE_VLOG_UELOG(GetOwner(), LogStateTree, Verbosity, TEXT("%s: ") Format, *GetInstanceDescription(), ##__VA_ARGS__)
+#define STATETREE_CLOG(Condition, Verbosity, Format, ...) UE_CVLOG_UELOG((Condition), GetOwner(), LogStateTree, Verbosity, TEXT("%s: ") Format, *GetInstanceDescription(), ##__VA_ARGS__)
+
+namespace UE::StateTree
+{
+	constexpr int32 DebugIndentSize = 2;
+}
+
+FStateTreeExecutionContext::FStateTreeExecutionContext()
+{
+}
+
+FStateTreeExecutionContext::~FStateTreeExecutionContext()
+{
+}
+
+bool FStateTreeExecutionContext::Init(UObject& InOwner, const UStateTree& InStateTree, const EStateTreeStorage InStorageType)
+{
+	// Set owner first for proper logging (it will be reset in case of failure) 
+	Owner = &InOwner;
+
+	if (!InStateTree.IsValidStateTree())
+	{
+		STATETREE_LOG(Error, TEXT("%s: StateTree asset '%s' is not valid."), ANSI_TO_TCHAR(__FUNCTION__), *InStateTree.GetName());
+		Reset();
+		return false;
+	}
+	StateTree = &InStateTree;
+	check(StateTree->IsV2());
+
+	StorageType = InStorageType;
+	if (StorageType == EStateTreeStorage::Internal)
+	{
+		// Duplicate runtime state.
+		StorageInstance = StateTree->RuntimeStorageDefaultValue;
+	}
+
+	// Initialize struct view for all possible source structs.
+	ItemViews.SetNum(StateTree->PropertyBindings.GetSourceStructNum());
+
+	// Initialize array to keep track of which states have been updated.
+	VisitedStates.Init(false, StateTree->States.Num());
+	
+	return true;
+}
+
+void FStateTreeExecutionContext::Start(FStateTreeItemView ExternalStorage)
+{
+	if (!Owner || !StateTree)
+	{
+		return;
+	}
+	FStateTreeItemView Storage = SelectMutableStorage(ExternalStorage);
+	FStateTreeExecutionState& Exec = GetExecState(Storage);
+
+	// Stop if still running 
+	if (Exec.TreeRunStatus == EStateTreeRunStatus::Running)
+	{
+		const FStateTreeTransitionResult Transition(FStateTreeStateStatus(Exec.CurrentState, Exec.LastTickStatus), FStateTreeHandle::Succeeded);
+		ExitState(Storage, Transition);
+	}
+
+	// Initialize to unset running state, tick will choose the first state.
+	Exec.TreeRunStatus = EStateTreeRunStatus::Running;
+	Exec.CurrentState = FStateTreeHandle::Invalid;
+	Exec.LastTickStatus = EStateTreeRunStatus::Unset;
+}
+
+void FStateTreeExecutionContext::Stop(FStateTreeItemView ExternalStorage)
+{
+	if (!Owner || !StateTree)
+	{
+		return;
+	}
+	FStateTreeItemView Storage = SelectMutableStorage(ExternalStorage);
+	FStateTreeExecutionState& Exec = GetExecState(Storage);
+
+	// Stop if still running.
+	if (Exec.TreeRunStatus == EStateTreeRunStatus::Running)
+	{
+		const FStateTreeTransitionResult Transition(FStateTreeStateStatus(Exec.CurrentState, Exec.LastTickStatus), FStateTreeHandle::Succeeded);
+		ExitState(Storage, Transition);
+	}
+	
+	Exec.TreeRunStatus = EStateTreeRunStatus::Succeeded;
+	Exec.CurrentState = FStateTreeHandle::Succeeded;
+	Exec.LastTickStatus = EStateTreeRunStatus::Unset;
+}
+
+EStateTreeRunStatus FStateTreeExecutionContext::Tick(const float DeltaTime, FStateTreeItemView ExternalStorage)
+{
+	CSV_SCOPED_TIMING_STAT_EXCLUSIVE(StateTreeCtxTick);
+
+	if (!Owner || !StateTree)
+	{
+		return EStateTreeRunStatus::Failed;
+	}
+	FStateTreeItemView Storage = SelectMutableStorage(ExternalStorage);
+	FStateTreeExecutionState& Exec = GetExecState(Storage);
+
+	// No ticking of the tree is done or stopped.
+	if (Exec.TreeRunStatus != EStateTreeRunStatus::Running)
+	{
+		return Exec.TreeRunStatus;
+	}
+
+	// Update the gated transition time.
+	if (Exec.GatedTransitionIndex != INDEX_NONE)
+	{
+		Exec.GatedTransitionTime -= DeltaTime;
+	}
+	
+	// The state selection is repeated up to MaxIteration time. This allows failed EnterState() to potentially find a new state immediately.
+	// This helps event driven StateTrees to not require another event/tick to find a suitable state.
+	static const int32 MaxIterations = 5;
+	for (int32 Iter = 0; Iter < MaxIterations; Iter++)
+	{
+		// Reset visited state, used to keep track which state's evaluators have been ticked.
+		VisitedStates.Init(false, StateTree->States.Num());
+
+		// Tick evaluators on active states. Further evaluator may be ticked during selection as non-active states are visited.
+		// Ticked states are kept track of and evaluators are ticked just once per frame.
+		if (Exec.LastTickStatus == EStateTreeRunStatus::Running)
+		{
+			TickEvaluators(Storage, Exec.CurrentState, EStateTreeEvaluationType::Tick, DeltaTime);
+		}
+
+		// Select initial state or trigger transitions.
+		FStateTreeTransitionResult Transition;
+		if (Exec.LastTickStatus == EStateTreeRunStatus::Unset && Exec.CurrentState == FStateTreeHandle::Invalid)
+		{
+			// No state has been selected yet, select starting from root.
+			static const FStateTreeHandle RootState = FStateTreeHandle(0);
+			Transition = SelectState(Storage, FStateTreeStateStatus(FStateTreeHandle::Invalid), RootState, RootState, 0);
+		}
+		else
+		{
+			// Trigger running transitions or state succeed/failed transitions.
+			// Transitions are triggered early in the frame, to prevent one frame lag in games which require quick state changes (i.e. combat).
+			Transition = TriggerTransitions(Storage, FStateTreeStateStatus(Exec.CurrentState, Exec.LastTickStatus), 0);
+		}
+
+		// Handle potential transition.
+		if (Transition.Next.IsValid())
+		{
+			ExitState(Storage, Transition);
+			
+			if (Transition.Next == FStateTreeHandle::Succeeded || Transition.Next == FStateTreeHandle::Failed)
+			{
+				// Transition to a terminal state (succeeded/failed), or default transition failed.
+				Exec.TreeRunStatus = Transition.Next == FStateTreeHandle::Succeeded ? EStateTreeRunStatus::Succeeded : EStateTreeRunStatus::Failed;
+				return Exec.TreeRunStatus;
+			}
+			
+			// Enter state tasks can fail/succeed, treat it same as tick.
+			Exec.LastTickStatus = EnterState(Storage, Transition);
+			Exec.CurrentState = Transition.Next;
+			// Report state completed immediately.
+			if (Exec.LastTickStatus != EStateTreeRunStatus::Running)
+			{
+				StateCompleted(Storage, Exec.CurrentState, Exec.LastTickStatus);
+			}
+		}
+		
+		// Stop as soon as have found a running state.
+		if (Exec.LastTickStatus == EStateTreeRunStatus::Running)
+		{
+			break;
+		}
+	}
+	
+	if (!Exec.CurrentState.IsValid())
+	{
+		// Should not happen. This may happen if initial state or default transition could not be selected. 
+		STATETREE_LOG(Error, TEXT("%s: Failed to select state on '%s' using StateTree '%s'."), ANSI_TO_TCHAR(__FUNCTION__), *GetNameSafe(Owner), *GetNameSafe(StateTree));
+		Exec.TreeRunStatus = EStateTreeRunStatus::Failed;
+		return Exec.TreeRunStatus;
+	}
+
+	// Tick tasks on active states.
+	if (Exec.LastTickStatus == EStateTreeRunStatus::Running)
+	{
+		Exec.LastTickStatus = TickTasks(Storage, Exec.CurrentState, DeltaTime);
+		// Report state completed immediately.
+		if (Exec.LastTickStatus != EStateTreeRunStatus::Running)
+		{
+			StateCompleted(Storage, Exec.CurrentState, Exec.LastTickStatus);
+		}
+	}
+	
+	
+	return Exec.TreeRunStatus;
+}
+
+void FStateTreeExecutionContext::Reset()
+{
+	StateTree = nullptr;
+	Owner = nullptr;
+	World = nullptr;
+	ItemViews.Reset();
+	StorageInstance.Reset();
+	StorageType = EStateTreeStorage::Internal;
+	VisitedStates.Reset();
+}
+
+EStateTreeRunStatus FStateTreeExecutionContext::EnterState(FStateTreeItemView Storage, const FStateTreeTransitionResult& Transition)
+{
+	if (!Transition.Next.IsValid())
+	{
+		return EStateTreeRunStatus::Failed;
+	}
+
+	// Activate evaluators along the active branch.
+	TStaticArray<FStateTreeHandle, 32> States;
+	const int32 NumStates = GetActiveStates(Transition.Next, States);
+	TStaticArray<FStateTreeHandle, 32> PrevStates;
+	const int32 NumPrevStates = GetActiveStates(Transition.Source.State, PrevStates);
+	// Pad prev states with invalid states if it's shorter.
+	for (int32 Index = NumPrevStates; Index < NumStates; Index++)
+	{
+		PrevStates[Index] = FStateTreeHandle::Invalid;
+	}
+
+	// On target branch means that the state is the target of current transition or child of it.
+	// States which were active before and will remain active, but are not on target branch will not get
+	// EnterState called. That is, a transition is handled as "replan from this state".
+	bool bOnTargetBranch = false;
+
+	FStateTreeTransitionResult CurrentTransition = Transition;
+	
+	EStateTreeRunStatus Result = EStateTreeRunStatus::Running;
+	EnterStateStatus = Result;
+
+	for (int32 Index = 0; Index < NumStates; Index++)
+	{
+		const FStateTreeHandle CurrentHandle = States[Index];
+		const FBakedStateTreeState& State = StateTree->States[CurrentHandle.Index];
+		bOnTargetBranch = bOnTargetBranch || CurrentHandle == Transition.Target;
+		const bool bWasActive = PrevStates[Index] == CurrentHandle;
+		if (bWasActive && !bOnTargetBranch)
+		{
+			// States which will keep on begin active and were not part of the transition will not get enter/exit state.
+			// Must update item views.
+			for (int32 EvalIndex = 0; EvalIndex < int32(State.EvaluatorsNum); EvalIndex++)
+			{
+				FStateTreeItemView EvalView = GetItem(Storage, int32(State.EvaluatorsBegin) + EvalIndex);
+				const FStateTreeEvaluator2Base& Eval = EvalView.Get<FStateTreeEvaluator2Base>();
+				ItemViews[Eval.SourceStructIndex] = EvalView;
+			}
+			for (int32 TaskIndex = 0; TaskIndex < int32(State.TasksNum); TaskIndex++)
+			{
+				FStateTreeItemView TaskView = GetItem(Storage, int32(State.TasksBegin) + TaskIndex);
+				const FStateTreeTask2Base& Task = TaskView.Get<FStateTreeTask2Base>();
+				ItemViews[Task.SourceStructIndex] = TaskView;
+			}
+			continue;
+		}
+
+		CurrentTransition.Current = CurrentHandle;
+		
+		const EStateTreeStateChangeType ChangeType = bWasActive ? EStateTreeStateChangeType::Sustained : EStateTreeStateChangeType::Changed;
+
+		STATETREE_LOG(Log, TEXT("%*sEnter state '%s' %s"), Index*UE::StateTree::DebugIndentSize, TEXT(""), *DebugGetStatePath(States, Index), *StaticEnum<EStateTreeStateChangeType>()->GetValueAsString(ChangeType));
+		
+		for (int32 EvalIndex = 0; EvalIndex < int32(State.EvaluatorsNum); EvalIndex++)
+		{
+			FStateTreeItemView EvalView = GetItem(Storage, int32(State.EvaluatorsBegin) + EvalIndex);
+			FStateTreeEvaluator2Base& Eval = EvalView.GetMutable<FStateTreeEvaluator2Base>();
+
+			// Copy bound properties.
+			if (Eval.BindingsBatch.IsValid())
+			{
+				StateTree->PropertyBindings.CopyTo(ItemViews, Eval.BindingsBatch.Index, EvalView);
+			}
+			STATETREE_LOG(Verbose, TEXT("%*s  Notify Evaluator '%s'."), Index*UE::StateTree::DebugIndentSize, TEXT(""), *Eval.Name.ToString());
+			Eval.EnterState(*this, ChangeType, CurrentTransition);
+
+			ItemViews[Eval.SourceStructIndex] = EvalView;
+		}
+		
+		// Activate tasks on current state.
+		for (int32 TaskIndex = 0; TaskIndex < int32(State.TasksNum); TaskIndex++)
+		{
+			FStateTreeItemView TaskView = GetItem(Storage, int32(State.TasksBegin) + TaskIndex);
+			FStateTreeTask2Base& Task = TaskView.GetMutable<FStateTreeTask2Base>();
+
+			// Copy bound properties.
+			if (Task.BindingsBatch.IsValid())
+			{
+				StateTree->PropertyBindings.CopyTo(ItemViews, Task.BindingsBatch.Index, TaskView);
+			}
+		
+			STATETREE_LOG(Verbose, TEXT("%*s  Notify Task '%s'"), Index*UE::StateTree::DebugIndentSize, TEXT(""), *Task.Name.ToString());
+			const EStateTreeRunStatus Status = Task.EnterState(*this, ChangeType, CurrentTransition);
+			if (Status == EStateTreeRunStatus::Failed)
+			{
+				Result = Status;
+				// @todo StateTree: Ideally we should break here, but it is commented out for now in order to keep symmetrical enter/exit calls.
+				// break;
+				// As a workaround, we mark EnterStateStatus as 'failed' in context so remaining calls to EnterState could be aware of the failure and act accordingly.
+				EnterStateStatus = Status;
+
+			}
+
+			ItemViews[Task.SourceStructIndex] = TaskView;
+		}
+	}
+
+	return Result;
+}
+
+void FStateTreeExecutionContext::ExitState(FStateTreeItemView Storage, const FStateTreeTransitionResult& Transition)
+{
+	if (!Transition.Source.State.IsValid())
+	{
+		return;
+	}
+
+	// Reset transition delay
+	FStateTreeExecutionState& Exec = GetExecState(Storage);
+	Exec.GatedTransitionIndex = INDEX_NONE;
+	Exec.GatedTransitionTime = 0.0f;
+
+	// Deactivate evaluators along the active branch.
+	TStaticArray<FStateTreeHandle, 32> States;
+	const int32 NumStates = GetActiveStates(Transition.Source.State, States);
+	TStaticArray<FStateTreeHandle, 32> NextStates;
+	const int32 NumNextStates = GetActiveStates(Transition.Next, NextStates);
+	// Pad next states with invalid states if it's shorter.
+	for (int32 Index = NumNextStates; Index < NumStates; Index++)
+	{
+		NextStates[Index] = FStateTreeHandle::Invalid;
+	}
+
+	// On target branch means that the state is the target of current transition or child of it.
+	// States which were active before and will remain active, but are not on target branch will not get
+	// EnterState called. That is, a transition is handled as "replan from this state".
+	bool bOnTargetBranch = false;
+
+	FStateTreeTransitionResult CurrentTransition = Transition;
+
+	// It would be more symmetrical, if the evals, tasks and states were executed in reverse order, but we need
+	// to do it like this because of the property copies.
+	for (int32 Index = 0; Index < NumStates; Index++)
+	{
+		const FStateTreeHandle CurrentHandle = States[Index];
+		const FStateTreeHandle NextHandle = NextStates[Index];
+		const FBakedStateTreeState& State = StateTree->States[CurrentHandle.Index];
+		const bool bRemainsActive = NextHandle == States[Index];
+		bOnTargetBranch = bOnTargetBranch || NextHandle == Transition.Target;
+		if (bRemainsActive && !bOnTargetBranch)
+		{
+			// States which will keep on begin active and were not part of the transition will not get enter/exit state.
+			// Must update item views.
+			for (int32 EvalIndex = 0; EvalIndex < int32(State.EvaluatorsNum); EvalIndex++)
+			{
+				FStateTreeItemView EvalView = GetItem(Storage, int32(State.EvaluatorsBegin) + EvalIndex);
+				const FStateTreeEvaluator2Base& Eval = EvalView.Get<FStateTreeEvaluator2Base>();
+				ItemViews[Eval.SourceStructIndex] = EvalView;
+			}
+			for (int32 TaskIndex = 0; TaskIndex < int32(State.TasksNum); TaskIndex++)
+			{
+				FStateTreeItemView TaskView = GetItem(Storage, int32(State.TasksBegin) + TaskIndex);
+				const FStateTreeTask2Base& Task = TaskView.Get<FStateTreeTask2Base>();
+				ItemViews[Task.SourceStructIndex] = TaskView;
+			}
+			continue;
+		}
+
+		CurrentTransition.Current = CurrentHandle;
+		
+		const EStateTreeStateChangeType ChangeType = bRemainsActive ? EStateTreeStateChangeType::Sustained : EStateTreeStateChangeType::Changed;
+
+		STATETREE_LOG(Log, TEXT("%*sExit state '%s' %s"), Index*UE::StateTree::DebugIndentSize, TEXT(""), *DebugGetStatePath(States, Index), *StaticEnum<EStateTreeStateChangeType>()->GetValueAsString(ChangeType));
+
+		for (int32 EvalIndex = 0; EvalIndex < int32(State.EvaluatorsNum); EvalIndex++)
+		{
+			FStateTreeItemView EvalView = GetItem(Storage, int32(State.EvaluatorsBegin) + EvalIndex);
+			FStateTreeEvaluator2Base& Eval = EvalView.GetMutable<FStateTreeEvaluator2Base>();
+
+			// Copy bound properties.
+			if (Eval.BindingsBatch.IsValid())
+			{
+				StateTree->PropertyBindings.CopyTo(ItemViews, Eval.BindingsBatch.Index, EvalView);
+			}
+			STATETREE_LOG(Verbose, TEXT("%*s  Notify Evaluator '%s'."), Index*UE::StateTree::DebugIndentSize, TEXT(""), *Eval.Name.ToString());
+			Eval.ExitState(*this, ChangeType, CurrentTransition);
+
+			ItemViews[Eval.SourceStructIndex] = EvalView;
+		}
+
+		// Deactivate tasks on current State
+		for (int32 TaskIndex = 0; TaskIndex < int32(State.TasksNum); TaskIndex++)
+		{
+			FStateTreeItemView TaskView = GetItem(Storage, int32(State.TasksBegin) + TaskIndex);
+			FStateTreeTask2Base& Task = TaskView.GetMutable<FStateTreeTask2Base>();
+
+			// Copy bound properties.
+			if (Task.BindingsBatch.IsValid())
+			{
+				StateTree->PropertyBindings.CopyTo(ItemViews, Task.BindingsBatch.Index, TaskView);
+			}
+
+			STATETREE_LOG(Verbose, TEXT("%*s  Notify Task '%s'"), Index*UE::StateTree::DebugIndentSize, TEXT(""), *Task.Name.ToString());
+			Task.ExitState(*this, ChangeType, CurrentTransition);
+
+			ItemViews[Task.SourceStructIndex] = TaskView;
+		}
+	}
+}
+
+void FStateTreeExecutionContext::StateCompleted(FStateTreeItemView Storage, const FStateTreeHandle CurrentState, const EStateTreeRunStatus CompletionStatus)
+{
+	if (!CurrentState.IsValid())
+	{
+		return;
+	}
+
+	TStaticArray<FStateTreeHandle, 32> States;
+	const int32 NumStates = GetActiveStates(CurrentState, States);
+
+	// Call from child towards root to allow to pass results back.
+	// Note: Completed is assumed to be called immediately after tick or enter state, so there's no property copying.
+	for (int32 Index = NumStates - 1; Index >= 0; Index--)
+	{
+		const FStateTreeHandle CurrentHandle = States[Index];
+		const FBakedStateTreeState& State = StateTree->States[CurrentHandle.Index];
+
+		STATETREE_LOG(Verbose, TEXT("%*sState Completed '%s' %s"), Index*UE::StateTree::DebugIndentSize, TEXT(""), *DebugGetStatePath(States, Index), *StaticEnum<EStateTreeRunStatus>()->GetValueAsString(CompletionStatus));
+
+		// Notify Tasks
+		for (int32 TaskIndex = int32(State.TasksNum) - 1; TaskIndex >= 0; TaskIndex--)
+		{
+			FStateTreeItemView TaskView = GetItem(Storage, int32(State.TasksBegin) + TaskIndex);
+			FStateTreeTask2Base& Task = TaskView.GetMutable<FStateTreeTask2Base>();
+
+			STATETREE_LOG(Verbose, TEXT("%*s  Notify Task '%s'"), Index*UE::StateTree::DebugIndentSize, TEXT(""), *Task.Name.ToString());
+			Task.StateCompleted(*this, CompletionStatus, CurrentState);
+		}
+
+		// Notify evaluators
+		for (int32 EvalIndex = int32(State.EvaluatorsNum) - 1; EvalIndex >= 0; EvalIndex--)
+		{
+			FStateTreeItemView EvalView = GetItem(Storage, int32(State.EvaluatorsBegin) + EvalIndex);
+			FStateTreeEvaluator2Base& Eval = EvalView.GetMutable<FStateTreeEvaluator2Base>();
+
+			STATETREE_LOG(Verbose, TEXT("%*s  Notify Evaluator '%s'"), Index*UE::StateTree::DebugIndentSize, TEXT(""), *Eval.Name.ToString());
+			Eval.StateCompleted(*this, CompletionStatus, CurrentState);
+		}
+	}
+}
+
+void FStateTreeExecutionContext::TickEvaluators(FStateTreeItemView Storage, const FStateTreeHandle CurrentState, const EStateTreeEvaluationType EvalType, const float DeltaTime)
+{
+	CSV_SCOPED_TIMING_STAT_EXCLUSIVE(StateTreeCtxTickEvaluators);
+
+	if (!CurrentState.IsValid())
+	{
+		return;
+	}
+
+	TStaticArray<FStateTreeHandle, 32> States;
+	const int32 NumStates = GetActiveStates(CurrentState, States);
+
+	for (int32 Index = 0; Index < NumStates; Index++)
+	{
+		const FStateTreeHandle CurrentHandle = States[Index];
+		if (VisitedStates[CurrentHandle.Index])
+		{
+			// Already ticked this frame.
+			continue;
+		}
+
+		const FBakedStateTreeState& State = StateTree->States[CurrentHandle.Index];
+
+		STATETREE_CLOG(State.EvaluatorsNum > 0, Verbose, TEXT("%*sTicking Evaluators of state '%s' %s"), Index*UE::StateTree::DebugIndentSize, TEXT(""), *DebugGetStatePath(States, Index), *StaticEnum<EStateTreeEvaluationType>()->GetValueAsString(EvalType));
+
+		// Tick evaluators
+		for (int32 EvalIndex = 0; EvalIndex < int32(State.EvaluatorsNum); EvalIndex++)
+		{
+			FStateTreeItemView EvalView = GetItem(Storage, int32(State.EvaluatorsBegin) + EvalIndex);
+			FStateTreeEvaluator2Base& Eval = EvalView.GetMutable<FStateTreeEvaluator2Base>();
+
+			// Copy bound properties.
+			if (Eval.BindingsBatch.IsValid())
+			{
+				StateTree->PropertyBindings.CopyTo(ItemViews, Eval.BindingsBatch.Index, EvalView);
+			}
+			STATETREE_LOG(Verbose, TEXT("%*s  Evaluate: '%s'"), Index*UE::StateTree::DebugIndentSize, TEXT(""), *Eval.Name.ToString());
+			Eval.Evaluate(*this, EvalType, DeltaTime);
+			ItemViews[Eval.SourceStructIndex] = EvalView;
+		}
+
+		VisitedStates[CurrentHandle.Index] = true;
+	}
+}
+
+EStateTreeRunStatus FStateTreeExecutionContext::TickTasks(FStateTreeItemView Storage, const FStateTreeHandle CurrentState, const float DeltaTime)
+{
+	CSV_SCOPED_TIMING_STAT_EXCLUSIVE(StateTreeCtxTickTasks);
+
+	if (!CurrentState.IsValid())
+	{
+		return EStateTreeRunStatus::Failed;
+	}
+
+	TStaticArray<FStateTreeHandle, 32> States;
+	const int32 NumStates = GetActiveStates(CurrentState, States);
+
+	EStateTreeRunStatus Result = EStateTreeRunStatus::Running;
+	int32 NumTotalTasks = 0;
+
+	for (int32 Index = 0; Index < NumStates && Result != EStateTreeRunStatus::Failed; Index++)
+	{
+		const FStateTreeHandle CurrentHandle = States[Index];
+		const FBakedStateTreeState& State = StateTree->States[CurrentHandle.Index];
+
+		STATETREE_CLOG(State.TasksNum > 0, Verbose, TEXT("%*sTicking Tasks of state '%s'"), Index*UE::StateTree::DebugIndentSize, TEXT(""), *DebugGetStatePath(States, Index));
+
+		// Tick Tasks
+		for (int32 TaskIndex = 0; TaskIndex < int32(State.TasksNum); TaskIndex++)
+		{
+			FStateTreeItemView TaskView = GetItem(Storage, int32(State.TasksBegin) + TaskIndex);
+			FStateTreeTask2Base& Task = TaskView.GetMutable<FStateTreeTask2Base>();
+
+			// Copy bound properties.
+			if (Task.BindingsBatch.IsValid())
+			{
+				StateTree->PropertyBindings.CopyTo(ItemViews, Task.BindingsBatch.Index, TaskView);
+			}
+			STATETREE_LOG(Verbose, TEXT("%*s  Tick: '%s'"), Index*UE::StateTree::DebugIndentSize, TEXT(""), *Task.Name.ToString());
+			const EStateTreeRunStatus TaskResult = Task.Tick(*this, DeltaTime);
+
+			// TODO: Add more control over which states can control the failed/succeeded result.
+			if (TaskResult != EStateTreeRunStatus::Running)
+			{
+				Result = TaskResult;
+			}
+			if (TaskResult == EStateTreeRunStatus::Failed)
+			{
+				break;
+			}
+
+			ItemViews[Task.SourceStructIndex] = TaskView;
+		}
+		NumTotalTasks += State.TasksNum;
+	}
+
+	if (NumTotalTasks == 0)
+	{
+		// No tasks, done ticking.
+		Result = EStateTreeRunStatus::Succeeded;
+	}
+
+	return Result;
+}
+
+bool FStateTreeExecutionContext::TestAllConditions(const uint32 ConditionsOffset, const uint32 ConditionsNum)
+{
+	CSV_SCOPED_TIMING_STAT_EXCLUSIVE(StateTreeCtxTestConditions);
+
+	uint8 CondCopyBuffer[128];	// This buffer is intentionally small, conditions be small.
+	
+	for (uint32 i = 0; i < ConditionsNum; i++)
+	{
+		// Copy struct values from the StateTree asset to a stack copy.
+		const FInstancedStruct& SourceCondPtr = StateTree->Conditions2[ConditionsOffset + i];
+		const FStateTreeConditionBase& SourceCond = SourceCondPtr.Get<FStateTreeConditionBase>();
+		const UScriptStruct* CondStruct = SourceCondPtr.GetScriptStruct();
+		check(CondStruct);
+		check(CondStruct->GetStructureSize() <= sizeof(CondCopyBuffer));
+
+		CondStruct->InitializeStruct(CondCopyBuffer);
+		CondStruct->CopyScriptStruct(CondCopyBuffer, &SourceCond);
+
+		// Copy properties over
+		FStateTreeItemView CondCopyView(CondStruct, CondCopyBuffer);
+		if (SourceCond.BindingsBatch.IsValid())
+		{
+			StateTree->PropertyBindings.CopyTo(ItemViews, SourceCond.BindingsBatch.Index, CondCopyView);
+		}
+
+		// Do test (virtual call)
+		const FStateTreeConditionBase& Cond = CondCopyView.Get<FStateTreeConditionBase>();
+		const bool bResult = Cond.TestCondition();
+
+		// Struct destructor
+		CondStruct->DestroyStruct(CondCopyBuffer);
+
+		if (!bResult)
+		{
+			return false;
+		}
+	}
+	return true;
+}
+
+FStateTreeTransitionResult FStateTreeExecutionContext::TriggerTransitions(FStateTreeItemView Storage, const FStateTreeStateStatus CurrentStatus, const int32 Depth)
+{
+	FStateTreeExecutionState& Exec = GetExecState(Storage);
+	EStateTreeTransitionEvent Event = EStateTreeTransitionEvent::OnCondition;
+	if (CurrentStatus.RunStatus == EStateTreeRunStatus::Succeeded)
+	{
+		Event = EStateTreeTransitionEvent::OnSucceeded;
+	}
+	else if (CurrentStatus.RunStatus == EStateTreeRunStatus::Failed)
+	{
+		Event = EStateTreeTransitionEvent::OnFailed;
+	}
+
+	// Walk towards root and check all transitions along the way.
+	for (FStateTreeHandle Handle = CurrentStatus.State; Handle.IsValid(); Handle = StateTree->States[Handle.Index].Parent)
+	{
+		const FBakedStateTreeState& State = StateTree->States[Handle.Index];
+		for (uint32 i = 0; i < State.TransitionsNum; i++)
+		{
+			// All transition conditions must pass
+			const int16 TransitionIndex = State.TransitionsBegin + i;
+			const FBakedStateTransition& Transition = StateTree->Transitions[TransitionIndex];
+			if (EnumHasAllFlags(Transition.Event, Event) && TestAllConditions(Transition.ConditionsBegin, Transition.ConditionsNum))
+			{
+				// If a transition has delay, we stop testing other transitions, but the transition will not pass the condition until the delay time passes.
+				if (Transition.GateDelay > 0)
+				{
+					if (Exec.GatedTransitionIndex != TransitionIndex)
+					{
+						Exec.GatedTransitionIndex = TransitionIndex;
+						Exec.GatedTransitionTime = FMath::RandRange(0.0f, Transition.GateDelay * 0.1f); // TODO: we need variance too.
+						BeginGatedTransition(Exec);
+						STATETREE_LOG(Verbose, TEXT("Gated transition triggered from '%s' (%s) -> '%s' %.1fs"), *GetSafeStateName(CurrentStatus.State), *State.Name.ToString(), *GetSafeStateName(Transition.State), Exec.GatedTransitionTime);
+					}
+
+					// Keep on updating current state, until we have tried to trigger
+					if (Exec.GatedTransitionTime > 0.0f)
+					{
+						return FStateTreeTransitionResult(CurrentStatus, FStateTreeHandle::Invalid);
+					}
+
+					STATETREE_LOG(Verbose, TEXT("Passed gated transition from '%s' (%s) -> '%s'"), *GetSafeStateName(CurrentStatus.State), *State.Name.ToString(), *GetSafeStateName(Transition.State));
+				}
+				
+				if (Transition.Type == EStateTreeTransitionType::GotoState || Transition.Type == EStateTreeTransitionType::NextState)
+				{
+					FStateTreeTransitionResult Result = SelectState(Storage, CurrentStatus, Transition.State, Transition.State, Depth + 1);
+					if (Result.Next.IsValid())
+					{
+						STATETREE_LOG(Verbose, TEXT("Transition on state '%s' (%s) -[%s]-> state '%s'"), *GetSafeStateName(CurrentStatus.State), *State.Name.ToString(), *GetSafeStateName(Result.Target), *GetSafeStateName(Result.Next));
+						return Result;
+					}
+				}
+				else if (Transition.Type == EStateTreeTransitionType::NotSet)
+				{
+					// NotSet is no-operation, but can be used to mask a transition at parent state. Returning unset keeps updating current state.
+					return FStateTreeTransitionResult(CurrentStatus, FStateTreeHandle::Invalid);
+				}
+				else if (Transition.Type == EStateTreeTransitionType::Succeeded)
+				{
+					STATETREE_LOG(Verbose, TEXT("Stop tree execution from state '%s' (%s): Succeeded"), *GetSafeStateName(CurrentStatus.State), *State.Name.ToString());
+					return FStateTreeTransitionResult(CurrentStatus, FStateTreeHandle::Succeeded);
+				}
+				else
+				{
+					STATETREE_LOG(Verbose, TEXT("Stop tree execution from state '%s' (%s): Failed"), *GetSafeStateName(CurrentStatus.State), *State.Name.ToString());
+					return FStateTreeTransitionResult(CurrentStatus, FStateTreeHandle::Failed);
+				}
+			}
+			else if (Exec.GatedTransitionIndex == TransitionIndex)
+			{
+				// If the current transition was gated transition, reset it if the condition failed.
+				Exec.GatedTransitionIndex = INDEX_NONE;
+				Exec.GatedTransitionTime = 0.0f;
+			}
+		}
+	}
+
+	if (CurrentStatus.RunStatus != EStateTreeRunStatus::Running)
+	{
+		STATETREE_LOG(Error, TEXT("%s: Default transitions on state '%s' could not be triggered. '%s' using StateTree '%s'."),
+            ANSI_TO_TCHAR(__FUNCTION__), *GetSafeStateName(CurrentStatus.State), *GetNameSafe(Owner), *GetNameSafe(StateTree));
+		return FStateTreeTransitionResult(CurrentStatus, FStateTreeHandle::Failed);
+	}
+
+	// No transition triggered, keep on updating current state.
+	return FStateTreeTransitionResult(CurrentStatus, FStateTreeHandle::Invalid);
+}
+
+FStateTreeTransitionResult FStateTreeExecutionContext::SelectState(FStateTreeItemView Storage, const FStateTreeStateStatus InitialStateStatus, const FStateTreeHandle InitialTargetState, const FStateTreeHandle NextState, const int Depth)
+{
+	if (!NextState.IsValid())
+	{
+		// Trying to select non-existing state.
+		STATETREE_LOG(Error, TEXT("%s: Trying to select invalid state from '%s' via '%s.  '%s' using StateTree '%s'."),
+            ANSI_TO_TCHAR(__FUNCTION__), *GetStateStatusString(InitialStateStatus), *GetSafeStateName(InitialTargetState), *GetNameSafe(Owner), *GetNameSafe(StateTree));
+		return FStateTreeTransitionResult(InitialStateStatus, InitialTargetState, FStateTreeHandle::Invalid);
+	}
+
+	const FBakedStateTreeState& State = StateTree->States[NextState.Index];
+
+	// Make sure all the evaluators for the target state are up to date.
+	TickEvaluators(Storage, NextState, EStateTreeEvaluationType::PreSelect, 0.0f);
+	
+	// Check that the state can be entered
+	if (TestAllConditions(State.EnterConditionsBegin, State.EnterConditionsNum))
+	{
+		// If the state has children, proceed to select children.
+		if (State.HasChildren())
+		{
+			for (uint16 ChildState = State.ChildrenBegin; ChildState < State.ChildrenEnd; ChildState = StateTree->States[ChildState].GetNextSibling())
+			{
+				const FStateTreeTransitionResult ChildResult = SelectState(Storage, InitialStateStatus, InitialTargetState, FStateTreeHandle(ChildState), Depth + 1);
+				if (ChildResult.Next.IsValid())
+				{
+					// Selection succeeded
+					return ChildResult;
+				}
+			}
+		}
+		else
+		{
+			// Select this state.
+			return FStateTreeTransitionResult(InitialStateStatus, InitialTargetState, NextState);
+		}
+	}
+	
+	// Nothing got selected.
+	return FStateTreeTransitionResult(InitialStateStatus, InitialTargetState, FStateTreeHandle::Invalid);
+}
+
+int32 FStateTreeExecutionContext::GetActiveStates(const FStateTreeHandle StateHandle, TStaticArray<FStateTreeHandle, 32>& OutStateHandles) const
+{
+	if (StateHandle == FStateTreeHandle::Succeeded || StateHandle == FStateTreeHandle::Failed || StateHandle == FStateTreeHandle::Invalid)
+	{
+		return 0;
+	}
+	
+	int32 NumStates = 0;
+	FStateTreeHandle CurrentHandle = StateHandle;
+	while (CurrentHandle.IsValid() && NumStates < OutStateHandles.Num())
+	{
+		OutStateHandles[NumStates++] = CurrentHandle;
+		CurrentHandle = StateTree->States[CurrentHandle.Index].Parent;
+	}
+
+	Algo::Reverse(OutStateHandles.GetData(), NumStates);
+	
+	return NumStates;
+}
+
+FString FStateTreeExecutionContext::GetSafeStateName(const FStateTreeHandle State) const
+{
+	check(StateTree);
+	if (State == FStateTreeHandle::Invalid)
+	{
+		return TEXT("(State Invalid)");
+	}
+	else if (State == FStateTreeHandle::Succeeded)
+	{
+		return TEXT("(State Succeeded)");
+	}
+	else if (State == FStateTreeHandle::Failed)
+	{
+		return TEXT("(State Failed)");
+	}
+	else if (StateTree->States.IsValidIndex(State.Index))
+	{
+		return *StateTree->States[State.Index].Name.ToString();
+	}
+	return TEXT("(Unknown)");
+}
+
+FString FStateTreeExecutionContext::DebugGetStatePath(const TArrayView<FStateTreeHandle> ActiveStateHandles, const int32 ActiveStateIndex) const
+{
+	FString StatePath;
+	if (!ensureMsgf(ActiveStateHandles.IsValidIndex(ActiveStateIndex), TEXT("Provided index must be valid")))
+	{
+		return StatePath;
+	}
+
+	for (int32 i = 0; i <= ActiveStateIndex; i++)
+	{
+        const FBakedStateTreeState& State = StateTree->States[ActiveStateHandles[i].Index];
+		StatePath.Appendf(TEXT("%s%s"), i == 0 ? TEXT("") : TEXT("."), *State.Name.ToString());
+	}
+	return StatePath;
+}
+
+FString FStateTreeExecutionContext::GetStateStatusString(const FStateTreeStateStatus StateStatus) const
+{
+	// Invalid (null) state in status refers to the whole tree.
+	const FString StateName = GetSafeStateName(StateStatus.State);
+	switch(StateStatus.RunStatus)
+	{
+	case EStateTreeRunStatus::Unset:
+		return StateName + TEXT("Unset");
+	case EStateTreeRunStatus::Running:
+		return StateName + TEXT("Running");
+	case EStateTreeRunStatus::Succeeded:
+		return StateName + TEXT("Succeeded");
+	case EStateTreeRunStatus::Failed:
+		return StateName + TEXT("Failed");
+	default:
+		return StateName + TEXT("(Unknown)");
+	}
+}
+
+EStateTreeRunStatus FStateTreeExecutionContext::GetLastTickStatus(FStateTreeItemView ExternalStorage) const
+{
+	const FStateTreeExecutionState& Exec = GetExecState(SelectStorage(ExternalStorage));
+	return Exec.LastTickStatus;
+}
+
+#if WITH_GAMEPLAY_DEBUGGER
+
+FString FStateTreeExecutionContext::GetDebugInfoString(FStateTreeItemView ExternalStorage)
+{
+	if (!StateTree)
+	{
+		return FString(TEXT("No StateTree asset."));
+	}
+	FStateTreeItemView Storage = SelectMutableStorage(ExternalStorage);
+	FStateTreeExecutionState& Exec = GetExecState(Storage);
+
+	FString DebugString = FString::Printf(TEXT("StateTree (asset: '%s')\n"), *GetNameSafe(StateTree));
+
+	DebugString += TEXT("Status: ");
+	switch (Exec.TreeRunStatus)
+	{
+	case EStateTreeRunStatus::Failed:
+		DebugString += TEXT("Failed\n");
+		break;
+	case EStateTreeRunStatus::Succeeded:
+		DebugString += TEXT("Succeeded\n");
+		break;
+	case EStateTreeRunStatus::Running:
+		DebugString += TEXT("Running\n");
+		break;
+	default:
+		DebugString += TEXT("--\n");
+	}
+
+	// Active States
+	TStaticArray<FStateTreeHandle, 32> ActiveStates;
+	const int32 NumStates = GetActiveStates(Exec.CurrentState, ActiveStates);
+
+	DebugString += TEXT("Current State:\n");
+	for (int32 Index = 0; Index < NumStates; Index++)
+	{
+		FStateTreeHandle Handle = ActiveStates[Index];
+		if (Handle.IsValid())
+		{
+			const FBakedStateTreeState& State = StateTree->States[Handle.Index];
+			DebugString += FString::Printf(TEXT("[%s]\n"), *State.Name.ToString());
+
+			if (State.TasksNum > 0)
+			{
+				DebugString += TEXT("\nTasks:\n");
+				for (int32 j = 0; j < int32(State.TasksNum); j++)
+				{
+					FStateTreeItemView TaskView = GetItem(Storage, int32(State.TasksBegin) + j);
+					const FStateTreeTask2Base& Task = TaskView.Get<FStateTreeTask2Base>();
+					Task.AppendDebugInfoString(DebugString, *this);
+				}
+			}
+			if (State.EvaluatorsNum > 0)
+			{
+				DebugString += TEXT("\nEvaluators:\n");
+				for (int32 EvalIndex = 0; EvalIndex < int32(State.EvaluatorsNum); EvalIndex++)
+				{
+					FStateTreeItemView EvalView = GetItem(Storage, int32(State.EvaluatorsBegin) + EvalIndex);
+					const FStateTreeEvaluator2Base& Eval = EvalView.Get<FStateTreeEvaluator2Base>();
+					Eval.AppendDebugInfoString(DebugString, *this);
+				}
+			}
+		}
+	}
+
+	return DebugString;
+}
+#endif // WITH_GAMEPLAY_DEBUGGER
+
+#if WITH_STATETREE_DEBUG
+void FStateTreeExecutionContext::DebugPrintInternalLayout(FStateTreeItemView ExternalStorage)
+{
+	LOG_SCOPE_VERBOSITY_OVERRIDE(LogStateTree, ELogVerbosity::Log);
+
+	if (StateTree == nullptr)
+	{
+		UE_LOG(LogStateTree, Log, TEXT("No StateTree asset."));
+		return;
+	}
+
+	FString DebugString = FString::Printf(TEXT("StateTree (asset: '%s')\n"), *GetNameSafe(StateTree));
+
+	// Runtime storage items (e.g. tasks, evaluators)
+	DebugString += FString::Printf(TEXT("\nItems(%d)\n"), StateTree->RuntimeStorageItems.Num());
+	for (const FInstancedStruct& Item : StateTree->RuntimeStorageItems)
+	{
+		DebugString += FString::Printf(TEXT("  %s\n"), Item.IsValid() ? *Item.GetScriptStruct()->GetName() : TEXT("null"));
+	}
+
+	// Runtime storage offsets (e.g. tasks, evaluators)
+	DebugString += FString::Printf(TEXT("\nOffsets(%d)\n  [ %-40s | %-6s ]\n"), StateTree->RuntimeStorageOffsets.Num(), TEXT("Name"), TEXT("Offset"));
+	for (const FStateTreeRuntimeStorageItemOffset& Offset : StateTree->RuntimeStorageOffsets)
+	{
+		DebugString += FString::Printf(TEXT("  | %-40s | %6d |\n"), Offset.Struct ? *Offset.Struct->GetName() : TEXT("null"),  Offset.Offset);
+	}
+
+	// Conditions
+	DebugString += FString::Printf(TEXT("\nConditions(%d)\n"), StateTree->Conditions2.Num());
+	for (const FInstancedStruct& Condition : StateTree->Conditions2)
+	{
+		DebugString += FString::Printf(TEXT("  %s\n"), Condition.IsValid() ? *Condition.GetScriptStruct()->GetName() : TEXT("null"));
+	}
+
+	// External items (e.g. fragments, subsystems)
+	DebugString += FString::Printf(TEXT("\nExternal items(%d)\n  [ %-40s | %-8s | %5s ]\n"), StateTree->ExternalItems.Num(), TEXT("Name"), TEXT("Optional"), TEXT("Index"));
+	for (const FStateTreeExternalItemDesc& Desc : StateTree->ExternalItems)
+	{
+		DebugString += FString::Printf(TEXT("  | %-40s | %8s | %5d |\n"), Desc.Struct ? *Desc.Struct->GetName() : TEXT("null"), *LexToString(Desc.bOptional), Desc.Handle.GetIndex());
+	}
+
+	// Bindings
+	StateTree->PropertyBindings.DebugPrintInternalLayout(DebugString);
+
+	// Transitions
+	DebugString += FString::Printf(TEXT("\nTransitions(%d)\n  [ %-3s | %15s | %-40s | %-40s | %-8s ]\n"), StateTree->Transitions.Num()
+		, TEXT("Idx"), TEXT("State"), TEXT("Transition Type"), TEXT("Transition Event"), TEXT("Num Cond"));
+	for (const FBakedStateTransition& Transition : StateTree->Transitions)
+	{
+		DebugString += FString::Printf(TEXT("  | %3d | %15s | %-40s | %-40s | %8d |\n"),
+									Transition.ConditionsBegin, *Transition.State.Describe(),
+									*StaticEnum<EStateTreeTransitionType>()->GetValueAsString(Transition.Type),
+									*StaticEnum<EStateTreeTransitionEvent>()->GetValueAsString(Transition.Event),
+									Transition.ConditionsNum);
+	}
+
+	// ItemViews
+	DebugString += FString::Printf(TEXT("\nItemViews(%d)\n"), ItemViews.Num());
+	for (const FStateTreeItemView& ItemView : ItemViews)
+	{
+		DebugString += FString::Printf(TEXT("  [%s]\n"), ItemView.IsValid() ? *ItemView.GetStruct()->GetName() : TEXT("null"));
+	}
+
+	// States
+	FStateTreeItemView Storage = SelectMutableStorage(ExternalStorage);
+	DebugString += FString::Printf(TEXT("\nStates(%d)\n"
+		"  [ %-30s | %15s | %5s [%3s:%-3s[ | Begin Idx : %4s %4s %4s %4s | Num : %4s %4s %4s %4s | Transitions : %-16s %-40s %-16s %-40s ]\n"),
+		StateTree->States.Num(),
+		TEXT("Name"), TEXT("Parent"), TEXT("Child"), TEXT("Beg"), TEXT("End"),
+		TEXT("Cond"), TEXT("Tr"), TEXT("Tsk"), TEXT("Evt"), TEXT("Cond"), TEXT("Tr"), TEXT("Tsk"), TEXT("Evt"),
+		TEXT("Done State"), TEXT("Done Type"), TEXT("Failed State"), TEXT("Failed Type")
+		);
+	for (const FBakedStateTreeState& State : StateTree->States)
+	{
+		DebugString += FString::Printf(TEXT("  | %-30s | %15s | %5s [%3d:%-3d[ | %9s   %4d %4d %4d %4d | %3s   %4d %4d %4d %4d | %11s   %-16s %-40s %-16s %-40s |\n"),
+									*State.Name.ToString(), *State.Parent.Describe(),
+									TEXT(""), State.ChildrenBegin, State.ChildrenEnd,
+									TEXT(""), State.EnterConditionsBegin, State.TransitionsBegin, State.TasksBegin, State.EvaluatorsBegin,
+									TEXT(""), State.EnterConditionsNum, State.TransitionsNum, State.TasksNum, State.EvaluatorsNum,
+									TEXT(""), *State.StateDoneTransitionState.Describe(), *StaticEnum<EStateTreeTransitionType>()->GetValueAsString(State.StateDoneTransitionType),
+									*State.StateFailedTransitionState.Describe(), *StaticEnum<EStateTreeTransitionType>()->GetValueAsString(State.StateFailedTransitionType));
+	}
+
+	DebugString += FString::Printf(TEXT("\nStates Items\n  [ %-30s | %-12s | %-30s | %15s | %10s ]\n"),
+	TEXT("State"), TEXT("Type"), TEXT("Name"), TEXT("Bindings"), TEXT("Struct Idx"));
+	for (const FBakedStateTreeState& State : StateTree->States)
+	{
+		// Tasks
+		if (State.TasksNum)
+		{
+			for (int32 j = 0; j < State.TasksNum; j++)
+			{
+				const FStateTreeTask2Base& Task = GetItem(Storage, State.TasksBegin + j).Get<FStateTreeTask2Base>();
+				DebugString += FString::Printf(TEXT("  | %-30s | %-12s | %-30s | %15s | %10d |\n"), *State.Name.ToString(),
+					TEXT("  Task"), *Task.Name.ToString(), *Task.BindingsBatch.Describe(), Task.SourceStructIndex);
+			}
+		}
+
+		// Evaluators
+		if (State.EvaluatorsNum)
+		{
+			for (int32 EvalIndex = 0; EvalIndex < State.EvaluatorsNum; EvalIndex++)
+			{
+				const FStateTreeEvaluator2Base& Eval = GetItem(Storage, State.EvaluatorsBegin + EvalIndex).Get<FStateTreeEvaluator2Base>();
+				DebugString += FString::Printf(TEXT("  | %-30s | %-12s | %-30s | %15s | %10d |\n"), *State.Name.ToString(),
+					TEXT("  Evaluator"), *Eval.Name.ToString(), *Eval.BindingsBatch.Describe(), Eval.SourceStructIndex);
+			}
+		}
+	}
+
+	UE_LOG(LogStateTree, Log, TEXT("%s"), *DebugString);
+}
+
+FString FStateTreeExecutionContext::GetActiveStateName(FStateTreeItemView ExternalStorage) const
+{
+	if (!StateTree)
+	{
+		return FString(TEXT("<None>"));
+	}
+	const FStateTreeItemView Storage = SelectStorage(ExternalStorage);
+	FStateTreeExecutionState& Exec = GetExecState(Storage);
+
+	FString FullStateName;
+	
+	// Active States
+	TStaticArray<FStateTreeHandle, 32> ActiveStates;
+	const int32 NumStates = GetActiveStates(Exec.CurrentState, ActiveStates);
+
+	for (int32 Index = 0; Index < NumStates; Index++)
+	{
+		FStateTreeHandle Handle = ActiveStates[Index];
+		if (Handle.IsValid())
+		{
+			const FBakedStateTreeState& State = StateTree->States[Handle.Index];
+			if (Index > 0)
+			{
+				FullStateName += TEXT("\n");
+			}
+			FullStateName += FString::Printf(TEXT("%*s-"), Index * 3, TEXT("")); // Indent
+			FullStateName += *State.Name.ToString();
+		}
+	}
+
+	switch (Exec.TreeRunStatus)
+	{
+	case EStateTreeRunStatus::Failed:
+		FullStateName += TEXT(" FAILED\n");
+		break;
+	case EStateTreeRunStatus::Succeeded:
+		FullStateName += TEXT(" SUCCEEDED\n");
+		break;
+	case EStateTreeRunStatus::Running:
+		// Empty
+		break;
+	default:
+		FullStateName += TEXT("--\n");
+	}
+
+	return FullStateName;
+}
+#endif // WITH_STATETREE_DEBUG
+
+#undef STATETREE_LOG
+#undef STATETREE_CLOG
diff --git a/Engine/Plugins/Runtime/StateTree/Source/StateTreeModule/Private/StateTreeInstance.cpp b/Engine/Plugins/Runtime/StateTree/Source/StateTreeModule/Private/StateTreeInstance.cpp
new file mode 100644
--- /dev/null
+++ b/Engine/Plugins/Runtime/StateTree/Source/StateTreeModule/Private/StateTreeInstance.cpp
@@ -0,0 +1,871 @@
+// Copyright Epic Games, Inc. All Rights Reserved.
+
+#include "StateTreeInstance.h"
+#include "StateTree.h"
+#include "StateTreeTaskBase.h"
+#include "StateTreeEvaluatorBase.h"
+#include "StateTreeVariableDesc.h"
+#include "StateTreeParameterStorage.h"
+#include "StateTreeParameterLayout.h"
+#include "StateTreeConditionBase.h"
+#include "CoreMinimal.h"
+#include "UObject/WeakObjectPtr.h"
+#include "GameFramework/Actor.h"
+#include "VisualLogger/VisualLoggerTypes.h"
+#include "VisualLogger/VisualLogger.h"
+
+#define STATETREE_LOG(Verbosity, Format, ...) UE_VLOG(GetOwner(), LogStateTree, Verbosity, Format, ##__VA_ARGS__)
+#define STATETREE_CLOG(Condition, Verbosity, Format, ...) UE_CVLOG((Condition), GetOwner(), LogStateTree, Verbosity, Format, ##__VA_ARGS__)
+
+//PRAGMA_DISABLE_OPTIMIZATION
+
+namespace FStateTreeHelpers
+{
+	bool CompareBools(const bool bLHS, const bool bRHS, const EGenericAICheck Operator)
+	{
+		switch (Operator)
+		{
+		case EGenericAICheck::Equal:
+			return bLHS == bRHS;
+		case EGenericAICheck::NotEqual:
+			return bLHS != bRHS;
+		default:
+			ensureMsgf(false, TEXT("Unhandled operator %d"), Operator);
+			break;
+		}
+		return false;
+	}
+
+	bool CompareInts(const int LHS, const int RHS, const EGenericAICheck Operator)
+	{
+		switch (Operator)
+		{
+		case EGenericAICheck::Equal:
+			return LHS == RHS;
+		case EGenericAICheck::NotEqual:
+			return LHS != RHS;
+		case EGenericAICheck::Less:
+			return LHS < RHS;
+		case EGenericAICheck::LessOrEqual:
+			return LHS <= RHS;
+		case EGenericAICheck::Greater:
+			return LHS > RHS;
+		case EGenericAICheck::GreaterOrEqual:
+			return LHS >= RHS;
+		default:
+			ensureMsgf(false, TEXT("Unhandled operator %d"), Operator);
+			break;
+		}
+		return false;
+	}
+
+	bool CompareFloats(const float LHS, const float RHS, const EGenericAICheck Operator)
+	{
+		switch (Operator)
+		{
+		case EGenericAICheck::Equal:
+			return LHS == RHS;
+		case EGenericAICheck::NotEqual:
+			return LHS != RHS;
+		case EGenericAICheck::Less:
+			return LHS < RHS;
+		case EGenericAICheck::LessOrEqual:
+			return LHS <= RHS;
+		case EGenericAICheck::Greater:
+			return LHS > RHS;
+		case EGenericAICheck::GreaterOrEqual:
+			return LHS >= RHS;
+		default:
+			ensureMsgf(false, TEXT("Unhandled operator %d"), Operator);
+			break;
+		}
+		return false;
+	}
+
+	bool CompareVectors(const FVector& LHS, const FVector& RHS, const EGenericAICheck Operator)
+	{
+		switch (Operator)
+		{
+		case EGenericAICheck::Equal:
+			return LHS.Equals(RHS);
+		case EGenericAICheck::NotEqual:
+			return !LHS.Equals(RHS);
+		default:
+			ensureMsgf(false, TEXT("Unhandled operator %d"), Operator);
+			break;
+		}
+		return false;
+	}
+
+	bool CompareObjects(const FWeakObjectPtr LHS, const FWeakObjectPtr RHS, const EGenericAICheck Operator)
+	{
+		switch (Operator)
+		{
+		case EGenericAICheck::Equal:
+			return LHS == RHS;
+		case EGenericAICheck::NotEqual:
+			return LHS != RHS;
+		default:
+			ensureMsgf(false, TEXT("Unhandled operator %d"), Operator);
+			break;
+		}
+		return false;
+	}
+};
+
+
+
+FStateTreeInstance::FStateTreeInstance()
+	: Owner(nullptr)
+	, StateTree(nullptr)
+	, CurrentState(FStateTreeHandle::Invalid)
+	, LastTickStatus(EStateTreeRunStatus::Failed)
+	, CurrentTask(FStateTreeHandle::Invalid)
+	, NextTask(FStateTreeHandle::Invalid)
+	, RunStatus(EStateTreeRunStatus::Failed)
+{
+}
+
+FStateTreeInstance::~FStateTreeInstance()
+{
+}
+
+void FStateTreeInstance::Reset()
+{
+	Owner = nullptr;
+	StateTree = nullptr;
+	Evaluators.Reset();
+	VariableMemoryPool.Reset();
+	CurrentState = FStateTreeHandle::Invalid;
+	CurrentTask = FStateTreeHandle::Invalid;
+	NextTask = FStateTreeHandle::Invalid;
+	RunStatus = EStateTreeRunStatus::Failed;
+}
+
+bool FStateTreeInstance::Init(UObject& InOwner, const UStateTree& InStateTree)
+{
+	if (!InStateTree.IsValidStateTree())
+	{
+		STATETREE_LOG(Error, TEXT("%s: StateTree asset '%' is not valid."), ANSI_TO_TCHAR(__FUNCTION__), *InStateTree.GetName());
+		Reset();
+		return false;
+	}
+
+	if (InStateTree.IsV2())
+	{
+		STATETREE_LOG(Error, TEXT("%s: V2 Assets are not supported."), ANSI_TO_TCHAR(__FUNCTION__), *InStateTree.GetName());
+		Reset();
+		return false;
+	}
+
+	Owner = &InOwner;
+	StateTree = &InStateTree;
+
+	CurrentState = FStateTreeHandle::Invalid;
+	CurrentTask = FStateTreeHandle::Invalid;
+	NextTask = FStateTreeHandle::Invalid;
+	RunStatus = EStateTreeRunStatus::Failed;
+
+	// Initialize variable memory
+	VariableMemoryPool.SetNumZeroed(int32(StateTree->Variables.GetMemoryUsage()));
+
+	// Instantiate Evaluators
+	for (const UStateTreeEvaluatorBase* Eval : StateTree->Evaluators)
+	{
+		UStateTreeEvaluatorBase* EvalInstance = Cast<UStateTreeEvaluatorBase>(StaticDuplicateObject(Eval, Owner));
+		if (!EvalInstance)
+		{
+			STATETREE_LOG(Error, TEXT("%s: Failed to create StateTree '%' Evaluator '%s' on '%s'."), ANSI_TO_TCHAR(__FUNCTION__), *GetNameSafe(StateTree), *GetNameSafe(Eval), *GetNameSafe(Owner));
+			Reset();
+			return false;
+		}
+		if (!EvalInstance->Initialize(*this))
+		{
+			STATETREE_LOG(Error, TEXT("%s: Failed to initialize StateTree '%' Evaluator '%s' on '%s'."), ANSI_TO_TCHAR(__FUNCTION__), *GetNameSafe(StateTree), *GetNameSafe(Eval), *GetNameSafe(Owner));
+			Reset();
+			return false;
+		}
+		Evaluators.Add(EvalInstance);
+	}
+
+	// Instantiate Tasks
+	for (const UStateTreeTaskBase* Task : StateTree->Tasks)
+	{
+		// TODO: move evaluators into structs too so that we can skip the owner here.
+		UStateTreeTaskBase* TaskInstance = Cast<UStateTreeTaskBase>(StaticDuplicateObject(Task, Owner));
+		if (!TaskInstance)
+		{
+			STATETREE_LOG(Error, TEXT("%s: Failed to create StateTree '%' Task '%s' on '%s'."), ANSI_TO_TCHAR(__FUNCTION__), *GetNameSafe(StateTree), *GetNameSafe(Task), *GetNameSafe(Owner));
+			Reset();
+			return false;
+		}
+		if (!TaskInstance->Initialize(*this))
+		{
+			STATETREE_LOG(Error, TEXT("%s: Failed to initialize StateTree '%' Task '%s' on '%s'."), ANSI_TO_TCHAR(__FUNCTION__), *GetNameSafe(StateTree), *GetNameSafe(Task), *GetNameSafe(Owner));
+			Reset();
+			return false;
+		}
+		Tasks.Add(TaskInstance);
+	}
+
+	return true;
+}
+
+bool FStateTreeInstance::TestCondition(const FStateTreeCondition& Cond)
+{
+	switch (Cond.Left.Type)
+	{
+	case EStateTreeVariableType::Bool:
+		return FStateTreeHelpers::CompareBools(GetValueBool(Cond.Left.Handle), GetValueBool(Cond.Right.Handle), Cond.Operator);
+	case EStateTreeVariableType::Int:
+		return FStateTreeHelpers::CompareInts(GetValueInt(Cond.Left.Handle), GetValueInt(Cond.Right.Handle), Cond.Operator);
+	case EStateTreeVariableType::Float:
+		return FStateTreeHelpers::CompareFloats(GetValueFloat(Cond.Left.Handle), GetValueFloat(Cond.Right.Handle), Cond.Operator);
+	case EStateTreeVariableType::Vector:
+		return FStateTreeHelpers::CompareVectors(GetValueVector(Cond.Left.Handle), GetValueVector(Cond.Right.Handle), Cond.Operator);
+	case EStateTreeVariableType::Object:
+		return FStateTreeHelpers::CompareObjects(GetValueObject(Cond.Left.Handle), GetValueObject(Cond.Right.Handle), Cond.Operator);
+	default:
+		ensureMsgf(false, TEXT("Unhandled type %d"), Cond.Left.Type);
+		return false;
+	}
+
+	return false;
+}
+
+bool FStateTreeInstance::TestAllConditions(const uint32 ConditionsOffset, const uint32 ConditionsNum)
+{
+	for (uint32 i = 0; i < ConditionsNum; i++)
+	{
+		if (!TestCondition(StateTree->Conditions[ConditionsOffset + i]))
+		{
+			return false;
+		}
+	}
+	return true;
+}
+
+bool FStateTreeInstance::GetValueBool(const FStateTreeHandle Handle, const bool DefaultValue) const
+{
+	if (Handle.Index < uint16(VariableMemoryPool.Num())) // Handles invalid handles too, Invalid handle is largest possible value.
+	{
+		return FStateTreeVariableHelpers::GetValueFromMemory<bool>(&VariableMemoryPool[Handle.Index]);
+	}
+	return StateTree ? StateTree->Constants.GetConstantBool(Handle, DefaultValue) : DefaultValue;
+}
+
+int32 FStateTreeInstance::GetValueInt(const FStateTreeHandle Handle, const int32 DefaultValue) const
+{
+	if (Handle.Index < uint16(VariableMemoryPool.Num())) // Handles invalid handles too, Invalid handle is largest possible value.
+	{
+		return FStateTreeVariableHelpers::GetValueFromMemory<int32>(&VariableMemoryPool[Handle.Index]);
+	}
+	return StateTree ? StateTree->Constants.GetConstantInt(Handle, DefaultValue) : DefaultValue;
+}
+
+float FStateTreeInstance::GetValueFloat(const FStateTreeHandle Handle, const float DefaultValue) const
+{
+	if (Handle.Index < uint16(VariableMemoryPool.Num())) // Handles invalid handles too, Invalid handle is largest possible value.
+	{
+		return FStateTreeVariableHelpers::GetValueFromMemory<float>(&VariableMemoryPool[Handle.Index]);
+	}
+	return StateTree ? StateTree->Constants.GetConstantFloat(Handle, DefaultValue) : DefaultValue;
+}
+
+const FVector& FStateTreeInstance::GetValueVector(const FStateTreeHandle Handle, const FVector& DefaultValue) const
+{
+	if (Handle.Index < uint16(VariableMemoryPool.Num())) // Handles invalid handles too, Invalid handle is largest possible value.
+	{
+		return FStateTreeVariableHelpers::GetValueFromMemory<FVector>(&VariableMemoryPool[Handle.Index]);
+	}
+	return StateTree ? StateTree->Constants.GetConstantVector(Handle, DefaultValue) : DefaultValue;
+}
+
+const FWeakObjectPtr& FStateTreeInstance::GetValueObject(const FStateTreeHandle Handle, const FWeakObjectPtr& DefaultValue) const
+{
+	if (Handle.Index < uint16(VariableMemoryPool.Num())) // Handles invalid handles too, Invalid handle is largest possible value.
+	{
+		return FStateTreeVariableHelpers::GetValueFromMemory<FWeakObjectPtr>(&VariableMemoryPool[Handle.Index]);
+	}
+	return DefaultValue;
+}
+
+bool FStateTreeInstance::GetValue(const FStateTreeHandle Handle, EStateTreeVariableType Type, uint8* Value) const
+{
+	if (Handle.Index < uint16(VariableMemoryPool.Num())) // Handles invalid handles too, Invalid handle is largest possible value.
+	{
+		FMemory::Memcpy(Value, &VariableMemoryPool[Handle.Index], FStateTreeVariableHelpers::GetVariableMemoryUsage(Type));
+		return true;
+	}
+	return StateTree ? StateTree->Constants.GetConstant(Handle, Type, Value) : false;
+}
+
+
+bool FStateTreeInstance::SetValueBool(const FStateTreeHandle Handle, const bool Value)
+{
+	if (Handle.Index < uint16(VariableMemoryPool.Num())) // Handles invalid handles too, Invalid handle is largest possible value.
+	{
+		return FStateTreeVariableHelpers::SetValueInMemory<bool>(&VariableMemoryPool[Handle.Index], Value);
+	}
+	return false;
+}
+
+bool FStateTreeInstance::SetValueInt(const FStateTreeHandle Handle, const int32 Value)
+{
+	if (Handle.Index < uint16(VariableMemoryPool.Num())) // Handles invalid handles too, Invalid handle is largest possible value.
+	{
+		return FStateTreeVariableHelpers::SetValueInMemory<int32>(&VariableMemoryPool[Handle.Index], Value);
+	}
+	return false;
+}
+
+bool FStateTreeInstance::SetValueFloat(const FStateTreeHandle Handle, const float Value)
+{
+	if (Handle.Index < uint16(VariableMemoryPool.Num())) // Handles invalid handles too, Invalid handle is largest possible value.
+	{
+		return FStateTreeVariableHelpers::SetValueInMemory<float>(&VariableMemoryPool[Handle.Index], Value);
+	}
+	return false;
+}
+
+bool FStateTreeInstance::SetValueVector(const FStateTreeHandle Handle, const FVector& Value)
+{
+	if (Handle.Index < uint16(VariableMemoryPool.Num())) // Handles invalid handles too, Invalid handle is largest possible value.
+	{
+		return FStateTreeVariableHelpers::SetValueInMemory<FVector>(&VariableMemoryPool[Handle.Index], Value);
+	}
+	return false;
+}
+
+bool FStateTreeInstance::SetValueObject(const FStateTreeHandle Handle, FWeakObjectPtr Value)
+{
+	if (Handle.Index < uint16(VariableMemoryPool.Num())) // Handles invalid handles too, Invalid handle is largest possible value.
+	{
+		return FStateTreeVariableHelpers::SetWeakObjectInMemory(&VariableMemoryPool[Handle.Index], Value);
+	}
+	return false;
+}
+
+bool FStateTreeInstance::SetValue(const FStateTreeHandle Handle, EStateTreeVariableType Type, const uint8* Value)
+{
+	if (Handle.Index < uint16(VariableMemoryPool.Num())) // Handles invalid handles too, Invalid handle is largest possible value.
+	{
+		FMemory::Memcpy(&VariableMemoryPool[Handle.Index], Value, FStateTreeVariableHelpers::GetVariableMemoryUsage(Type));
+		return true;
+	}
+	return false;
+}
+
+bool FStateTreeInstance::ReadParametersFromStorage(const FStateTreeParameterLayout& ParameterLayout, const FStateTreeParameterStorage& ParameterStorage)
+{
+	// Expect reader to know the layout
+	if (uint32(ParameterLayout.Parameters.Num()) != ParameterStorage.GetParameterNum())
+	{
+		return false;
+	}
+
+	for (int32 i = 0; i < ParameterLayout.Parameters.Num(); i++)
+	{
+		const FStateTreeParameter& Param = ParameterLayout.Parameters[i];
+		const uint8* ValuePtr = ParameterStorage.GetParameterPtr(uint32(i), Param.Variable.Type);
+		if (!ValuePtr)
+		{
+			return false;
+		}
+		SetValue(Param.Variable.Handle, Param.Variable.Type, ValuePtr);
+	}
+	return true;
+}
+
+bool FStateTreeInstance::WriteParametersToStorage(const FStateTreeParameterLayout& ParameterLayout, FStateTreeParameterStorage& ParameterStorage) const
+{
+	// Expect writer to know the layout
+	if (uint32(ParameterLayout.Parameters.Num()) != ParameterStorage.GetParameterNum())
+	{
+		return false;
+	}
+
+	for (int32 i = 0; i < ParameterLayout.Parameters.Num(); i++)
+	{
+		const FStateTreeParameter& Param = ParameterLayout.Parameters[i];
+		uint8* ValuePtr = ParameterStorage.GetParameterPtr(uint32(i), Param.Variable.Type);
+		if (!ValuePtr)
+		{
+			return false;
+		}
+		GetValue(Param.Variable.Handle, Param.Variable.Type, ValuePtr);
+	}
+	return true;
+}
+
+void FStateTreeInstance::ResetState()
+{
+	if (CurrentState.IsValid())
+	{
+		ExitState(CurrentState);
+		CurrentTask = FStateTreeHandle::Invalid;
+		NextTask = FStateTreeHandle::Invalid;
+	}
+
+	CurrentState = FStateTreeHandle::Invalid;
+}
+
+void FStateTreeInstance::ActivateEvaluators()
+{
+	for (UStateTreeEvaluatorBase* Eval : Evaluators)
+	{
+		Eval->Activate(*this);
+	}
+}
+
+void FStateTreeInstance::DeactivateEvaluators()
+{
+	for (UStateTreeEvaluatorBase* Eval : Evaluators)
+	{
+		Eval->Deactivate(*this);
+	}
+}
+
+void FStateTreeInstance::TickEvaluators(const float DeltaTime)
+{
+	for (UStateTreeEvaluatorBase* Eval : Evaluators)
+	{
+		Eval->Tick(*this, DeltaTime);
+	}
+}
+
+
+void FStateTreeInstance::Start()
+{
+	if (!Owner || !StateTree)
+	{
+		return;
+	}
+
+	ActivateEvaluators();
+	ResetState();
+	RunStatus = EStateTreeRunStatus::Running;
+}
+
+void FStateTreeInstance::Stop()
+{
+	if (!Owner || !StateTree)
+	{
+		return;
+	}
+
+	DeactivateEvaluators();
+	ResetState();
+	if (RunStatus == EStateTreeRunStatus::Running)
+	{
+		RunStatus = EStateTreeRunStatus::Succeeded;
+	}
+}
+
+void FStateTreeInstance::EnterState(const FStateTreeHandle StateHandle)
+{
+	if (!StateHandle.IsValid())
+	{
+		return;
+	}
+	const FBakedStateTreeState& State = StateTree->States[StateHandle.Index];
+
+	if (State.TasksNum == 0)
+	{
+		STATETREE_LOG(Error, TEXT("%s: Trying to enter State '%s' without Tasks on '%s' using StateTree %s."), ANSI_TO_TCHAR(__FUNCTION__), *State.Name.ToString(), *GetNameSafe(Owner), *GetNameSafe(StateTree));
+	}
+
+	CurrentTask = FStateTreeHandle::Invalid;
+	NextTask = FStateTreeHandle::Invalid;
+}
+
+void FStateTreeInstance::ExitState(const FStateTreeHandle StateHandle)
+{
+	if (!StateHandle.IsValid())
+	{
+		return;
+	}
+	const FBakedStateTreeState& State = StateTree->States[StateHandle.Index];
+
+	// Stop any tasks in progress.
+	if (CurrentTask.IsValid())
+	{
+		Tasks[State.TasksBegin + CurrentTask.Index]->Deactivate(*this);
+	}
+}
+
+EStateTreeRunStatus FStateTreeInstance::TickState(const FStateTreeHandle StateHandle, const float DeltaTime)
+{
+	if (!StateHandle.IsValid())
+	{
+		return EStateTreeRunStatus::Failed;
+	}
+	const FBakedStateTreeState& State = StateTree->States[StateHandle.Index];
+	EStateTreeRunStatus Result = EStateTreeRunStatus::Running;
+
+	// Invalid NextTasks means restart of the sequence (i.e. enter state or restart sequence)
+	if (!NextTask.IsValid())
+	{
+		// Stop current task if running.
+		if (CurrentTask.IsValid())
+		{
+			Tasks[State.TasksBegin + CurrentTask.Index]->Deactivate(*this);
+			CurrentTask = FStateTreeHandle::Invalid;
+		}
+		// Start from first task
+		NextTask = FStateTreeHandle(0);
+	}
+
+	// Change task if requested
+	if (CurrentTask.Index != NextTask.Index)
+	{
+		if (CurrentTask.IsValid())
+		{
+			Tasks[State.TasksBegin + CurrentTask.Index]->Deactivate(*this);
+		}
+
+		CurrentTask = NextTask;
+
+		Tasks[State.TasksBegin + CurrentTask.Index]->Activate(*this);
+	}
+
+	if (CurrentTask.IsValid())
+	{
+		Result = Tasks[State.TasksBegin + CurrentTask.Index]->Tick(*this, DeltaTime);
+	}
+
+	// If task succeeded, goto next task or report that the sequence has succeeded.
+	if (Result == EStateTreeRunStatus::Succeeded)
+	{
+		// Goto next Task
+		NextTask.Index = CurrentTask.Index + 1;
+		if (NextTask.Index < State.TasksNum)
+		{
+			Result = EStateTreeRunStatus::Running;
+		}
+		else
+		{
+			// Mark to restart sequence in case the same State gets selected again.
+			NextTask = FStateTreeHandle::Invalid;
+			Result = EStateTreeRunStatus::Succeeded;
+		}
+	}
+
+	return Result;
+}
+
+FStateTreeInstance::FStateSelectionResult FStateTreeInstance::TriggerTransitions(const FStateTreeHandle CurrentStateHandle)
+{
+	// Walk towards root and check all transitions along the way.
+	for (FStateTreeHandle Handle = CurrentStateHandle; Handle.IsValid(); Handle = StateTree->States[Handle.Index].Parent)
+	{
+		const FBakedStateTreeState& State = StateTree->States[Handle.Index];
+		for (uint32 i = 0; i < State.TransitionsNum; i++)
+		{
+			// All transition conditions must pass
+			const FBakedStateTransition& Transition = StateTree->Transitions[State.TransitionsBegin + i];
+			if (TestAllConditions(Transition.ConditionsBegin, Transition.ConditionsNum))
+			{
+				if (Transition.Type == EStateTreeTransitionType::GotoState || Transition.Type == EStateTreeTransitionType::NextState)
+				{
+					FStateSelectionResult SelectedState = SelectState(Transition.State);
+					if (SelectedState.IsValid())
+					{
+						return SelectedState;
+					}
+				}
+				else if (Transition.Type == EStateTreeTransitionType::SelectChildState)
+				{
+					FStateSelectionResult SelectedState = SelectChildState(CurrentStateHandle);
+					if (SelectedState.IsValid())
+					{
+						return SelectedState;
+					}
+				}
+				else if (Transition.Type == EStateTreeTransitionType::NotSet)
+				{
+					// NotSet is no-operation, but can be used to mask a transition at parent state. Returning invalid keeps updating current state.
+					return FStateSelectionResult(FStateTreeHandle::Invalid);
+				}
+				else if (Transition.Type == EStateTreeTransitionType::Succeeded)
+				{
+					return FStateSelectionResult(EStateTreeRunStatus::Succeeded);
+				}
+				else
+				{
+					return FStateSelectionResult(EStateTreeRunStatus::Failed);
+				}
+			}
+		}
+	}
+
+	// No transition triggered, keep on updating current state.
+	return FStateSelectionResult(FStateTreeHandle::Invalid);
+}
+
+FStateTreeInstance::FStateSelectionResult FStateTreeInstance::TriggerDefaultTransition(const FStateTreeHandle StateHandle, EStateTreeRunStatus StateTickStatus)
+{
+	const FBakedStateTreeState& State = StateTree->States[StateHandle.Index];
+
+	EStateTreeTransitionType TransitionType = State.StateDoneTransitionType;
+	FStateTreeHandle TransitionState = State.StateDoneTransitionState;
+	
+	// Use State Failed transition if state has failed, and failed transition is set.
+	if (StateTickStatus == EStateTreeRunStatus::Failed && State.StateFailedTransitionType != EStateTreeTransitionType::NotSet)
+	{
+		TransitionType = State.StateFailedTransitionType;
+		TransitionState = State.StateFailedTransitionState;
+	}
+	
+	if ((TransitionType == EStateTreeTransitionType::GotoState || TransitionType == EStateTreeTransitionType::NextState) && TransitionState.IsValid())
+	{
+		// If default transition is set to specific state, run selection logic from there.
+		FStateTreeInstance::FStateSelectionResult Result = SelectState(TransitionState);
+		if (Result.IsValid())
+		{
+			return Result;
+		}
+		else
+		{
+			STATETREE_LOG(Error, TEXT("%s: Default transition on state '%s' could not be triggered. '%s' using StateTree %s."), ANSI_TO_TCHAR(__FUNCTION__), *State.Name.ToString(), *GetNameSafe(Owner), *GetNameSafe(StateTree));
+			return FStateSelectionResult(EStateTreeRunStatus::Failed);
+		}
+	}
+	else if (TransitionType == EStateTreeTransitionType::SelectChildState)
+	{
+		// If no default transition specified, run selection logic on children.
+		FStateTreeInstance::FStateSelectionResult Result = SelectChildState(StateHandle);
+		if (Result.IsValid())
+		{
+			return Result;
+		}
+		else
+		{
+			STATETREE_LOG(Error, TEXT("%s: Default transition on state '%s' could not be triggered. '%s' using StateTree %s."), ANSI_TO_TCHAR(__FUNCTION__), *State.Name.ToString(), *GetNameSafe(Owner), *GetNameSafe(StateTree));
+			return FStateSelectionResult(EStateTreeRunStatus::Failed);
+		}
+	}
+	else if (TransitionType == EStateTreeTransitionType::Succeeded)
+	{
+		return FStateSelectionResult(EStateTreeRunStatus::Succeeded);
+	}
+
+	return FStateSelectionResult(EStateTreeRunStatus::Failed);
+}
+
+FStateTreeInstance::FStateSelectionResult FStateTreeInstance::SelectState(const FStateTreeHandle StateHandle)
+{
+	const FBakedStateTreeState& State = StateTree->States[StateHandle.Index];
+
+	// Check that the state can be entered
+	if (!TestAllConditions(State.EnterConditionsBegin, State.EnterConditionsNum))
+	{
+		return FStateSelectionResult(FStateTreeHandle::Invalid);
+	}
+
+	// If the state has any tasks, select it.
+	if (State.TasksNum > 0)
+	{
+		return FStateSelectionResult(StateHandle);
+	}
+	// Else proceed to select any of the children.
+	return SelectChildState(StateHandle);
+}
+
+FStateTreeInstance::FStateSelectionResult FStateTreeInstance::SelectChildState(const FStateTreeHandle StateHandle)
+{
+	const FBakedStateTreeState& State = StateTree->States[StateHandle.Index];
+
+	// Visit children in order.
+	for (uint16 ChildState = State.ChildrenBegin; ChildState < State.ChildrenEnd;)
+	{
+		FStateSelectionResult SelectedState = SelectState(FStateTreeHandle(ChildState));
+		if (SelectedState.IsValid())
+		{
+			return SelectedState;
+		}
+		ChildState = StateTree->States[ChildState].GetNextSibling();
+	}
+
+	return FStateSelectionResult(FStateTreeHandle::Invalid);
+}
+
+EStateTreeRunStatus FStateTreeInstance::Tick(const float DeltaTime, FStateTreeParameterStorage* InputParameters)
+{
+	if (RunStatus != EStateTreeRunStatus::Running)
+	{
+		return RunStatus;
+	}
+
+	if (!Owner || !StateTree)
+	{
+		RunStatus = EStateTreeRunStatus::Failed;
+		return RunStatus;
+	}
+
+	if (InputParameters)
+	{
+		FStateTreeParameterLayout MappedParameters;
+		InputParameters->MapVariables(MappedParameters, StateTree->Parameters);
+		ReadParametersFromStorage(MappedParameters, *InputParameters);
+	}
+
+	TickEvaluators(DeltaTime);
+
+	FStateSelectionResult NextState(EStateTreeRunStatus::Failed);
+
+	if (CurrentState == FStateTreeHandle::Invalid)
+	{
+		// No state has been selected yet, select starting from root.
+		static const FStateTreeHandle RootState = FStateTreeHandle(0);
+		NextState = SelectState(RootState);
+	}
+	else
+	{
+		if (LastTickStatus == EStateTreeRunStatus::Running)
+		{
+			// Check transitions while running, returns "running/invalid" if no transition is triggered.
+			NextState = TriggerTransitions(CurrentState);
+		}
+		else
+		{
+			// State completed or failed on last tick, check if we can proceed to the default next state, returns "failed" if transition cannot be triggered.
+			NextState = TriggerDefaultTransition(CurrentState, LastTickStatus);
+		}
+	}
+
+	if (NextState.Status != EStateTreeRunStatus::Running)
+	{
+		// Transition to a terminal state (succeeded/failed), or default transition failed.
+		ExitState(CurrentState);
+		RunStatus = NextState.Status;
+		return RunStatus;
+	}
+	if (NextState.IsValid())
+	{
+		// New state has become active, switch state.
+		ExitState(CurrentState);
+		CurrentState = NextState.State;
+		EnterState(CurrentState);
+	}
+	if (!CurrentState.IsValid())
+	{
+		// Should not happen.
+		STATETREE_LOG(Error, TEXT("%s: Failed to select state on '%s' using StateTree %s."), ANSI_TO_TCHAR(__FUNCTION__), *GetNameSafe(Owner), *GetNameSafe(StateTree));
+		RunStatus = EStateTreeRunStatus::Failed;
+		return RunStatus;
+	}
+
+	// Run state
+	LastTickStatus = TickState(CurrentState, DeltaTime);
+
+	return RunStatus;
+}
+
+#if ENABLE_VISUAL_LOG
+void FStateTreeInstance::DescribeSelfToVisLog(struct FVisualLogEntry& Snapshot) const
+{
+}
+#endif // ENABLE_VISUAL_LOG
+
+#if WITH_GAMEPLAY_DEBUGGER
+FString FStateTreeInstance::GetDebugInfoString() const
+{
+	if (!StateTree)
+	{
+		return FString(TEXT("No StateTree asset."));
+	}
+
+	FString DebugString = FString::Printf(TEXT("StateTree (asset: %s)\n"), *GetNameSafe(StateTree));
+
+	DebugString += TEXT("Status: ");
+	switch (RunStatus)
+	{
+	case EStateTreeRunStatus::Failed:
+		DebugString += TEXT("Failed\n");
+		break;
+	case EStateTreeRunStatus::Succeeded:
+		DebugString += TEXT("Succeeded\n");
+		break;
+	case EStateTreeRunStatus::Running:
+		DebugString += TEXT("Running\n");
+		break;
+	default:
+		DebugString += TEXT("--\n");
+	}
+
+	// Active States
+	TArray<FStateTreeHandle> ActiveStates;
+	for (FStateTreeHandle Handle = CurrentState; Handle.IsValid(); Handle = StateTree->States[Handle.Index].Parent)
+	{
+		ActiveStates.Insert(Handle, 0);
+	}
+
+	DebugString += TEXT("Current State:\n");
+	for (FStateTreeHandle Handle : ActiveStates)
+	{
+		if (Handle.IsValid())
+		{
+			const FBakedStateTreeState& State = StateTree->States[Handle.Index];
+			DebugString += FString::Printf(TEXT("[%s] "), *State.Name.ToString());
+			if (!StateTree->IsV2())
+			{
+				if (Handle == CurrentState)
+				{
+					for (uint32 i = 0; i < State.TasksNum; i++)
+					{
+						DebugString += FString::Printf(TEXT(" [%s%s] "), i == (uint32)CurrentTask.Index ? TEXT(">>") : TEXT(""), *Tasks[State.TasksBegin + i]->Name.ToString());
+					}
+				}
+			}
+			DebugString += TEXT("\n");
+		}
+	}
+
+	if (CurrentState.IsValid() && CurrentTask.IsValid())
+	{
+		const FBakedStateTreeState& State = StateTree->States[CurrentState.Index];
+		DebugString += FString::Printf(TEXT("\nCurrent Task %s (%d):\n"), *Tasks[State.TasksBegin + CurrentTask.Index]->Name.ToString(), CurrentTask.Index);
+		Tasks[State.TasksBegin + CurrentTask.Index]->AppendDebugInfoString(DebugString, *this);
+	}
+
+	DebugString += TEXT("\nEvaluators:\n");
+	for (const UStateTreeEvaluatorBase* Eval : Evaluators)
+	{
+		Eval->AppendDebugInfoString(DebugString, *this);
+	}
+
+	DebugString += TEXT("\nVariables:\n");
+	for (const FStateTreeVariableDesc& Var : StateTree->Variables.Variables)
+	{
+		const FStateTreeHandle Handle(Var.Offset);
+		switch (Var.Type)
+		{
+		case EStateTreeVariableType::Bool:
+			DebugString += FString::Printf(TEXT("    %s = %s\n"), *Var.Name.ToString(), GetValueBool(Handle) ? TEXT("true") : TEXT("false"));
+			break;
+		case EStateTreeVariableType::Float:
+			DebugString += FString::Printf(TEXT("    %s = %f\n"), *Var.Name.ToString(), GetValueFloat(Handle));
+			break;
+		case EStateTreeVariableType::Int:
+			DebugString += FString::Printf(TEXT("    %s = %d\n"), *Var.Name.ToString(), GetValueInt(Handle));
+			break;
+		case EStateTreeVariableType::Object:
+			DebugString += FString::Printf(TEXT("    %s = %s\n"), *Var.Name.ToString(), *GetNameSafe(GetValueObject(Handle).Get()));
+			break;
+		case EStateTreeVariableType::Vector:
+			DebugString += FString::Printf(TEXT("    %s = %s\n"), *Var.Name.ToString(), *GetValueVector(Handle).ToString());
+			break;
+		default:
+			DebugString += FString::Printf(TEXT("    %s = ??\n"), *Var.Name.ToString());
+			break;
+		}
+	}
+
+	return DebugString;
+}
+
+#endif // WITH_GAMEPLAY_DEBUGGER
+
+#undef STATETREE_LOG
+#undef STATETREE_CLOG
diff --git a/Engine/Plugins/Runtime/StateTree/Source/StateTreeModule/Private/StateTreeModule.cpp b/Engine/Plugins/Runtime/StateTree/Source/StateTreeModule/Private/StateTreeModule.cpp
new file mode 100644
--- /dev/null
+++ b/Engine/Plugins/Runtime/StateTree/Source/StateTreeModule/Private/StateTreeModule.cpp
@@ -0,0 +1,25 @@
+// Copyright Epic Games, Inc. All Rights Reserved.
+
+#include "StateTreeModule.h"
+#include "Modules/ModuleManager.h"
+
+#define LOCTEXT_NAMESPACE "StateTree"
+
+class FStateTreeModule : public IStateTreeModule
+{
+	// Begin IModuleInterface
+	virtual void StartupModule() override;
+	virtual void ShutdownModule() override;
+};
+
+IMPLEMENT_MODULE(FStateTreeModule, StateTreeModule)
+
+void FStateTreeModule::StartupModule()
+{
+}
+
+void FStateTreeModule::ShutdownModule()
+{
+}
+
+#undef LOCTEXT_NAMESPACE
diff --git a/Engine/Plugins/Runtime/StateTree/Source/StateTreeModule/Private/StateTreeParameter.cpp b/Engine/Plugins/Runtime/StateTree/Source/StateTreeModule/Private/StateTreeParameter.cpp
new file mode 100644
--- /dev/null
+++ b/Engine/Plugins/Runtime/StateTree/Source/StateTreeModule/Private/StateTreeParameter.cpp
@@ -0,0 +1,13 @@
+// Copyright Epic Games, Inc. All Rights Reserved.
+
+#include "StateTreeParameter.h"
+#include "StateTreeVariableDesc.h"
+
+#define LOCTEXT_NAMESPACE "StateTreeEditor"
+
+#if WITH_EDITOR
+bool FStateTreeParameter::IsCompatible(const FStateTreeVariableDesc& Desc) const
+{
+	return Name == Desc.Name && ID == Desc.ID && Variable.Type == Desc.Type && Variable.BaseClass == Desc.BaseClass;
+}
+#endif
\ No newline at end of file
diff --git a/Engine/Plugins/Runtime/StateTree/Source/StateTreeModule/Private/StateTreeParameterLayout.cpp b/Engine/Plugins/Runtime/StateTree/Source/StateTreeModule/Private/StateTreeParameterLayout.cpp
new file mode 100644
--- /dev/null
+++ b/Engine/Plugins/Runtime/StateTree/Source/StateTreeModule/Private/StateTreeParameterLayout.cpp
@@ -0,0 +1,76 @@
+// Copyright Epic Games, Inc. All Rights Reserved.
+
+#include "StateTreeParameterLayout.h"
+#include "StateTreeVariableLayout.h"
+#include "StateTreeConstantStorage.h"
+
+#define LOCTEXT_NAMESPACE "StateTreeEditor"
+
+
+#if WITH_EDITOR
+bool FStateTreeParameterLayout::ResolveVariables(const FStateTreeVariableLayout& Variables, FStateTreeConstantStorage& Constants)
+{
+	for (FStateTreeParameter& Param : Parameters)
+	{
+		if (!Param.Variable.ResolveHandle(Variables, Constants))
+		{
+			return false;
+		}
+	}
+
+	return true;
+}
+
+bool FStateTreeParameterLayout::IsCompatible(const FStateTreeVariableLayout* ParameterLayout) const
+{
+	if (!ParameterLayout)
+	{
+		return Parameters.Num() == 0;
+	}
+
+	if (Parameters.Num() != ParameterLayout->Variables.Num())
+	{
+		return false;
+	}
+	for (int32 i = 0; i < Parameters.Num(); i++)
+	{
+		if (!Parameters[i].IsCompatible(ParameterLayout->Variables[i]))
+		{
+			return false;
+		}
+	}
+	return true;
+}
+
+void FStateTreeParameterLayout::UpdateLayout(const FStateTreeVariableLayout* NewLayout)
+{
+	if (!NewLayout)
+	{
+		Reset();
+		return;
+	}
+
+	TMap<FGuid, FStateTreeVariable> OldVariables;
+	for (const FStateTreeParameter& Param : Parameters)
+	{
+		OldVariables.Add(Param.ID, Param.Variable);
+	}
+
+	Parameters.SetNum(NewLayout->Variables.Num());
+	for (int32 i = 0; i < NewLayout->Variables.Num(); i++)
+	{
+		FStateTreeParameter& Param = Parameters[i];
+		const FStateTreeVariableDesc& Var = NewLayout->Variables[i];
+		Param.Name = Var.Name;
+		Param.ID = Var.ID;
+		Param.Variable.InitializeFromDesc(EStateTreeVariableBindingMode::Typed, Var);
+
+		// Try to carry over the previous values.
+		FStateTreeVariable* OldVar = OldVariables.Find(Param.ID);
+		if (OldVar)
+		{
+			Param.Variable.CopyValueFrom(*OldVar);
+		}
+	}
+}
+#endif // WITH_EDITOR
diff --git a/Engine/Plugins/Runtime/StateTree/Source/StateTreeModule/Private/StateTreeParameterStorage.cpp b/Engine/Plugins/Runtime/StateTree/Source/StateTreeModule/Private/StateTreeParameterStorage.cpp
new file mode 100644
--- /dev/null
+++ b/Engine/Plugins/Runtime/StateTree/Source/StateTreeModule/Private/StateTreeParameterStorage.cpp
@@ -0,0 +1,109 @@
+// Copyright Epic Games, Inc. All Rights Reserved.
+
+#include "StateTreeParameterStorage.h"
+#include "StateTreeVariableLayout.h"
+#include "StateTreeInstance.h"
+#include "StateTreeTypes.h"
+
+void FStateTreeParameterStorage::Reset()
+{
+	Layout.Reset();
+	Memory.Reset(0);
+}
+
+void FStateTreeParameterStorage::SetLayout(const FStateTreeParameterLayout& NewLayout)
+{
+	// Create partial variable map from parameer set binding
+	Layout.Variables.Reset(NewLayout.Parameters.Num());
+	for (const FStateTreeParameter& Param : NewLayout.Parameters)
+	{
+		FStateTreeVariableDesc& Desc = Layout.Variables.AddDefaulted_GetRef();
+		Desc.Name = Param.Name;
+		Desc.Type = Param.Variable.Type;
+#if WITH_EDITORONLY_DATA
+		Desc.BaseClass = Param.Variable.BaseClass;
+		Desc.ID = FGuid();
+#endif
+	}
+
+	const uint16 MemoryUsage = Layout.CalculateOffsetsAndMemoryUsage();
+	Memory.SetNumZeroed(MemoryUsage);
+}
+
+void FStateTreeParameterStorage::SetLayout(const FStateTreeVariableLayout& NewLayout)
+{
+	Layout = NewLayout;
+#if WITH_EDITORONLY_DATA
+	// Cannot bind to storage with ID, make sure we don't accidentally carry these around.
+	for (FStateTreeVariableDesc& Desc : Layout.Variables)
+	{
+		Desc.ID = FGuid();
+	}
+#endif
+	uint16 MemoryUsage = Layout.CalculateOffsetsAndMemoryUsage();
+	Memory.SetNumZeroed(MemoryUsage);
+}
+
+void FStateTreeParameterStorage::MapVariables(FStateTreeParameterLayout& MappedLayout, const FStateTreeParameterLayout& ExpectedLayout) const
+{
+	MappedLayout.Parameters.Reset(Layout.Variables.Num());
+
+	// Optimize for the case that the bindings are sequentially equal in layout.
+	int32 Index = 0;
+	const int32 ExpectedParametersNum = ExpectedLayout.Parameters.Num();
+
+	// Create binding for all parameters in storages parameter set, find matching binding from input.
+	for (const FStateTreeVariableDesc& Desc : Layout.Variables)
+	{
+		FStateTreeParameter& Param = MappedLayout.Parameters.AddDefaulted_GetRef();
+		Param.Name = Desc.Name;
+		Param.Variable.Type = Desc.Type;
+		Param.Variable.Handle = FStateTreeHandle::Invalid;
+
+		// Try to find matching binding
+		for (int32 j = 0; j < ExpectedParametersNum; j++)
+		{
+			const FStateTreeParameter& ExpectedParam = ExpectedLayout.Parameters[Index];
+			Index = (Index + 1) < ExpectedParametersNum ? (Index + 1) : 0;	// Incement first, so that the counters advance together on success.
+			if (ExpectedParam.Variable.Type == Param.Variable.Type && ExpectedParam.Name == Param.Name)
+			{
+				Param.Variable.Handle = ExpectedParam.Variable.Handle;
+				break;
+			}
+		}
+	}
+}
+
+bool FStateTreeParameterStorage::IsCompatible(const FStateTreeParameterLayout& ParameterLayout) const
+{
+	if (Layout.Variables.Num() != ParameterLayout.Parameters.Num())
+	{
+		return false;
+	}
+	for (int32 i = 0; i < Layout.Variables.Num(); i++)
+	{
+		if (Layout.Variables[i].Type != ParameterLayout.Parameters[i].Variable.Type)
+		{
+			return false;
+		}
+	}
+	return true;
+}
+
+uint8* FStateTreeParameterStorage::GetParameterPtr(uint32 Index, EStateTreeVariableType ExpectedType)
+{
+	if (Index >= (uint32)Layout.Variables.Num() || Layout.Variables[Index].Type != ExpectedType)
+	{
+		return nullptr;
+	}
+	return &Memory[Layout.Variables[Index].Offset];
+}
+
+const uint8* FStateTreeParameterStorage::GetParameterPtr(uint32 Index, EStateTreeVariableType ExpectedType) const
+{
+	if (Index >= (uint32)Layout.Variables.Num() || Layout.Variables[Index].Type != ExpectedType)
+	{
+		return nullptr;
+	}
+	return &Memory[Layout.Variables[Index].Offset];
+}
