Moved GameplayBehaviors out of restricted folder over to Experimental

diff --git a/Engine/Plugins/Runtime/StateTree/Source/StateTreeModule/Private/StateTreePropertyBindings.cpp b/Engine/Plugins/Runtime/StateTree/Source/StateTreeModule/Private/StateTreePropertyBindings.cpp
new file mode 100644
--- /dev/null
+++ b/Engine/Plugins/Runtime/StateTree/Source/StateTreeModule/Private/StateTreePropertyBindings.cpp
@@ -0,0 +1,650 @@
+// Copyright Epic Games, Inc. All Rights Reserved.
+#include "StateTreePropertyBindings.h"
+#include "StateTreeTypes.h"
+
+//----------------------------------------------------------------//
+//  FStateTreePropertyBindings
+//----------------------------------------------------------------//
+
+void FStateTreePropertyBindings::Reset()
+{
+	SourceStructs.Reset();
+	CopyBatches.Reset();
+	PropertyBindings.Reset();
+	PropertyPaths.Reset();
+	PropertySegments.Reset();
+	PropertyCopies.Reset();
+	PropertyAccesses.Reset();
+	PropertyIndirections.Reset();
+	bBindingsResolved = false;
+}
+
+bool FStateTreePropertyBindings::ResolvePaths()
+{
+	PropertyIndirections.Reset();
+	PropertyCopies.SetNum(PropertyBindings.Num());
+	PropertyAccesses.SetNum(PropertyPaths.Num());
+
+	bBindingsResolved = true;
+
+	bool bResult = true;
+	
+	for (const FStateTreePropCopyBatch& Batch : CopyBatches)
+	{
+		for (int32 i = Batch.BindingsBegin; i != Batch.BindingsEnd; i++)
+		{
+			const FStateTreePropertyBinding& Binding = PropertyBindings[i];
+			FStateTreePropCopy& Copy = PropertyCopies[i];
+
+			Copy.SourceAccessIndex = Binding.SourcePathIndex;
+			Copy.TargetAccessIndex = Binding.TargetPathIndex;
+			Copy.SourceStructIndex = Binding.SourceStructIndex;
+			Copy.Type = Binding.CopyType;
+
+			const UStruct* SourceStruct = SourceStructs[Binding.SourceStructIndex].Struct;
+			const UStruct* TargetStruct = Batch.TargetStruct.Struct;
+			if (!SourceStruct || !TargetStruct)
+			{
+				bBindingsResolved = false;
+				break;
+			}
+
+			// Resolve paths and validate the copy. Stops on first failure.
+			bool bSuccess = true;
+			bSuccess = bSuccess && ResolvePath(SourceStruct, PropertyPaths[Binding.SourcePathIndex], PropertyAccesses[Copy.SourceAccessIndex]);
+			bSuccess = bSuccess && ResolvePath(TargetStruct, PropertyPaths[Binding.TargetPathIndex], PropertyAccesses[Copy.TargetAccessIndex]);
+			bSuccess = bSuccess && ValidateCopy(Copy);
+			if (!bSuccess)
+			{
+				// Resolving or validating failed, make the copy a nop.
+				Copy.Type = EStateTreePropertyCopyType::None;
+				bResult = false;
+			}
+		}
+	}
+
+	return bResult;
+}
+
+bool FStateTreePropertyBindings::ResolvePath(const UStruct* Struct, const FStateTreePropertyPath& Path, FStateTreePropertyAccess& OutAccess)
+{
+	if (!Struct)
+	{
+		UE_LOG(LogStateTree, Error, TEXT("%s: '%s' Invalid source struct."), ANSI_TO_TCHAR(__FUNCTION__), *GetPathAsString(Path));
+		return false;
+	}
+
+	const UStruct* CurrentStruct = Struct;
+	FProperty* LeafProperty = nullptr;
+	const int32 IndirectionsBegin = PropertyIndirections.Num();
+
+	for (int32 i = Path.SegmentsBegin; i < Path.SegmentsEnd; i++)
+	{
+		const FStateTreePropertySegment& Segment = PropertySegments[i];
+		FStateTreePropertyIndirection& Indirection = PropertyIndirections.AddDefaulted_GetRef();
+
+		const bool bFinalSegment = i == (Path.SegmentsEnd - 1);
+
+		if (!ensure(CurrentStruct))
+		{
+			UE_LOG(LogStateTree, Error, TEXT("%s: '%s' Invalid struct."), ANSI_TO_TCHAR(__FUNCTION__), *GetPathAsString(Path, i, TEXT("<"), TEXT(">")));
+			return false;
+		}
+		FProperty* Property = CurrentStruct->FindPropertyByName(Segment.Name);
+		if (!Property)
+		{
+			// TODO: use core redirects to fix up the name.
+			UE_LOG(LogStateTree, Error, TEXT("%s: Malformed path '%s', could not to find property '%s%s.%s'."),
+				ANSI_TO_TCHAR(__FUNCTION__), *GetPathAsString(Path, i, TEXT("<"), TEXT(">")),
+				CurrentStruct->GetPrefixCPP(), *CurrentStruct->GetName(), *Segment.Name.ToString());
+			return false;
+		}
+		Indirection.ArrayIndex = Segment.ArrayIndex == INDEX_NONE ? 0 : Segment.ArrayIndex;
+		Indirection.Type = Segment.Type;
+		Indirection.Offset = Property->GetOffset_ForInternal() + Property->ElementSize * Indirection.ArrayIndex;
+
+		// Check to see if it is an array access first.
+		if (FArrayProperty* ArrayProperty = CastField<FArrayProperty>(Property))
+		{
+			if (FStructProperty* ArrayOfStructsProperty = CastField<FStructProperty>(ArrayProperty->Inner))
+			{
+				CurrentStruct = ArrayOfStructsProperty->Struct;
+			}
+			else if (FObjectPropertyBase* ArrayOfObjectsProperty = CastField<FObjectPropertyBase>(ArrayProperty->Inner))
+			{
+				CurrentStruct = ArrayOfObjectsProperty->PropertyClass;
+			}
+			Indirection.ArrayProperty = ArrayProperty;
+		}
+		// Leaf segments all get treated the same, plain, struct or object
+		else if (bFinalSegment)
+		{
+			CurrentStruct = nullptr;
+		}
+		// Check to see if this is a simple structure (eg. not an array of structures)
+		else if (FStructProperty* StructProperty = CastField<FStructProperty>(Property))
+		{
+			CurrentStruct = StructProperty->Struct;
+		}
+		// Check to see if this is a simple object (eg. not an array of objects)
+		else if (FObjectProperty* ObjectProperty = CastField<FObjectProperty>(Property))
+		{
+			CurrentStruct = ObjectProperty->PropertyClass;
+		}
+		// Check to see if this is a simple weak object property (eg. not an array of weak objects).
+		else if (FWeakObjectProperty* WeakObjectProperty = CastField<FWeakObjectProperty>(Property))
+		{
+			CurrentStruct = WeakObjectProperty->PropertyClass;
+		}
+		// Check to see if this is a simple soft object property (eg. not an array of soft objects).
+		else if (FSoftObjectProperty* SoftObjectProperty = CastField<FSoftObjectProperty>(Property))
+		{
+			CurrentStruct = SoftObjectProperty->PropertyClass;
+		}
+		else
+		{
+			UE_LOG(LogStateTree, Error, TEXT("%s: Unsupported segment %s in path '%s'."),
+			    ANSI_TO_TCHAR(__FUNCTION__), *Property->GetCPPType(), *GetPathAsString(Path, i, TEXT("<"), TEXT(">")));
+		}
+
+		if (bFinalSegment)
+		{
+			LeafProperty = Property;
+		}
+	}
+
+	int32 IndirectionsEnd = PropertyIndirections.Num();
+
+	// Collapse adjacent offset indirections
+	for (int32 i = IndirectionsBegin; i < IndirectionsEnd; i++)
+	{
+		FStateTreePropertyIndirection& Indirection = PropertyIndirections[i];
+		if (Indirection.Type == EStateTreePropertyAccessType::Offset && (i + 1) < IndirectionsEnd)
+		{
+			FStateTreePropertyIndirection& NextIndirection = PropertyIndirections[i + 1];
+			if (NextIndirection.Type == EStateTreePropertyAccessType::Offset)
+			{
+				Indirection.Offset += NextIndirection.Offset;
+				PropertyIndirections.RemoveAt(i + 1);
+				i--;
+				IndirectionsEnd--;
+			}
+		}
+	}
+
+	OutAccess.IndirectionsBegin = IndirectionsBegin;
+	OutAccess.IndirectionsEnd = IndirectionsEnd;
+	OutAccess.LeafProperty = LeafProperty;
+
+	return true;
+}
+
+bool FStateTreePropertyBindings::ValidateCopy(FStateTreePropCopy& Copy) const
+{
+	const FProperty* SourceProperty = PropertyAccesses[Copy.SourceAccessIndex].LeafProperty;
+	const FProperty* TargetProperty = PropertyAccesses[Copy.TargetAccessIndex].LeafProperty;
+
+	if (!SourceProperty || !TargetProperty)
+	{
+		return false;
+	}
+
+	// Extract underlying types for enums
+	if (const FEnumProperty* EnumPropertyA = CastField<const FEnumProperty>(SourceProperty))
+	{
+		SourceProperty = EnumPropertyA->GetUnderlyingProperty();
+	}
+
+	if (const FEnumProperty* EnumPropertyB = CastField<const FEnumProperty>(TargetProperty))
+	{
+		TargetProperty = EnumPropertyB->GetUnderlyingProperty();
+	}
+
+	bool bResult = true;
+	switch (Copy.Type)
+	{
+	case EStateTreePropertyCopyType::CopyPlain:
+		Copy.CopySize = PropertyAccesses[Copy.SourceAccessIndex].LeafProperty->ElementSize * PropertyAccesses[Copy.SourceAccessIndex].LeafProperty->ArrayDim;
+		bResult = (SourceProperty->PropertyFlags & CPF_IsPlainOldData) != 0 && (TargetProperty->PropertyFlags & CPF_IsPlainOldData) != 0;
+		break;
+	case EStateTreePropertyCopyType::CopyComplex:
+		bResult = true;
+		break;
+	case EStateTreePropertyCopyType::CopyBool:
+		bResult = SourceProperty->IsA<FBoolProperty>() && TargetProperty->IsA<FBoolProperty>();
+		break;
+	case EStateTreePropertyCopyType::CopyStruct:
+		bResult = SourceProperty->IsA<FStructProperty>() && TargetProperty->IsA<FStructProperty>();
+		break;
+	case EStateTreePropertyCopyType::CopyObject:
+		bResult = SourceProperty->IsA<FObjectPropertyBase>() && TargetProperty->IsA<FObjectPropertyBase>();
+		break;
+	case EStateTreePropertyCopyType::CopyName:
+		bResult = SourceProperty->IsA<FNameProperty>() && TargetProperty->IsA<FNameProperty>();
+		break;
+	case EStateTreePropertyCopyType::CopyFixedArray:
+		bResult = SourceProperty->IsA<FArrayProperty>() && TargetProperty->IsA<FArrayProperty>();
+		break;
+	case EStateTreePropertyCopyType::PromoteBoolToByte:
+		bResult = SourceProperty->IsA<FBoolProperty>() && TargetProperty->IsA<FByteProperty>();
+		break;
+	case EStateTreePropertyCopyType::PromoteBoolToInt32:
+		bResult = SourceProperty->IsA<FBoolProperty>() && TargetProperty->IsA<FIntProperty>();
+		break;
+	case EStateTreePropertyCopyType::PromoteBoolToUInt32:
+		bResult = SourceProperty->IsA<FBoolProperty>() && TargetProperty->IsA<FUInt32Property>();
+		break;
+	case EStateTreePropertyCopyType::PromoteBoolToInt64:
+		bResult = SourceProperty->IsA<FBoolProperty>() && TargetProperty->IsA<FInt64Property>();
+		break;
+	case EStateTreePropertyCopyType::PromoteBoolToFloat:
+		bResult = SourceProperty->IsA<FBoolProperty>() && TargetProperty->IsA<FFloatProperty>();
+		break;
+	case EStateTreePropertyCopyType::PromoteBoolToDouble:
+		bResult = SourceProperty->IsA<FBoolProperty>() && TargetProperty->IsA<FDoubleProperty>();
+		break;
+	case EStateTreePropertyCopyType::PromoteByteToInt32:
+		bResult = SourceProperty->IsA<FByteProperty>() && TargetProperty->IsA<FIntProperty>();
+		break;
+	case EStateTreePropertyCopyType::PromoteByteToUInt32:
+		bResult = SourceProperty->IsA<FByteProperty>() && TargetProperty->IsA<FUInt32Property>();
+		break;
+	case EStateTreePropertyCopyType::PromoteByteToInt64:
+		bResult = SourceProperty->IsA<FByteProperty>() && TargetProperty->IsA<FInt64Property>();
+		break;
+	case EStateTreePropertyCopyType::PromoteByteToFloat:
+		bResult = SourceProperty->IsA<FByteProperty>() && TargetProperty->IsA<FFloatProperty>();
+		break;
+	case EStateTreePropertyCopyType::PromoteByteToDouble:
+		bResult = SourceProperty->IsA<FByteProperty>() && TargetProperty->IsA<FDoubleProperty>();
+		break;
+	case EStateTreePropertyCopyType::PromoteInt32ToInt64:
+		bResult = SourceProperty->IsA<FIntProperty>() && TargetProperty->IsA<FInt64Property>();
+		break;
+	case EStateTreePropertyCopyType::PromoteInt32ToFloat:
+		bResult = SourceProperty->IsA<FIntProperty>() && TargetProperty->IsA<FFloatProperty>();
+		break;
+	case EStateTreePropertyCopyType::PromoteInt32ToDouble:
+		bResult = SourceProperty->IsA<FIntProperty>() && TargetProperty->IsA<FDoubleProperty>();
+		break;
+	case EStateTreePropertyCopyType::PromoteFloatToDouble:
+		bResult = SourceProperty->IsA<FFloatProperty>() && TargetProperty->IsA<FDoubleProperty>();
+		break;
+	case EStateTreePropertyCopyType::DemoteDoubleToFloat:
+		bResult = SourceProperty->IsA<FDoubleProperty>() && TargetProperty->IsA<FFloatProperty>();
+		break;
+	case EStateTreePropertyCopyType::PromoteUInt32ToInt64:
+		bResult = SourceProperty->IsA<FUInt32Property>() && TargetProperty->IsA<FInt64Property>();
+		break;
+	case EStateTreePropertyCopyType::PromoteUInt32ToFloat:
+		bResult = SourceProperty->IsA<FUInt32Property>() && TargetProperty->IsA<FFloatProperty>();
+		break;
+	default:
+		UE_LOG(LogStateTree, Error, TEXT("FStateTreePropertyBindings::ValidateCopy: Unhandled copy type %s between '%s' and '%s'"),
+			*StaticEnum<EStateTreePropertyCopyType>()->GetValueAsString(Copy.Type), *SourceProperty->GetNameCPP(), *TargetProperty->GetNameCPP());
+		bResult = false;
+		break;
+	}
+
+	UE_CLOG(!bResult, LogStateTree, Error, TEXT("FStateTreePropertyBindings::ValidateCopy: Failed to validate copy type %s between '%s' and '%s'"),
+		*StaticEnum<EStateTreePropertyCopyType>()->GetValueAsString(Copy.Type), *SourceProperty->GetNameCPP(), *TargetProperty->GetNameCPP());
+	
+	return bResult;
+}
+
+uint8* FStateTreePropertyBindings::GetAddress(FStateTreeItemView InStructView, const FStateTreePropertyAccess& InAccess) const
+{
+	check(InStructView.IsValid());
+
+	uint8* Address = InStructView.GetMutableMemory();
+	for (int32 i = InAccess.IndirectionsBegin; Address != nullptr && i < InAccess.IndirectionsEnd; i++)
+	{
+		const FStateTreePropertyIndirection& Indirection = PropertyIndirections[i];
+		switch (Indirection.Type)
+		{
+		case EStateTreePropertyAccessType::Offset:
+		{
+			Address = Address + Indirection.Offset;
+			break;
+		}
+		case EStateTreePropertyAccessType::Object:
+		{
+			UObject* Object = *reinterpret_cast<UObject**>(Address + Indirection.Offset);
+			Address = reinterpret_cast<uint8*>(Object);
+			break;
+		}
+		case EStateTreePropertyAccessType::WeakObject:
+		{
+			TWeakObjectPtr<UObject>& WeakObjectPtr = *reinterpret_cast<TWeakObjectPtr<UObject>*>(Address + Indirection.Offset);
+			UObject* Object = WeakObjectPtr.Get();
+			Address = reinterpret_cast<uint8*>(Object);
+			break;
+		}
+		case EStateTreePropertyAccessType::SoftObject:
+		{
+			FSoftObjectPtr& SoftObjectPtr = *reinterpret_cast<FSoftObjectPtr*>(Address + Indirection.Offset);
+			UObject* Object = SoftObjectPtr.Get();
+			Address = reinterpret_cast<uint8*>(Object);
+			break;
+		}
+		case EStateTreePropertyAccessType::IndexArray:
+		{
+			check(Indirection.ArrayProperty);
+			FScriptArrayHelper Helper(Indirection.ArrayProperty, Address + Indirection.Offset);
+			if (Helper.IsValidIndex(Indirection.ArrayIndex))
+			{
+				Address = reinterpret_cast<uint8*>(Helper.GetRawPtr(Indirection.ArrayIndex));
+			}
+			else
+			{
+				Address = nullptr;
+			}
+			break;
+		}
+		default:
+			ensureMsgf(false, TEXT("FStateTreePropertyBindings::GetAddress: Unhandled indirection type %s for '%s'"),
+				*StaticEnum<EStateTreePropertyAccessType>()->GetValueAsString(Indirection.Type), *InAccess.LeafProperty->GetNameCPP());
+		}
+	}
+
+	return Address;
+}
+
+void FStateTreePropertyBindings::PerformCopy(const FStateTreePropCopy& Copy, const FProperty* SourceProperty, const uint8* SourceAddress, const FProperty* TargetProperty, uint8* TargetAddress) const
+{
+	check(SourceProperty);
+	check(SourceAddress);
+	check(TargetProperty);
+	check(TargetAddress);
+
+	switch (Copy.Type)
+	{
+	case EStateTreePropertyCopyType::CopyPlain:
+		FMemory::Memcpy(TargetAddress, SourceAddress, Copy.CopySize);
+		break;
+	case EStateTreePropertyCopyType::CopyComplex:
+		SourceProperty->CopyCompleteValue(TargetAddress, SourceAddress);
+		break;
+	case EStateTreePropertyCopyType::CopyBool:
+		static_cast<const FBoolProperty*>(TargetProperty)->SetPropertyValue(TargetAddress, static_cast<const FBoolProperty*>(SourceProperty)->GetPropertyValue(SourceAddress));
+		break;
+	case EStateTreePropertyCopyType::CopyStruct:
+		static_cast<const FStructProperty*>(TargetProperty)->Struct->CopyScriptStruct(TargetAddress, SourceAddress);
+		break;
+	case EStateTreePropertyCopyType::CopyObject:
+		static_cast<const FObjectPropertyBase*>(TargetProperty)->SetObjectPropertyValue(TargetAddress, static_cast<const FObjectPropertyBase*>(SourceProperty)->GetObjectPropertyValue(SourceAddress));
+		break;
+	case EStateTreePropertyCopyType::CopyName:
+		static_cast<const FNameProperty*>(TargetProperty)->SetPropertyValue(TargetAddress, static_cast<const FNameProperty*>(SourceProperty)->GetPropertyValue(SourceAddress));
+		break;
+	case EStateTreePropertyCopyType::CopyFixedArray:
+	{
+		// Copy into fixed sized array (EditFixedSize). Resizable arrays are copied as Complex, and regular fixed sizes arrays via the regular copies (dim specifies array size).
+		const FArrayProperty* SourceArrayProperty = static_cast<const FArrayProperty*>(SourceProperty);
+		const FArrayProperty* TargetArrayProperty = static_cast<const FArrayProperty*>(TargetProperty);
+		FScriptArrayHelper SourceArrayHelper(SourceArrayProperty, SourceAddress);
+		FScriptArrayHelper TargetArrayHelper(TargetArrayProperty, TargetAddress);
+			
+		const int32 MinSize = FMath::Min(SourceArrayHelper.Num(), TargetArrayHelper.Num());
+		for (int32 ElementIndex = 0; ElementIndex < MinSize; ++ElementIndex)
+		{
+			SourceArrayProperty->Inner->CopySingleValue(TargetArrayHelper.GetRawPtr(ElementIndex), SourceArrayHelper.GetRawPtr(ElementIndex));
+		}
+		break;
+	}
+	case EStateTreePropertyCopyType::PromoteBoolToByte:
+		*reinterpret_cast<uint8*>(TargetAddress) = (uint8)static_cast<const FBoolProperty*>(SourceProperty)->GetPropertyValue(SourceAddress);
+		break;
+	case EStateTreePropertyCopyType::PromoteBoolToInt32:
+		*reinterpret_cast<int32*>(TargetAddress) = (int32)static_cast<const FBoolProperty*>(SourceProperty)->GetPropertyValue(SourceAddress);
+		break;
+	case EStateTreePropertyCopyType::PromoteBoolToUInt32:
+		*reinterpret_cast<uint32*>(TargetAddress) = (uint32)static_cast<const FBoolProperty*>(SourceProperty)->GetPropertyValue(SourceAddress);
+		break;
+	case EStateTreePropertyCopyType::PromoteBoolToInt64:
+		*reinterpret_cast<int64*>(TargetAddress) = (int64)static_cast<const FBoolProperty*>(SourceProperty)->GetPropertyValue(SourceAddress);
+		break;
+	case EStateTreePropertyCopyType::PromoteBoolToFloat:
+		*reinterpret_cast<float*>(TargetAddress) = (float)static_cast<const FBoolProperty*>(SourceProperty)->GetPropertyValue(SourceAddress);
+		break;
+	case EStateTreePropertyCopyType::PromoteBoolToDouble:
+		*reinterpret_cast<double*>(TargetAddress) = (double)static_cast<const FBoolProperty*>(SourceProperty)->GetPropertyValue(SourceAddress);
+		break;
+	case EStateTreePropertyCopyType::PromoteByteToInt32:
+		*reinterpret_cast<int32*>(TargetAddress) = (int32)*reinterpret_cast<const uint8*>(SourceAddress);
+		break;
+	case EStateTreePropertyCopyType::PromoteByteToUInt32:
+		*reinterpret_cast<uint32*>(TargetAddress) = (uint32)*reinterpret_cast<const uint8*>(SourceAddress);
+		break;
+	case EStateTreePropertyCopyType::PromoteByteToInt64:
+		*reinterpret_cast<int64*>(TargetAddress) = (int64)*reinterpret_cast<const uint8*>(SourceAddress);
+		break;
+	case EStateTreePropertyCopyType::PromoteByteToFloat:
+		*reinterpret_cast<float*>(TargetAddress) = (float)*reinterpret_cast<const uint8*>(SourceAddress);
+		break;
+	case EStateTreePropertyCopyType::PromoteByteToDouble:
+		*reinterpret_cast<double*>(TargetAddress) = (double)*reinterpret_cast<const uint8*>(SourceAddress);
+		break;
+	case EStateTreePropertyCopyType::PromoteInt32ToInt64:
+		*reinterpret_cast<int64*>(TargetAddress) = (int64)*reinterpret_cast<const int32*>(SourceAddress);
+		break;
+	case EStateTreePropertyCopyType::PromoteInt32ToFloat:
+		*reinterpret_cast<float*>(TargetAddress) = (float)*reinterpret_cast<const int32*>(SourceAddress);
+		break;
+	case EStateTreePropertyCopyType::PromoteInt32ToDouble:
+		*reinterpret_cast<double*>(TargetAddress) = (double)*reinterpret_cast<const int32*>(SourceAddress);
+		break;
+	case EStateTreePropertyCopyType::PromoteFloatToDouble:
+		*reinterpret_cast<double*>(TargetAddress) = (double)*reinterpret_cast<const float*>(SourceAddress);
+		break;
+	case EStateTreePropertyCopyType::DemoteDoubleToFloat:
+		*reinterpret_cast<float*>(TargetAddress) = (float)*reinterpret_cast<const double*>(SourceAddress);
+		break;
+	case EStateTreePropertyCopyType::PromoteUInt32ToInt64:
+		*reinterpret_cast<int64*>(TargetAddress) = (int64)*reinterpret_cast<const uint32*>(SourceAddress);
+		break;
+	case EStateTreePropertyCopyType::PromoteUInt32ToFloat:
+		*reinterpret_cast<float*>(TargetAddress) = (float)*reinterpret_cast<const uint32*>(SourceAddress);
+		break;
+	default:
+		ensureMsgf(false, TEXT("FStateTreePropertyBindings::PerformCopy: Unhandled copy type %s between '%s' and '%s'"),
+			*StaticEnum<EStateTreePropertyCopyType>()->GetValueAsString(Copy.Type), *SourceProperty->GetNameCPP(), *TargetProperty->GetNameCPP());
+		break;
+	}
+}
+
+void FStateTreePropertyBindings::CopyTo(TConstArrayView<FStateTreeItemView> SourceStructViews, const int32 TargetBatchIndex, FStateTreeItemView TargetStructView) const
+{
+	// This is made ensure so that the programmers have the change to catch it (it's usually programming error not to call ResolvePaths(), and it wont spam log for others.
+	if (!ensureMsgf(bBindingsResolved, TEXT("Bindings must be resolved successfully before copying. See ResolvePaths()")))
+	{
+		return;
+	}
+
+	if (TargetBatchIndex == INDEX_NONE)
+	{
+		return;
+	}
+
+	check(CopyBatches.IsValidIndex(TargetBatchIndex));
+	const FStateTreePropCopyBatch& Batch = CopyBatches[TargetBatchIndex];
+
+	check(TargetStructView.IsValid());
+	check(TargetStructView.GetStruct() == Batch.TargetStruct.Struct);
+
+	for (int32 i = Batch.BindingsBegin; i != Batch.BindingsEnd; i++)
+	{
+		const FStateTreePropCopy& Copy = PropertyCopies[i];
+		// Copies that fail to be resolved (i.e. property path does not resolve, types changed) will be marked as None, skip them.
+		if (Copy.Type == EStateTreePropertyCopyType::None)
+		{
+			continue;
+		}
+		const FStateTreePropertyAccess& SourceAccess = PropertyAccesses[Copy.SourceAccessIndex];
+		const FStateTreePropertyAccess& TargetAccess = PropertyAccesses[Copy.TargetAccessIndex];
+		check(SourceStructViews[Copy.SourceStructIndex].GetStruct() == SourceStructs[Copy.SourceStructIndex].Struct);
+		const uint8* SourceAddress = GetAddress(SourceStructViews[Copy.SourceStructIndex], SourceAccess);
+		uint8* TargetAddress = GetAddress(TargetStructView, TargetAccess);
+		if (SourceAddress && TargetAddress)
+		{
+			PerformCopy(Copy, SourceAccess.LeafProperty, SourceAddress, TargetAccess.LeafProperty, TargetAddress);
+		}
+	}
+}
+
+void FStateTreePropertyBindings::DebugPrintInternalLayout(FString& OutString) const
+{
+	/** Array of expected source structs. */
+	OutString += FString::Printf(TEXT("\nBindableStructDesc (%d)\n  [ %-40s | %-40s ]\n"), SourceStructs.Num(), TEXT("Type"), TEXT("Name"));
+	for (const FStateTreeBindableStructDesc& BindableStructDesc : SourceStructs)
+	{
+		OutString += FString::Printf(TEXT("  | %-40s | %-40s |\n"),
+									 BindableStructDesc.Struct ? *BindableStructDesc.Struct->GetName() : TEXT("null"),
+									 *BindableStructDesc.Name.ToString());
+	}
+
+	/** Array of copy batches. */
+	OutString += FString::Printf(TEXT("\nCopyBatches (%d)\n  [ %-40s | %-40s | %-8s [%-3s:%-3s[ ]\n"), CopyBatches.Num(),
+		TEXT("Target Type"), TEXT("Target Name"), TEXT("Bindings"), TEXT("Beg"), TEXT("End"));
+	for (const FStateTreePropCopyBatch& CopyBatch : CopyBatches)
+	{
+		OutString += FString::Printf(TEXT("  | %-40s | %-40s | %8s [%3d:%-3d[ |\n"),
+									 CopyBatch.TargetStruct.Struct ? *CopyBatch.TargetStruct.Struct->GetName() : TEXT("null"),
+									 *CopyBatch.TargetStruct.Name.ToString(),
+									 TEXT(""), CopyBatch.BindingsBegin, CopyBatch.BindingsEnd);
+	}
+
+	/** Array of property bindings, resolved into arrays of copies before use. */
+	OutString += FString::Printf(TEXT("\nPropertyBindings (%d)\n  [ %-8s | %-8s | %-10s | %-50s ]\n"), PropertyBindings.Num(),
+	TEXT("Src Path"), TEXT("Tgt Path"), TEXT("Src Struct"), TEXT("Copy Type"));
+	for (const FStateTreePropertyBinding& PropertyBinding : PropertyBindings)
+	{
+		OutString += FString::Printf(TEXT("  | %8d | %8d | %10d | %-50s |\n"),
+									 PropertyBinding.SourcePathIndex, PropertyBinding.TargetPathIndex, PropertyBinding.SourceStructIndex,
+									 *StaticEnum<EStateTreePropertyCopyType>()->GetValueAsString(PropertyBinding.CopyType));
+	}
+
+	/** Array of property bindings, indexed by property paths. */
+	OutString += FString::Printf(TEXT("\nPropertyPaths (%d)\n  [ %-4s [%3s:%3s[ |\n"), PropertyPaths.Num(),
+		TEXT("Seg."), TEXT("Beg"), TEXT("End"));
+	for (const FStateTreePropertyPath& PropertyPath : PropertyPaths)
+	{
+		OutString += FString::Printf(TEXT("  | %4s [%3d:%-3d[ |\n"), TEXT(""), PropertyPath.SegmentsBegin, PropertyPath.SegmentsEnd);
+	}
+
+	/** Array of property segments, indexed by property paths. */
+	OutString += FString::Printf(TEXT("\nPropertySegments (%d)\n  [ %-20s | %3s | %-40s ]\n"), PropertySegments.Num(),
+		TEXT("Name"), TEXT("Idx"), TEXT("Access Type"));
+	for (const FStateTreePropertySegment& PropertySegment : PropertySegments)
+	{
+		OutString += FString::Printf(TEXT("  | %-20s | %3d | %-40s |\n"),
+									 *PropertySegment.Name.ToString(),
+									 PropertySegment.ArrayIndex, *StaticEnum<EStateTreePropertyAccessType>()->GetValueAsString(PropertySegment.Type));
+	}
+
+	/** Array of property copies */
+	OutString += FString::Printf(TEXT("\nPropertyCopies (%d)\n  [ %-7s | %-7s | %-7s | %-50s | %-4s ]\n"), PropertyCopies.Num(),
+		TEXT("Src Idx"), TEXT("Tgt Idx"), TEXT("Struct"), TEXT("Copy Type"), TEXT("Size"));
+	for (const FStateTreePropCopy& PropertyCopy : PropertyCopies)
+	{
+		OutString += FString::Printf(TEXT("  | %7d | %7d | %7d | %-50s | %4d |\n"),
+					PropertyCopy.SourceAccessIndex, PropertyCopy.TargetAccessIndex, PropertyCopy.SourceStructIndex,
+					*StaticEnum<EStateTreePropertyCopyType>()->GetValueAsString(PropertyCopy.Type),
+					PropertyCopy.CopySize);
+	}
+
+	/** Array of property accesses, indexed by copies*/
+	OutString += FString::Printf(TEXT("\nPropertyAccesses (%d)\n  [ %-8s [%-4s:%-4s [ |\n"), PropertyAccesses.Num(),
+		TEXT("Indirect"), TEXT("Beg"), TEXT("End"));
+	for (const FStateTreePropertyAccess& PropertyAccess : PropertyAccesses)
+	{
+		OutString += FString::Printf(TEXT("  | %8s [%4d:%-4d [ |\n"), TEXT(""), PropertyAccess.IndirectionsBegin, PropertyAccess.IndirectionsEnd);
+	}
+
+	/** Array of property indirections, indexed by accesses*/
+	OutString += FString::Printf(TEXT("\nPropertyIndirections (%d)\n  [ %-4s | %-4s | %-40s ] \n"), PropertyIndirections.Num(),
+		TEXT("Idx"), TEXT("Off."), TEXT("Access Type"));
+	for (const FStateTreePropertyIndirection& PropertyIndirection : PropertyIndirections)
+	{
+		OutString += FString::Printf(TEXT("  | %4d | %4d | %-40s |\n"),
+					PropertyIndirection.ArrayIndex, PropertyIndirection.Offset,*StaticEnum<EStateTreePropertyAccessType>()->GetValueAsString(PropertyIndirection.Type));
+	}
+}
+
+FString FStateTreePropertyBindings::GetPathAsString(const FStateTreePropertyPath& Path, const int32 HighlightedSegment, const TCHAR* HighlightPrefix, const TCHAR* HighlightPostfix)
+{
+	FString Result;
+	for (int32 i = Path.SegmentsBegin; i < Path.SegmentsEnd; i++)
+	{
+		const FStateTreePropertySegment& Segment = PropertySegments[i];
+		if (i > 0)
+		{
+			Result += TEXT(".");
+		}
+		
+		if (i == HighlightedSegment && HighlightPrefix)
+		{
+			Result += HighlightPrefix;
+		}
+
+		Result += Segment.Name.ToString();
+
+		if (i == HighlightedSegment && HighlightPostfix)
+		{
+			Result += HighlightPostfix;
+		}
+	}
+	return Result;
+}
+
+
+//----------------------------------------------------------------//
+//  FStateTreeEditorPropertyPath
+//----------------------------------------------------------------//
+
+FString FStateTreeEditorPropertyPath::ToString(const int32 HighlightedSegment, const TCHAR* HighlightPrefix, const TCHAR* HighlightPostfix) const
+{
+	FString Result;
+	for (int32 i = 0; i < Path.Num(); i++)
+	{
+		if (i > 0)
+		{
+			Result += TEXT(".");
+		}
+		if (i == HighlightedSegment && HighlightPrefix)
+		{
+			Result += HighlightPrefix;
+		}
+
+		Result += Path[i];
+
+		if (i == HighlightedSegment && HighlightPostfix)
+		{
+			Result += HighlightPostfix;
+		}
+	}
+	return Result;
+}
+
+bool FStateTreeEditorPropertyPath::operator==(const FStateTreeEditorPropertyPath& RHS) const
+{
+	if (StructID != RHS.StructID)
+	{
+		return false;
+	}
+
+	if (Path.Num() != RHS.Path.Num())
+	{
+		return false;
+	}
+
+	for (int32 i = 0; i < Path.Num(); i++)
+	{
+		if (Path[i] != RHS.Path[i])
+		{
+			return false;
+		}
+	}
+
+	return true;
+}
diff --git a/Engine/Plugins/Runtime/StateTree/Source/StateTreeModule/Private/StateTreeTaskBase.cpp b/Engine/Plugins/Runtime/StateTree/Source/StateTreeModule/Private/StateTreeTaskBase.cpp
new file mode 100644
--- /dev/null
+++ b/Engine/Plugins/Runtime/StateTree/Source/StateTreeModule/Private/StateTreeTaskBase.cpp
@@ -0,0 +1,35 @@
+// Copyright Epic Games, Inc. All Rights Reserved.
+
+#include "StateTreeTaskBase.h"
+#include "CoreMinimal.h"
+
+// Base class for all evaluators
+UStateTreeTaskBase::UStateTreeTaskBase(const FObjectInitializer& ObjectInitializer)
+	: Super(ObjectInitializer)
+	, ID(FGuid::NewGuid())
+{
+}
+
+#if WITH_EDITOR
+void UStateTreeTaskBase::PostLoad()
+{
+	Super::PostLoad();
+	if (!ID.IsValid())
+	{
+		ID = FGuid::NewGuid();
+	}
+}
+#endif
+
+#if WITH_GAMEPLAY_DEBUGGER
+void UStateTreeTaskBase::AppendDebugInfoString(FString& DebugString, const FStateTreeInstance& StateTreeInstance) const
+{
+}
+#endif
+
+#if WITH_GAMEPLAY_DEBUGGER
+void FStateTreeTask2Base::AppendDebugInfoString(FString& DebugString, const FStateTreeExecutionContext& Context) const
+{
+	DebugString += FString::Printf(TEXT("[%s]\n"), *Name.ToString());
+}
+#endif
diff --git a/Engine/Plugins/Runtime/StateTree/Source/StateTreeModule/Private/StateTreeTypes.cpp b/Engine/Plugins/Runtime/StateTree/Source/StateTreeModule/Private/StateTreeTypes.cpp
new file mode 100644
--- /dev/null
+++ b/Engine/Plugins/Runtime/StateTree/Source/StateTreeModule/Private/StateTreeTypes.cpp
@@ -0,0 +1,9 @@
+// Copyright Epic Games, Inc. All Rights Reserved.
+
+#include "StateTreeTypes.h"
+
+DEFINE_LOG_CATEGORY(LogStateTree);
+
+const FStateTreeHandle FStateTreeHandle::Invalid = FStateTreeHandle();
+const FStateTreeHandle FStateTreeHandle::Succeeded = FStateTreeHandle(FStateTreeHandle::SucceededIndex);
+const FStateTreeHandle FStateTreeHandle::Failed = FStateTreeHandle(FStateTreeHandle::FailedIndex);
diff --git a/Engine/Plugins/Runtime/StateTree/Source/StateTreeModule/Private/StateTreeVariable.cpp b/Engine/Plugins/Runtime/StateTree/Source/StateTreeModule/Private/StateTreeVariable.cpp
new file mode 100644
--- /dev/null
+++ b/Engine/Plugins/Runtime/StateTree/Source/StateTreeModule/Private/StateTreeVariable.cpp
@@ -0,0 +1,168 @@
+// Copyright Epic Games, Inc. All Rights Reserved.
+
+#include "StateTreeVariable.h"
+#include "StateTreeVariableDesc.h"
+#include "StateTreeVariableLayout.h"
+#include "StateTreeConstantStorage.h"
+
+#define LOCTEXT_NAMESPACE "StateTreeEditor"
+
+#if WITH_EDITORONLY_DATA
+// Make sure the supported variable types are POD and can fit into the storage.
+static_assert(sizeof(bool) <= sizeof(FStateTreeVariable::Value) && TIsPODType<bool>::Value, "bool does not fit into the Value storage or is not POD type.");
+static_assert(sizeof(int32) <= sizeof(FStateTreeVariable::Value) && TIsPODType<int32>::Value, "int32 does not fit into the Value storage or is not POD type.");
+static_assert(sizeof(float) <= sizeof(FStateTreeVariable::Value) && TIsPODType<float>::Value, "float does not fit into the Value storage or is not POD type.");
+static_assert(sizeof(FVector) <= sizeof(FStateTreeVariable::Value) && TIsPODType<FVector>::Value, "FVector does not fit into the Value storage or is not POD type.");
+#endif
+
+FStateTreeVariable::FStateTreeVariable()
+{
+#if WITH_EDITORONLY_DATA
+	BindingMode = EStateTreeVariableBindingMode::Any;
+	BaseClass = nullptr;
+	FMemory::Memzero(Value);
+#endif
+	Handle = FStateTreeHandle::Invalid;
+	Type = EStateTreeVariableType::Void;
+}
+
+FStateTreeVariable::FStateTreeVariable(const EStateTreeVariableBindingMode InMode, const EStateTreeVariableType InType, TSubclassOf<UObject> InBaseClass)
+{
+#if WITH_EDITORONLY_DATA
+	BindingMode = InMode;
+	BaseClass = InBaseClass;
+	FMemory::Memzero(Value);
+	if (BindingMode == EStateTreeVariableBindingMode::Definition)
+	{
+		ID = FGuid::NewGuid();
+	}
+#endif
+	Handle = FStateTreeHandle::Invalid;
+	Type = InType;
+}
+
+FStateTreeVariable::FStateTreeVariable(const EStateTreeVariableBindingMode InMode, const FName& InName, const EStateTreeVariableType InType, TSubclassOf<UObject> InBaseClass)
+{
+#if WITH_EDITORONLY_DATA
+	BindingMode = InMode;
+	BaseClass = InBaseClass;
+	FMemory::Memzero(Value);
+	if (BindingMode == EStateTreeVariableBindingMode::Definition)
+	{
+		ID = FGuid::NewGuid();
+	}
+	Name = InName;
+#endif
+	Handle = FStateTreeHandle::Invalid;
+	Type = InType;
+}
+
+#if WITH_EDITOR
+FStateTreeVariableDesc FStateTreeVariable::AsVariableDesc() const
+{
+	FStateTreeVariableDesc Desc;
+	Desc.Name = Name;
+	Desc.ID = ID;
+	Desc.BaseClass = BaseClass;
+	Desc.Offset = 0;
+	Desc.Type = Type;
+	return Desc;
+}
+
+FText FStateTreeVariable::GetDescription() const
+{
+	if (IsBound())
+	{
+		return  FText::FromName(Name);
+	}
+
+	switch (Type)
+	{
+	case EStateTreeVariableType::Void:
+		return FText(LOCTEXT("Void", "Void"));
+	case EStateTreeVariableType::Bool:
+		return GetBoolValue() ? FText(LOCTEXT("True", "True")) : FText(LOCTEXT("False", "False"));
+	case EStateTreeVariableType::Int:
+		return FText::AsNumber(GetIntValue());
+	case EStateTreeVariableType::Float:
+		return FText::AsNumber(GetFloatValue());
+	case EStateTreeVariableType::Vector:
+		return GetVectorValue().ToText();
+	case EStateTreeVariableType::Object:
+		return FText(LOCTEXT("Empty", "Empty"));
+	default:
+		return FText(LOCTEXT("Undefined", "Undefined"));
+		break;
+	}
+}
+
+bool FStateTreeVariable::ResolveHandle(const FStateTreeVariableLayout& Variables, FStateTreeConstantStorage& Constants)
+{
+	if (IsBound())
+	{
+		Handle = Variables.GetVariableHandle(ID);
+	}
+	else
+	{
+		switch (Type)
+		{
+		case EStateTreeVariableType::Bool:
+			Handle = Constants.AddConstantBool(GetBoolValue());
+			break;
+		case EStateTreeVariableType::Float:
+			Handle = Constants.AddConstantFloat(GetFloatValue());
+			break;
+		case EStateTreeVariableType::Int:
+			Handle = Constants.AddConstantInt(GetIntValue());
+			break;
+		case EStateTreeVariableType::Vector:
+			Handle = Constants.AddConstantVector(GetVectorValue());
+			break;
+		default:
+			Handle = FStateTreeHandle::Invalid;
+			break;
+		}
+	}
+	return true;
+}
+
+void FStateTreeVariable::InitializeFromDesc(EStateTreeVariableBindingMode InMode, const FStateTreeVariableDesc& Desc)
+{
+	BindingMode = InMode;
+	Type = Desc.Type;
+	BaseClass = Desc.BaseClass;
+	if (BindingMode == EStateTreeVariableBindingMode::Definition)
+	{
+		// Defines parameter
+		ID = Desc.ID;
+		Name = Desc.Name;
+	}
+	else
+	{
+		// Refers to parameter
+		ID.Invalidate();
+		Name = FName();
+	}
+	FMemory::Memzero(Value);
+}
+
+bool FStateTreeVariable::HasSameType(const FStateTreeVariable& Other) const
+{
+	return Other.Type == Type && Other.BaseClass == BaseClass;
+}
+
+bool FStateTreeVariable::CopyValueFrom(FStateTreeVariable& Other)
+{
+	if (!HasSameType(Other))
+	{
+		return false;
+	}
+
+	ID = Other.ID;
+	Name = Other.Name;
+	FMemory::Memcpy(Value, Other.Value, FStateTreeVariable::ValueStorageSizeBytes);
+
+	return true;
+}
+
+#endif
diff --git a/Engine/Plugins/Runtime/StateTree/Source/StateTreeModule/Private/StateTreeVariableDesc.cpp b/Engine/Plugins/Runtime/StateTree/Source/StateTreeModule/Private/StateTreeVariableDesc.cpp
new file mode 100644
--- /dev/null
+++ b/Engine/Plugins/Runtime/StateTree/Source/StateTreeModule/Private/StateTreeVariableDesc.cpp
@@ -0,0 +1,4 @@
+// Copyright Epic Games, Inc. All Rights Reserved.
+
+#include "StateTreeVariableDesc.h"
+#include "StateTreeTypes.h"
diff --git a/Engine/Plugins/Runtime/StateTree/Source/StateTreeModule/Private/Tasks/StateTreeTask_SubStateTree.cpp b/Engine/Plugins/Runtime/StateTree/Source/StateTreeModule/Private/Tasks/StateTreeTask_SubStateTree.cpp
new file mode 100644
--- /dev/null
+++ b/Engine/Plugins/Runtime/StateTree/Source/StateTreeModule/Private/Tasks/StateTreeTask_SubStateTree.cpp
@@ -0,0 +1,91 @@
+// Copyright Epic Games, Inc. All Rights Reserved.
+
+#include "Tasks/StateTreeTask_SubStateTree.h"
+#include "CoreMinimal.h"
+#include "GameFramework/Actor.h"
+#include "StateTreeDelegates.h"
+#include "StateTreeInstance.h"
+#include "StateTreeVariableLayout.h"
+#include "StateTreeConstantStorage.h"
+
+UStateTreeTask_SubStateTree::UStateTreeTask_SubStateTree(const FObjectInitializer& ObjectInitializer)
+	:Super(ObjectInitializer)
+{
+}
+
+
+bool UStateTreeTask_SubStateTree::Initialize(FStateTreeInstance& StateTreeInstance)
+{
+	if (StateTree == nullptr || StateTreeInstance.GetOwner() == nullptr)
+	{
+		return false;
+	}
+
+	ParameterStorage.SetLayout(Parameters);
+
+	return SubStateTreeInstance.Init(*StateTreeInstance.GetOwner(), *StateTree);
+}
+
+void UStateTreeTask_SubStateTree::Activate(FStateTreeInstance& StateTreeInstance)
+{
+	SubStateTreeInstance.Start();
+}
+
+void UStateTreeTask_SubStateTree::Deactivate(FStateTreeInstance& StateTreeInstance)
+{
+	SubStateTreeInstance.Stop();
+}
+
+EStateTreeRunStatus UStateTreeTask_SubStateTree::Tick(FStateTreeInstance& StateTreeInstance, const float DeltaTime)
+{
+	StateTreeInstance.WriteParametersToStorage(Parameters, ParameterStorage);
+	EStateTreeRunStatus Status = SubStateTreeInstance.Tick(DeltaTime, &ParameterStorage);
+	return Status;
+}
+
+#if WITH_GAMEPLAY_DEBUGGER
+void UStateTreeTask_SubStateTree::AppendDebugInfoString(FString& DebugString, const FStateTreeInstance& StateTreeInstance) const
+{
+	Super::AppendDebugInfoString(DebugString, StateTreeInstance);
+	DebugString += SubStateTreeInstance.GetDebugInfoString();
+}
+#endif
+
+#if WITH_EDITOR
+bool UStateTreeTask_SubStateTree::ResolveVariables(const FStateTreeVariableLayout& Variables, FStateTreeConstantStorage& Constants, UObject* Outer)
+{
+	return Parameters.ResolveVariables(Variables, Constants);
+}
+
+bool UStateTreeTask_SubStateTree::ValidateParameterLayout()
+{
+	const FStateTreeVariableLayout* ParameterLayout = StateTree ? &StateTree->GetInputParameterLayout() : nullptr;
+	if (!Parameters.IsCompatible(ParameterLayout))
+	{
+		Modify();
+		Parameters.UpdateLayout(ParameterLayout);
+		return true;
+	}
+	return false;
+}
+
+void UStateTreeTask_SubStateTree::PostEditChangeProperty(FPropertyChangedEvent& PropertyChangedEvent)
+{
+	if (PropertyChangedEvent.Property)
+	{
+		if (PropertyChangedEvent.Property->GetFName() == GET_MEMBER_NAME_CHECKED(UStateTreeTask_SubStateTree, StateTree))
+		{
+			// Update parameters from asset
+			const FStateTreeVariableLayout* ParameterLayout = StateTree ? &StateTree->GetInputParameterLayout() : nullptr;
+			Parameters.UpdateLayout(ParameterLayout);
+
+			// Update UI
+			UStateTree* ParentStateTree = GetTypedOuter<UStateTree>();
+			if (ParentStateTree)
+			{
+				UE::StateTree::Delegates::OnParametersInvalidated.Broadcast(*ParentStateTree);
+			}
+		}
+	}
+}
+#endif
diff --git a/Engine/Plugins/Runtime/StateTree/Source/StateTreeModule/Private/Tasks/StateTreeTask_Test.cpp b/Engine/Plugins/Runtime/StateTree/Source/StateTreeModule/Private/Tasks/StateTreeTask_Test.cpp
new file mode 100644
--- /dev/null
+++ b/Engine/Plugins/Runtime/StateTree/Source/StateTreeModule/Private/Tasks/StateTreeTask_Test.cpp
@@ -0,0 +1,75 @@
+// Copyright Epic Games, Inc. All Rights Reserved.
+
+#include "Tasks/StateTreeTask_Test.h"
+#include "Engine/Engine.h"
+
+
+UStateTreeTask_Test::UStateTreeTask_Test(const FObjectInitializer& ObjectInitializer)
+	:Super(ObjectInitializer)
+{
+}
+
+bool UStateTreeTask_Test::Initialize(FStateTreeInstance& StateTreeInstance)
+{
+	return true;
+}
+
+void UStateTreeTask_Test::Activate(FStateTreeInstance& StateTreeInstance)
+{
+	Time = 0.0f;
+}
+
+void UStateTreeTask_Test::Deactivate(FStateTreeInstance& StateTreeInstance)
+{
+}
+
+EStateTreeRunStatus UStateTreeTask_Test::Tick(FStateTreeInstance& StateTreeInstance, const float DeltaTime)
+{
+	Time += DeltaTime;
+	if (Time > Duration)
+	{
+		return bFail ? EStateTreeRunStatus::Failed : EStateTreeRunStatus::Succeeded;		
+	}
+	return EStateTreeRunStatus::Running;
+}
+
+#if WITH_GAMEPLAY_DEBUGGER
+void UStateTreeTask_Test::AppendDebugInfoString(FString& DebugString, const FStateTreeInstance& StateTreeInstance) const
+{
+	Super::AppendDebugInfoString(DebugString, StateTreeInstance);
+	DebugString += FString::Printf(TEXT("Time= %f/%f\n"), Time, Duration);
+}
+#endif
+
+#if WITH_EDITOR
+bool UStateTreeTask_Test::ResolveVariables(const FStateTreeVariableLayout& Variables, FStateTreeConstantStorage& Constants, UObject* Outer)
+{
+	return true;
+}
+#endif
+
+
+
+FStateTreeTask2_Test::FStateTreeTask2_Test()
+{
+}
+
+FStateTreeTask2_Test::~FStateTreeTask2_Test()
+{
+}
+
+EStateTreeRunStatus FStateTreeTask2_Test::Tick(FStateTreeExecutionContext& Context, const float DeltaTime)
+{
+	// Keeping this here until we get Gameplay Debugger to work again.
+//	GEngine->AddOnScreenDebugMessage(INDEX_NONE, 1.0f, FColor::Blue, *FString::Printf(TEXT("[%s] Value=%f IntVal=%d bBoolVal=%s\n"), *Name.ToString(), Value, IntVal, *LexToString(bBoolVal)));
+
+	return EStateTreeRunStatus::Running;
+}
+
+#if WITH_GAMEPLAY_DEBUGGER
+void FStateTreeTask2_Test::AppendDebugInfoString(FString& DebugString, const FStateTreeExecutionContext& Context) const
+{
+	DebugString += FString::Printf(TEXT("[%s] Value=%f IntVal=%d bBoolVal=%s\n"), *Name.ToString(), Value, IntVal, *LexToString(bBoolVal));
+}
+#endif
+
diff --git a/Engine/Plugins/Runtime/StateTree/Source/StateTreeModule/Public/Conditions/StateTreeCondition_Common.h b/Engine/Plugins/Runtime/StateTree/Source/StateTreeModule/Public/Conditions/StateTreeCondition_Common.h
new file mode 100644
--- /dev/null
+++ b/Engine/Plugins/Runtime/StateTree/Source/StateTreeModule/Public/Conditions/StateTreeCondition_Common.h
@@ -0,0 +1,176 @@
+// Copyright Epic Games, Inc. All Rights Reserved.
+
+#pragma once
+
+#include "CoreMinimal.h"
+#include "AITypes.h"
+#include "StateTreeConditionBase.h"
+#include "StateTreeAnyEnum.h"
+#include "StateTreeCondition_Common.generated.h"
+
+/**
+ * Condition comparing two integers.
+ */
+USTRUCT(DisplayName="Integer Compare")
+struct STATETREEMODULE_API FStateTreeCondition_CompareInt : public FStateTreeConditionBase
+{
+	GENERATED_BODY()
+
+	FStateTreeCondition_CompareInt() : FStateTreeConditionBase() { }
+	FStateTreeCondition_CompareInt(const int32 InLeft, const EGenericAICheck InOperator, const int32 InRight)
+		: FStateTreeConditionBase()
+		, Left(InLeft)
+		, Operator(InOperator)
+		, Right(InRight)
+	{}
+
+#if WITH_EDITOR
+	virtual FText GetDescription(const IStateTreeBindingLookup& BindingLookup) const override;
+#endif
+	virtual bool TestCondition() const override;
+
+	UPROPERTY(EditAnywhere, Category = Condition, meta = (Bindable))
+	int32 Left = 0;
+
+	UPROPERTY(EditAnywhere, Category = Condition)
+	EGenericAICheck Operator = EGenericAICheck::Equal;
+
+	UPROPERTY(EditAnywhere, Category = Condition, meta = (Bindable))
+	int32 Right = 0;
+
+	UPROPERTY(EditAnywhere, Category = Condition)
+	bool bInvert = false;
+};
+
+/**
+ * Condition comparing two floats.
+ */
+USTRUCT(DisplayName = "Float Compare")
+struct STATETREEMODULE_API FStateTreeCondition_CompareFloat : public FStateTreeConditionBase
+{
+	GENERATED_BODY()
+
+	FStateTreeCondition_CompareFloat() : FStateTreeConditionBase() {}
+	FStateTreeCondition_CompareFloat(const float InLeft, const EGenericAICheck InOperator, const float InRight)
+        : FStateTreeConditionBase()
+        , Left(InLeft)
+        , Operator(InOperator)
+        , Right(InRight)
+	{}
+
+#if WITH_EDITOR
+	virtual FText GetDescription(const IStateTreeBindingLookup& BindingLookup) const override;
+#endif
+	virtual bool TestCondition() const override;
+
+	UPROPERTY(EditAnywhere, Category = Condition, meta = (Bindable))
+	float Left = 0.0f;
+
+	UPROPERTY(EditAnywhere, Category = Condition)
+	EGenericAICheck Operator = EGenericAICheck::Equal;
+
+	UPROPERTY(EditAnywhere, Category = Condition, meta = (Bindable))
+	float Right = 0.0f;
+
+	UPROPERTY(EditAnywhere, Category = Condition)
+	bool bInvert = false;
+};
+
+/**
+ * Condition comparing two booleans.
+ */
+USTRUCT(DisplayName = "Bool Compare")
+struct STATETREEMODULE_API FStateTreeCondition_CompareBool : public FStateTreeConditionBase
+{
+	GENERATED_BODY()
+
+	FStateTreeCondition_CompareBool() : FStateTreeConditionBase() {}
+	FStateTreeCondition_CompareBool(const bool bInLeft, const bool bInRight)
+        : FStateTreeConditionBase()
+        , bLeft(bInLeft)
+        , bRight(bInRight)
+	{}
+
+#if WITH_EDITOR
+	virtual FText GetDescription(const IStateTreeBindingLookup& BindingLookup) const override;
+#endif
+	virtual bool TestCondition() const override;
+
+	UPROPERTY(EditAnywhere, Category = Condition, meta = (Bindable))
+	bool bLeft = false;
+
+	UPROPERTY(EditAnywhere, Category = Condition, meta = (Bindable))
+	bool bRight = false;
+
+	UPROPERTY(EditAnywhere, Category = Condition)
+	bool bInvert = false;
+};
+
+/**
+ * Condition comparing two enums.
+ */
+USTRUCT(DisplayName = "Enum Compare")
+struct STATETREEMODULE_API FStateTreeCondition_CompareEnum : public FStateTreeConditionBase
+{
+	GENERATED_BODY()
+
+	FStateTreeCondition_CompareEnum() : FStateTreeConditionBase() {}
+	FStateTreeCondition_CompareEnum(const FStateTreeAnyEnum& InLeft, const FStateTreeAnyEnum& InRight)
+        : FStateTreeConditionBase()
+        , Left(InLeft)
+        , Right(InRight)
+	{}
+
+#if WITH_EDITOR
+	virtual FText GetDescription(const IStateTreeBindingLookup& BindingLookup) const override;
+	virtual void OnBindingChanged(const FStateTreeEditorPropertyPath& SourcePath, const FStateTreeEditorPropertyPath& TargetPath, const IStateTreeBindingLookup& BindingLookup) override;
+#endif
+	virtual bool TestCondition() const override;
+
+	UPROPERTY(EditAnywhere, Category = Condition, meta = (Bindable))
+	FStateTreeAnyEnum Left;
+
+	UPROPERTY(EditAnywhere, Category = Condition, meta = (Bindable))
+	FStateTreeAnyEnum Right;
+
+	UPROPERTY(EditAnywhere, Category = Condition)
+	bool bInvert = false;
+};
+
+/**
+ * Condition comparing distance between two vectors.
+ */
+USTRUCT(DisplayName = "Distance Compare")
+struct STATETREEMODULE_API FStateTreeCondition_CompareDistance : public FStateTreeConditionBase
+{
+	GENERATED_BODY()
+
+	FStateTreeCondition_CompareDistance() : FStateTreeConditionBase() { }
+	FStateTreeCondition_CompareDistance(const FVector& InSource, const FVector& InTarget, const EGenericAICheck InOperator, const float InDistance)
+        : FStateTreeConditionBase()
+        , Source(InSource)
+        , Target(InTarget)
+        , Operator(InOperator)
+        , Distance(InDistance)
+	{}
+
+#if WITH_EDITOR
+	virtual FText GetDescription(const IStateTreeBindingLookup& BindingLookup) const override;
+#endif
+	virtual bool TestCondition() const override;
+
+	UPROPERTY(EditAnywhere, Category = Condition, meta = (Bindable))
+	FVector Source = FVector(EForceInit::ForceInitToZero);
+
+	UPROPERTY(EditAnywhere, Category = Condition, meta = (Bindable))
+	FVector Target = FVector(EForceInit::ForceInitToZero);
+
+	UPROPERTY(EditAnywhere, Category = Condition)
+	EGenericAICheck Operator = EGenericAICheck::Equal;
+
+	UPROPERTY(EditAnywhere, Category = Condition, meta = (Bindable))
+	float Distance = 0.0f;
+
+	UPROPERTY(EditAnywhere, Category = Condition)
+	bool bInvert = false;
+};
diff --git a/Engine/Plugins/Runtime/StateTree/Source/StateTreeModule/Public/Evaluators/StateTreeEvaluator_Test.h b/Engine/Plugins/Runtime/StateTree/Source/StateTreeModule/Public/Evaluators/StateTreeEvaluator_Test.h
new file mode 100644
--- /dev/null
+++ b/Engine/Plugins/Runtime/StateTree/Source/StateTreeModule/Public/Evaluators/StateTreeEvaluator_Test.h
@@ -0,0 +1,142 @@
+// Copyright Epic Games, Inc. All Rights Reserved.
+
+#pragma once
+
+#include "CoreMinimal.h"
+#include "Misc/Guid.h"
+#include "StateTreeTypes.h"
+#include "StateTreeVariable.h"
+#include "StateTreeEvaluatorBase.h"
+#include "StateTreeEvaluator_Test.generated.h"
+
+struct FStateTreeVariableLayout;
+struct FStateTreeConstantStorage;
+class UStateTreeEvaluatorInstanceBase;
+
+/**
+ * Test evaluator
+ */
+
+UCLASS(BlueprintType, EditInlineNew, CollapseCategories)
+class STATETREEMODULE_API UStateTreeEvaluator_Test : public UStateTreeEvaluatorBase
+{
+	GENERATED_BODY()
+
+public:
+	UStateTreeEvaluator_Test(const FObjectInitializer& ObjectInitializer = FObjectInitializer::Get());
+
+	// Called when instantiated first time.
+	virtual bool Initialize(FStateTreeInstance& StateTreeInstance) override;
+	// Called when evaluator becomes active/ticking.
+	virtual void Activate(FStateTreeInstance& StateTreeInstance) override;
+	// Called when evaluator becomes inactive.
+	virtual void Deactivate(FStateTreeInstance& StateTreeInstance) override;
+
+	// Called on each tick.
+	virtual void Tick(FStateTreeInstance& StateTreeInstance, const float DeltaTime) override;
+
+#if WITH_GAMEPLAY_DEBUGGER
+	virtual void AppendDebugInfoString(FString& DebugString, const FStateTreeInstance& StateTreeInstance) const override;
+#endif // WITH_GAMEPLAY_DEBUGGER
+
+#if WITH_EDITOR
+	virtual void DefineOutputVariables(FStateTreeVariableLayout& Variables) const override;
+	virtual bool ResolveVariables(const FStateTreeVariableLayout& Variables, FStateTreeConstantStorage& Constants, UObject* Outer) override;
+#endif
+
+protected:
+	// State
+	float Timer;
+
+	// UI properties
+	UPROPERTY(EditDefaultsOnly, Category = Test)
+	float PlainNumber;
+
+	UPROPERTY(EditDefaultsOnly, Category = Test)
+	FStateTreeVariable SomeFloat;
+
+	// Output variables
+	UPROPERTY()
+	FStateTreeVariable BoolVar;
+	UPROPERTY()
+	FStateTreeVariable FloatVar;
+	UPROPERTY()
+	FStateTreeVariable IntVar;
+	UPROPERTY()
+	FStateTreeVariable VectorVar;
+	UPROPERTY()
+	FStateTreeVariable Object1Var;
+	UPROPERTY()
+	FStateTreeVariable Object2Var;
+	UPROPERTY()
+	FStateTreeVariable ObjectRVar;
+};
+
+
+// STATETREE_V2
+
+UENUM()
+namespace EStateTreeOldEnumTest
+{
+	enum Type
+	{
+		Foo,
+		Bar
+	};
+}
+
+USTRUCT()
+struct FDummyStateTreeResult : public FStateTreeResult
+{
+	GENERATED_BODY()
+
+	virtual ~FDummyStateTreeResult() {}
+	virtual const UScriptStruct& GetStruct() const override { return *StaticStruct(); }
+
+	UPROPERTY(EditAnywhere, Category = Value)
+	FName Name;
+
+	UPROPERTY(EditAnywhere, Category = Value)
+	int32 Count = 0;
+};
+
+USTRUCT(meta=(DisplayName="Evaluator Test2"))
+struct STATETREEMODULE_API FStateTreeEvaluator2_Test : public FStateTreeEvaluator2Base
+{
+	GENERATED_BODY()
+
+	FStateTreeEvaluator2_Test();
+	virtual ~FStateTreeEvaluator2_Test();
+	virtual void EnterState(FStateTreeExecutionContext& Context, const EStateTreeStateChangeType ChangeType, const FStateTreeTransitionResult& Transition) override;
+	virtual void ExitState(FStateTreeExecutionContext& Context, const EStateTreeStateChangeType ChangeType, const FStateTreeTransitionResult& Transition) override;
+	virtual void Evaluate(FStateTreeExecutionContext& Context, const EStateTreeEvaluationType EvalType, const float DeltaTime) override;
+
+protected:
+	UPROPERTY(EditAnywhere, Category = Test, meta = (Bindable))
+	float Foo = 0.0f;
+
+	UPROPERTY(EditAnywhere, Category = Test)
+	int32 IntVal = 0;
+
+	UPROPERTY(EditAnywhere, Category = Test)
+	bool bBoolVal = false;
+
+	UPROPERTY(EditAnywhere, Category = Test)
+	FVector VecVal = FVector::ZeroVector;
+
+	UPROPERTY(EditAnywhere, Category = Test, Meta = (Bindable))
+	EStateTreeVariableType EnumVal = EStateTreeVariableType::Void;
+
+	UPROPERTY(EditAnywhere, Category = Test, Meta = (Bindable))
+	TEnumAsByte<enum EStateTreeOldEnumTest::Type> OldEnumVal = EStateTreeOldEnumTest::Foo;
+
+	UPROPERTY(EditAnywhere, Category = Test)
+	FStateTreeResultRef OutResult;
+
+	UPROPERTY(EditAnywhere, Category = Test, Meta = (Bindable))
+	FStateTreeResultRef InResult;
+	
+	FDummyStateTreeResult DummyResult;
+};
+
+// ~STATETREE_V2
diff --git a/Engine/Plugins/Runtime/StateTree/Source/StateTreeModule/Public/StateTree.h b/Engine/Plugins/Runtime/StateTree/Source/StateTreeModule/Public/StateTree.h
new file mode 100644
--- /dev/null
+++ b/Engine/Plugins/Runtime/StateTree/Source/StateTreeModule/Public/StateTree.h
@@ -0,0 +1,159 @@
+// Copyright Epic Games, Inc. All Rights Reserved.
+
+#pragma once
+
+#include "CoreMinimal.h"
+#include "Engine/DataAsset.h"
+#include "StateTreeTypes.h"
+#include "StateTreeSchema.h"
+#include "StateTreeCondition.h"
+#include "StateTreeVariableLayout.h"
+#include "StateTreeParameterLayout.h"
+#include "StateTreeConstantStorage.h"
+#include "InstancedStruct.h"
+#include "StateTreePropertyBindings.h"
+#include "StateTree.generated.h"
+
+class UStateTreeEvaluatorBase;
+class UStateTreeTaskBase;
+
+/**
+ * StateTree asset. Contains the StateTree definition in both editor and runtime (baked) formats.
+ */
+UCLASS(BlueprintType)
+class STATETREEMODULE_API UStateTree : public UDataAsset
+{
+	GENERATED_BODY()
+
+public:
+
+	UStateTree(const FObjectInitializer& ObjectInitializer = FObjectInitializer::Get());
+
+	const FStateTreeVariableLayout& GetInputParameterLayout() const { return InputParameter; }
+
+	/** @return Script Struct that can be used to instantiate the runtime storage */
+	const UScriptStruct* GetRuntimeStorageStruct() const { return RuntimeStorageStruct; }
+
+	/** @return Instance of the runtime storage that contains the default values */
+	const FInstancedStruct& GetRuntimeStorageDefaultValue() const { return RuntimeStorageDefaultValue; }
+
+	/** @return List of external items required by the state tree */
+	TConstArrayView<FStateTreeExternalItemDesc> GetExternalItems() const { return ExternalItems; }
+
+	/** @return Schema describing which inputs, evaluators, and tasks a StateTree can contain */
+	const UStateTreeSchema* GetSchema() const { return Schema; }
+	void SetSchema(UStateTreeSchema* InSchema) { Schema = InSchema; }
+
+	/** @return true is the tree asset is considered valid (e.g. at least one state) */
+	bool IsValidStateTree() const;
+
+	// STATETREE_V2
+	void ResolvePropertyPaths();
+	bool IsV2() const { return Schema ? Schema->IsV2() : false; }
+	// ~STATETREE_V2
+
+#if WITH_EDITOR
+	/** Resets the baked data to empty. */
+	void ResetBaked();
+#endif
+
+#if WITH_EDITORONLY_DATA
+	// Edit time data for the StateTree, instance of UStateTreeEditorData
+	UPROPERTY()
+	UObject* EditorData;
+#endif
+
+protected:
+	virtual void BeginDestroy() override;
+	virtual void PostLoad() override;
+	
+#if WITH_EDITOR
+	virtual void PostEditChangeProperty(struct FPropertyChangedEvent& PropertyChangedEvent) override;
+	virtual void GetAssetRegistryTags(TArray<FAssetRegistryTag>& OutTags) const override;
+#endif
+
+	/** Initializes the types and default values related to StateTree runtime storage */
+	void InitRuntimeStorage();
+
+	/** @return Number of items (Evaluators & Tasks) in the runtime storage. */
+	int32 GetRuntimeStorageItemCount() const { return RuntimeStorageOffsets.Num(); }
+
+private:
+
+	// Properties
+
+	UPROPERTY(EditDefaultsOnly, Category = Common, Instanced)
+	UStateTreeSchema* Schema = nullptr;
+
+	// TODO: find a nice way to combine this with Parameters. The data is almost the same, but the semantics different:
+	// InputParameter defines layout with offsets, while Parameters defines how they connect to Variables.
+	// This could be editor only, maybe.
+	UPROPERTY(EditDefaultsOnly, Category = InputParameters)
+	FStateTreeVariableLayout InputParameter;
+
+	// Baked data
+	UPROPERTY()
+	FStateTreeVariableLayout Variables;
+
+	UPROPERTY()
+	FStateTreeConstantStorage Constants;
+
+	UPROPERTY()
+	TArray<UStateTreeTaskBase*> Tasks;
+
+	UPROPERTY()
+	TArray<UStateTreeEvaluatorBase*> Evaluators;
+
+// STATETREE_V2
+	/** Evaluators and Tasks that require runtime state */
+	UPROPERTY()
+	TArray<FInstancedStruct> RuntimeStorageItems;
+
+	/** Script Struct that can be used to instantiate the runtime storage */
+	UPROPERTY()
+	UScriptStruct* RuntimeStorageStruct;
+
+	/** Offsets into the runtime type to quickly get a struct view to a specific Task or Evaluator */
+	TArray<FStateTreeRuntimeStorageItemOffset> RuntimeStorageOffsets;
+
+	/** Instance of the runtime storage that contains the default values. */
+	UPROPERTY()
+	FInstancedStruct RuntimeStorageDefaultValue;
+
+	UPROPERTY(meta = (BaseStruct = "StateTreeConditionBase", ExcludeBaseStruct))
+	TArray<FInstancedStruct> Conditions2;
+
+	/** List of external items required by the state tree */
+	UPROPERTY()
+	TArray<FStateTreeExternalItemDesc> ExternalItems;
+
+	UPROPERTY()
+	FStateTreePropertyBindings PropertyBindings;
+	
+// ~STATETREE_V2
+
+	UPROPERTY()
+	TArray<FBakedStateTreeState> States;
+
+	UPROPERTY()
+	TArray<FStateTreeCondition> Conditions;
+
+	UPROPERTY()
+	TArray<FBakedStateTransition> Transitions;
+
+	UPROPERTY()
+	FStateTreeParameterLayout Parameters;
+
+	friend struct FStateTreeInstance;
+	friend struct FStateTreeExecutionContext;
+#if WITH_EDITORONLY_DATA
+	friend struct FStateTreeBaker;
+#endif
+};
+
+
+
+
+
+
+
diff --git a/Engine/Plugins/Runtime/StateTree/Source/StateTreeModule/Public/StateTreeCondition.h b/Engine/Plugins/Runtime/StateTree/Source/StateTreeModule/Public/StateTreeCondition.h
new file mode 100644
--- /dev/null
+++ b/Engine/Plugins/Runtime/StateTree/Source/StateTreeModule/Public/StateTreeCondition.h
@@ -0,0 +1,34 @@
+// Copyright Epic Games, Inc. All Rights Reserved.
+
+#pragma once
+
+#include "CoreMinimal.h"
+#include "AITypes.h"
+#include "StateTreeVariable.h"
+#include "StateTreeCondition.generated.h"
+
+/**
+ * Condition contains left and right variable references and comparison operator.
+ * The variable chosen left will dictate which type is available at right.
+ * This is done in the editor code, see StateTreeConditionDetails.
+*/
+USTRUCT()
+struct STATETREEMODULE_API FStateTreeCondition
+{
+	GENERATED_BODY()
+public:
+	FStateTreeCondition();
+
+#if WITH_EDITOR
+	FText GetDescription() const;
+#endif
+
+	UPROPERTY(EditDefaultsOnly, Category = Condition)
+	FStateTreeVariable Left;
+
+	UPROPERTY(EditDefaultsOnly, Category = Condition)
+	FStateTreeVariable Right;
+
+	UPROPERTY(EditDefaultsOnly, Category = Condition)
+	EGenericAICheck Operator;
+};
diff --git a/Engine/Plugins/Runtime/StateTree/Source/StateTreeModule/Public/StateTreeConditionBase.h b/Engine/Plugins/Runtime/StateTree/Source/StateTreeModule/Public/StateTreeConditionBase.h
new file mode 100644
--- /dev/null
+++ b/Engine/Plugins/Runtime/StateTree/Source/StateTreeModule/Public/StateTreeConditionBase.h
@@ -0,0 +1,50 @@
+// Copyright Epic Games, Inc. All Rights Reserved.
+
+#pragma once
+
+#include "CoreMinimal.h"
+#include "EngineDefines.h"
+#include "StateTreeTypes.h"
+#include "StateTreeConditionBase.generated.h"
+
+struct IStateTreeBindingLookup;
+struct FStateTreeEditorPropertyPath;
+
+/**
+ * Base struct for all conditions.
+ */
+USTRUCT()
+struct STATETREEMODULE_API FStateTreeConditionBase
+{
+	GENERATED_BODY()
+
+	FStateTreeConditionBase()
+	{
+	}
+
+	virtual ~FStateTreeConditionBase() {}
+
+#if WITH_EDITOR
+	/** @return Rich text description of the condition. */
+	virtual FText GetDescription(const IStateTreeBindingLookup& BindingLookup) const { return FText::GetEmpty(); }
+	/**
+	 * Called when binding of any of the properties in the condition changes.
+	 * @param SourcePath Source path of the new binding.
+	 * @param TargetPath Target path of the new binding (the property in the condition).
+	 * @param BindingLookup Reference to binding lookup which can be used to reason about property paths.
+	 */
+	virtual void OnBindingChanged(const FStateTreeEditorPropertyPath& SourcePath, const FStateTreeEditorPropertyPath& TargetPath, const IStateTreeBindingLookup& BindingLookup) {}
+#endif
+	/** @return True if the condition passes. */
+	virtual bool TestCondition() const { return false; }
+
+#if WITH_EDITORONLY_DATA
+	/** Edit time struct ID for property binding. */
+	UPROPERTY(EditDefaultsOnly, Category = Evaluator, meta = (IgnoreForMemberInitializationTest, EditCondition = "false", EditConditionHides))	// Hack, we want the ID to be found as IPropertyHandle, but do not want to display it.
+	FGuid ID;
+#endif
+
+	/** Property binding copy batch handle. */
+	UPROPERTY()
+	FStateTreeHandle BindingsBatch = FStateTreeHandle::Invalid;
+};
diff --git a/Engine/Plugins/Runtime/StateTree/Source/StateTreeModule/Public/StateTreeConstantStorage.h b/Engine/Plugins/Runtime/StateTree/Source/StateTreeModule/Public/StateTreeConstantStorage.h
new file mode 100644
--- /dev/null
+++ b/Engine/Plugins/Runtime/StateTree/Source/StateTreeModule/Public/StateTreeConstantStorage.h
@@ -0,0 +1,126 @@
+// Copyright Epic Games, Inc. All Rights Reserved.
+
+#pragma once
+
+#include "CoreMinimal.h"
+#include "StateTreeTypes.h"
+#include "StateTreeVariableDesc.h"
+#include "StateTreeConstantStorage.generated.h"
+
+/**
+ * Storage of all the constants used in the StateTree.
+ * StateTree variables are handles can refer to items in StateTree instances memory or constants.
+ */
+USTRUCT()
+struct STATETREEMODULE_API FStateTreeConstantStorage
+{
+	GENERATED_BODY()
+
+	FStateTreeConstantStorage()
+	: ConstantBaseOffset(0)
+	{
+	}
+
+	UPROPERTY()
+	TArray<uint8> Memory;
+
+	UPROPERTY()
+	TArray<EStateTreeVariableType> Types;
+
+	UPROPERTY()
+	uint16 ConstantBaseOffset;
+
+	bool GetConstant(FStateTreeHandle Handle, EStateTreeVariableType Type, uint8* Value) const
+	{
+		const uint16 Offset = Handle.Index - ConstantBaseOffset;
+		if (Offset < (uint16)Memory.Num())
+		{
+			FMemory::Memcpy(Value, &Memory[Offset], FStateTreeVariableHelpers::GetVariableMemoryUsage(Type));
+			return true;
+		}
+		return false;
+	}
+
+	float GetConstantFloat(FStateTreeHandle Handle, const float DefaultValue = 0.0f) const
+	{
+		return GetConstant<float>(Handle, DefaultValue);
+	}
+
+	int GetConstantInt(FStateTreeHandle Handle, const int DefaultValue = 0) const
+	{
+		return GetConstant<int>(Handle, DefaultValue);
+	}
+
+	bool GetConstantBool(FStateTreeHandle Handle, const bool DefaultValue = false) const
+	{
+		return GetConstant<bool>(Handle, DefaultValue);
+	}
+
+	const FVector& GetConstantVector(FStateTreeHandle Handle, const FVector& DefaultValue = FVector::ZeroVector) const
+	{
+		return GetConstant<FVector>(Handle, DefaultValue);
+	}
+
+	template <typename T>
+	const T& GetConstant(FStateTreeHandle Handle, const T& DefaultValue) const
+	{
+		const uint16 Offset = Handle.Index - ConstantBaseOffset;
+		if (Offset < (uint16)Memory.Num())
+		{
+			return FStateTreeVariableHelpers::GetValueFromMemory<T>(&Memory[Offset]);
+		}
+		return DefaultValue;
+	}
+
+
+#if WITH_EDITOR
+	void Reset()
+	{
+		Memory.Reset();
+		Types.Reset();
+		ConstantBaseOffset = 0;
+	}
+
+	FStateTreeHandle AddConstantFloat(const float Value)
+	{
+		return AddConstant<float, EStateTreeVariableType::Float>(Value);
+	}
+
+	FStateTreeHandle AddConstantInt(const int32 Value)
+	{
+		return AddConstant<int32, EStateTreeVariableType::Int>(Value);
+	}
+
+	FStateTreeHandle AddConstantBool(const bool Value)
+	{
+		return AddConstant<bool, EStateTreeVariableType::Bool>(Value);
+	}
+
+	FStateTreeHandle AddConstantVector(const FVector& Value)
+	{
+		return AddConstant<FVector, EStateTreeVariableType::Vector>(Value);
+	}
+
+	template <typename T, EStateTreeVariableType TType>
+	FStateTreeHandle AddConstant(const T& Value)
+	{
+		uint16 Offset = 0;
+		for (EStateTreeVariableType Type : Types)
+		{
+			if (Type == TType && FStateTreeVariableHelpers::GetValueFromMemory<T>(&Memory[Offset]) == Value)
+			{
+				// Found same constant, return it.
+				return FStateTreeHandle(ConstantBaseOffset + Offset);
+			}
+			Offset += FStateTreeVariableHelpers::GetVariableMemoryUsage(Type);
+		}
+		// Add new constant
+		Offset = (uint16)Memory.Num();
+		Memory.AddZeroed(FStateTreeVariableHelpers::GetVariableMemoryUsage(TType));
+		FStateTreeVariableHelpers::SetValueInMemory<T>(&Memory[Offset], Value);
+
+		return FStateTreeHandle(ConstantBaseOffset + Offset);
+	}
+#endif
+};
+
diff --git a/Engine/Plugins/Runtime/StateTree/Source/StateTreeModule/Public/StateTreeDelegates.h b/Engine/Plugins/Runtime/StateTree/Source/StateTreeModule/Public/StateTreeDelegates.h
new file mode 100644
--- /dev/null
+++ b/Engine/Plugins/Runtime/StateTree/Source/StateTreeModule/Public/StateTreeDelegates.h
@@ -0,0 +1,31 @@
+// Copyright Epic Games, Inc. All Rights Reserved.
+
+#pragma once
+
+#include "CoreMinimal.h"
+
+class UStateTree;
+
+namespace UE::StateTree::Delegates
+{
+	
+#if WITH_EDITOR
+	// Called when linkable name in a StateTree has changed.
+	DECLARE_MULTICAST_DELEGATE_OneParam(FOnIdentifierChanged, const UStateTree& /*StateTree*/);
+	extern STATETREEMODULE_API FOnIdentifierChanged OnIdentifierChanged;
+
+	// Called when parameters are invalidated, i.e. linked asset's parameter layout changed.
+	DECLARE_MULTICAST_DELEGATE_OneParam(FOnParametersInvalidated, const UStateTree& /*StateTree*/);
+	extern STATETREEMODULE_API FOnParametersInvalidated OnParametersInvalidated;
+
+	// Called when parameter layout for a StateTree has changed.
+	DECLARE_MULTICAST_DELEGATE_OneParam(FOnParameterLayoutChanged, const UStateTree& /*StateTree*/);
+	extern STATETREEMODULE_API FOnParameterLayoutChanged OnParameterLayoutChanged;
+	
+	// Called when schema of the StateTree has changed.
+	DECLARE_MULTICAST_DELEGATE_OneParam(FOnSchemaChanged, const UStateTree& /*StateTree*/);
+	extern STATETREEMODULE_API FOnSchemaChanged OnSchemaChanged;
+#endif
+
+}; // UE::StateTree::Delegates
+
diff --git a/Engine/Plugins/Runtime/StateTree/Source/StateTreeModule/Public/StateTreeEvaluatorBase.h b/Engine/Plugins/Runtime/StateTree/Source/StateTreeModule/Public/StateTreeEvaluatorBase.h
new file mode 100644
--- /dev/null
+++ b/Engine/Plugins/Runtime/StateTree/Source/StateTreeModule/Public/StateTreeEvaluatorBase.h
@@ -0,0 +1,144 @@
+// Copyright Epic Games, Inc. All Rights Reserved.
+
+#pragma once
+
+#include "CoreMinimal.h"
+#include "Misc/Guid.h"
+#include "StateTreeTypes.h"
+#include "StateTreePropertyBindings.h"
+#include "StateTreeEvaluatorBase.generated.h"
+
+struct FStateTreeInstance;
+struct FStateTreeExecutionContext;
+struct FStateTreeVariableLayout;
+struct FStateTreeConstantStorage;
+
+/**
+ * Base class of StateTree Evaluators.
+ * Evaluators calculate and expose data to be used for decision making in a StateTree.
+ */
+UCLASS(Abstract)
+class STATETREEMODULE_API UStateTreeEvaluatorBase : public UObject
+{
+	GENERATED_BODY()
+
+public:
+	UStateTreeEvaluatorBase(const FObjectInitializer& ObjectInitializer);
+
+	// Called when instantiated first time.
+	virtual bool Initialize(FStateTreeInstance& StateTreeInstance) PURE_VIRTUAL(UStateTreeEvaluatorBase::Initialize, return false; );
+	// Called when evaluator becomes active/ticking.
+	virtual void Activate(FStateTreeInstance & StateTreeInstance) PURE_VIRTUAL(UStateTreeEvaluatorBase::Activate, return; );
+	// Called when evaluator becomes inactive.
+	virtual void Deactivate(FStateTreeInstance & StateTreeInstance) PURE_VIRTUAL(UStateTreeEvaluatorBase::Deactivate, return; );
+
+	// Called on each tick.
+	virtual void Tick(FStateTreeInstance & StateTreeInstance, const float DeltaTime) PURE_VIRTUAL(UStateTreeEvaluatorBase::Tick, return; );
+
+	// Returns ID of the Evaluator
+	const FName& GetName() const { return Name; }
+
+#if WITH_GAMEPLAY_DEBUGGER
+	virtual void AppendDebugInfoString(FString & DebugString, const FStateTreeInstance & StateTreeInstance) const;
+#endif // WITH_GAMEPLAY_DEBUGGER
+
+#if WITH_EDITOR
+	// Defines output variables in variable map.
+	virtual void DefineOutputVariables(FStateTreeVariableLayout& Variables) const {};
+	// Create Evaluator instance template, set default variables, and resolve variable handles.
+	virtual bool ResolveVariables(const FStateTreeVariableLayout& Variables, FStateTreeConstantStorage& Constants, UObject* Outer) PURE_VIRTUAL(UStateTreeEvaluatorBase::Resolve, return false; );
+	// Returns ture if changed.
+	virtual bool ValidateParameterLayout() { return false; }
+
+	const FGuid GetID() const { return ID; }
+
+	void SetNewUniqueID() { ID = FGuid::NewGuid(); }
+	void SetName(const FName& InName) { Name = InName; }
+
+	virtual void PostEditChangeProperty(struct FPropertyChangedEvent& PropertyChangedEvent) override;
+#endif
+
+protected:
+	UPROPERTY(EditDefaultsOnly, Category = Evaluator)
+	FName Name;
+
+#if WITH_EDITORONLY_DATA
+	UPROPERTY(meta = (IgnoreForMemberInitializationTest))
+	FGuid ID;
+#endif
+};
+
+
+// STATETREE_V2
+
+/**
+ * Base struct of StateTree Evaluators.
+ * Evaluators calculate and expose data to be used for decision making in a StateTree.
+ */
+USTRUCT()
+struct STATETREEMODULE_API FStateTreeEvaluator2Base // TODO: change to FStateTreeEvaluatorBase once UStateTreeEvaluatorBase is removed
+{
+	GENERATED_BODY()
+
+	FStateTreeEvaluator2Base() = default;
+
+	virtual ~FStateTreeEvaluator2Base() {}
+
+	/**
+	 * Called when a new state is entered and evaluator is part of active states. The change type parameter describes if the evaluator's state
+	 * was previously part of the list of active states (Sustained), or if it just became active (Changed).
+	 * @param Context Reference to current execution context.
+	 * @param ChangeType Describes the change type (Changed/Sustained).
+	 * @param Transition Describes the states involved in the transition
+	 */
+	virtual void EnterState(FStateTreeExecutionContext& Context, const EStateTreeStateChangeType ChangeType, const FStateTreeTransitionResult& Transition) {}
+
+	/**
+	 * Called when a current state is exited and evaluator is part of active states. The change type parameter describes if the evaluator's state
+	 * will be active after the transition (Sustained), or if it will became inactive (Changed).
+	 * @param Context Reference to current execution context.
+	 * @param ChangeType Describes the change type (Changed/Sustained).
+	 * @param Transition Describes the states involved in the transition
+	 */
+	virtual void ExitState(FStateTreeExecutionContext& Context, const EStateTreeStateChangeType ChangeType, const FStateTreeTransitionResult& Transition) {}
+
+	/**
+	 * Called Right after a state has been completed. StateCompleted is called in reverse order to allow to propagate state to Evaluators and Tasks that
+	 * are executed earlier in the tree. Note that StateCompleted is not called if conditional transition changes the state.
+	 * @param Context Reference to current execution context.
+	 * @param CompletionStatus Describes the running status of the completed state (Succeeded/Failed).
+	 * @param CompletedState Handle of the state that was completed.
+	 */
+	virtual void StateCompleted(FStateTreeExecutionContext& Context, const EStateTreeRunStatus CompletionStatus, const FStateTreeHandle CompletedState) {}
+	
+	/**
+	 * Called when evaluator needs to be updated. EvalType describes if the tick happens during state tree tick when the evaluator is on active state (Tick),
+	 * or during state selection process when the evaluator's state is visited while it's inactive (PreSelection).
+	 * That is, type "Tick" means that the call happens between EnterState()/ExitState() pair, "PreSelection" is used otherwise.
+	 * @param Context Reference to current execution context.
+	 * @param EvalType Describes tick type.
+	 * @param DeltaTime Time since last StateTree tick, or 0 if called during preselection.
+	 */
+	virtual void Evaluate(FStateTreeExecutionContext& Context, const EStateTreeEvaluationType EvalType, const float DeltaTime) {}
+
+#if WITH_GAMEPLAY_DEBUGGER
+	virtual void AppendDebugInfoString(FString& DebugString, const FStateTreeExecutionContext& Context) const;
+#endif // WITH_GAMEPLAY_DEBUGGER
+
+	UPROPERTY(EditDefaultsOnly, Category = Evaluator, meta=(EditCondition = "false", EditConditionHides))
+	FName Name;
+
+#if WITH_EDITORONLY_DATA
+	UPROPERTY(EditDefaultsOnly, Category = Evaluator, meta=(IgnoreForMemberInitializationTest, EditCondition="false", EditConditionHides))	// Hack, we want the ID to be found as IPropertyHandle, but do not want to display it.
+	FGuid ID;
+#endif
+
+	UPROPERTY()
+	FStateTreeHandle BindingsBatch = FStateTreeHandle::Invalid;		// Property binding copy batch handle.
+
+	UPROPERTY()
+	uint16 SourceStructIndex = 0;									// Property binding Source Struct index of the evaluator.
+};
+template<> struct TStructOpsTypeTraits<FStateTreeEvaluator2Base> : public TStructOpsTypeTraitsBase2<FStateTreeEvaluator2Base> { enum { WithPureVirtual = true, }; };
+
+// ~STATETREE_V2
diff --git a/Engine/Plugins/Runtime/StateTree/Source/StateTreeModule/Public/StateTreeExecutionContext.h b/Engine/Plugins/Runtime/StateTree/Source/StateTreeModule/Public/StateTreeExecutionContext.h
new file mode 100644
--- /dev/null
+++ b/Engine/Plugins/Runtime/StateTree/Source/StateTreeModule/Public/StateTreeExecutionContext.h
@@ -0,0 +1,282 @@
+// Copyright Epic Games, Inc. All Rights Reserved.
+
+#pragma once
+
+#include "StateTree.h"
+#include "InstancedStruct.h"
+#include "Containers/StaticArray.h"
+#include "StateTreePropertyBindings.h"
+#include "StateTreeExecutionContext.generated.h"
+
+
+USTRUCT()
+struct STATETREEMODULE_API FStateTreeExecutionState
+{
+	GENERATED_BODY()
+
+	/** Currently active state */
+	FStateTreeHandle CurrentState = FStateTreeHandle::Invalid;
+
+	/** Result of last tick */
+	EStateTreeRunStatus LastTickStatus = EStateTreeRunStatus::Failed;
+
+	/** Running status of the instance */
+	EStateTreeRunStatus TreeRunStatus = EStateTreeRunStatus::Unset;
+
+	/** Delayed transition handle, if exists */
+	int16 GatedTransitionIndex = INDEX_NONE;
+
+	/** Running time of the delayed transition */
+	float GatedTransitionTime = 0.0f;
+};
+
+UENUM()
+enum class EStateTreeStorage : uint8
+{
+	/** Execution context has internal storage */
+	Internal,
+	/** Execution context assumes external storage */
+	External,
+};
+
+/**
+ * Runs StateTrees defined in UStateTree asset.
+ * Uses constant data from StateTree, keeps local storage of variables, and creates instanced Evaluators and Tasks.
+ */
+USTRUCT()
+struct STATETREEMODULE_API FStateTreeExecutionContext
+{
+	GENERATED_BODY()
+
+public:
+	FStateTreeExecutionContext();
+	virtual ~FStateTreeExecutionContext();
+
+	/** Initializes the StateTree instance to be used with specific owner and StateTree asset. */
+	bool Init(UObject& InOwner, const UStateTree& InStateTree, const EStateTreeStorage InStorageType);
+
+	/** Returns the StateTree asset in use. */
+	const UStateTree* GetStateTree() const { return StateTree; }
+
+	UObject* GetOwner() const { return Owner; }
+	UWorld* GetWorld() const { return World; };
+	void SetWorld(UWorld* InWorld) { World = InWorld; };
+
+	/** Start executing. */
+	void Start(FStateTreeItemView ExternalStorage = FStateTreeItemView());
+	/** Stop executing. */
+	void Stop(FStateTreeItemView ExternalStorage = FStateTreeItemView());
+
+	/** Tick the state tree logic. */
+	EStateTreeRunStatus Tick(const float DeltaTime, FStateTreeItemView ExternalStorage = FStateTreeItemView());
+
+	/** @return Pointer to a State or null if state not found */ 
+	const FBakedStateTreeState* GetStateFromHandle(const FStateTreeHandle StateHandle) const
+	{
+		return (StateTree && StateTree->States.IsValidIndex(StateHandle.Index)) ? &StateTree->States[StateHandle.Index] : nullptr;
+	}
+
+	/** @return Array view to external item descriptors associated with this context. Note: Init() must be called before calling this method. */
+	TConstArrayView<FStateTreeExternalItemDesc> GetExternalItems() const
+	{
+		check(StateTree);
+		return StateTree->ExternalItems;
+	}
+
+	/** @return True if all required external item pointers are set. */ 
+	bool AreExternalItemsValid() const
+	{
+		check(StateTree);
+		bool bResult = true;
+		for (const FStateTreeExternalItemDesc& ItemDesc : StateTree->ExternalItems)
+		{
+			const FStateTreeItemView& ItemView = ItemViews[ItemDesc.Handle.GetIndex()];
+			if (!ItemDesc.bOptional && (ItemView.GetMemory() == nullptr || ItemView.GetStruct() != ItemDesc.Struct))
+			{
+				bResult = false;
+				break;
+			}
+		}
+		return bResult;
+	}
+
+	/** @return Handle to external item of type InStruct, or invalid handle if not found */ 
+	FStateTreeExternalItemHandle GetExternalItemHandleByStruct(const UStruct* InStruct) const
+	{
+		check(StateTree);
+		const FStateTreeExternalItemDesc* Item = StateTree->ExternalItems.FindByPredicate([InStruct](const FStateTreeExternalItemDesc& Item) { return Item.Struct == InStruct; });
+		return Item != nullptr ? Item->Handle : FStateTreeExternalItemHandle();
+	}
+
+	/** Sets external item view based on handle */ 
+	void SetExternalItem(const FStateTreeExternalItemHandle ItemHandle, FStateTreeItemView Item)
+	{
+		check(StateTree);
+		check(ItemHandle.IsValid());
+		ItemViews[ItemHandle.GetIndex()] = Item;
+	}
+
+	/** @return item view to external item based on handle */ 
+	FStateTreeItemView GetExternalItem(const FStateTreeExternalItemHandle ItemHandle) const
+	{
+		check(StateTree);
+		check(ItemHandle.IsValid());
+		return ItemViews[ItemHandle.GetIndex()];
+	}
+
+	EStateTreeRunStatus GetLastTickStatus(FStateTreeItemView ExternalStorage = FStateTreeItemView()) const;
+	EStateTreeRunStatus GetEnterStateStatus() const { return EnterStateStatus; }
+
+#if WITH_GAMEPLAY_DEBUGGER
+	/** @return Debug string describing the current state of the execution */
+	FString GetDebugInfoString(FStateTreeItemView ExternalStorage = FStateTreeItemView());
+#endif // WITH_GAMEPLAY_DEBUGGER
+
+#if WITH_STATETREE_DEBUG
+	FString GetActiveStateName(FStateTreeItemView ExternalStorage = FStateTreeItemView()) const;
+	
+	void DebugPrintInternalLayout(FStateTreeItemView ExternalStorage);
+#endif
+	
+protected:
+
+	/** @return Prefix that will be used by STATETREE_LOG and STATETREE_CLOG, empty by default. */
+	virtual FString GetInstanceDescription() const { return TEXT(""); }
+
+	/** Callback when gated transition is triggered. Contexts that are event based can use this to trigger a future event. */
+	virtual void BeginGatedTransition(const FStateTreeExecutionState& Exec) {};
+	
+	/**
+	 * Resets the instance to initial empty state. Note: Does not call ExitState().
+	 */
+	void Reset();
+
+	/**
+	 * Handles logic for entering State. EnterState is called on new active Evaluators and Tasks that are part of the re-planned tree.
+	 * Re-planned tree is from the transition target up to the leaf state. States that are parent to the transition target state
+	 * and still active after the transition will remain intact.
+	 * @return Run status returned by the tasks.
+	 */
+	EStateTreeRunStatus EnterState(FStateTreeItemView Storage, const FStateTreeTransitionResult& Transition);
+
+	/**
+	 * Handles logic for exiting State. ExitState is called on current active Evaluators and Tasks that are part of the re-planned tree.
+	 * Re-planned tree is from the transition target up to the leaf state. States that are parent to the transition target state
+	 * and still active after the transition will remain intact.
+	 */
+	void ExitState(FStateTreeItemView Storage, const FStateTreeTransitionResult& Transition);
+
+	/**
+	 * Handles logic for exiting State. ExitState is called on current active Evaluators and Tasks in reverse order (from leaf to root).
+	 */
+	void StateCompleted(FStateTreeItemView Storage, const FStateTreeHandle CurrentState, const EStateTreeRunStatus CompletionStatus);
+
+	/**
+	 * Ticks evaluators of all active states starting from current state by delta time.
+	 * If TickEvaluators() is called multiple times per frame (i.e. during selection when visiting new states), each state and evaluator is ticked only once.
+	 */
+	void TickEvaluators(FStateTreeItemView Storage, const FStateTreeHandle CurrentState, const EStateTreeEvaluationType EvalType, const float DeltaTime);
+
+	/**
+	 * Ticks tasks of all active states starting from current state by delta time.
+	 * @return Run status returned by the tasks.
+	 */
+	EStateTreeRunStatus TickTasks(FStateTreeItemView Storage, const FStateTreeHandle CurrentState, const float DeltaTime);
+
+	/**
+	 * Checks all conditions at given range
+	 * @return True if all conditions pass.
+	 */
+	bool TestAllConditions(const uint32 ConditionsOffset, const uint32 ConditionsNum);
+
+	/**
+	 * Triggers transitions based on current run status. CurrentStatus is used to select which transitions events are triggered.
+	 * If CurrentStatus is "Running", "Conditional" transitions pass, "Completed/Failed" will trigger "OnCompleted/OnSucceeded/OnFailed" transitions.
+	 * Transition target state can point to a selector state. For that reason the result contains both the target state, as well ass
+	 * the actual next state returned by the selector.
+	 * @return Transition result describing the source state, state transitioned to, and next selected state.
+	 */
+	FStateTreeTransitionResult TriggerTransitions(FStateTreeItemView Storage, const FStateTreeStateStatus CurrentStatus, const int Depth);
+
+	/**
+	 * Runs state selection logic starting at the specified state, walking towards the leaf states.
+	 * If the preconditions of NextState are not met, "Invalid" is returned. 
+	 * If NextState is a selector state, SelectState is called recursively (depth-first) to all child states (where NextState will be one of child states).
+	 * If NextState is a leaf state, the NextState is returned.
+	 * @param Storage View representing all runtime data used by tasks and evaluators
+	 * @param InitialStateStatus Describes the current state and running status (will be passed intact to next selector)
+	 * @param InitialTargetState The state the initial transition target state (will be passed intact to next selector) 
+	 * @param NextState The state which we try to select next.
+	 * @param Depth Depth of recursion.
+	 * @return Transition result describing the source state, transition target state, and next selected state.
+	 */
+	FStateTreeTransitionResult SelectState(FStateTreeItemView Storage, const FStateTreeStateStatus InitialStateStatus, const FStateTreeHandle InitialTargetState, const FStateTreeHandle NextState, const int Depth);
+
+	/** @return State handles from specified state handle back to the root, specified handle included. */
+	int32 GetActiveStates(const FStateTreeHandle StateHandle, TStaticArray<FStateTreeHandle, 32>& OutStateHandles) const;
+
+	/** @return Mutable storage based on storage settings. */
+	FStateTreeItemView SelectMutableStorage(FStateTreeItemView ExternalStorage)
+	{
+		return StorageType == EStateTreeStorage::External ? ExternalStorage : FStateTreeItemView(StorageInstance);
+	}
+
+	/** @return Const storage based on storage settings. */
+	const FStateTreeItemView SelectStorage(FStateTreeItemView ExternalStorage) const
+	{
+		return StorageType == EStateTreeStorage::External ? ExternalStorage : FStateTreeItemView(StorageInstance);
+	}
+
+	/** @return View to an Evaluator or a Task. */
+	FStateTreeItemView GetItem(FStateTreeItemView Storage, const int32 Index)
+	{
+		const FStateTreeRuntimeStorageItemOffset& ItemOffset = StateTree->RuntimeStorageOffsets[Index];
+		return FStateTreeItemView(ItemOffset.Struct, Storage.GetMutableMemory() + ItemOffset.Offset);
+	}
+
+	/** @return StateTree execution state from the runtime storage. */
+	FStateTreeExecutionState& GetExecState(FStateTreeItemView Storage) const
+	{
+		const FStateTreeRuntimeStorageItemOffset& ItemOffset = StateTree->RuntimeStorageOffsets[0];
+		check(ItemOffset.Struct == FStateTreeExecutionState::StaticStruct());
+		return *reinterpret_cast<FStateTreeExecutionState*>(Storage.GetMutableMemory() + ItemOffset.Offset);
+	}
+
+	/** @return String describing state status for logging and debug. */
+	FString GetStateStatusString(const FStateTreeStateStatus StateStatus) const;
+
+	/** @return String describing state name for logging and debug. */
+	FString GetSafeStateName(const FStateTreeHandle State) const;
+
+	/** @return String describing full path of an activate state for logging and debug. */
+	FString DebugGetStatePath(TArrayView<FStateTreeHandle> ActiveStateHandles, int32 ActiveStateIndex) const;
+
+	/** The StateTree asset the context is initialized for */
+	UPROPERTY()
+	const UStateTree* StateTree = nullptr;
+
+	UPROPERTY()
+	UObject* Owner = nullptr;
+	
+	UPROPERTY()
+	UWorld* World = nullptr;
+
+	/** States visited during a tick while updating evaluators. Initialized to match the number of states in the asset. */ 
+	TArray<bool> VisitedStates;
+
+	/** Array of item pointers (external items, tasks, evaluators), used during evaluation. Initialized to match the number of items in the asset. */
+	TArray<FStateTreeItemView> ItemViews;
+
+	/** Optional Instance of the storage */
+	FInstancedStruct StorageInstance;
+
+	/** Storage type of the context */
+	EStateTreeStorage StorageType = EStateTreeStorage::Internal;
+
+	/**
+	 * Temporary status held within the context when calling 'EnterState' on multiple tasks.
+	 * Since it is called on all tasks even if a failed status was returned, this allow other tasks to act accordingly.
+	 * Note. This should be replaced by symmetrical unrolling of tasks on failure.
+	 */
+	EStateTreeRunStatus EnterStateStatus = EStateTreeRunStatus::Unset;
+};
diff --git a/Engine/Plugins/Runtime/StateTree/Source/StateTreeModule/Public/StateTreeInstance.h b/Engine/Plugins/Runtime/StateTree/Source/StateTreeModule/Public/StateTreeInstance.h
new file mode 100644
--- /dev/null
+++ b/Engine/Plugins/Runtime/StateTree/Source/StateTreeModule/Public/StateTreeInstance.h
@@ -0,0 +1,165 @@
+// Copyright Epic Games, Inc. All Rights Reserved.
+
+#pragma once
+
+#include "EngineDefines.h"
+#include "StateTree.h"
+#include "StateTreeInstance.generated.h"
+
+class UStateTreeEvaluatorBase;
+class UStateTreeTaskBase;
+struct FStateTreeParameterStorage;
+struct FStateTreeParameterLayout;
+struct FStateTreeCondition;
+
+/**
+ * Runs StateTrees defined in UStateTree asset.
+ * Uses constant data from StateTree, keeps local storage of variables, and creates instanced Evaluators and Tasks.
+ */
+USTRUCT()
+struct STATETREEMODULE_API FStateTreeInstance
+{
+	GENERATED_BODY()
+
+public:
+	FStateTreeInstance();
+	~FStateTreeInstance();
+
+	// Initializes the StateTree instance to be used with specific owner and StateTree asset.
+	bool Init(UObject& Owner, const UStateTree& InStateTree);
+
+	// Returns the StateTree asset in use.
+	const UStateTree* GetStateTree() const { return StateTree; }
+
+	// Returns owner of the StateTree instance.
+	UObject* GetOwner() const { return Owner; }
+
+	// Start executing.
+	void Start();
+	// Stop executing.
+	void Stop();
+
+	// Tick the state tree logic.
+	EStateTreeRunStatus Tick(const float DeltaTime, FStateTreeParameterStorage* InputParameters = nullptr);
+
+	// Returns value of a variable, or default value if handle is not valid.
+	bool GetValueBool(const FStateTreeHandle Handle, const bool DefaultValue = false) const;
+	int32 GetValueInt(const FStateTreeHandle Handle, const int32 DefaultValue = 0) const;
+	float GetValueFloat(const FStateTreeHandle Handle, const float DefaultValue = 0.0f) const;
+	const FVector& GetValueVector(const FStateTreeHandle Handle, const FVector& DefaultValue = FVector::ZeroVector) const;
+	const FWeakObjectPtr& GetValueObject(const FStateTreeHandle Handle, const FWeakObjectPtr& DefaultValue = nullptr) const;
+	bool GetValue(const FStateTreeHandle Handle, EStateTreeVariableType Type, uint8* Value) const;
+
+	// Sets value of a variable, returns true if the value changed.
+	bool SetValueBool(const FStateTreeHandle Handle, const bool Value);
+	bool SetValueInt(const FStateTreeHandle Handle, const int32 Value);
+	bool SetValueFloat(const FStateTreeHandle Handle, const float Value);
+	bool SetValueVector(const FStateTreeHandle Handle, const FVector& Value);
+	bool SetValueObject(const FStateTreeHandle Handle, FWeakObjectPtr Value);
+	bool SetValue(const FStateTreeHandle Handle, EStateTreeVariableType Type, const uint8* Value);
+
+	bool ReadParametersFromStorage(const FStateTreeParameterLayout& ParameterLayout, const FStateTreeParameterStorage& ParameterStorage);
+	bool WriteParametersToStorage(const FStateTreeParameterLayout& ParameterLayout, FStateTreeParameterStorage& ParameterStorage) const;
+
+#if ENABLE_VISUAL_LOG
+	void DescribeSelfToVisLog(struct FVisualLogEntry& Snapshot) const;
+#endif // ENABLE_VISUAL_LOG
+
+#if WITH_GAMEPLAY_DEBUGGER
+	FString GetDebugInfoString() const;
+#endif // WITH_GAMEPLAY_DEBUGGER
+
+protected:
+
+	// Describes result from state selection logic, which can be a state to enter to, or Done/Failed status for the whole tree.
+	struct FStateSelectionResult
+	{
+		FStateSelectionResult() : Status(EStateTreeRunStatus::Failed) {}
+		FStateSelectionResult(FStateTreeHandle InState, EStateTreeRunStatus InStatus) : State(InState), Status(InStatus) {}
+		FStateSelectionResult(FStateTreeHandle InState) : State(InState), Status(EStateTreeRunStatus::Running) {}
+		FStateSelectionResult(EStateTreeRunStatus InStatus) : State(), Status(InStatus) {}
+
+		bool IsValid() const { return Status == EStateTreeRunStatus::Running ? State.IsValid() : true; }
+
+		FStateTreeHandle State;
+		EStateTreeRunStatus Status;
+	};
+
+	// Resets the instance to initial empty state.
+	void Reset();
+
+	// Resets current active state, and calls ExitState() if needed.
+	void ResetState();
+
+	// Calls Activate on all Evaluators.
+	void ActivateEvaluators();
+	// Calls Deactivate on all Evaluators.
+	void DeactivateEvaluators();
+	// Ticks Evaluators by delta time.
+	void TickEvaluators(const float DeltaTime);
+
+	// Ticks specified state by delta time.
+	EStateTreeRunStatus TickState(const FStateTreeHandle StateHandle, const float DeltaTime);
+
+	// Handles logic for entering State.
+	void EnterState(const FStateTreeHandle StateHandle);
+	// Handles logic for exiting State.
+	void ExitState(const FStateTreeHandle StateHandle);
+
+	// Returns true if the specified condition passes.
+	bool TestCondition(const FStateTreeCondition& Condition);
+
+	// Returns true if all conditions pass.
+	bool TestAllConditions(const uint32 ConditionsOffset, const uint32 ConditionsNum);
+
+	// Walks towards the Root state and triggers the first transition which conditions pass and target state can be entered.
+	// On valid transition SelectState() is executed starting at the transition target state, and resulting state is returned.
+	// If no transition can be triggered, returns CurrentStateHandle.
+	FStateSelectionResult TriggerTransitions(const FStateTreeHandle CurrentStateHandle);
+
+	// Triggers default transitions for specified state. If the state has State Failed transition specified,
+	// and StateTickStatus is Failed, the failure transition is tried to trigger. In other cases State Done
+	// transition is tried to trigger (including the case that State Failed is unset). 
+	// If target state cannot be entered or transition not valid, sets execution status to failed.
+	// On valid transition SelectState() is executed starting at the transition target state, and resulting state is returned.
+	FStateSelectionResult TriggerDefaultTransition(const FStateTreeHandle StateHandle, EStateTreeRunStatus StateTickStatus);
+
+	// Runs state selection logic starting at the specified state, walking towards the leaf nodes.
+	FStateSelectionResult SelectState(const FStateTreeHandle StateHandle);
+
+	// Runs state selection logic on the children of the specified state, walking towards the leaf nodes.
+	FStateSelectionResult SelectChildState(const FStateTreeHandle StateHandle);
+
+	// Owner of the StateTree
+	UPROPERTY(Transient)
+	UObject* Owner;
+
+	// The StateTree asset.
+	UPROPERTY()
+	const UStateTree* StateTree;
+
+	// Evaluator instances
+	UPROPERTY(Transient)
+	TArray<UStateTreeEvaluatorBase*> Evaluators;
+
+	// Task instances
+	UPROPERTY(Transient)
+	TArray<UStateTreeTaskBase*> Tasks;
+
+	// Variable memory of the instance
+	TArray<uint8> VariableMemoryPool;
+
+	// Currently active state
+	FStateTreeHandle CurrentState;
+
+	EStateTreeRunStatus LastTickStatus;
+
+	// Currently active task.
+	FStateTreeHandle CurrentTask;
+
+	// Task to active the next update.
+	FStateTreeHandle NextTask;
+
+	// Running status of the instance.
+	EStateTreeRunStatus RunStatus;
+};
diff --git a/Engine/Plugins/Runtime/StateTree/Source/StateTreeModule/Public/StateTreeParameter.h b/Engine/Plugins/Runtime/StateTree/Source/StateTreeModule/Public/StateTreeParameter.h
new file mode 100644
--- /dev/null
+++ b/Engine/Plugins/Runtime/StateTree/Source/StateTreeModule/Public/StateTreeParameter.h
@@ -0,0 +1,35 @@
+// Copyright Epic Games, Inc. All Rights Reserved.
+
+#pragma once
+
+#include "CoreMinimal.h"
+#include "StateTreeTypes.h"
+#include "Misc/Guid.h"
+#include "StateTreeVariable.h"
+#include "StateTreeParameter.generated.h"
+
+struct FStateTreeVariableDesc;
+
+/**
+ * Describes a named reference to a variable in the StateTree.
+ */
+USTRUCT()
+struct STATETREEMODULE_API FStateTreeParameter
+{
+	GENERATED_BODY()
+
+#if WITH_EDITOR
+	bool IsCompatible(const FStateTreeVariableDesc& Desc) const;
+#endif
+
+	UPROPERTY(EditDefaultsOnly, Category = Parameters);
+	FName Name;
+
+	UPROPERTY(EditDefaultsOnly, Category = Parameters);
+	FStateTreeVariable Variable;
+
+#if WITH_EDITORONLY_DATA
+	UPROPERTY(EditDefaultsOnly, Category = Value)
+	FGuid ID;
+#endif
+};
diff --git a/Engine/Plugins/Runtime/StateTree/Source/StateTreeModule/Public/StateTreeParameterLayout.h b/Engine/Plugins/Runtime/StateTree/Source/StateTreeModule/Public/StateTreeParameterLayout.h
new file mode 100644
--- /dev/null
+++ b/Engine/Plugins/Runtime/StateTree/Source/StateTreeModule/Public/StateTreeParameterLayout.h
@@ -0,0 +1,33 @@
+// Copyright Epic Games, Inc. All Rights Reserved.
+
+#pragma once
+
+#include "CoreMinimal.h"
+#include "StateTreeTypes.h"
+#include "StateTreeParameter.h"
+#include "StateTreeParameterLayout.generated.h"
+
+struct FStateTreeVariableLayout;
+struct FStateTreeConstantStorage;
+
+/**
+ * Defines layout of parameters.
+ * Can be used in the UI to define a how StateTree variables or constants are bound to specific parameters, see StateTreeParameterLayoutDetails.
+ */
+USTRUCT()
+struct FStateTreeParameterLayout
+{
+	GENERATED_BODY()
+
+public:
+	void Reset() { Parameters.Reset(); }
+
+#if WITH_EDITOR
+	bool ResolveVariables(const FStateTreeVariableLayout& Variables, FStateTreeConstantStorage& Constants);
+	bool IsCompatible(const FStateTreeVariableLayout* Layout) const;
+	void UpdateLayout(const FStateTreeVariableLayout* NewLayout);
+#endif
+
+	UPROPERTY(EditDefaultsOnly, Category = Parameters, EditFixedSize, meta = (EditFixedOrder));
+	TArray<FStateTreeParameter> Parameters;
+};
diff --git a/Engine/Plugins/Runtime/StateTree/Source/StateTreeModule/Public/StateTreeParameterStorage.h b/Engine/Plugins/Runtime/StateTree/Source/StateTreeModule/Public/StateTreeParameterStorage.h
new file mode 100644
--- /dev/null
+++ b/Engine/Plugins/Runtime/StateTree/Source/StateTreeModule/Public/StateTreeParameterStorage.h
@@ -0,0 +1,36 @@
+// Copyright Epic Games, Inc. All Rights Reserved.
+
+#pragma once
+
+#include "CoreMinimal.h"
+#include "StateTreeTypes.h"
+#include "StateTreeParameterLayout.h"
+#include "StateTreeVariableLayout.h"
+#include "StateTreeParameterStorage.generated.h"
+
+USTRUCT()
+struct STATETREEMODULE_API FStateTreeParameterStorage
+{
+	GENERATED_BODY()
+
+public:
+	void Reset();
+	void SetLayout(const FStateTreeParameterLayout& NewLayout);
+	void SetLayout(const FStateTreeVariableLayout& NewLayout);
+
+	void MapVariables(FStateTreeParameterLayout& MappedLayout, const FStateTreeParameterLayout& ExpectedLayout) const;
+	bool IsCompatible(const FStateTreeParameterLayout& ParameterLayout) const;
+
+	uint32 GetParameterNum() const { return (uint32)Layout.Variables.Num(); }
+	uint8* GetParameterPtr(uint32 Index, EStateTreeVariableType ExpectedType);
+	const uint8* GetParameterPtr(uint32 Index, EStateTreeVariableType ExpectedType) const;
+
+	// TODO: manual read/write API
+
+protected:
+	UPROPERTY(Transient)
+	TArray<uint8> Memory;
+
+	UPROPERTY()
+	FStateTreeVariableLayout Layout;
+};
