Moved GameplayBehaviors out of restricted folder over to Experimental

diff --git a/Engine/Plugins/Runtime/StateTree/Source/StateTreeModule/Public/StateTreePropertyBindings.h b/Engine/Plugins/Runtime/StateTree/Source/StateTreeModule/Public/StateTreePropertyBindings.h
new file mode 100644
--- /dev/null
+++ b/Engine/Plugins/Runtime/StateTree/Source/StateTreeModule/Public/StateTreePropertyBindings.h
@@ -0,0 +1,535 @@
+// Copyright Epic Games, Inc. All Rights Reserved.
+
+#pragma once
+
+#include "CoreMinimal.h"
+#include "Misc/Guid.h"
+#include "InstancedStruct.h"
+#include "StateTreePropertyBindings.generated.h"
+
+class FProperty;
+
+
+/**
+ * Short lived pointer to an UOBJECT() or USTRUCT().
+ * The item view expects a type (UStruct) when you pass in a valid memory. In case of null, the type can be empty too.
+ */
+struct STATETREEMODULE_API FStateTreeItemView
+{
+	FStateTreeItemView() = default;
+
+	// USTRUCT() constructor.
+	FStateTreeItemView(const UScriptStruct* InScriptStruct, uint8* InMemory) : Struct(InScriptStruct), Memory(InMemory)
+	{
+		// Must have type with valid pointer.
+		check(!Memory || (Memory && Struct));
+	}
+
+	// UOBJECT() constructor.
+	FStateTreeItemView(UObject* Object) : Struct(Object ? Object->GetClass() : nullptr), Memory(reinterpret_cast<uint8*>(Object))
+	{
+		// Must have type with valid pointer.
+		check(!Memory || (Memory && Struct));
+	}
+
+	// USTRUCT() from a StructView.
+	FStateTreeItemView(FStructView StructView) : Struct(StructView.GetScriptStruct()), Memory(StructView.GetMutableMemory())
+	{
+		// Must have type with valid pointer.
+		check(!Memory || (Memory && Struct));
+	}
+
+	/**
+	 * Check is the view is valid (both pointer and type are set). On valid views it is safe to call the Get<>() methods returning a reference.
+	 * @return True if the view is valid.
+	*/
+	bool IsValid() const
+	{
+		return Memory != nullptr && Struct != nullptr;
+	}
+
+	/*
+	 * UOBJECT() getters (reference & pointer, const & mutable)
+	 */
+	template <typename T>
+    typename TEnableIf<TIsDerivedFrom<T, UObject>::IsDerived, const T&>::Type Get() const
+	{
+		check(Memory != nullptr);
+		check(Struct != nullptr);
+		check(Struct->IsChildOf(T::StaticClass()));
+		return *((T*)Memory);
+	}
+
+	template <typename T>
+	typename TEnableIf<TIsDerivedFrom<T, UObject>::IsDerived, T&>::Type GetMutable() const
+	{
+		check(Memory != nullptr);
+		check(Struct != nullptr);
+		check(Struct->IsChildOf(T::StaticClass()));
+		return *((T*)Memory);
+	}
+
+	template <typename T>
+	typename TEnableIf<TIsDerivedFrom<T, UObject>::IsDerived, const T*>::Type GetPtr() const
+	{
+		// If Memory is set, expect Struct too. Otherwise, let nulls pass through.
+		check(!Memory || (Memory && Struct));
+		check(!Struct || Struct->IsChildOf(T::StaticClass()));
+		return ((T*)Memory);
+	}
+
+	template <typename T>
+    typename TEnableIf<TIsDerivedFrom<T, UObject>::IsDerived, T*>::Type GetMutablePtr() const
+	{
+		// If Memory is set, expect Struct too. Otherwise, let nulls pass through.
+		check(!Memory || (Memory && Struct));
+		check(!Struct || Struct->IsChildOf(T::StaticClass()));
+		return ((T*)Memory);
+	}
+
+	/*
+	 * USTRUCT() getters (reference & pointer, const & mutable)
+	 */
+	template <typename T>
+	typename TEnableIf<!TIsDerivedFrom<T, UObject>::IsDerived, const T&>::Type Get() const
+	{
+		check(Memory != nullptr);
+		check(Struct != nullptr);
+		check(Struct->IsChildOf(T::StaticStruct()));
+		return *((T*)Memory);
+	}
+
+	template <typename T>
+    typename TEnableIf<!TIsDerivedFrom<T, UObject>::IsDerived, T&>::Type GetMutable() const
+	{
+		check(Memory != nullptr);
+		check(Struct != nullptr);
+		check(Struct->IsChildOf(T::StaticStruct()));
+		return *((T*)Memory);
+	}
+
+	template <typename T>
+    typename TEnableIf<!TIsDerivedFrom<T, UObject>::IsDerived, const T*>::Type GetPtr() const
+	{
+		// If Memory is set, expect Struct too. Otherwise, let nulls pass through.
+		check(!Memory || (Memory && Struct));
+		check(!Struct || Struct->IsChildOf(T::StaticStruct()));
+		return ((T*)Memory);
+	}
+
+	template <typename T>
+    typename TEnableIf<!TIsDerivedFrom<T, UObject>::IsDerived, T*>::Type GetMutablePtr() const
+	{
+		// If Memory is set, expect Struct too. Otherwise, let nulls pass through.
+		check(!Memory || (Memory && Struct));
+		check(!Struct || Struct->IsChildOf(T::StaticStruct()));
+		return ((T*)Memory);
+	}
+
+	/** @return Struct describing the item type. */
+	const UStruct* GetStruct() const { return Struct; }
+
+	/** @return Raw const pointer to the item. */
+	const uint8* GetMemory() const { return Memory; }
+
+	/** @return Raw mutable pointer to the item. */
+	uint8* GetMutableMemory() const { return Memory; }
+	
+protected:
+	/** Class or struct of the external item. */
+	const UStruct* Struct = nullptr;
+
+	/** Memory pointing at the class or struct */
+	uint8* Memory = nullptr;
+};
+
+
+/**
+ * Descriptor for a struct or class that can be a binding source or target.
+ * Each struct has unique identifier, which is used to distinguish them, and name that is mostly for debugging and UI.
+ */
+USTRUCT()
+struct STATETREEMODULE_API FStateTreeBindableStructDesc
+{
+	GENERATED_BODY()
+
+#if WITH_EDITORONLY_DATA
+	bool operator==(const FStateTreeBindableStructDesc& RHS) const
+	{
+		return ID == RHS.ID && Struct == RHS.Struct; // Not checking name, it's cosmetic.
+	}
+#endif
+
+	/** The type of the struct or class. */
+	UPROPERTY()
+	const UStruct* Struct = nullptr;
+
+	/** Name of the container (cosmetic). */
+	UPROPERTY()
+	FName Name;
+
+#if WITH_EDITORONLY_DATA
+	/** Unique identifier of the struct. */
+	UPROPERTY()
+	FGuid ID;
+#endif
+};
+
+
+/**
+ * The type of access a property path describes.
+ */
+UENUM()
+enum class EStateTreePropertyAccessType : uint8
+{
+	Offset,			// Access node is a simple basePtr + offset
+	Object,			// Access node needs to dereference an object at its current address
+	WeakObject,		// Access is a weak object
+	SoftObject,		// Access is a soft object
+	IndexArray,		// Access node indexes a dynamic array
+};
+
+/**
+ * Describes how the copy should be performed.
+ */
+UENUM()
+enum class EStateTreePropertyCopyType : uint8
+{
+	None,						// No copying
+	
+	CopyPlain,					// For plain old data types, we do a simple memcpy.
+	CopyComplex,				// For more complex data types, we need to call the properties copy function
+	CopyBool,					// Read and write properties using bool property helpers, as source/dest could be bitfield or boolean
+	CopyStruct,					// Use struct copy operation, as this needs to correctly handle CPP struct ops
+	CopyObject,					// Read and write properties using object property helpers, as source/dest could be regular/weak/soft etc.
+	CopyName,					// FName needs special case because its size changes between editor/compiler and runtime.
+	CopyFixedArray,				// Array needs special handling for fixed size TArrays
+
+	/* Promote the type during the copy */
+
+	/* Bool promotions */
+	PromoteBoolToByte,
+	PromoteBoolToInt32,
+	PromoteBoolToUInt32,
+	PromoteBoolToInt64,
+	PromoteBoolToFloat,
+	PromoteBoolToDouble,
+
+	/* Byte promotions */
+	PromoteByteToInt32,
+	PromoteByteToUInt32,
+	PromoteByteToInt64,
+	PromoteByteToFloat,
+	PromoteByteToDouble,
+
+	/* Int32 promotions */
+	PromoteInt32ToInt64,
+	PromoteInt32ToFloat,	// This is strictly sketchy because of potential data loss, but it is usually OK in the general case
+	PromoteInt32ToDouble,
+
+	/* UInt32 promotions */
+	PromoteUInt32ToInt64,
+    PromoteUInt32ToFloat,	// This is strictly sketchy because of potential data loss, but it is usually OK in the general case
+
+	/* Float promotions */
+	// LWC_TODO: Float/double should become synonyms in state trees? Certainly BPs.
+	PromoteFloatToDouble,
+	DemoteDoubleToFloat,	// LWC_TODO: This should not ship!
+};
+
+
+/**
+ * Describes a segment of a property path (see FStateTreePropertyPath). Used for storage only.
+ */
+USTRUCT()
+struct STATETREEMODULE_API FStateTreePropertySegment
+{
+	GENERATED_BODY()
+
+	/** Property name. */
+	UPROPERTY()
+	FName Name;
+
+	/** Index in the array the property points at. */
+	UPROPERTY()
+	int32 ArrayIndex = 0;
+
+	/** Type of access/indirection. */
+	UPROPERTY()
+	EStateTreePropertyAccessType Type = EStateTreePropertyAccessType::Offset;
+};
+
+/**
+ * Describes a property path made of segments (i.e. Foo.Bar consists of two segments). Used for storage only.
+ */
+USTRUCT()
+struct STATETREEMODULE_API FStateTreePropertyPath
+{
+	GENERATED_BODY()
+
+	/** Index to first segment of the property path. */
+	UPROPERTY()
+	int32 SegmentsBegin = 0;
+
+	/** Index to one past the last segment in the property path. */
+	UPROPERTY()
+	int32 SegmentsEnd = 0;
+};
+
+/**
+ * Describes property binding, the property from source path is copied into the property at the target path.
+ */
+USTRUCT()
+struct STATETREEMODULE_API FStateTreePropertyBinding
+{
+	GENERATED_BODY()
+
+	/** Index to source property path. */
+	UPROPERTY()
+	int32 SourcePathIndex = 0;
+
+	/** Index to target property path. */
+	UPROPERTY()
+	int32 TargetPathIndex = 0;
+
+	/** Index to the source struct the source path refers to, sources are stored in FStateTreePropertyBindings. */
+	UPROPERTY()
+	int32 SourceStructIndex = 0;
+
+	/** The type of copy to use between the properties. */
+	UPROPERTY()
+	EStateTreePropertyCopyType CopyType = EStateTreePropertyCopyType::None;
+};
+
+/**
+ * Property indirection is a resolved property path segment, used for accessing properties in structs.
+ */
+USTRUCT()
+struct STATETREEMODULE_API FStateTreePropertyIndirection
+{
+	GENERATED_BODY()
+
+	/** Index in the array the property points at. */
+	UPROPERTY()
+	int32 ArrayIndex = 0;
+
+	/** Cached offset of the property */
+	UPROPERTY()
+	int32 Offset = 0;
+
+	/** Type of access/indirection. */
+	UPROPERTY()
+	EStateTreePropertyAccessType Type = EStateTreePropertyAccessType::Offset;
+
+	/** Cached array property. */
+	const FArrayProperty* ArrayProperty = nullptr;
+};
+
+/**
+ * Property access is a resolved property path, used for accessing properties in structs via one or more indirections.
+ */
+USTRUCT()
+struct STATETREEMODULE_API FStateTreePropertyAccess
+{
+	GENERATED_BODY()
+
+	/** Index to first indirection of the property access. */
+	UPROPERTY()
+	int32 IndirectionsBegin = 0;
+
+	/** Index to one past the last indirection in the property access. */
+	UPROPERTY()
+	int32 IndirectionsEnd = 0;
+
+	/** Cached pointer to the leaf property of the access. */
+	const FProperty* LeafProperty = nullptr;
+};
+
+
+/**
+ * Describes property copy, the property from source is copied into the property at the target.
+ * Copy target struct is described in the property copy batch.
+ */
+USTRUCT()
+struct STATETREEMODULE_API FStateTreePropCopy
+{
+	GENERATED_BODY()
+
+	/** Index to source property access. */
+	UPROPERTY()
+	int32 SourceAccessIndex = 0;
+
+	/** Index to target property access. */
+	UPROPERTY()
+	int32 TargetAccessIndex = 0;
+
+	/** Index to the struct the source path refers to, sources are stored in FStateTreePropertyBindings. */
+	UPROPERTY()
+	int32 SourceStructIndex = 0;
+
+	/** Type of the copy */
+	UPROPERTY()
+	EStateTreePropertyCopyType Type = EStateTreePropertyCopyType::None;
+
+	/** Cached property element size * dim. */
+	UPROPERTY()
+	int32 CopySize = 0;
+};
+
+
+/**
+ * Describes a batch of property copies from many sources to one target struct.
+ * Note: The batch is used to reference both bindings and copies (a binding turns into copy when resolved).
+ */
+USTRUCT()
+struct STATETREEMODULE_API FStateTreePropCopyBatch
+{
+	GENERATED_BODY()
+
+	/** Expected target struct */
+	UPROPERTY()
+	FStateTreeBindableStructDesc TargetStruct;
+
+	/** Index to first binding/copy. */
+	UPROPERTY()
+	int32 BindingsBegin = 0;
+
+	/** Index to one past the last binding/copy. */
+	UPROPERTY()
+	int32 BindingsEnd = 0;
+};
+
+/**
+ * Runtime storage and execution of property bindings.
+ */
+USTRUCT()
+struct STATETREEMODULE_API FStateTreePropertyBindings
+{
+	GENERATED_BODY()
+
+	/**
+	 * Clears all bindings.
+	 */
+	void Reset();
+
+	/**
+	 * Resolves paths to indirections.
+	 * @return True if resolve succeeded.
+	 */
+	bool ResolvePaths();
+
+	/**
+	 * @return True if bindings have been resolved.
+	 */
+	bool IsValid() const { return bBindingsResolved; }
+
+	/**
+	 * @return Number of source structs the copy expects.
+	 */
+	int32 GetSourceStructNum() const { return SourceStructs.Num(); }
+
+	/**
+	 * Copies a batch of properties from source structs to target struct.
+	 * @param SourceStructViews Views to structs where properties are copied from.
+	 * @param TargetBatchIndex Batch index to copy (see FStateTreePropertyBindingCompiler).
+	 * @param TargetStructView View to struct where properties are copied to.
+	 */
+	void CopyTo(TConstArrayView<FStateTreeItemView> SourceStructViews, const int32 TargetBatchIndex, FStateTreeItemView TargetStructView) const;
+
+	void DebugPrintInternalLayout(FString& OutString) const;
+
+protected:
+	bool ResolvePath(const UStruct* Struct, const FStateTreePropertyPath& Path, FStateTreePropertyAccess& OutAccess);
+	bool ValidateCopy(FStateTreePropCopy& Copy) const;
+	void PerformCopy(const FStateTreePropCopy& Copy, const FProperty* SourceProperty, const uint8* SourceAddress, const FProperty* TargetProperty, uint8* TargetAddress) const;
+	uint8* GetAddress(FStateTreeItemView InStructView, const FStateTreePropertyAccess& InAccess) const;
+	FString GetPathAsString(const FStateTreePropertyPath& Path, const int32 HighlightedSegment = INDEX_NONE, const TCHAR* HighlightPrefix = nullptr, const TCHAR* HighlightPostfix = nullptr);
+
+	/** Array of expected source structs. */
+	UPROPERTY()
+	TArray<FStateTreeBindableStructDesc> SourceStructs;
+
+	/** Array of copy batches. */
+	UPROPERTY()
+	TArray<FStateTreePropCopyBatch> CopyBatches;
+
+	/** Array of property bindings, resolved into arrays of copies before use. */
+	UPROPERTY()
+	TArray<FStateTreePropertyBinding> PropertyBindings;
+
+	/** Array of property bindings, indexed by property paths. */
+	UPROPERTY()
+	TArray<FStateTreePropertyPath> PropertyPaths;
+
+	/** Array of property segments, indexed by property paths. */
+	UPROPERTY()
+	TArray<FStateTreePropertySegment> PropertySegments;
+
+	/** Array of property copies */
+	UPROPERTY(Transient)
+	TArray<FStateTreePropCopy> PropertyCopies;
+
+	/** Array of property accesses, indexed by copies*/
+	UPROPERTY(Transient)
+	TArray<FStateTreePropertyAccess> PropertyAccesses;
+
+	/** Array of property indirections, indexed by accesses*/
+	UPROPERTY(Transient)
+	TArray<FStateTreePropertyIndirection> PropertyIndirections;
+
+	/** Flag indicating if the properties has been resolved successfully . */
+	bool bBindingsResolved = false;
+
+	friend struct FStateTreePropertyBindingCompiler;
+};
+
+
+
+/**
+ * Editor representation of a property path in StateTree.
+ * Note: This is defined here for IStateTreeBindingLookup.
+ */
+USTRUCT()
+struct STATETREEMODULE_API FStateTreeEditorPropertyPath
+{
+	GENERATED_BODY()
+
+	FStateTreeEditorPropertyPath() = default;
+	FStateTreeEditorPropertyPath(const FGuid& InStructID, const TCHAR* PropertyName) : StructID(InStructID) { Path.Add(PropertyName); }
+	FStateTreeEditorPropertyPath(const FStateTreeEditorPropertyPath& InPath) : StructID(InPath.StructID), Path(InPath.Path) {}
+
+	/**
+	 * Returns the property path as a one string. Highlight allows to decorate a specific segment.
+	 * @param HighlightedSegment Index of the highlighted path segment
+	 * @param HighlightPrefix String to append before highlighted segment
+	 * @param HighlightPostfix String to append after highlighted segment
+	 */
+	FString ToString(const int32 HighlightedSegment = INDEX_NONE, const TCHAR* HighlightPrefix = nullptr, const TCHAR* HighlightPostfix = nullptr) const;
+
+	/** Handle of the struct this property path is relative to. */
+	UPROPERTY()
+	FGuid StructID;
+
+	/** Property path segments */
+	UPROPERTY()
+	TArray<FString> Path;
+
+	bool IsValid() const { return StructID.IsValid() && Path.Num() > 0; }
+
+	bool operator==(const FStateTreeEditorPropertyPath& RHS) const;
+};
+
+/**
+ * Helper interface to reason about bound properties. The implementation is in the editor plugin.
+ */
+struct STATETREEMODULE_API IStateTreeBindingLookup
+{
+	/** @return Source path for given target path, or null if binding does not exists. */
+	virtual const FStateTreeEditorPropertyPath* GetPropertyBindingSource(const FStateTreeEditorPropertyPath& InTargetPath) const PURE_VIRTUAL(IStateTreeBindingLookup::GetPropertyBindingSource, return nullptr; );
+
+	/** @return Display name given property path. */
+	virtual FText GetPropertyPathDisplayName(const FStateTreeEditorPropertyPath& InPath) const PURE_VIRTUAL(IStateTreeBindingLookup::GetPropertyPathDisplayName, return FText::GetEmpty(); );
+
+	/** @return Leaf property based on property path. */
+	virtual const FProperty* GetPropertyPathLeafProperty(const FStateTreeEditorPropertyPath& InPath) const PURE_VIRTUAL(IStateTreeBindingLookup::GetPropertyPathLeafProperty, return nullptr; );
+};
diff --git a/Engine/Plugins/Runtime/StateTree/Source/StateTreeModule/Public/StateTreeSchema.h b/Engine/Plugins/Runtime/StateTree/Source/StateTreeModule/Public/StateTreeSchema.h
new file mode 100644
--- /dev/null
+++ b/Engine/Plugins/Runtime/StateTree/Source/StateTreeModule/Public/StateTreeSchema.h
@@ -0,0 +1,37 @@
+// Copyright Epic Games, Inc. All Rights Reserved.
+
+#pragma once
+
+#include "CoreMinimal.h"
+#include "Engine/DataAsset.h"
+#include "StateTreeTypes.h"
+#include "StateTreeSchema.generated.h"
+
+/**
+ * Schema describing which inputs, evaluators, and tasks a StateTree can contain.
+ * Each StateTree asset saves the schema class name in asset data tags, which can be
+ * used to limit which StatTree assets can be selected per use case, i.e.:
+ *
+ *	UPROPERTY(EditDefaultsOnly, Category = AI, meta=(RequiredAssetDataTags="Schema=StateTreeSchema_SupaDupa"))
+ *	UStateTree* StateTree;
+ *
+ */
+UCLASS(Abstract)
+class STATETREEMODULE_API UStateTreeSchema : public UObject
+{
+	GENERATED_BODY()
+
+public:
+
+	/** @return Returns the script struct the storage struct will be derived from. */
+	virtual UScriptStruct* GetStorageSuperStruct() const { return nullptr; }
+
+	/** @return True if specified struct is supported */
+	virtual bool IsStructAllowed(const UScriptStruct* InScriptStruct) const PURE_VIRTUAL(UStateTreeSchema::IsStructAllowed, return false; );
+
+	/** @return True if specified struct/class is supported as external item */
+	virtual bool IsExternalItemAllowed(const UStruct& InStruct) const { return false; };
+
+	/** @return True if we should use StateTree V2 */
+	virtual bool IsV2() const { return false; }
+};
diff --git a/Engine/Plugins/Runtime/StateTree/Source/StateTreeModule/Public/StateTreeTaskBase.h b/Engine/Plugins/Runtime/StateTree/Source/StateTreeModule/Public/StateTreeTaskBase.h
new file mode 100644
--- /dev/null
+++ b/Engine/Plugins/Runtime/StateTree/Source/StateTreeModule/Public/StateTreeTaskBase.h
@@ -0,0 +1,131 @@
+// Copyright Epic Games, Inc. All Rights Reserved.
+
+#pragma once
+
+#include "CoreMinimal.h"
+#include "Misc/Guid.h"
+#include "StateTreeTypes.h"
+#include "StateTreeTaskBase.generated.h"
+
+struct FStateTreeInstance;
+struct FStateTreeExecutionContext;
+struct FStateTreeVariableLayout;
+struct FStateTreeConstantStorage;
+
+/**
+ * Base class for StateTree Tasks.
+ * Tasks are logic executed at an active state.
+ */
+UCLASS(Abstract)
+class STATETREEMODULE_API UStateTreeTaskBase : public UObject
+{
+	GENERATED_BODY()
+
+public:
+	UStateTreeTaskBase(const FObjectInitializer& ObjectInitializer);
+
+	// Called when instantiated first time.
+	virtual bool Initialize(FStateTreeInstance& StateTreeInstance) PURE_VIRTUAL(UStateTreeTaskBase::Initialize, return false; );
+	// Called when evaluator becomes active/ticking.
+	virtual void Activate(FStateTreeInstance& StateTreeInstance) PURE_VIRTUAL(UStateTreeTaskBase::Activate, return; );
+	// Called when evaluator becomes inactive.
+	virtual void Deactivate(FStateTreeInstance& StateTreeInstance) PURE_VIRTUAL(UStateTreeTaskBase::Deactivate, return; );
+
+	// Called on each tick.
+	virtual EStateTreeRunStatus Tick(FStateTreeInstance& StateTreeInstance, const float DeltaTime) PURE_VIRTUAL(UStateTreeTaskBase::Tick, return EStateTreeRunStatus::Failed; );
+
+#if WITH_GAMEPLAY_DEBUGGER
+	virtual void AppendDebugInfoString(FString& DebugString, const FStateTreeInstance& StateTreeInstance) const;
+#endif
+
+#if WITH_EDITOR
+	// Create Task instance template, set default variables, and resolve variable handles.
+	virtual bool ResolveVariables(const FStateTreeVariableLayout& Variables, FStateTreeConstantStorage& Constants, UObject* Outer) PURE_VIRTUAL(UStateTreeTaskBase::Resolve, return false; );
+	// Returns ture if changed.
+	virtual bool ValidateParameterLayout() { return false; }
+
+	virtual void PostLoad() override;
+#endif
+
+	void SetNewUniqueID() { ID = FGuid::NewGuid(); }
+
+	UPROPERTY(EditAnywhere, Category = Task)
+	FName Name;
+
+	UPROPERTY(meta = (IgnoreForMemberInitializationTest))
+	FGuid ID;
+};
+
+
+// STATETREE_V2
+
+/**
+ * Base stuct for StateTree Tasks.
+ * Tasks are logic executed in an active state.
+ */
+USTRUCT()
+struct STATETREEMODULE_API FStateTreeTask2Base // TODO: change to FStateTreeTaskBase once UStateTreeTaskBase is removed
+{
+	GENERATED_BODY()
+
+	FStateTreeTask2Base() = default;
+
+	virtual ~FStateTreeTask2Base() {}
+
+	/**
+	 * Called when a new state is entered and task is part of active states. The change type parameter describes if the task's state
+	 * was previously part of the list of active states (Sustained), or if it just became active (Changed).
+	 * @param Context Reference to current execution context.
+	 * @param ChangeType Describes the change type (Changed/Sustained).
+	 * @param Transition Describes the states involved in the transition
+	 * @return Succeed/Failed will end the state immediately and trigger to select new state, Running will carry on to tick the state.
+	 */
+	virtual EStateTreeRunStatus EnterState(FStateTreeExecutionContext& Context, const EStateTreeStateChangeType ChangeType, const FStateTreeTransitionResult& Transition) { return EStateTreeRunStatus::Running; }
+
+	/**
+	* Called when a current state is exited and task is part of active states. The change type parameter describes if the task's state
+	* will be active after the transition (Sustained), or if it will became inactive (Changed).
+	* @param Context Reference to current execution context.
+	* @param ChangeType Describes the change type (Changed/Sustained).
+	* @param Transition Describes the states involved in the transition
+	*/
+	virtual void ExitState(FStateTreeExecutionContext& Context, const EStateTreeStateChangeType ChangeType, const FStateTreeTransitionResult& Transition) {}
+
+	/**
+	 * Called Right after a state has been completed. StateCompleted is called in reverse order to allow to propagate state to Evaluators and Tasks that
+	 * are executed earlier in the tree. Note that StateCompleted is not called if conditional transition changes the state.
+	 * @param Context Reference to current execution context.
+	 * @param CompletionStatus Describes the running status of the completed state (Succeeded/Failed).
+	 * @param CompletedState Handle of the state that was completed.
+	 */
+	virtual void StateCompleted(FStateTreeExecutionContext& Context, const EStateTreeRunStatus CompletionStatus, const FStateTreeHandle CompletedState) {}
+
+	/**
+	* Called during state tree tick when the task is on active state.
+	* @param Context Reference to current execution context.
+	* @param DeltaTime Time since last StateTree tick.
+	* @return Running status of the state: Running if still in progress, Succeeded if execution is done and succeeded, Failed if execution is done and failed.
+	*/
+	virtual EStateTreeRunStatus Tick(FStateTreeExecutionContext& Context, const float DeltaTime) { return EStateTreeRunStatus::Failed; };
+
+#if WITH_GAMEPLAY_DEBUGGER
+	virtual void AppendDebugInfoString(FString& DebugString, const FStateTreeExecutionContext& Context) const;
+#endif
+
+	UPROPERTY(EditAnywhere, Category = Task, meta = (EditCondition = "false", EditConditionHides))
+	FName Name;
+
+#if WITH_EDITORONLY_DATA
+	UPROPERTY(EditAnywhere, Category = Task, meta = (IgnoreForMemberInitializationTest, EditCondition = "false", EditConditionHides))	// Hack, we want the ID to be found as IPropertyHandle, but do not want to display it.
+	FGuid ID;
+#endif
+
+	UPROPERTY()
+	FStateTreeHandle BindingsBatch = FStateTreeHandle::Invalid;	// Property binding copy batch handle.
+
+	UPROPERTY()
+	uint16 SourceStructIndex = 0;								// Property binding Source Struct index of the task.
+};
+template<> struct TStructOpsTypeTraits<FStateTreeTask2Base> : public TStructOpsTypeTraitsBase2<FStateTreeTask2Base> { enum { WithPureVirtual = true, }; };
+
+// ~STATETREE_V2
diff --git a/Engine/Plugins/Runtime/StateTree/Source/StateTreeModule/Public/StateTreeTypes.h b/Engine/Plugins/Runtime/StateTree/Source/StateTreeModule/Public/StateTreeTypes.h
new file mode 100644
--- /dev/null
+++ b/Engine/Plugins/Runtime/StateTree/Source/StateTreeModule/Public/StateTreeTypes.h
@@ -0,0 +1,463 @@
+// Copyright Epic Games, Inc. All Rights Reserved.
+
+#pragma once
+
+#include "CoreMinimal.h"
+#include "StateTreeTypes.generated.h"
+
+STATETREEMODULE_API DECLARE_LOG_CATEGORY_EXTERN(LogStateTree, Warning, All);
+
+#ifndef WITH_STATETREE_DEBUG
+#define WITH_STATETREE_DEBUG (!(UE_BUILD_SHIPPING || UE_BUILD_SHIPPING_WITH_EDITOR || UE_BUILD_TEST) && 1)
+#endif // WITH_STATETREE_DEBUG
+
+/** Types of supported variables */
+UENUM()
+enum class EStateTreeVariableType : uint8
+{
+	Void,
+	Float,
+	Int,
+	Bool,
+	Vector,
+	Object,
+};
+
+/**
+ * Status describing current ticking state. 
+ */
+UENUM()
+enum class EStateTreeRunStatus : uint8
+{
+	Unset,				/** Status not set. */
+	Failed,				/** Tree execution has stopped on failure. */
+	Succeeded,			/** Tree execution has stopped on success. */
+	Running,			/** Tree is still running. */
+};
+
+/**  Evaluator evaluation type. */
+UENUM()
+enum class EStateTreeEvaluationType : uint8
+{
+	PreSelect,			/** Called during selection process on states that have not been visited yet. */
+    Tick,				/** Called during tick on active states. */
+};
+
+/**  State change type. Passed to EnterState() and ExitState() to indicate how the state change affects the state and Evaluator or Task is on. */
+UENUM()
+enum class EStateTreeStateChangeType : uint8
+{
+	None,				/** Not an activation */
+	Changed,			/** The state became activated or deactivated. */
+    Sustained,			/** The state is parent of new active state and sustained previous active state. */
+};
+
+/** Transitions behavior. */
+UENUM()
+enum class EStateTreeTransitionType : uint8
+{
+	Succeeded,			// Signal StateTree execution succeeded.
+	Failed,				// Signal StateTree execution failed.
+	SelectChildState,	// Select child state.
+	GotoState,			// Transition to specified state.
+	NotSet,				// No transition.
+	NextState,			// Goto next sibling state.
+};
+
+
+/** Transitions event. */
+UENUM()
+enum class EStateTreeTransitionEvent : uint8
+{
+	None = 0 UMETA(Hidden),
+	OnCompleted = 0x1 | 0x2,
+    OnSucceeded = 0x1,
+    OnFailed = 0x2,
+    OnCondition = 0x4,
+};
+ENUM_CLASS_FLAGS(EStateTreeTransitionEvent)
+
+/**
+ * Handle that is used to refer baked state tree data.
+ */
+USTRUCT()
+struct STATETREEMODULE_API FStateTreeHandle
+{
+	GENERATED_BODY()
+
+	static const uint16 InvalidIndex = uint16(-1);		// Index value indicating invalid item.
+	static const uint16 SucceededIndex = uint16(-2);	// Index value indicating a Succeeded item.
+	static const uint16 FailedIndex = uint16(-3);		// Index value indicating a Failed item.
+	
+	static const FStateTreeHandle Invalid;
+	static const FStateTreeHandle Succeeded;
+	static const FStateTreeHandle Failed;
+
+	FStateTreeHandle() = default;
+	explicit FStateTreeHandle(uint16 InIndex) : Index(InIndex) {}
+
+	bool IsValid() const { return Index != InvalidIndex; }
+
+	bool operator==(const FStateTreeHandle& RHS) const { return Index == RHS.Index; }
+	bool operator!=(const FStateTreeHandle& RHS) const { return Index != RHS.Index; }
+
+	FString Describe() const
+	{
+		switch (Index)
+		{
+		case InvalidIndex:		return TEXT("Invalid Item");
+		case SucceededIndex: 	return TEXT("Succeeded Item");
+		case FailedIndex: 		return TEXT("Failed Item");
+		default: 				return FString::Printf(TEXT("%d"), Index);
+		}
+	}
+
+	UPROPERTY()
+	uint16 Index = InvalidIndex;
+};
+
+UENUM()
+enum class EStateTreeResultStatus : uint8
+{
+	Unset,
+	Available,
+    InUse,
+	Succeeded,
+    Failed,
+};
+
+USTRUCT()
+struct STATETREEMODULE_API FStateTreeResult
+{
+	GENERATED_BODY()
+	virtual ~FStateTreeResult() {}
+	virtual const UScriptStruct& GetStruct() const PURE_VIRTUAL(FStateTreeResult::GetStruct, return *FStateTreeResult::StaticStruct(); );
+};
+template<> struct TStructOpsTypeTraits<FStateTreeResult> : public TStructOpsTypeTraitsBase2<FStateTreeResult> { enum { WithPureVirtual = true, }; };
+
+USTRUCT()
+struct STATETREEMODULE_API FStateTreeResultRef
+{
+	GENERATED_BODY()
+
+	FStateTreeResultRef() = default;
+	explicit FStateTreeResultRef(FStateTreeResult* InResult) : Result(InResult) {}
+	FStateTreeResultRef(const FStateTreeResultRef& InOther) : FStateTreeResultRef(InOther.Result) {}
+	~FStateTreeResultRef()
+	{
+		Result = nullptr;
+	}
+
+	FStateTreeResultRef& operator=(FStateTreeResult* InResult)
+	{
+		Result = InResult;
+		return *this;
+	}
+
+	FStateTreeResultRef& operator=(const FStateTreeResultRef& InOther)
+	{
+		Result = InOther.Result;
+		return *this;
+	}
+
+	bool IsValid() const { return Result != nullptr; }
+
+	void Reset()
+	{
+		Result = nullptr;
+	}
+
+	template<typename T>
+	bool IsA() const
+	{
+		const UScriptStruct* ScriptStruct = Result != nullptr ? &Result->GetStruct() : nullptr;
+		return ScriptStruct != nullptr && ScriptStruct->IsChildOf(T::StaticStruct());
+	}
+	
+	// Returns mutable reference to the struct, this getter assumes that all data is valid
+	template<typename T>
+    T& GetMutable()
+	{
+		check(Result != nullptr);
+		const UScriptStruct& ScriptStruct = Result->GetStruct();
+		check(ScriptStruct.IsChildOf(T::StaticStruct()));
+		return *static_cast<T*>(Result);
+	}
+
+	// Returns const reference to the struct, this getter assumes that all data is valid
+	template<typename T>
+    const T& Get() const
+	{
+		check(Result != nullptr);
+		const UScriptStruct& ScriptStruct = Result->GetStruct();
+		check(ScriptStruct.IsChildOf(T::StaticStruct()));
+		return *static_cast<T*>(Result);
+	}
+
+	// Returns mutable pointer to the struct, or nullptr if cast is not valid.
+	template<typename T>
+    T* GetMutablePtr()
+	{
+		const UScriptStruct* ScriptStruct = Result != nullptr ? &Result->GetStruct() : nullptr;
+		if (ScriptStruct != nullptr && ScriptStruct->IsChildOf(T::StaticStruct()))
+		{
+			return static_cast<T*>(Result);
+		}
+		return nullptr;
+	}
+
+	// Returns const pointer to the struct, or nullptr if cast is not valid.
+	template<typename T>
+    const T* GetPtr() const
+	{
+		const UScriptStruct* ScriptStruct = Result != nullptr ? &Result->GetStruct() : nullptr;
+		if (ScriptStruct != nullptr && ScriptStruct->IsChildOf(T::StaticStruct()))
+		{
+			return static_cast<T*>(Result);
+		}
+		return nullptr;
+	}
+	
+private:
+	FStateTreeResult* Result = nullptr;
+};
+
+
+/**
+ * Describes current status of a running state or desired state.
+ */
+USTRUCT()
+struct STATETREEMODULE_API FStateTreeStateStatus
+{
+	GENERATED_BODY()
+
+	FStateTreeStateStatus() = default;
+	FStateTreeStateStatus(const FStateTreeHandle InState, const EStateTreeRunStatus InStatus) : State(InState), RunStatus(InStatus) {}
+	FStateTreeStateStatus(const FStateTreeHandle InState) : State(InState), RunStatus(EStateTreeRunStatus::Running) {}
+	FStateTreeStateStatus(const EStateTreeRunStatus InStatus) : State(), RunStatus(InStatus) {}
+
+	bool IsSet() const { return RunStatus != EStateTreeRunStatus::Unset; }
+	
+	UPROPERTY()
+	FStateTreeHandle State = FStateTreeHandle::Invalid;
+
+	UPROPERTY()
+	EStateTreeRunStatus RunStatus = EStateTreeRunStatus::Unset;
+};
+
+/**
+ * Describes a state tree transition. Source is the state where the transition started, Target describes the state where the transition pointed at,
+ * and Next describes the selected state. The reason Transition and Next are different is that Transition state can be a selector state,
+ * in which case the children will be visited until a leaf state is found, which will be the next state.
+ */
+USTRUCT()
+struct STATETREEMODULE_API FStateTreeTransitionResult
+{
+	GENERATED_BODY()
+
+	FStateTreeTransitionResult() = default;
+	FStateTreeTransitionResult(const FStateTreeStateStatus InSource, const FStateTreeHandle InTransitionAndNext) : Source(InSource), Target(InTransitionAndNext), Next(InTransitionAndNext) {}
+	FStateTreeTransitionResult(const FStateTreeStateStatus InSource, const FStateTreeHandle InTransition, const FStateTreeHandle InNext) : Source(InSource), Target(InTransition), Next(InNext) {}
+
+	/** State where the transition started. */
+	UPROPERTY()
+	FStateTreeStateStatus Source = FStateTreeStateStatus();
+
+	/** Transition target state */
+	UPROPERTY()
+	FStateTreeHandle Target = FStateTreeHandle::Invalid;
+
+	/** Selected state, can be different from Transition, if Transition is a selector state. */
+	UPROPERTY()
+	FStateTreeHandle Next = FStateTreeHandle::Invalid;
+
+	/** Current state, update as we execute the tree. */
+	UPROPERTY()
+	FStateTreeHandle Current = FStateTreeHandle::Invalid;
+};
+
+
+/**
+ *  Runtime representation of a StateTree transition.
+ */
+USTRUCT()
+struct STATETREEMODULE_API FBakedStateTransition
+{
+	GENERATED_BODY()
+
+	UPROPERTY()
+	uint16 ConditionsBegin = 0;							// Index to first condition to test
+	UPROPERTY()
+	FStateTreeHandle State = FStateTreeHandle::Invalid;	// Target state of the transition
+	UPROPERTY()
+	EStateTreeTransitionType Type = EStateTreeTransitionType::NotSet;	// Type of the transition.
+	UPROPERTY()
+	EStateTreeTransitionEvent Event = EStateTreeTransitionEvent::None;	// Type of the transition event.
+	UPROPERTY()
+	uint8 GateDelay = 0;								// The time the conditions need to hold true for the transition to become active, in tenths of a seconds.
+	UPROPERTY()
+	uint8 ConditionsNum = 0;							// Number of conditions to test.
+};
+
+/**
+ *  Runtime representation of a StateTree state.
+ */
+USTRUCT()
+struct STATETREEMODULE_API FBakedStateTreeState
+{
+	GENERATED_BODY()
+
+	/** @return Index to the next sibling state. */
+	uint16 GetNextSibling() const { return ChildrenEnd; }
+
+	/** @return True if the state has any child states */
+	bool HasChildren() const { return ChildrenEnd > ChildrenBegin; }
+
+	UPROPERTY()
+	FName Name;											// Name of the State
+
+	UPROPERTY()
+	FStateTreeHandle Parent = FStateTreeHandle::Invalid;	// Parent state
+	UPROPERTY()
+	uint16 ChildrenBegin = 0;							// Index to first child state
+	UPROPERTY()
+	uint16 ChildrenEnd = 0;								// Index one past the last child state
+
+	UPROPERTY()
+	FStateTreeHandle StateDoneTransitionState = FStateTreeHandle::Invalid;		// State to transition to when the state execution is done. See also StateDoneTransitionType.
+	UPROPERTY()
+	FStateTreeHandle StateFailedTransitionState = FStateTreeHandle::Invalid;	// State to transition to if the state execution fails. See also StateFailedTransitionType.
+	
+	UPROPERTY()
+	uint16 EnterConditionsBegin = 0;					// Index to first state enter condition
+	UPROPERTY()
+	uint16 TransitionsBegin = 0;						// Index to first transition
+	UPROPERTY()
+	uint16 TasksBegin = 0;								// Index to first task
+	UPROPERTY()
+	uint16 EvaluatorsBegin = 0;							// Index to first evaluator
+
+	UPROPERTY()
+	uint8 EnterConditionsNum = 0;						// Number of enter conditions
+	UPROPERTY()
+	uint8 TransitionsNum = 0;							// Number of transitions
+	UPROPERTY()
+	uint8 TasksNum = 0;									// Number of tasks
+	UPROPERTY()
+	uint8 EvaluatorsNum = 0;							// Number of evaluators
+
+	UPROPERTY()
+	EStateTreeTransitionType StateDoneTransitionType = EStateTreeTransitionType::NotSet;		// Type of the State Done transition. See also StateDoneTransitionState.
+	UPROPERTY()
+	EStateTreeTransitionType StateFailedTransitionType = EStateTreeTransitionType::NotSet;		// Type of the State Failed transition. See also StateFailedTransitionState.
+};
+
+/** An offset into the StateTree runtime storage type to get a struct view to a specific Task or Evaluator. */
+struct FStateTreeRuntimeStorageItemOffset
+{
+	FStateTreeRuntimeStorageItemOffset() = default;
+	FStateTreeRuntimeStorageItemOffset(const UScriptStruct* InStruct, const int32 InOffset) : Struct(InStruct), Offset(InOffset) {}
+
+	/** Struct of the item */
+	const UScriptStruct* Struct = nullptr;
+	/** Offset within the storage struct */
+	int32 Offset = 0;
+};
+
+
+
+/**
+ * Handle to an external struct or object. 
+ */
+USTRUCT()
+struct STATETREEMODULE_API FStateTreeExternalItemHandle
+{
+	GENERATED_BODY()
+
+	FStateTreeExternalItemHandle() = default;
+	explicit FStateTreeExternalItemHandle(const uint16 InIndex) : ItemIndex(InIndex) {}
+	
+	static const uint16 InvalidIndex = uint16(-1);		// Index value indicating invalid item.
+
+	bool IsValid() const { return ItemIndex != InvalidIndex; }
+
+	void SetIndex(const uint16 Index) { ItemIndex = Index; }
+	uint16 GetIndex() const { return ItemIndex; }
+	
+protected:
+	UPROPERTY();
+	uint16 ItemIndex = InvalidIndex;
+};
+
+
+USTRUCT()
+struct STATETREEMODULE_API FStateTreeExternalItemDesc
+{
+	GENERATED_BODY()
+
+	FStateTreeExternalItemDesc() = default;
+	FStateTreeExternalItemDesc(const UStruct* InStruct, const bool bInOptional) : Struct(InStruct), bOptional(bInOptional) {}
+	
+	/** Class or struct of the external item. */
+	UPROPERTY();
+	const UStruct* Struct = nullptr;
+
+	/** Handle to the StateTreeExecutionContext item views array */
+	UPROPERTY();
+	FStateTreeExternalItemHandle Handle;
+
+	/** If true, the extern item is optional (can be null). */
+	UPROPERTY();
+	bool bOptional = false;
+};
+
+
+namespace FStateTreeVariableHelpers
+{
+	// Returns size of a variable type in bytes.
+	inline uint16 GetVariableMemoryUsage(EStateTreeVariableType Type)
+	{
+		static const uint64 Alignment = 4;
+		switch (Type)
+		{
+		case EStateTreeVariableType::Float:
+			return uint16(Align(sizeof(float), Alignment));
+		case EStateTreeVariableType::Int:
+			return uint16(Align(sizeof(int32), Alignment));
+		case EStateTreeVariableType::Bool:
+			return uint16(Align(sizeof(bool), Alignment));
+		case EStateTreeVariableType::Vector:
+			return uint16(Align(sizeof(FVector), Alignment));
+		case EStateTreeVariableType::Object:
+			return uint16(Align(sizeof(FWeakObjectPtr), Alignment));
+		default:
+			return 0;
+		}
+	}
+
+	// Get simple value based on pointer in byte array.
+	template<typename T>
+	static const T& GetValueFromMemory(const uint8* MemoryBlock)
+	{
+		return *((T*)MemoryBlock);
+	}
+
+	// Helper function for writing typed data to memory block, returns true if value has changed.
+	template<typename T>
+	static bool SetValueInMemory(uint8* MemoryBlock, const T& Value)
+	{
+		const bool bChanged = *((T*)MemoryBlock) != Value;
+		*((T*)MemoryBlock) = Value;
+
+		return bChanged;
+	}
+
+	// Helper function for writing weak object data to memory block, returns true if value has changed.
+	static bool SetWeakObjectInMemory(uint8* MemoryBlock, const FWeakObjectPtr& Value)
+	{
+		FWeakObjectPtr* PrevValue = (FWeakObjectPtr*)MemoryBlock;
+		const bool bChanged = !PrevValue->HasSameIndexAndSerialNumber(Value);
+		*((FWeakObjectPtr*)MemoryBlock) = Value;
+		return bChanged;
+	}
+};
\ No newline at end of file
diff --git a/Engine/Plugins/Runtime/StateTree/Source/StateTreeModule/Public/StateTreeVariable.h b/Engine/Plugins/Runtime/StateTree/Source/StateTreeModule/Public/StateTreeVariable.h
new file mode 100644
--- /dev/null
+++ b/Engine/Plugins/Runtime/StateTree/Source/StateTreeModule/Public/StateTreeVariable.h
@@ -0,0 +1,86 @@
+// Copyright Epic Games, Inc. All Rights Reserved.
+
+#pragma once
+
+#include "CoreMinimal.h"
+#include "Misc/Guid.h"
+#include "UObject/ObjectMacros.h"
+#include "Templates/SubclassOf.h"
+#include "StateTreeTypes.h"
+#include "StateTreeVariableDesc.h"
+#include "StateTreeVariable.generated.h"
+
+
+UENUM()
+enum class EStateTreeVariableBindingMode : uint8
+{
+	Any,		// The variable can be bound to any type variable, type and base class are copied from the bound variable.
+	Typed,		// The variable can be bound to only the type of variable described in type and base class.
+	Definition,	// The variable is used as defining a variable in state tree.
+};
+
+/**
+ * Describes a reference to a variable or constant in the StateTree.
+ * See StateTreeVariableDetails for the UI.
+ */
+USTRUCT()
+struct STATETREEMODULE_API FStateTreeVariable
+{
+	GENERATED_BODY()
+
+public:
+	FStateTreeVariable();
+	FStateTreeVariable(const EStateTreeVariableBindingMode InMode, const EStateTreeVariableType InType, TSubclassOf<UObject> InBaseClass = nullptr);
+	FStateTreeVariable(const EStateTreeVariableBindingMode InMode, const FName& InName, const EStateTreeVariableType InType, TSubclassOf<UObject> InBaseClass = nullptr);
+
+	// Initialize binding to match variable descriptor.
+	void InitializeFromDesc(EStateTreeVariableBindingMode InMode, const FStateTreeVariableDesc& Desc);
+
+#if WITH_EDITOR
+	bool IsBound() const { return ID.IsValid(); }
+
+	FStateTreeVariableDesc AsVariableDesc() const;
+
+	float GetFloatValue() const { return *reinterpret_cast<const float*>(Value); }
+	int32 GetIntValue() const { return *reinterpret_cast<const int32*>(Value); }
+	bool GetBoolValue() const { return *reinterpret_cast<const bool*>(Value); }
+	FVector GetVectorValue() const { return *reinterpret_cast<const FVector*>(Value); }
+	FWeakObjectPtr GetObjectValue() const { return nullptr; } // For completeness, we don't have a way to set an object value on asset.
+
+	FText GetDescription() const;
+
+	bool ResolveHandle(const struct FStateTreeVariableLayout& Variables, struct FStateTreeConstantStorage& Constants);
+
+	bool HasSameType(const FStateTreeVariable& Other) const;
+	bool CopyValueFrom(FStateTreeVariable& Other);
+#endif
+
+	UPROPERTY()
+	FStateTreeHandle Handle;
+
+	UPROPERTY(EditDefaultsOnly, Category = Value)
+	EStateTreeVariableType Type;
+
+#if WITH_EDITORONLY_DATA
+	UPROPERTY(EditDefaultsOnly, Category = Value, Transient)	// Do not serialize binding mode
+	EStateTreeVariableBindingMode BindingMode;
+
+	UPROPERTY(EditDefaultsOnly, Category = Value)
+	TSubclassOf<UObject> BaseClass;
+
+	UPROPERTY(EditDefaultsOnly, Category = Value)
+	FGuid ID;
+
+	UPROPERTY(EditDefaultsOnly, Category = Value)
+	FName Name;
+
+	static const int ValueStorageSizeBytes = 32;	// LWC_TODO: Perf pessimization. Was 16 - sizeof(FVector3f)
+
+	/* Common buffer to store all the variable types.
+	 * Note: Variable types must be POD type and fit into the 16 bytes.
+	 * Note: Adding new types require to update FStateTreeVariableDetails too.
+	 */
+	UPROPERTY(EditDefaultsOnly, Category = Value)
+	uint8 Value[ValueStorageSizeBytes];
+#endif
+};
diff --git a/Engine/Plugins/Runtime/StateTree/Source/StateTreeModule/Public/StateTreeVariableDesc.h b/Engine/Plugins/Runtime/StateTree/Source/StateTreeModule/Public/StateTreeVariableDesc.h
new file mode 100644
--- /dev/null
+++ b/Engine/Plugins/Runtime/StateTree/Source/StateTreeModule/Public/StateTreeVariableDesc.h
@@ -0,0 +1,72 @@
+// Copyright Epic Games, Inc. All Rights Reserved.
+
+#pragma once
+
+#include "CoreMinimal.h"
+#include "Templates/SubclassOf.h"
+#include "UObject/ObjectMacros.h"
+#if WITH_EDITORONLY_DATA
+	#include "Templates/SubclassOf.h"
+	#include "Misc/Guid.h"
+#endif
+#include "StateTreeTypes.h"
+#include "StateTreeVariableDesc.generated.h"
+
+/**
+ * Description of a StateTree variable.
+ * Contains name and type information. The editor data allows to track variable name changes,
+ * and additional type information allows more fine grained UI for Object variable type.
+*/
+USTRUCT(BlueprintType)
+struct STATETREEMODULE_API FStateTreeVariableDesc
+{
+	GENERATED_BODY()
+
+	FStateTreeVariableDesc()
+		: Name()
+		, Offset(0)
+		, Type(EStateTreeVariableType::Void)
+#if WITH_EDITORONLY_DATA
+		, ID(FGuid::NewGuid())
+		, BaseClass(nullptr)
+#endif
+	{}
+
+	FStateTreeVariableDesc(const FName& InName, const EStateTreeVariableType InType)
+		: Name(InName)
+		, Offset(0)
+		, Type(InType)
+#if WITH_EDITORONLY_DATA
+		, ID(FGuid::NewGuid())
+		, BaseClass(nullptr)
+#endif
+	{}
+
+	FStateTreeVariableDesc(const FName& InName, const EStateTreeVariableType InType, TSubclassOf<UObject> InBaseClass)
+		: Name(InName)
+		, Offset(0)
+		, Type(InType)
+#if WITH_EDITORONLY_DATA
+		, ID(FGuid::NewGuid())
+		, BaseClass(InBaseClass)
+#endif
+	{}
+
+	UPROPERTY(EditDefaultsOnly, Category = Value)
+	FName Name;
+
+	UPROPERTY(EditDefaultsOnly, Category = Value)
+	uint16 Offset;
+
+	UPROPERTY(EditDefaultsOnly, Category = Value)
+	EStateTreeVariableType Type;
+
+#if WITH_EDITORONLY_DATA
+	UPROPERTY(EditDefaultsOnly, Category = Value, meta = (IgnoreForMemberInitializationTest))
+	FGuid ID;
+
+	UPROPERTY(EditDefaultsOnly, Category = Value)
+	TSubclassOf<UObject> BaseClass;
+#endif
+};
+
diff --git a/Engine/Plugins/Runtime/StateTree/Source/StateTreeModule/Public/StateTreeVariableLayout.h b/Engine/Plugins/Runtime/StateTree/Source/StateTreeModule/Public/StateTreeVariableLayout.h
new file mode 100644
--- /dev/null
+++ b/Engine/Plugins/Runtime/StateTree/Source/StateTreeModule/Public/StateTreeVariableLayout.h
@@ -0,0 +1,91 @@
+// Copyright Epic Games, Inc. All Rights Reserved.
+
+#pragma once
+
+#include "CoreMinimal.h"
+#if WITH_EDITOR
+#include "Misc/Guid.h"
+#endif
+#include "StateTreeTypes.h"
+#include "StateTreeVariableDesc.h"
+#include "StateTreeVariable.h"
+#include "StateTreeVariableLayout.generated.h"
+
+/**
+ * Defines layout of variable descriptors.
+ * Can be used to define a variable layout in the UI too, see StateTreeVariableLayoutDetails.
+ */
+USTRUCT()
+struct STATETREEMODULE_API FStateTreeVariableLayout
+{
+	GENERATED_BODY()
+
+public:
+
+	uint16 GetMemoryUsage() const
+	{
+		if (Variables.Num() > 0)
+		{
+			return Variables.Last().Offset + FStateTreeVariableHelpers::GetVariableMemoryUsage(Variables.Last().Type);
+		}
+		return 0;
+	}
+
+	uint16 CalculateOffsetsAndMemoryUsage()
+	{
+		uint16 MemoryUsage = 0;
+
+		for (FStateTreeVariableDesc& Var : Variables)
+		{
+			Var.Offset = MemoryUsage;
+			MemoryUsage += FStateTreeVariableHelpers::GetVariableMemoryUsage(Var.Type);
+		}
+
+		return MemoryUsage;
+	}
+
+	void Reset()
+	{
+		Variables.Reset();
+	}
+
+#if WITH_EDITOR
+	FStateTreeHandle GetVariableHandle(const FGuid& VariableID) const
+	{
+		const FStateTreeVariableDesc* Var = Variables.FindByPredicate([&VariableID](const FStateTreeVariableDesc& Var) -> bool { return Var.ID == VariableID; });
+		return Var != nullptr ? FStateTreeHandle(Var->Offset) : FStateTreeHandle::Invalid;
+	}
+
+	FStateTreeVariableDesc& DefineVariable(const FStateTreeVariableDesc& Desc)
+	{
+		const FGuid ID = Desc.ID;
+		int32 Index = Variables.IndexOfByPredicate([ID](const FStateTreeVariableDesc& Desc) -> bool { return Desc.ID == ID; });
+		if (Index == INDEX_NONE)
+		{
+			Index = Variables.Add(Desc);
+		}
+		return Variables[Index];
+	}
+
+	FStateTreeVariableDesc& DefineVariable(const FName& ScopeName, const FStateTreeVariableDesc& Desc)
+	{
+		FStateTreeVariableDesc& NewVar = DefineVariable(Desc);
+		NewVar.Name = FName(ScopeName.ToString() + TEXT(".") + Desc.Name.ToString());
+		return NewVar;
+	}
+
+	FStateTreeVariableDesc& DefineVariable(const FStateTreeVariable& Var)
+	{
+		return DefineVariable(Var.AsVariableDesc());
+	}
+
+	FStateTreeVariableDesc& DefineVariable(const FName& ScopeName, const FStateTreeVariable& Var)
+	{
+		return DefineVariable(ScopeName, Var.AsVariableDesc());
+	}
+
+#endif
+
+	UPROPERTY(EditDefaultsOnly, Category = Layout)
+	TArray<FStateTreeVariableDesc> Variables;
+};
diff --git a/Engine/Plugins/Runtime/StateTree/Source/StateTreeModule/Public/Tasks/StateTreeTask_SubStateTree.h b/Engine/Plugins/Runtime/StateTree/Source/StateTreeModule/Public/Tasks/StateTreeTask_SubStateTree.h
new file mode 100644
--- /dev/null
+++ b/Engine/Plugins/Runtime/StateTree/Source/StateTreeModule/Public/Tasks/StateTreeTask_SubStateTree.h
@@ -0,0 +1,68 @@
+// Copyright Epic Games, Inc. All Rights Reserved.
+
+#pragma once
+
+#include "CoreMinimal.h"
+#include "Misc/Guid.h"
+#include "StateTreeTypes.h"
+#include "StateTreeInstance.h"
+#include "StateTreeVariable.h"
+#include "StateTreeParameterLayout.h"
+#include "StateTreeParameterStorage.h"
+#include "StateTreeTaskBase.h"
+#include "StateTreeTask_SubStateTree.generated.h"
+
+struct FStateTreeVariableLayout;
+struct FStateTreeConstantStorage;
+class UStateTree;
+
+/**
+ * Sub StateTree Task
+ * Runs specified StateTree while the state is active.
+ */
+
+UCLASS(BlueprintType, EditInlineNew, CollapseCategories)
+class STATETREEMODULE_API UStateTreeTask_SubStateTree : public UStateTreeTaskBase
+{
+	GENERATED_BODY()
+
+	public:
+	UStateTreeTask_SubStateTree(const FObjectInitializer& ObjectInitializer = FObjectInitializer::Get());
+
+	// Called when instantiated first time.
+	virtual bool Initialize(FStateTreeInstance& StateTreeInstance) override;
+	// Called when evaluator becomes active/ticking.
+	virtual void Activate(FStateTreeInstance& StateTreeInstance) override;
+	// Called when evaluator becomes inactive.
+	virtual void Deactivate(FStateTreeInstance& StateTreeInstance) override;
+
+	// Called on each tick.
+	virtual EStateTreeRunStatus Tick(FStateTreeInstance& StateTreeInstance, const float DeltaTime) override;
+
+#if WITH_GAMEPLAY_DEBUGGER
+	virtual void AppendDebugInfoString(FString& DebugString, const FStateTreeInstance& StateTreeInstance) const;
+#endif
+
+#if WITH_EDITOR
+	void PostEditChangeProperty(FPropertyChangedEvent& PropertyChangedEvent);
+
+	virtual bool ResolveVariables(const FStateTreeVariableLayout& Variables, FStateTreeConstantStorage& Constants, UObject* Outer) override;
+	virtual bool ValidateParameterLayout() override;
+#endif
+
+protected:
+
+	UPROPERTY(EditDefaultsOnly, Category = Asset)
+	UStateTree* StateTree;
+
+	UPROPERTY(EditDefaultsOnly, Category = Asset)
+	FStateTreeParameterLayout Parameters;
+
+	UPROPERTY(Transient)
+	FStateTreeParameterStorage ParameterStorage;
+
+	UPROPERTY(Transient)
+	FStateTreeInstance SubStateTreeInstance;
+};
+
+
diff --git a/Engine/Plugins/Runtime/StateTree/Source/StateTreeModule/Public/Tasks/StateTreeTask_Test.h b/Engine/Plugins/Runtime/StateTree/Source/StateTreeModule/Public/Tasks/StateTreeTask_Test.h
new file mode 100644
--- /dev/null
+++ b/Engine/Plugins/Runtime/StateTree/Source/StateTreeModule/Public/Tasks/StateTreeTask_Test.h
@@ -0,0 +1,76 @@
+// Copyright Epic Games, Inc. All Rights Reserved.
+
+#pragma once
+
+#include "CoreMinimal.h"
+#include "StateTreeTypes.h"
+#include "StateTreeTaskBase.h"
+#include "StateTreeTask_Test.generated.h"
+
+
+UCLASS(BlueprintType, EditInlineNew, CollapseCategories)
+class STATETREEMODULE_API UStateTreeTask_Test : public UStateTreeTaskBase
+{
+	GENERATED_BODY()
+
+protected:
+	UStateTreeTask_Test(const FObjectInitializer& ObjectInitializer);
+
+	/** Called when instantiated first time. */
+	virtual bool Initialize(FStateTreeInstance& StateTreeInstance) override;
+	/** Called when evaluator becomes active/ticking. */
+	virtual void Activate(FStateTreeInstance& StateTreeInstance)  override;
+	/** Called when evaluator becomes inactive. */
+	virtual void Deactivate(FStateTreeInstance& StateTreeInstance)  override;
+
+	/** Called on each tick. */
+	virtual EStateTreeRunStatus Tick(FStateTreeInstance& StateTreeInstance, const float DeltaTime) override;
+
+#if WITH_GAMEPLAY_DEBUGGER
+	virtual void AppendDebugInfoString(FString& DebugString, const FStateTreeInstance& StateTreeInstance) const;
+#endif
+
+#if WITH_EDITOR
+	/** Create Task instance template, set default variables, and resolve variable handles. */
+	virtual bool ResolveVariables(const FStateTreeVariableLayout& Variables, FStateTreeConstantStorage& Constants, UObject* Outer) override;
+#endif
+
+	float Time = 0.0f;
+
+	UPROPERTY(EditAnywhere, Category = Task)
+	float Duration = 0.0f;
+
+	UPROPERTY(EditAnywhere, Category = Task)
+	bool bFail = false;
+};
+
+USTRUCT(meta = (DisplayName = "Task Test2"))
+struct STATETREEMODULE_API FStateTreeTask2_Test : public FStateTreeTask2Base
+{
+	GENERATED_BODY()
+
+	FStateTreeTask2_Test();
+	virtual ~FStateTreeTask2_Test();
+	
+	virtual EStateTreeRunStatus Tick(FStateTreeExecutionContext& Context, const float DeltaTime) override;
+
+#if WITH_GAMEPLAY_DEBUGGER
+	virtual void AppendDebugInfoString(FString& DebugString, const FStateTreeExecutionContext& Context) const;
+#endif
+
+protected:
+	UPROPERTY(EditAnywhere, Category = Task, meta = (Bindable))
+	float Value = 0.0f;
+
+	UPROPERTY(EditAnywhere, Category = Task, meta = (Bindable))
+	FVector VectorValue = FVector::ZeroVector;
+
+	UPROPERTY(EditAnywhere, Category = Task, meta = (Bindable))
+	int32 IntVal = 0;
+
+	UPROPERTY(EditAnywhere, Category = Task)
+	int32 IntVal2 = 0;
+
+	UPROPERTY(EditAnywhere, Category = Task, meta = (Bindable))
+	bool bBoolVal = false;
+};
diff --git a/Engine/Plugins/Runtime/StateTree/Source/StateTreeModule/StateTreeModule.Build.cs b/Engine/Plugins/Runtime/StateTree/Source/StateTreeModule/StateTreeModule.Build.cs
new file mode 100644
--- /dev/null
+++ b/Engine/Plugins/Runtime/StateTree/Source/StateTreeModule/StateTreeModule.Build.cs
@@ -0,0 +1,37 @@
+// Copyright Epic Games, Inc. All Rights Reserved.
+
+namespace UnrealBuildTool.Rules
+{
+	public class StateTreeModule : ModuleRules
+	{
+		public StateTreeModule(ReadOnlyTargetRules Target) : base(Target)
+		{
+			PublicIncludePaths.AddRange(
+			new string[] {
+			}
+			);
+
+			PublicDependencyModuleNames.AddRange(
+			new string[] {
+				"Core",
+				"CoreUObject",
+				"Engine",
+				"AIModule",
+				"GameplayBehaviorsModule",
+				"GameplayTasks",
+				"GameplayTags",
+				"GameplayAbilities",
+				"SmartObjectsModule",
+				"StructUtils",
+			}
+			);
+
+			PrivateDependencyModuleNames.AddRange(
+			new string[] {
+				"RenderCore",
+			}
+			);
+		}
+
+	}
+}
diff --git a/Engine/Plugins/Runtime/StateTree/Source/StateTreeTestSuite/Private/StateTreeTest.cpp b/Engine/Plugins/Runtime/StateTree/Source/StateTreeTestSuite/Private/StateTreeTest.cpp
new file mode 100644
--- /dev/null
+++ b/Engine/Plugins/Runtime/StateTree/Source/StateTreeTestSuite/Private/StateTreeTest.cpp
@@ -0,0 +1,339 @@
+// Copyright Epic Games, Inc. All Rights Reserved.
+
+#include "StateTreeTest.h"
+#include "CoreMinimal.h"
+#include "AITestsCommon.h"
+#include "Misc/App.h"
+#include "StateTreeEditorData.h"
+#include "StateTreeState.h"
+#include "StateTreeBaker.h"
+#include "Conditions/StateTreeCondition_Common.h"
+#include "StateTreeEvaluatorBase.h"
+#include "StateTreeTaskBase.h"
+#include "StateTree.h"
+#include "StateTreeTestTypes.h"
+#include "StateTreeExecutionContext.h"
+
+#define LOCTEXT_NAMESPACE "AITestSuite_StateTreeTest"
+
+PRAGMA_DISABLE_OPTIMIZATION
+
+namespace UE::StateTree::Tests
+{
+	UStateTree& NewStateTree()
+	{
+		UStateTree* StateTree = NewObject<UStateTree>();
+		check(StateTree);
+		UStateTreeEditorData* EditorData = NewObject<UStateTreeEditorData>(StateTree);
+		check(EditorData);
+		StateTree->EditorData = EditorData;
+		UStateTreeSchema* Schema = NewObject<UStateTreeTestSchema>();
+		StateTree->SetSchema(Schema);
+		return *StateTree;
+	}
+}
+
+struct FStateTreeTest_MakeAndBakeStateTree : FAITestBase
+{
+	virtual bool InstantTest() override
+	{
+		UStateTree& StateTree = UE::StateTree::Tests::NewStateTree();
+		UStateTreeEditorData& EditorData = *Cast<UStateTreeEditorData>(StateTree.EditorData);
+		
+		UStateTreeState& Root = EditorData.AddSubTree(FName(TEXT("Root")));
+		UStateTreeState& StateA = Root.AddChildState(FName(TEXT("A")));
+		UStateTreeState& StateB = Root.AddChildState(FName(TEXT("B")));
+
+		// Root
+		FTestEval_A& EvalA = Root.AddEvaluator<FTestEval_A>();
+		Root.StateDoneTransition.Set(EStateTreeTransitionType::SelectChildState);
+
+		// State A
+		FTestTask_B& TaskB1 = StateA.AddTask<FTestTask_B>();
+		EditorData.AddPropertyBinding(STATETREE_PROPPATH_CHECKED(EvalA, IntA), STATETREE_PROPPATH_CHECKED(TaskB1, IntB));
+
+		FStateTreeCondition_CompareInt& IntCond = StateA.AddEnterCondition<FStateTreeCondition_CompareInt>(0, EGenericAICheck::Less, 2);
+		EditorData.AddPropertyBinding(STATETREE_PROPPATH_CHECKED(EvalA, IntA), STATETREE_PROPPATH_CHECKED(IntCond, Left));
+
+		StateA.StateDoneTransition.Set(EStateTreeTransitionType::GotoState, &StateB);
+
+		// State B
+		FTestTask_B& TaskB2 = StateB.AddTask<FTestTask_B>();
+		EditorData.AddPropertyBinding(STATETREE_PROPPATH_CHECKED(EvalA, bBoolA), STATETREE_PROPPATH_CHECKED(TaskB2, bBoolB));
+
+		FStateTreeTransition2& Trans = StateB.AddTransition(EStateTreeTransitionEvent::OnCondition, EStateTreeTransitionType::GotoState, &Root);
+		FStateTreeCondition_CompareFloat& TransFloatCond = Trans.AddCondition<FStateTreeCondition_CompareFloat>(0.0f, EGenericAICheck::Less, 13.0f);
+		EditorData.AddPropertyBinding(STATETREE_PROPPATH_CHECKED(EvalA, FloatA), STATETREE_PROPPATH_CHECKED(TransFloatCond, Left));
+
+		StateB.StateDoneTransition.Set(EStateTreeTransitionType::Succeeded);
+
+		FStateTreeBaker Baker;
+		bool bResult = Baker.Bake2(StateTree);
+
+		AITEST_TRUE("StateTree should get baked", bResult);
+
+		return true;
+	}
+};
+IMPLEMENT_AI_INSTANT_TEST(FStateTreeTest_MakeAndBakeStateTree, "System.AI.StateTree.MakeAndBakeStateTree");
+
+
+struct FStateTreeTest_WanderLoop : FAITestBase
+{
+	virtual bool InstantTest() override
+	{
+		UStateTree& StateTree = UE::StateTree::Tests::NewStateTree();
+		UStateTreeEditorData& EditorData = *Cast<UStateTreeEditorData>(StateTree.EditorData);
+		
+		UStateTreeState& Root = EditorData.AddSubTree(FName(TEXT("Root")));
+		UStateTreeState& Wander = Root.AddChildState(FName(TEXT("Wander")));
+		
+		UStateTreeState& UseSmartObjectOnLane = Wander.AddChildState(FName(TEXT("UseSmartObjectOnLane")));
+		UStateTreeState& WalkAlongLane = Wander.AddChildState(FName(TEXT("WalkAlongLane")));
+		UStateTreeState& StandOnLane = Wander.AddChildState(FName(TEXT("StandOnLane")));
+		
+		UStateTreeState& WalkToSO = UseSmartObjectOnLane.AddChildState(FName(TEXT("WalkToSO")));
+		UStateTreeState& UseSO = UseSmartObjectOnLane.AddChildState(FName(TEXT("UseSO")));
+
+		// - Root
+		Root.StateDoneTransition.Set(EStateTreeTransitionType::SelectChildState);
+
+		//   \- Wander
+		FTestEval_Wander& WanderEval = Wander.AddEvaluator<FTestEval_Wander>(FName(TEXT("WanderEval")));
+		FTestEval_SmartObjectSensor& SmartObjectEval = Wander.AddEvaluator<FTestEval_SmartObjectSensor>();
+
+		//      |- UseSmartObjectOnLane
+		FTestTask_ReserveSmartObject& ReserveSOTask = UseSmartObjectOnLane.AddTask<FTestTask_ReserveSmartObject>(FName(TEXT("ReserveSOTask")));
+		EditorData.AddPropertyBinding(STATETREE_PROPPATH_CHECKED(SmartObjectEval, PotentialSmartObjects), STATETREE_PROPPATH_CHECKED(ReserveSOTask, PotentialSmartObjects));
+		FStateTreeCondition_CompareBool& ReserveHasSmartObjects = UseSmartObjectOnLane.AddEnterCondition<FStateTreeCondition_CompareBool>(false, true);
+		EditorData.AddPropertyBinding(STATETREE_PROPPATH_CHECKED(SmartObjectEval, bHasSmartObjects), STATETREE_PROPPATH_CHECKED(ReserveHasSmartObjects, bLeft));
+		UseSmartObjectOnLane.AddTransition(EStateTreeTransitionEvent::OnCompleted, EStateTreeTransitionType::GotoState, &Wander); // This catches child states too
+
+		//      |  |- WalkToSO
+		FTestTask_MoveTo& MoveToSOTask = WalkToSO.AddTask<FTestTask_MoveTo>(FName(TEXT("MoveToSOTask")));
+		MoveToSOTask.TicksToCompletion = 2;
+		EditorData.AddPropertyBinding(STATETREE_PROPPATH_CHECKED(ReserveSOTask, ReservedSmartObjectLocation), STATETREE_PROPPATH_CHECKED(MoveToSOTask, MoveLocation));
+		WalkToSO.AddTransition(EStateTreeTransitionEvent::OnSucceeded, EStateTreeTransitionType::NextState);
+
+		//      |  \- UseSO
+		FTestTask_UseSmartObject& UseSOTask = UseSO.AddTask<FTestTask_UseSmartObject>(FName(TEXT("UseSOTask")));
+		UseSOTask.TicksToCompletion = 2;
+		EditorData.AddPropertyBinding(STATETREE_PROPPATH_CHECKED(ReserveSOTask, ReservedSmartObject), STATETREE_PROPPATH_CHECKED(UseSOTask, SmartObject));
+		// UseSO uses UseSmartObjectOnLane completed transition.
+
+		//      |- Walk Along Lane
+		FTestTask_MoveTo& MoveAlongLaneTask = WalkAlongLane.AddTask<FTestTask_MoveTo>(FName(TEXT("MoveAlongLaneTask")));
+		MoveAlongLaneTask.TicksToCompletion = 2;
+		EditorData.AddPropertyBinding(STATETREE_PROPPATH_CHECKED(WanderEval, WanderLocation), STATETREE_PROPPATH_CHECKED(MoveAlongLaneTask, MoveLocation));
+		FStateTreeCondition_CompareBool& MoveHasWanderLoc = WalkAlongLane.AddEnterCondition<FStateTreeCondition_CompareBool>(false, true);
+		EditorData.AddPropertyBinding(STATETREE_PROPPATH_CHECKED(WanderEval, bHasWanderLocation), STATETREE_PROPPATH_CHECKED(MoveHasWanderLoc, bLeft));
+		WalkAlongLane.AddTransition(EStateTreeTransitionEvent::OnCompleted, EStateTreeTransitionType::GotoState, &Wander);
+
+		//      \- StandOnLane
+		FTestTask_Stand& StandTask = StandOnLane.AddTask<FTestTask_Stand>(FName(TEXT("StandTask")));
+		StandTask.TicksToCompletion = 2;
+		StandOnLane.AddTransition(EStateTreeTransitionEvent::OnCompleted, EStateTreeTransitionType::GotoState, &Wander);
+
+		FStateTreeBaker Baker;
+		bool bResult = Baker.Bake2(StateTree);
+
+		AITEST_TRUE("StateTree should get baked", bResult);
+
+		EStateTreeRunStatus Status = EStateTreeRunStatus::Unset;
+		FTestStateTreeExecutionContext Exec;
+		const bool bInitSucceeded = Exec.Init(StateTree, StateTree, EStateTreeStorage::Internal);
+		AITEST_TRUE("StateTree should init", bInitSucceeded);
+
+		const FString TickStr(TEXT("Tick"));
+		const FString EnterStateStr(TEXT("EnterState"));
+		const FString ExitStateStr(TEXT("ExitState"));
+		const FString ReserveFailedStr(TEXT("ReserveFailed"));
+		
+		Exec.Start();
+
+		// No SOs on first tick, we should stand.
+		Status = Exec.Tick(0.1f);
+		AITEST_TRUE("StateTree StandTask should tick", Exec.Expect(StandTask.Name, TickStr));
+		Exec.LogClear();
+
+		Status = Exec.Tick(0.1f);
+		AITEST_TRUE("StateTree StandTask should tick", Exec.Expect(StandTask.Name, TickStr));
+		Exec.LogClear();
+
+		// Stand completed, SO sensor should have found SOs, Expect Move to SO
+		Status = Exec.Tick(0.1f);
+		AITEST_TRUE("StateTree MoveToSOTask should enter state, and tick", Exec.Expect(MoveToSOTask.Name, EnterStateStr).Then(MoveToSOTask.Name, TickStr));
+		Exec.LogClear();
+
+		Status = Exec.Tick(0.1f);
+		AITEST_TRUE("StateTree MoveToSOTask should tick", Exec.Expect(MoveToSOTask.Name, TickStr));
+		Exec.LogClear();
+
+		// Use SO
+		Status = Exec.Tick(0.1f);
+		AITEST_TRUE("StateTree UseSOTask should enter state, and tick", Exec.Expect(UseSOTask.Name, EnterStateStr).Then(UseSOTask.Name, TickStr));
+		Exec.LogClear();
+
+		Status = Exec.Tick(0.1f);
+		AITEST_TRUE("StateTree UseSOTask should tick", Exec.Expect(UseSOTask.Name, TickStr));
+		Exec.LogClear();
+
+		// SO done, should select Wander>UseSmartObjectOnLane>WalkToSO.
+		// The next SO on the sensor is invalid, we try to reserve invalid SO, fail, then fallback to select Stand.
+		Status = Exec.Tick(0.1f);
+		AITEST_TRUE("StateTree ReserveSOTask should fail enter state", Exec.Expect(ReserveSOTask.Name, EnterStateStr).Then(ReserveSOTask.Name, ReserveFailedStr));
+		AITEST_TRUE("StateTree StandTask should tick", Exec.Expect(StandTask.Name, TickStr));
+		Exec.LogClear();
+
+		Status = Exec.Tick(0.1f);
+		AITEST_TRUE("StateTree StandTask should tick", Exec.Expect(StandTask.Name, TickStr));
+		Exec.LogClear();
+
+		// We should end up in stand again because of the failed cool down on SO sensor.
+		Status = Exec.Tick(0.1f);
+		AITEST_TRUE("StateTree StandTask should tick", Exec.Expect(StandTask.Name, TickStr));
+		Exec.LogClear();
+
+		Status = Exec.Tick(0.1f);
+		AITEST_TRUE("StateTree StandTask should tick", Exec.Expect(StandTask.Name, TickStr));
+		Exec.LogClear();
+		
+		// After stand is done we should have new SO ready to use.
+		Status = Exec.Tick(0.1f);
+		AITEST_TRUE("StateTree MoveToSOTask should enter state, and tick", Exec.Expect(MoveToSOTask.Name, EnterStateStr).Then(MoveToSOTask.Name, TickStr));
+		Exec.LogClear();
+
+		return true;
+	}
+};
+IMPLEMENT_AI_INSTANT_TEST(FStateTreeTest_WanderLoop, "System.AI.StateTree.WanderLoop");
+
+
+struct FStateTreeTest_Sequence : FAITestBase
+{
+	virtual bool InstantTest() override
+	{
+		UStateTree& StateTree = UE::StateTree::Tests::NewStateTree();
+		UStateTreeEditorData& EditorData = *Cast<UStateTreeEditorData>(StateTree.EditorData);
+		
+		UStateTreeState& Root = EditorData.AddSubTree(FName(TEXT("Root")));
+		UStateTreeState& State1 = Root.AddChildState(FName(TEXT("State1")));
+		UStateTreeState& State2 = Root.AddChildState(FName(TEXT("State2")));
+
+		FTestTask_Stand& Task1 = State1.AddTask<FTestTask_Stand>(FName(TEXT("Task1")));
+		State1.AddTransition(EStateTreeTransitionEvent::OnCompleted, EStateTreeTransitionType::NextState);
+
+		FTestTask_Stand& Task2 = State2.AddTask<FTestTask_Stand>(FName(TEXT("Task2")));
+		State2.AddTransition(EStateTreeTransitionEvent::OnCompleted, EStateTreeTransitionType::Succeeded);
+
+		FStateTreeBaker Baker;
+		const bool bResult = Baker.Bake2(StateTree);
+		AITEST_TRUE("StateTree should get baked", bResult);
+
+		EStateTreeRunStatus Status = EStateTreeRunStatus::Unset;
+		FTestStateTreeExecutionContext Exec;
+		const bool bInitSucceeded = Exec.Init(StateTree, StateTree, EStateTreeStorage::Internal);
+		AITEST_TRUE("StateTree should init", bInitSucceeded);
+
+		const FString TickStr(TEXT("Tick"));
+		const FString EnterStateStr(TEXT("EnterState"));
+		const FString ExitStateStr(TEXT("ExitState"));
+		
+		Exec.Start();
+
+		Status = Exec.Tick(0.1f);
+		AITEST_TRUE("StateTree Task1 should enter state and tick", Exec.Expect(Task1.Name, EnterStateStr).Then(Task1.Name, TickStr));
+		AITEST_FALSE("StateTree Task2 should not tick", Exec.Expect(Task2.Name, TickStr));
+		AITEST_TRUE("StateTree should be running", Status == EStateTreeRunStatus::Running);
+		Exec.LogClear();
+		
+		Status = Exec.Tick(0.1f);
+		AITEST_TRUE("StateTree Task1 should exit state", Exec.Expect(Task1.Name, ExitStateStr));
+		AITEST_FALSE("StateTree Task1 should not tick", Exec.Expect(Task1.Name, TickStr));
+        AITEST_TRUE("StateTree Task2 should enter state and tick", Exec.Expect(Task2.Name, EnterStateStr).Then(Task2.Name, TickStr));
+        AITEST_TRUE("StateTree should be running", Status == EStateTreeRunStatus::Running);
+        Exec.LogClear();
+		
+		Status = Exec.Tick(0.1f);
+        AITEST_TRUE("StateTree Task2 should exit state", Exec.Expect(Task2.Name, ExitStateStr));
+		AITEST_FALSE("StateTree Task1 should not tick", Exec.Expect(Task1.Name, TickStr));
+        AITEST_FALSE("StateTree Task2 should not tick", Exec.Expect(Task2.Name, TickStr));
+        AITEST_TRUE("StateTree should be completed", Status == EStateTreeRunStatus::Succeeded);
+
+		return true;
+	}
+};
+IMPLEMENT_AI_INSTANT_TEST(FStateTreeTest_Sequence, "System.AI.StateTree.Sequence");
+
+struct FStateTreeTest_Select : FAITestBase
+{
+	virtual bool InstantTest() override
+	{
+		UStateTree& StateTree = UE::StateTree::Tests::NewStateTree();
+		UStateTreeEditorData& EditorData = *Cast<UStateTreeEditorData>(StateTree.EditorData);
+		
+		UStateTreeState& Root = EditorData.AddSubTree(FName(TEXT("Root")));
+		UStateTreeState& State1 = Root.AddChildState(FName(TEXT("State1")));
+		UStateTreeState& State1A = State1.AddChildState(FName(TEXT("State1A")));
+
+		FTestTask_Stand& TaskRoot = Root.AddTask<FTestTask_Stand>(FName(TEXT("TaskRoot")));
+		TaskRoot.TicksToCompletion = 2;
+
+		FTestTask_Stand& Task1 = State1.AddTask<FTestTask_Stand>(FName(TEXT("Task1")));
+		Task1.TicksToCompletion = 2;
+
+		FTestTask_Stand& Task1A = State1A.AddTask<FTestTask_Stand>(FName(TEXT("Task1A")));
+		Task1A.TicksToCompletion = 2;
+		State1A.AddTransition(EStateTreeTransitionEvent::OnCompleted, EStateTreeTransitionType::GotoState, &State1);
+
+		FStateTreeBaker Baker;
+		const bool bResult = Baker.Bake2(StateTree);
+		AITEST_TRUE("StateTree should get baked", bResult);
+
+		EStateTreeRunStatus Status = EStateTreeRunStatus::Unset;
+		FTestStateTreeExecutionContext Exec;
+		const bool bInitSucceeded = Exec.Init(StateTree, StateTree, EStateTreeStorage::Internal);
+		AITEST_TRUE("StateTree should init", bInitSucceeded);
+
+		const FString TickStr(TEXT("Tick"));
+		const FString EnterStateStr(TEXT("EnterState"));
+		const FString ExitStateStr(TEXT("ExitState"));
+
+		Exec.Start();
+		
+		Status = Exec.Tick(0.1f);
+		AITEST_TRUE("StateTree TaskRoot should enter state, and tick", Exec.Expect(TaskRoot.Name, EnterStateStr).Then(TaskRoot.Name, TickStr));
+		AITEST_TRUE("StateTree Task1 should enter state, and tick", Exec.Expect(Task1.Name, EnterStateStr).Then(Task1.Name, TickStr));
+		AITEST_TRUE("StateTree Task1A should enter state, and tick", Exec.Expect(Task1A.Name, EnterStateStr).Then(Task1A.Name, TickStr));
+		AITEST_TRUE("StateTree should be running", Status == EStateTreeRunStatus::Running);
+		Exec.LogClear();
+
+		// Regular tick
+		Status = Exec.Tick(0.1f);
+		AITEST_TRUE("StateTree tasks should update in order", Exec.Expect(TaskRoot.Name, TickStr).Then(Task1.Name, TickStr).Then(Task1A.Name, TickStr));
+		AITEST_FALSE("StateTree TaskRoot should not EnterState", Exec.Expect(TaskRoot.Name, EnterStateStr));
+		AITEST_FALSE("StateTree Task1 should not EnterState", Exec.Expect(Task1.Name, EnterStateStr));
+		AITEST_FALSE("StateTree Task1A should not EnterState", Exec.Expect(Task1A.Name, EnterStateStr));
+		AITEST_FALSE("StateTree TaskRoot should not ExitState", Exec.Expect(TaskRoot.Name, ExitStateStr));
+		AITEST_FALSE("StateTree Task1 should not ExitState", Exec.Expect(Task1.Name, ExitStateStr));
+		AITEST_FALSE("StateTree Task1A should not ExitState", Exec.Expect(Task1A.Name, ExitStateStr));
+		AITEST_TRUE("StateTree should be running", Status == EStateTreeRunStatus::Running);
+		Exec.LogClear();
+
+		// Partial reselect, Root should not get EnterState
+		Status = Exec.Tick(0.1f);
+		AITEST_FALSE("StateTree TaskRoot should not enter state", Exec.Expect(TaskRoot.Name, EnterStateStr));
+		AITEST_TRUE("StateTree Task1 should exit state, and enter state", Exec.Expect(Task1A.Name, ExitStateStr).Then(Task1.Name, EnterStateStr));
+		AITEST_TRUE("StateTree Task1A should exit state, enter state, and tick", Exec.Expect(Task1A.Name, ExitStateStr).Then(Task1A.Name, EnterStateStr).Then(Task1A.Name, TickStr));
+		AITEST_TRUE("StateTree should be running", Status == EStateTreeRunStatus::Running);
+        Exec.LogClear();
+
+		return true;
+	}
+};
+IMPLEMENT_AI_INSTANT_TEST(FStateTreeTest_Select, "System.AI.StateTree.Select");
+
+
+PRAGMA_ENABLE_OPTIMIZATION
+
+#undef LOCTEXT_NAMESPACE
diff --git a/Engine/Plugins/Runtime/StateTree/Source/StateTreeTestSuite/Private/StateTreeTest.h b/Engine/Plugins/Runtime/StateTree/Source/StateTreeTestSuite/Private/StateTreeTest.h
new file mode 100644
--- /dev/null
+++ b/Engine/Plugins/Runtime/StateTree/Source/StateTreeTestSuite/Private/StateTreeTest.h
@@ -0,0 +1,18 @@
+// Copyright Epic Games, Inc. All Rights Reserved.
+
+#pragma once
+
+#include "StateTreeSchema.h"
+
+#include "StateTreeTest.generated.h"
+
+UCLASS()
+class UStateTreeTestSchema : public UStateTreeSchema
+{
+	GENERATED_BODY()
+
+public:
+	/** @return True if we should use StateTree V2 */
+	virtual bool IsV2() const { return true; }
+};
+
diff --git a/Engine/Plugins/Runtime/StateTree/Source/StateTreeTestSuite/Private/StateTreeTestTypes.h b/Engine/Plugins/Runtime/StateTree/Source/StateTreeTestSuite/Private/StateTreeTestTypes.h
new file mode 100644
--- /dev/null
+++ b/Engine/Plugins/Runtime/StateTree/Source/StateTreeTestSuite/Private/StateTreeTestTypes.h
@@ -0,0 +1,685 @@
+// Copyright Epic Games, Inc. All Rights Reserved.
+
+#pragma once
+
+#include "AITestsCommon.h"
+#include "StateTreeTaskBase.h"
+#include "StateTreeEvaluatorBase.h"
+#include "StateTreeExecutionContext.h"
+#include "StateTreeTestTypes.generated.h"
+
+/**
+ * This macro can be used to create property path for property binding. Checks that the member exists.
+ * Example:
+ *		EditorData.AddPropertyBinding(STATETREE_PROPPATH_CHECKED(EvalA, IntA), STATETREE_PROPPATH_CHECKED(TaskB1, IntB));
+ **/
+#define STATETREE_PROPPATH_CHECKED(Struct, MemberName) \
+FStateTreeEditorPropertyPath(Struct.ID, ((void)sizeof(UEAsserts_Private::GetMemberNameCheckedJunk((&Struct)->MemberName)), TEXT(#MemberName)))
+
+
+
+USTRUCT()
+struct FTestStateTreeExecutionContext : public FStateTreeExecutionContext
+{
+	GENERATED_BODY()
+
+	struct FLogItem
+	{
+		FLogItem() = default;
+		FLogItem(const FName& InName, const FString& InMessage) : Name(InName), Message(InMessage) {}
+		FName Name;
+		FString Message; 
+	};
+	TArray<FLogItem> LogItems;
+	
+	void Log(const FName& Name, const FString& Message)
+	{
+		LogItems.Emplace(Name, Message);
+	}
+
+	void LogClear()
+	{
+		LogItems.Empty();
+	}
+
+	struct FLogOrder
+	{
+		FLogOrder(const FTestStateTreeExecutionContext& InContext, const int32 InIndex) : Context(InContext), Index(InIndex) {}
+
+		FLogOrder Then(const FName& Name, const FString& Message) const
+		{
+			int32 NextIndex = Index;
+			while (NextIndex < Context.LogItems.Num())
+			{
+				const FLogItem& Item = Context.LogItems[NextIndex];
+				if (Item.Name == Name && Item.Message == Message)
+				{
+					break;
+				}
+				NextIndex++;
+			}
+			return FLogOrder(Context, NextIndex);
+		}
+
+		operator bool() const { return Index < Context.LogItems.Num(); }
+		
+		const FTestStateTreeExecutionContext& Context;
+		int32 Index = 0;
+	};
+
+	FLogOrder Expect(const FName& Name, const FString& Message) const
+	{
+		return FLogOrder(*this, 0).Then(Name, Message);
+	}
+	
+};
+
+
+USTRUCT()
+struct FTestEval_A : public FStateTreeEvaluator2Base
+{
+	GENERATED_BODY()
+
+	FTestEval_A() = default;
+	virtual ~FTestEval_A() {}
+
+	UPROPERTY(EditAnywhere, Category = Test, meta = (Bindable))
+	float FloatA = 0.0f;
+
+	UPROPERTY(EditAnywhere, Category = Test, meta = (Bindable))
+	int32 IntA = 0;
+
+	UPROPERTY(EditAnywhere, Category = Test, meta = (Bindable))
+	bool bBoolA = false;
+};
+
+USTRUCT()
+struct FTestTask_B : public FStateTreeTask2Base
+{
+	GENERATED_BODY()
+
+	FTestTask_B() = default;
+	virtual ~FTestTask_B() {}
+	
+	UPROPERTY(EditAnywhere, Category = Task, meta = (Bindable))
+	float FloatB = 0.0f;
+
+	UPROPERTY(EditAnywhere, Category = Task, meta = (Bindable))
+	int32 IntB = 0;
+
+	UPROPERTY(EditAnywhere, Category = Task, meta = (Bindable))
+	bool bBoolB = false;
+};
+
+
+
+
+USTRUCT()
+struct FTestMoveLocationResult : public FStateTreeResult
+{
+	GENERATED_BODY()
+
+	virtual ~FTestMoveLocationResult() {}
+	virtual const UScriptStruct& GetStruct() const override { return *StaticStruct(); }
+
+	void Reset()
+	{
+		Location = FVector::ZeroVector;
+		Direction = FVector::ZeroVector;
+		Status = EStateTreeResultStatus::Unset;
+	}
+	
+	UPROPERTY(EditAnywhere, Category = Value)
+	FVector Location = FVector::ZeroVector;
+
+	UPROPERTY(EditAnywhere, Category = Value)
+	FVector Direction = FVector::ZeroVector;
+
+	EStateTreeResultStatus Status = EStateTreeResultStatus::Unset;
+};
+
+USTRUCT()
+struct FTestPotentialSmartObjectsResult : public FStateTreeResult
+{
+	GENERATED_BODY()
+
+	virtual ~FTestPotentialSmartObjectsResult() {}
+	virtual const UScriptStruct& GetStruct() const override { return *StaticStruct(); }
+
+	bool HasSmartObjects() const { return NumSmartObjects > 0 && FailedCoolDown == 0; }
+	
+	FName PopSmartObjectName()
+	{
+		FName Result;
+		if (NumSmartObjects > 0)
+		{
+			Result = SmartObjectNames[NumSmartObjects - 1];
+			NumSmartObjects--;
+		}
+		return Result;
+	}
+
+	void Reset()
+	{
+		NumSmartObjects = 0;
+	}
+	
+	void SetFailedCoolDown()
+	{
+		FailedCoolDown = 4;
+	}
+
+	UPROPERTY(EditAnywhere, Category = Value)
+	FName SmartObjectNames[4];
+
+	UPROPERTY(EditAnywhere, Category = Value)
+	int32 NumSmartObjects = 0;
+
+	int32 FailedCoolDown = 0;
+};
+
+USTRUCT()
+struct FTestSmartObjectResult : public FStateTreeResult
+{
+	GENERATED_BODY()
+
+	virtual ~FTestSmartObjectResult() {}
+	virtual const UScriptStruct& GetStruct() const override { return *StaticStruct(); }
+
+	void Reset()
+	{
+		AnimationName = FName();
+		SmartObjectHandle = INDEX_NONE;
+		Status = EStateTreeResultStatus::Unset;
+	}
+	
+	UPROPERTY(EditAnywhere, Category = Value)
+	FName AnimationName;
+
+	UPROPERTY(EditAnywhere, Category = Value)
+	int32 SmartObjectHandle = INDEX_NONE;
+
+	EStateTreeResultStatus Status = EStateTreeResultStatus::Unset;
+};
+
+
+USTRUCT()
+struct FTestEval_Wander : public FStateTreeEvaluator2Base
+{
+	GENERATED_BODY()
+
+	FTestEval_Wander() = default;
+	FTestEval_Wander(const FName InName) { Name = InName; }
+	virtual ~FTestEval_Wander() {}
+
+	virtual void EnterState(FStateTreeExecutionContext& Context, const EStateTreeStateChangeType ChangeType, const FStateTreeTransitionResult& Transition) override
+	{
+		FTestStateTreeExecutionContext& TestContext = static_cast<FTestStateTreeExecutionContext&>(Context);
+		TestContext.Log(Name, TEXT("EnterState"));
+
+		WanderLocation = &NextMoveLocation;
+	}
+
+	virtual void ExitState(FStateTreeExecutionContext& Context, const EStateTreeStateChangeType ChangeType, const FStateTreeTransitionResult& Transition) override
+	{
+		FTestStateTreeExecutionContext& TestContext = static_cast<FTestStateTreeExecutionContext&>(Context);
+		TestContext.Log(Name, TEXT("ExitState"));
+		
+		WanderLocation = &NextMoveLocation;
+	}
+
+	virtual void Evaluate(FStateTreeExecutionContext& Context, const EStateTreeEvaluationType EvalType, const float DeltaTime) override
+	{
+		FTestStateTreeExecutionContext& TestContext = static_cast<FTestStateTreeExecutionContext&>(Context);
+		TestContext.Log(Name, TEXT("Evaluate"));
+
+		WanderLocation = &NextMoveLocation;
+	}
+
+	// True if wander location available
+	UPROPERTY(EditAnywhere, Category = Eval)
+	bool bHasWanderLocation = false;
+
+	UPROPERTY(EditAnywhere, Category = Eval)
+	FStateTreeResultRef WanderLocation;
+	
+	FTestMoveLocationResult NextMoveLocation;
+};
+
+USTRUCT()
+struct FTestEval_SmartObjectSensor : public FStateTreeEvaluator2Base
+{
+	GENERATED_BODY()
+
+	FTestEval_SmartObjectSensor() = default;
+	FTestEval_SmartObjectSensor(const FName InName) { Name = InName; }
+	virtual ~FTestEval_SmartObjectSensor() {}
+
+	virtual void EnterState(FStateTreeExecutionContext& Context, const EStateTreeStateChangeType ChangeType, const FStateTreeTransitionResult& Transition) override
+	{
+		FTestStateTreeExecutionContext& TestContext = static_cast<FTestStateTreeExecutionContext&>(Context);
+		TestContext.Log(Name, TEXT("EnterState"));
+
+		// We expect the state of the evaluator to be valid on EnterState.
+		// Evaluate() have been called during the state selection process, and some properties may have been used as conditions during select.
+		
+		PotentialSmartObjects = &SmartObjects;
+	}
+	
+	virtual void ExitState(FStateTreeExecutionContext& Context, const EStateTreeStateChangeType ChangeType, const FStateTreeTransitionResult& Transition) override
+	{
+		FTestStateTreeExecutionContext& TestContext = static_cast<FTestStateTreeExecutionContext&>(Context);
+		TestContext.Log(Name, TEXT("ExitState"));
+
+		// Cleanup when we stop being used, this way the eval will be in known state when we will be ticked again later
+		// during state selection.
+		if (ChangeType == EStateTreeStateChangeType::Changed)
+		{
+			bHasSmartObjects = false;
+			bIsQueryingSmartObjects = false; // Cancel a query
+			SmartObjects.Reset(); // Note: this does not reset the cool down.
+		}
+
+		PotentialSmartObjects = &SmartObjects;
+	}
+
+	virtual void Evaluate(FStateTreeExecutionContext& Context, const EStateTreeEvaluationType EvalType, const float DeltaTime) override
+	{
+		FTestStateTreeExecutionContext& TestContext = static_cast<FTestStateTreeExecutionContext&>(Context);
+		TestContext.Log(Name, TEXT("Evaluate"));
+
+		// The cool down should be a timestamp instead of a counter, to catch a case where we re-enter the evaluator after being in different state.
+		if (SmartObjects.FailedCoolDown > 0)
+		{
+			SmartObjects.FailedCoolDown--;
+		}
+		bHasSmartObjects = SmartObjects.HasSmartObjects();
+		
+		if (bIsQueryingSmartObjects)
+		{
+			// Simulate Async query result
+			bIsQueryingSmartObjects = false;
+			bHasSmartObjects = true;
+			SmartObjects.NumSmartObjects = 2;
+			SmartObjects.SmartObjectNames[0] = FName();
+			SmartObjects.SmartObjectNames[1] = FName(TEXT("Foo"));
+		}
+		if (!bHasSmartObjects && !bIsQueryingSmartObjects && SmartObjects.FailedCoolDown == 0)
+		{
+			// Simulate Async query start
+			bIsQueryingSmartObjects = true;
+		}
+
+		// Set ref
+		PotentialSmartObjects = &SmartObjects;
+	}
+
+	// Result output
+	UPROPERTY(EditAnywhere, Category = Eval, meta=(Output, Struct="TestPotentialSmartObjectsResult"))
+	FStateTreeResultRef PotentialSmartObjects;
+
+	// True if any SmartObjects available
+	UPROPERTY(EditAnywhere, Category = Eval)
+	bool bHasSmartObjects = false;
+
+	bool bIsQueryingSmartObjects = false;
+	FTestPotentialSmartObjectsResult SmartObjects;
+};
+
+
+USTRUCT()
+struct FTestTask_ReserveSmartObject : public FStateTreeTask2Base
+{
+	GENERATED_BODY()
+
+	FTestTask_ReserveSmartObject() = default;
+	FTestTask_ReserveSmartObject(const FName InName) { Name = InName; }
+	virtual ~FTestTask_ReserveSmartObject() {}
+
+	static int32 AcquireSmartObject(const FName Name)
+	{
+		return Name.IsNone() ? INDEX_NONE : 1;
+	}
+
+	static void ReleaseSmartObject(const int32 Handle)
+	{
+		// empty
+	}
+	
+	virtual EStateTreeRunStatus EnterState(FStateTreeExecutionContext& Context, const EStateTreeStateChangeType ChangeType, const FStateTreeTransitionResult& Transition) override
+	{
+		FTestStateTreeExecutionContext& TestContext = static_cast<FTestStateTreeExecutionContext&>(Context);
+		TestContext.Log(Name, TEXT("EnterState"));
+
+		FTestPotentialSmartObjectsResult* PotentialSmartObjectsPtr = PotentialSmartObjects.GetMutablePtr<FTestPotentialSmartObjectsResult>();
+		if (PotentialSmartObjectsPtr == nullptr)
+		{
+			return EStateTreeRunStatus::Failed;
+		}
+
+		// Reset outputs
+		SmartObject.Reset();
+		SmartObjectLocation.Reset();
+
+		// This consumes the sensor inputs
+		bool bFoundSO = false;
+		while (PotentialSmartObjectsPtr->HasSmartObjects())
+		{
+			const FName SOName = PotentialSmartObjectsPtr->PopSmartObjectName();
+			const int32 NewHandle = AcquireSmartObject(SOName);
+			if (NewHandle != INDEX_NONE)
+			{
+				// Acquire and expose one SmartObject and the location
+				SmartObject.AnimationName = SOName;
+				SmartObject.SmartObjectHandle = NewHandle;
+				SmartObject.Status = EStateTreeResultStatus::Unset; // Not available just yet
+
+				SmartObjectLocation.Location = FVector(10,10,0);
+				SmartObjectLocation.Direction = FVector::ForwardVector;
+				SmartObjectLocation.Status = EStateTreeResultStatus::Available;
+
+				bFoundSO = true;
+				break;
+			}
+		}
+
+		if (!bFoundSO)
+		{
+			TestContext.Log(Name, TEXT("ReserveFailed"));
+			PotentialSmartObjectsPtr->SetFailedCoolDown();
+			return EStateTreeRunStatus::Failed;
+		}
+
+		ReservedSmartObject = &SmartObject;
+		ReservedSmartObjectLocation = &SmartObjectLocation;
+
+		return EStateTreeRunStatus::Running;
+	}
+
+	virtual void ExitState(FStateTreeExecutionContext& Context, const EStateTreeStateChangeType ChangeType, const FStateTreeTransitionResult& Transition) override
+	{
+		FTestStateTreeExecutionContext& TestContext = static_cast<FTestStateTreeExecutionContext&>(Context);
+		TestContext.Log(Name, TEXT("ExitState"));
+		
+		if (SmartObjectLocation.Status != EStateTreeResultStatus::Succeeded
+			|| SmartObject.Status != EStateTreeResultStatus::Succeeded)
+		{
+			// Did not fully complete the SO
+		}
+
+		// Clean up
+		if (SmartObject.SmartObjectHandle != INDEX_NONE)
+		{
+			ReleaseSmartObject(SmartObject.SmartObjectHandle);
+			SmartObject.SmartObjectHandle = INDEX_NONE;
+		}
+		
+		ReservedSmartObject = &SmartObject;
+		ReservedSmartObjectLocation = &SmartObjectLocation;
+	}
+
+	virtual void StateCompleted(FStateTreeExecutionContext& Context, const EStateTreeRunStatus CompletionStatus, const FStateTreeHandle StateCompleted) override
+	{
+		// Make SmartObject available only if successfully moved to it.
+		if (SmartObjectLocation.Status == EStateTreeResultStatus::Succeeded)
+		{
+			SmartObject.Status = EStateTreeResultStatus::Available;
+			SmartObjectLocation.Status = EStateTreeResultStatus::Unset;
+		}
+
+		if (SmartObjectLocation.Status == EStateTreeResultStatus::Failed
+			|| SmartObject.Status == EStateTreeResultStatus::Failed)
+		{
+			// Failed to use the SmartObject, set cool down here to prevent SO being used in next State Select.
+			// The Cooldown probably needs to be 
+			FTestPotentialSmartObjectsResult* PotentialSmartObjectsPtr = PotentialSmartObjects.GetMutablePtr<FTestPotentialSmartObjectsResult>();
+			if (PotentialSmartObjectsPtr)
+			{
+				PotentialSmartObjectsPtr->SetFailedCoolDown();
+			}
+		}
+	}
+
+	virtual EStateTreeRunStatus Tick(FStateTreeExecutionContext& Context, const float DeltaTime) override
+	{
+		FTestStateTreeExecutionContext& TestContext = static_cast<FTestStateTreeExecutionContext&>(Context);
+		TestContext.Log(Name, TEXT("Tick"));
+
+		// Need to update on each update, i.e. fragment may have moved.
+		ReservedSmartObject = &SmartObject;
+		ReservedSmartObjectLocation = &SmartObjectLocation;
+
+		return EStateTreeRunStatus::Running;
+	}
+
+	
+	UPROPERTY(EditAnywhere, Category = Eval, meta = (Bindable, Struct="TestPotentialSmartObjectsResult"))  // In
+	FStateTreeResultRef PotentialSmartObjects;
+
+	// Because reserve is a task, these can be seen only by other Task (not evals, not conditions)
+	UPROPERTY(EditAnywhere, Category = Eval, meta = (Output, Struct="TestMoveLocationResult")) // Out
+	FStateTreeResultRef ReservedSmartObjectLocation;
+
+	UPROPERTY(EditAnywhere, Category = Eval, meta = (Output, Struct="TestSmartObjectResult")) // Out
+	FStateTreeResultRef ReservedSmartObject;
+
+	FTestMoveLocationResult SmartObjectLocation;
+	FTestSmartObjectResult SmartObject;
+};
+
+
+USTRUCT()
+struct FTestTask_MoveTo : public FStateTreeTask2Base
+{
+	GENERATED_BODY()
+
+	FTestTask_MoveTo() = default;
+	FTestTask_MoveTo(const FName InName) { Name = InName; }
+	virtual ~FTestTask_MoveTo() {}
+
+	virtual EStateTreeRunStatus EnterState(FStateTreeExecutionContext& Context, const EStateTreeStateChangeType ChangeType, const FStateTreeTransitionResult& Transition) override
+	{
+		FTestStateTreeExecutionContext& TestContext = static_cast<FTestStateTreeExecutionContext&>(Context);
+		TestContext.Log(Name, TEXT("EnterState"));
+
+		// Move restarts only when state changes for good, this allows to have child states running too.
+		FTestMoveLocationResult* MoveLocationPtr = MoveLocation.GetMutablePtr<FTestMoveLocationResult>();
+		if (MoveLocationPtr == nullptr)
+		{
+			return EStateTreeRunStatus::Failed;
+		}
+		if (MoveLocationPtr->Status != EStateTreeResultStatus::Available)
+		{
+			MoveLocationPtr->Status = EStateTreeResultStatus::Failed; // This probably should be a function we call, to allow prioritize the return values.
+			return EStateTreeRunStatus::Failed;
+		}
+		MoveLocationPtr->Status = EStateTreeResultStatus::InUse;
+
+		CurrentTick = 0;
+
+		return EnterStateResult;
+	}
+
+
+	virtual void ExitState(FStateTreeExecutionContext& Context, const EStateTreeStateChangeType ChangeType, const FStateTreeTransitionResult& Transition) override
+	{
+		FTestStateTreeExecutionContext& TestContext = static_cast<FTestStateTreeExecutionContext&>(Context);
+		TestContext.Log(Name, TEXT("ExitState"));
+
+		FTestMoveLocationResult* MoveLocationPtr = MoveLocation.GetMutablePtr<FTestMoveLocationResult>();
+		if (MoveLocationPtr != nullptr)
+		{
+			if (MoveLocationPtr->Status == EStateTreeResultStatus::InUse)
+			{
+				// We got interrupted
+				MoveLocationPtr->Status = EStateTreeResultStatus::Failed;
+			}
+		}
+	}
+
+	virtual EStateTreeRunStatus Tick(FStateTreeExecutionContext& Context, const float DeltaTime) override
+	{
+		FTestStateTreeExecutionContext& TestContext = static_cast<FTestStateTreeExecutionContext&>(Context);
+		TestContext.Log(Name, TEXT("Tick"));
+
+		FTestMoveLocationResult* MoveLocationPtr = MoveLocation.GetMutablePtr<FTestMoveLocationResult>();
+		if (MoveLocationPtr == nullptr)
+		{
+			return EStateTreeRunStatus::Failed;
+		}
+
+		CurrentTick++;
+		const bool bDone = CurrentTick >= TicksToCompletion;
+		if (bDone)
+		{
+			MoveLocationPtr->Status = TickResult == EStateTreeRunStatus::Succeeded ? EStateTreeResultStatus::Succeeded : EStateTreeResultStatus::Failed;
+		}
+		
+		return bDone ? TickResult : EStateTreeRunStatus::Running;
+	}
+
+	int32 CurrentTick = 0;
+
+	UPROPERTY(EditAnywhere, Category = Eval, meta = (Bindable))
+	FStateTreeResultRef MoveLocation;
+
+	UPROPERTY(EditAnywhere, Category = Task, meta = (Bindable))
+	int32 TicksToCompletion = 2;
+
+	UPROPERTY(EditAnywhere, Category = Task, meta = (Bindable))
+	EStateTreeRunStatus TickResult = EStateTreeRunStatus::Succeeded;
+
+	UPROPERTY(EditAnywhere, Category = Task, meta = (Bindable))
+	EStateTreeRunStatus EnterStateResult = EStateTreeRunStatus::Running;
+};
+
+USTRUCT()
+struct FTestTask_Stand : public FStateTreeTask2Base
+{
+	GENERATED_BODY()
+
+	FTestTask_Stand() = default;
+	FTestTask_Stand(const FName InName) { Name = InName; }
+	virtual ~FTestTask_Stand() {}
+
+	virtual EStateTreeRunStatus EnterState(FStateTreeExecutionContext& Context, const EStateTreeStateChangeType ChangeType, const FStateTreeTransitionResult& Transition) override
+	{
+		FTestStateTreeExecutionContext& TestContext = static_cast<FTestStateTreeExecutionContext&>(Context);
+		TestContext.Log(Name, TEXT("EnterState"));
+		
+		if (ChangeType == EStateTreeStateChangeType::Changed)
+		{
+			CurrentTick = 0;
+		}
+		return EnterStateResult;
+	}
+
+	virtual void ExitState(FStateTreeExecutionContext& Context, const EStateTreeStateChangeType ChangeType, const FStateTreeTransitionResult& Transition) override
+	{
+		FTestStateTreeExecutionContext& TestContext = static_cast<FTestStateTreeExecutionContext&>(Context);
+		TestContext.Log(Name, TEXT("ExitState"));
+	}
+
+	virtual EStateTreeRunStatus Tick(FStateTreeExecutionContext& Context, const float DeltaTime) override
+	{
+		FTestStateTreeExecutionContext& TestContext = static_cast<FTestStateTreeExecutionContext&>(Context);
+		TestContext.Log(Name, TEXT("Tick"));
+		
+		CurrentTick++;
+		return (CurrentTick >= TicksToCompletion) ? TickResult : EStateTreeRunStatus::Running;
+	};
+
+	int32 CurrentTick = 0;
+	
+	UPROPERTY(EditAnywhere, Category = Task, meta = (Bindable))
+	int32 TicksToCompletion = 1;
+
+	UPROPERTY(EditAnywhere, Category = Task, meta = (Bindable))
+	EStateTreeRunStatus TickResult = EStateTreeRunStatus::Succeeded;
+
+	UPROPERTY(EditAnywhere, Category = Task, meta = (Bindable))
+	EStateTreeRunStatus EnterStateResult = EStateTreeRunStatus::Running;
+};
+
+USTRUCT()
+struct FTestTask_UseSmartObject : public FStateTreeTask2Base
+{
+	GENERATED_BODY()
+
+	FTestTask_UseSmartObject() = default;
+	FTestTask_UseSmartObject(const FName InName) { Name = InName; }
+	virtual ~FTestTask_UseSmartObject() {}
+
+	virtual EStateTreeRunStatus EnterState(FStateTreeExecutionContext& Context, const EStateTreeStateChangeType ChangeType, const FStateTreeTransitionResult& Transition) override
+	{
+		FTestStateTreeExecutionContext& TestContext = static_cast<FTestStateTreeExecutionContext&>(Context);
+		TestContext.Log(Name, TEXT("EnterState"));
+		
+		FTestSmartObjectResult* SmartObjectPtr = SmartObject.GetMutablePtr<FTestSmartObjectResult>();
+		if (SmartObjectPtr == nullptr)
+		{
+			return EStateTreeRunStatus::Failed;
+		}
+		if (SmartObjectPtr->Status != EStateTreeResultStatus::Available)
+		{
+			SmartObjectPtr->Status = EStateTreeResultStatus::Failed;
+			return EStateTreeRunStatus::Failed;
+		}
+		CurrentTick = 0;
+		SmartObjectPtr->Status = EStateTreeResultStatus::InUse;
+
+		return EnterStateResult;
+	}
+
+	virtual void ExitState(FStateTreeExecutionContext& Context, const EStateTreeStateChangeType ChangeType, const FStateTreeTransitionResult& Transition) override
+	{
+		FTestStateTreeExecutionContext& TestContext = static_cast<FTestStateTreeExecutionContext&>(Context);
+		TestContext.Log(Name, TEXT("ExitState"));
+
+		FTestSmartObjectResult* SmartObjectPtr = SmartObject.GetMutablePtr<FTestSmartObjectResult>();
+		if (SmartObjectPtr != nullptr)
+		{
+			if (SmartObjectPtr->Status != EStateTreeResultStatus::Succeeded)
+			{
+				// We got interrupted
+				SmartObjectPtr->Status = EStateTreeResultStatus::Failed;
+			}
+		}
+	}
+
+	virtual EStateTreeRunStatus Tick(FStateTreeExecutionContext& Context, const float DeltaTime) override
+	{
+		FTestStateTreeExecutionContext& TestContext = static_cast<FTestStateTreeExecutionContext&>(Context);
+		TestContext.Log(Name, TEXT("Tick"));
+
+		FTestSmartObjectResult* SmartObjectPtr = SmartObject.GetMutablePtr<FTestSmartObjectResult>();
+		if (SmartObjectPtr == nullptr)
+		{
+			return EStateTreeRunStatus::Failed;
+		}
+		
+		CurrentTick++;
+		
+		const bool bDone = CurrentTick >= TicksToCompletion;
+		if (bDone)
+		{
+			SmartObjectPtr->Status = TickResult == EStateTreeRunStatus::Succeeded ? EStateTreeResultStatus::Succeeded : EStateTreeResultStatus::Failed;
+		}
+		
+		return bDone ? TickResult : EStateTreeRunStatus::Running;
+	}
+
+	int32 CurrentTick = 0;
+
+	UPROPERTY(EditAnywhere, Category = Eval, meta = (Bindable, Struct="TestSmartObjectResult")) // In
+	FStateTreeResultRef SmartObject;
+
+	UPROPERTY(EditAnywhere, Category = Task)
+	int32 TicksToCompletion = 2;
+
+	UPROPERTY(EditAnywhere, Category = Task)
+	EStateTreeRunStatus TickResult = EStateTreeRunStatus::Succeeded;
+
+	UPROPERTY(EditAnywhere, Category = Task)
+	EStateTreeRunStatus EnterStateResult = EStateTreeRunStatus::Running;
+};
\ No newline at end of file
diff --git a/Engine/Plugins/Runtime/ZoneGraph/Config/BaseZoneGraph.ini b/Engine/Plugins/Runtime/ZoneGraph/Config/BaseZoneGraph.ini
new file mode 100644
--- /dev/null
+++ b/Engine/Plugins/Runtime/ZoneGraph/Config/BaseZoneGraph.ini
@@ -0,0 +1,17 @@
+[CoreRedirects]
+
++EnumRedirects=(OldName="FZoneShapePointType" ValueChanges=(("FZoneShapePointType::LaneSegment","FZoneShapePointType::LaneProfile")))
+
++StructRedirects=(OldName="ZoneLaneTemplate",NewName="ZoneLaneProfile")
++StructRedirects=(OldName="ZoneLaneTemplateRef",NewName="ZoneLaneProfileRef")
+
++PropertyRedirects=(OldName="ZoneGraphSettings.LaneTemplates",NewName="LaneProfiles")
++PropertyRedirects=(OldName="ZoneShapeComponent.LaneTemplate",NewName="LaneProfile")
++PropertyRedirects=(OldName="ZoneShapeComponent.PerPointLaneTemplates",NewName="PerPointLaneProfiles")
++PropertyRedirects=(OldName="ZoneShapePoint.LaneTemplate",NewName="LaneProfile")
+
++PropertyRedirects=(OldName="ZoneShapePoint.InControlPoint",NewName="InControlPoint_DEPRECATED")
++PropertyRedirects=(OldName="ZoneShapePoint.OutControlPoint",NewName="OutControlPoint_DEPRECATED")
+
++PropertyRedirects=(OldName="ZoneGraphTagFilter.IncludeTags",NewName="AnyTags")
++PropertyRedirects=(OldName="ZoneGraphTagFilter.ExcludeTags",NewName="NotTags")
\ No newline at end of file
diff --git a/Engine/Plugins/Runtime/ZoneGraph/Source/ZoneGraph/Public/BezierUtilities.h b/Engine/Plugins/Runtime/ZoneGraph/Source/ZoneGraph/Public/BezierUtilities.h
new file mode 100644
--- /dev/null
+++ b/Engine/Plugins/Runtime/ZoneGraph/Source/ZoneGraph/Public/BezierUtilities.h
@@ -0,0 +1,57 @@
+// Copyright Epic Games, Inc. All Rights Reserved.
+
+#pragma once
+
+#include "Math/Vector.h"
+#include "Math/Box.h"
+#include "Math/UnrealMathUtility.h"
+
+namespace UE
+{
+namespace CubicBezier
+{
+
+/** Evaluates cubic bezier value at specified time */
+inline FVector Eval(const FVector& P0, const FVector& P1, const FVector& P2, const FVector& P3, const float t)
+{
+	// Interpolate using De Casteljau algorithm.
+	const FVector P01 = FMath::Lerp(P0, P1, t);
+	const FVector P12 = FMath::Lerp(P1, P2, t);
+	const FVector P23 = FMath::Lerp(P2, P3, t);
+	const FVector P012 = FMath::Lerp(P01, P12, t);
+	const FVector P123 = FMath::Lerp(P12, P23, t);
+	return FMath::Lerp(P012, P123, t);
+}
+
+/** Evaluates cubic bezier derivative at specified time */
+inline FVector EvalDerivate(const FVector& P0, const FVector& P1, const FVector& P2, const FVector& P3, const float t)
+{
+	const FVector P01 = FMath::Lerp(P0, P1, t);
+	const FVector P12 = FMath::Lerp(P1, P2, t);
+	const FVector P23 = FMath::Lerp(P2, P3, t);
+	return 3.0f * (FMath::Lerp(P12, P23, t) - FMath::Lerp(P01, P12, t));
+}
+
+/** Splits cubic bezier curve at specified time, returns new points.
+  *	The new curves are (P0, Result[0], Result[1], Result[2]) and (P0, Result[2], Result[3], Result[4], P3) */
+ZONEGRAPH_API void SplitAt(const FVector& P0, const FVector& P1, const FVector& P2, const FVector& P3, const float t, FVector OutResult[5]);
+
+/** Returns bounding box of cubic bezier curve. */
+ZONEGRAPH_API FBox CalcBounds(const FVector& P0, const FVector& P1, const FVector& P2, const FVector& P3);
+
+/** Returns closest point to cubic bezier segment. The curve is split to up to Steps linear segments for approximate calculation. */
+ZONEGRAPH_API void ClosestPointApproximate(const FVector& FromPoint, const FVector& P0, const FVector& P1, const FVector& P2, const FVector& P3, FVector& OutClosestPoint, float& OutClosestT, const int Steps = 16);
+
+/** Returns closest point between a linear segment and cubic bezier segment. The curve is split to up to Steps linear segments for approximate calculation. */
+ZONEGRAPH_API void SegmentClosestPointApproximate(const FVector& SegStart, const FVector& SegEnd, const FVector& P0, const FVector& P1, const FVector& P2, const FVector& P3, FVector& OutClosestPoint, float& OutClosestT, const int Steps = 16);
+
+/** Returns approximate arc length of cubic bezier. The length is calculated based on control point hull with one subdivision. */
+ZONEGRAPH_API float ArcLengthApproximate(const FVector& P0, const FVector& P1, const FVector& P2, const FVector& P3);
+
+/** Recursively tessellate a bezier curve until the segments are within Tolerace from the curve, or until max recursion level is reached.
+  * Output will contain points describing linear segments of the flattened curve.
+  * Note that the points are appended, and that the first point is omitted so that the function can be used to flatten a bezier shape. */
+ZONEGRAPH_API void Tessellate(TArray<FVector>& Output, const FVector& P0, const FVector& P1, const FVector& P2, const FVector& P3, const float Tolerance, const int MaxLevel = 6);
+
+} // namespace CubicBezier
+} // namespace UE
