StateTree: Remove V1

diff --git a/Engine/Plugins/AI/MassAI/Source/MassAIBehavior/Public/MassStateTreeSchema.h b/Engine/Plugins/AI/MassAI/Source/MassAIBehavior/Public/MassStateTreeSchema.h
--- a/Engine/Plugins/AI/MassAI/Source/MassAIBehavior/Public/MassStateTreeSchema.h
+++ b/Engine/Plugins/AI/MassAI/Source/MassAIBehavior/Public/MassStateTreeSchema.h
@@ -18,6 +18,5 @@ protected:
 	virtual UScriptStruct* GetStorageSuperStruct() const override;
 	virtual bool IsStructAllowed(const UScriptStruct* InScriptStruct) const override;
 	virtual bool IsExternalItemAllowed(const UStruct& InStruct) const override;
-	virtual bool IsV2() const override { return true; }
 };
 
diff --git a/Engine/Plugins/Runtime/MassEntity/Config/BaseMassEntity.ini b/Engine/Plugins/Runtime/MassEntity/Config/BaseMassEntity.ini
--- a/Engine/Plugins/Runtime/MassEntity/Config/BaseMassEntity.ini
+++ b/Engine/Plugins/Runtime/MassEntity/Config/BaseMassEntity.ini
@@ -1,3 +1,4 @@
 [CoreRedirects]
 +PackageRedirects=(OldName="/Script/AggregateTicking",NewName="/Script/MassEntity")
 +PackageRedirects=(OldName="/Script/Pipe",NewName="/Script/MassEntity")
++StructRedirects=(OldName="LWEntity",NewName="MassEntityHandle")
diff --git a/Engine/Plugins/Runtime/StateTree/Config/BaseStateTree.ini b/Engine/Plugins/Runtime/StateTree/Config/BaseStateTree.ini
--- a/Engine/Plugins/Runtime/StateTree/Config/BaseStateTree.ini
+++ b/Engine/Plugins/Runtime/StateTree/Config/BaseStateTree.ini
@@ -5,5 +5,14 @@
 +PropertyRedirects=(OldName="BakedStateTreeState.DefaultTransitionType",NewName="StateDoneTransitionType")
 +PropertyRedirects=(OldName="StateTreeBindableStructDesc.ScriptStruct",NewName="StateTreeBindableStructDesc.Struct")
 
++StructRedirects=(OldName="StateTreeTransition2",NewName="StateTreeTransition")
+
++PropertyRedirects=(OldName="StateTreeState.Transitions2",NewName="StateTreeState.Transitions")
++PropertyRedirects=(OldName="StateTreeState.EnterConditions2",NewName="StateTreeState.EnterConditions")
++PropertyRedirects=(OldName="StateTreeState.Evaluators2",NewName="StateTreeState.Evaluators")
++PropertyRedirects=(OldName="StateTreeState.Tasks2",NewName="StateTreeState.Tasks")
+
++PropertyRedirects=(OldName="StateTree.Conditions2",NewName="StateTree.Conditions")
+
 +EnumRedirects=(OldName="EStateTreePropertyCopyType",ValueChanges=(("Plain","CopyPlain"),("Complex","CopyComplex"),("Bool","CopyBool"),("Struct","CopyStruct"),("Object","CopyObject"),("Name","CopyName"),("FixedArray","CopyFixedArray")))
 +EnumRedirects=(OldName="EStateTreePropertyAccessType",ValueChanges=(("Array","IndexArray")))
diff --git a/Engine/Plugins/Runtime/StateTree/Source/StateTreeEditorModule/Private/Customizations/StateTreeConditionDetails.cpp b/Engine/Plugins/Runtime/StateTree/Source/StateTreeEditorModule/Private/Customizations/StateTreeConditionDetails.cpp
deleted file mode 100644
diff --git a/Engine/Plugins/Runtime/StateTree/Source/StateTreeEditorModule/Private/Customizations/StateTreeConditionDetails.h b/Engine/Plugins/Runtime/StateTree/Source/StateTreeEditorModule/Private/Customizations/StateTreeConditionDetails.h
deleted file mode 100644
diff --git a/Engine/Plugins/Runtime/StateTree/Source/StateTreeEditorModule/Private/Customizations/StateTreeEditorPropertyPathDetails.cpp b/Engine/Plugins/Runtime/StateTree/Source/StateTreeEditorModule/Private/Customizations/StateTreeEditorPropertyPathDetails.cpp
deleted file mode 100644
diff --git a/Engine/Plugins/Runtime/StateTree/Source/StateTreeEditorModule/Private/Customizations/StateTreeEditorPropertyPathDetails.h b/Engine/Plugins/Runtime/StateTree/Source/StateTreeEditorModule/Private/Customizations/StateTreeEditorPropertyPathDetails.h
deleted file mode 100644
diff --git a/Engine/Plugins/Runtime/StateTree/Source/StateTreeEditorModule/Private/Customizations/StateTreeParameterLayoutDetails.cpp b/Engine/Plugins/Runtime/StateTree/Source/StateTreeEditorModule/Private/Customizations/StateTreeParameterLayoutDetails.cpp
deleted file mode 100644
diff --git a/Engine/Plugins/Runtime/StateTree/Source/StateTreeEditorModule/Private/Customizations/StateTreeParameterLayoutDetails.h b/Engine/Plugins/Runtime/StateTree/Source/StateTreeEditorModule/Private/Customizations/StateTreeParameterLayoutDetails.h
deleted file mode 100644
diff --git a/Engine/Plugins/Runtime/StateTree/Source/StateTreeEditorModule/Private/Customizations/StateTreeStateDetails.cpp b/Engine/Plugins/Runtime/StateTree/Source/StateTreeEditorModule/Private/Customizations/StateTreeStateDetails.cpp
--- a/Engine/Plugins/Runtime/StateTree/Source/StateTreeEditorModule/Private/Customizations/StateTreeStateDetails.cpp
+++ b/Engine/Plugins/Runtime/StateTree/Source/StateTreeEditorModule/Private/Customizations/StateTreeStateDetails.cpp
@@ -45,40 +45,9 @@ TSharedRef<IDetailCustomization> FStateTreeStateDetails::MakeInstance()
 void FStateTreeStateDetails::CustomizeDetails(IDetailLayoutBuilder& DetailBuilder)
 {
 	TSharedPtr<IPropertyHandle> TasksProperty = DetailBuilder.GetProperty(TEXT("Tasks"));
-	TSharedPtr<IPropertyHandle> Tasks2Property = DetailBuilder.GetProperty(TEXT("Tasks2"));
 	TSharedPtr<IPropertyHandle> EnterConditionsProperty = DetailBuilder.GetProperty(TEXT("EnterConditions"));
-	TSharedPtr<IPropertyHandle> EnterConditions2Property = DetailBuilder.GetProperty(TEXT("EnterConditions2"));
 	TSharedPtr<IPropertyHandle> EvaluatorsProperty = DetailBuilder.GetProperty(TEXT("Evaluators"));
-	TSharedPtr<IPropertyHandle> Evaluators2Property = DetailBuilder.GetProperty(TEXT("Evaluators2"));
 	TSharedPtr<IPropertyHandle> TransitionsProperty = DetailBuilder.GetProperty(TEXT("Transitions"));
-	TSharedPtr<IPropertyHandle> Transitions2Property = DetailBuilder.GetProperty(TEXT("Transitions2"));
-	TSharedPtr<IPropertyHandle> StateDoneTransitionProperty = DetailBuilder.GetProperty(TEXT("StateDoneTransition"));
-	TSharedPtr<IPropertyHandle> StateFailedTransitionProperty = DetailBuilder.GetProperty(TEXT("StateFailedTransition"));
-
-	TArray<TWeakObjectPtr<UObject>> Objects;
-	DetailBuilder.GetObjectsBeingCustomized(Objects);
-	UObject* Object = Objects.IsEmpty() ? nullptr : Objects[0].Get();
-	
-	if (IsV2(Object))
-	{
-		EnterConditionsProperty->MarkHiddenByCustomization();
-		EnterConditionsProperty = EnterConditions2Property;
-		EvaluatorsProperty->MarkHiddenByCustomization();
-		EvaluatorsProperty = Evaluators2Property;
-		TasksProperty->MarkHiddenByCustomization();
-		TasksProperty = Tasks2Property;
-		TransitionsProperty->MarkHiddenByCustomization();
-		TransitionsProperty = Transitions2Property;
-		StateDoneTransitionProperty->MarkHiddenByCustomization();
-		StateFailedTransitionProperty->MarkHiddenByCustomization();
-	}
-	else
-	{
-		EnterConditions2Property->MarkHiddenByCustomization();
-		Evaluators2Property->MarkHiddenByCustomization();
-		Tasks2Property->MarkHiddenByCustomization();
-		Transitions2Property->MarkHiddenByCustomization();
-	}
 
 	IDetailCategoryBuilder& StateCategory = DetailBuilder.EditCategory(TEXT("State"), LOCTEXT("StateDetailsState", "State"));
 	StateCategory.SetSortOrder(0);
@@ -117,18 +86,5 @@ void FStateTreeStateDetails::GenerateArrayElementWidget(TSharedRef<IPropertyHand
 	ChildrenBuilder.AddProperty(PropertyHandle);
 };
 
-bool FStateTreeStateDetails::IsV2(const UObject* Object) const
-{
-	if (!Object)
-	{
-		return false;
-	}
-	if (const UStateTree* StateTree = Object->GetTypedOuter<UStateTree>())
-	{
-		return StateTree->IsV2();
-	}
-	return false;
-}
-
 
 #undef LOCTEXT_NAMESPACE
diff --git a/Engine/Plugins/Runtime/StateTree/Source/StateTreeEditorModule/Private/Customizations/StateTreeStateDetails.h b/Engine/Plugins/Runtime/StateTree/Source/StateTreeEditorModule/Private/Customizations/StateTreeStateDetails.h
--- a/Engine/Plugins/Runtime/StateTree/Source/StateTreeEditorModule/Private/Customizations/StateTreeStateDetails.h
+++ b/Engine/Plugins/Runtime/StateTree/Source/StateTreeEditorModule/Private/Customizations/StateTreeStateDetails.h
@@ -22,6 +22,4 @@ public:
 protected:
 	void MakeArrayCategory(IDetailLayoutBuilder& DetailBuilder, FName CategoryName, const FText& DisplayName, int32 SortOrder, TSharedPtr<IPropertyHandle> PropertyHandle);
 	void GenerateArrayElementWidget(TSharedRef<IPropertyHandle> PropertyHandle, int32 ArrayIndex, IDetailChildrenBuilder& ChildrenBuilder);
-
-	bool IsV2(const UObject* Object) const;
 };
diff --git a/Engine/Plugins/Runtime/StateTree/Source/StateTreeEditorModule/Private/Customizations/StateTreeStateLinkDetails.cpp b/Engine/Plugins/Runtime/StateTree/Source/StateTreeEditorModule/Private/Customizations/StateTreeStateLinkDetails.cpp
--- a/Engine/Plugins/Runtime/StateTree/Source/StateTreeEditorModule/Private/Customizations/StateTreeStateLinkDetails.cpp
+++ b/Engine/Plugins/Runtime/StateTree/Source/StateTreeEditorModule/Private/Customizations/StateTreeStateLinkDetails.cpp
@@ -90,7 +90,6 @@ void FStateTreeStateLinkDetails::CacheStates()
 		UStateTree* OuterStateTree = OuterObjects[ObjectIdx]->GetTypedOuter<UStateTree>();
 		if (OuterStateTree)
 		{
-			bIsV2 = OuterStateTree->IsV2();
 			if (UStateTreeEditorData* TreeData = Cast<UStateTreeEditorData>(OuterStateTree->EditorData))
 			{
 				for (const UStateTreeState* Routine : TreeData->Routines)
@@ -128,9 +127,6 @@ void FStateTreeStateLinkDetails::OnStateComboChange(int Idx)
 			case ComboNextState:
 				TypeProperty->SetValue((uint8)EStateTreeTransitionType::NextState);
 				break;
-			case ComboSelectChildState:
-				TypeProperty->SetValue((uint8)EStateTreeTransitionType::SelectChildState);
-				break;
 			case ComboNotSet:
 			default:
 				TypeProperty->SetValue((uint8)EStateTreeTransitionType::NotSet);
@@ -163,12 +159,6 @@ TSharedRef<SWidget> FStateTreeStateLinkDetails::OnGetStateContent() const
 	FUIAction NextItemAction(FExecuteAction::CreateSP(const_cast<FStateTreeStateLinkDetails*>(this), &FStateTreeStateLinkDetails::OnStateComboChange, ComboNextState));
 	MenuBuilder.AddMenuEntry(LOCTEXT("TransitionNextState", "Next State"), LOCTEXT("TransitionNextTooltip", "Goto next sibling State."), FSlateIcon(), NextItemAction);
 
-	if (!bIsV2)
-	{
-		FUIAction SelectItemAction(FExecuteAction::CreateSP(const_cast<FStateTreeStateLinkDetails*>(this), &FStateTreeStateLinkDetails::OnStateComboChange, ComboSelectChildState));
-		MenuBuilder.AddMenuEntry(LOCTEXT("TransitionSelect", "Select"), LOCTEXT("TransitionSelectTooltip", "Select a child State to run."), FSlateIcon(), SelectItemAction);
-	}
-
 	if (CachedNames.Num() > 0)
 	{
 		MenuBuilder.BeginSection(FName(), LOCTEXT("TransitionGotoState", "Goto State"));
@@ -208,10 +198,6 @@ FText FStateTreeStateLinkDetails::GetCurrentStateDesc() const
 	{
 		return LOCTEXT("TransitionNextState", "Next State");
 	}
-	else if (TransitionType == EStateTreeTransitionType::SelectChildState)
-	{
-		return LOCTEXT("TransitionSelect", "Select");
-	}
 	else if (TransitionType == EStateTreeTransitionType::GotoState)
 	{
 		FName OldName;
diff --git a/Engine/Plugins/Runtime/StateTree/Source/StateTreeEditorModule/Private/Customizations/StateTreeStateLinkDetails.h b/Engine/Plugins/Runtime/StateTree/Source/StateTreeEditorModule/Private/Customizations/StateTreeStateLinkDetails.h
--- a/Engine/Plugins/Runtime/StateTree/Source/StateTreeEditorModule/Private/Customizations/StateTreeStateLinkDetails.h
+++ b/Engine/Plugins/Runtime/StateTree/Source/StateTreeEditorModule/Private/Customizations/StateTreeStateLinkDetails.h
@@ -32,7 +32,6 @@ private:
 	static const int ComboSucceeded = -2;
 	static const int ComboFailed = -3;
 	static const int ComboNextState = -4;
-	static const int ComboSelectChildState = -5;
 
 	void CacheStates();
 	void CacheStates(const UStateTreeState* State);
@@ -45,7 +44,6 @@ private:
 
 	TArray<FName> CachedNames;
 	TArray<FGuid> CachedIDs;
-	bool bIsV2;
 
 	TSharedPtr<IPropertyHandle> NameProperty;
 	TSharedPtr<IPropertyHandle> IDProperty;
diff --git a/Engine/Plugins/Runtime/StateTree/Source/StateTreeEditorModule/Private/Customizations/StateTreeTransition2Details.cpp b/Engine/Plugins/Runtime/StateTree/Source/StateTreeEditorModule/Private/Customizations/StateTreeTransition2Details.cpp
deleted file mode 100644
diff --git a/Engine/Plugins/Runtime/StateTree/Source/StateTreeEditorModule/Private/Customizations/StateTreeTransition2Details.h b/Engine/Plugins/Runtime/StateTree/Source/StateTreeEditorModule/Private/Customizations/StateTreeTransition2Details.h
deleted file mode 100644
diff --git a/Engine/Plugins/Runtime/StateTree/Source/StateTreeEditorModule/Private/Customizations/StateTreeTransitionDetails.cpp b/Engine/Plugins/Runtime/StateTree/Source/StateTreeEditorModule/Private/Customizations/StateTreeTransitionDetails.cpp
--- a/Engine/Plugins/Runtime/StateTree/Source/StateTreeEditorModule/Private/Customizations/StateTreeTransitionDetails.cpp
+++ b/Engine/Plugins/Runtime/StateTree/Source/StateTreeEditorModule/Private/Customizations/StateTreeTransitionDetails.cpp
@@ -4,14 +4,10 @@
 #include "Widgets/DeclarativeSyntaxSupport.h"
 #include "DetailWidgetRow.h"
 #include "DetailLayoutBuilder.h"
-#include "IPropertyUtilities.h"
 #include "IDetailPropertyRow.h"
 #include "IDetailChildrenBuilder.h"
 #include "StateTree.h"
 #include "StateTreeState.h"
-#include "Widgets/Input/SComboButton.h"
-#include "Framework/MultiBox/MultiBoxBuilder.h"
-#include "ScopedTransaction.h"
 
 #define LOCTEXT_NAMESPACE "StateTreeEditor"
 
@@ -25,7 +21,9 @@ void FStateTreeTransitionDetails::CustomizeHeader(TSharedRef<class IPropertyHand
 	StructProperty = StructPropertyHandle;
 	PropUtils = StructCustomizationUtils.GetPropertyUtilities().Get();
 
+	EventProperty = StructProperty->GetChildHandle(TEXT("Event"));
 	StateProperty = StructProperty->GetChildHandle(TEXT("State"));
+	GateDelayProperty = StructProperty->GetChildHandle(TEXT("GateDelay"));
 	ConditionsProperty = StructProperty->GetChildHandle(TEXT("Conditions"));
 
 	HeaderRow
@@ -52,9 +50,19 @@ void FStateTreeTransitionDetails::CustomizeHeader(TSharedRef<class IPropertyHand
 
 void FStateTreeTransitionDetails::CustomizeChildren(TSharedRef<class IPropertyHandle> StructPropertyHandle, class IDetailChildrenBuilder& StructBuilder, IPropertyTypeCustomizationUtils& StructCustomizationUtils)
 {
+	if (EventProperty)
+	{
+		StructBuilder.AddProperty(EventProperty.ToSharedRef());
+	}
+
+	if (GateDelayProperty)
+	{
+		StructBuilder.AddProperty(GateDelayProperty.ToSharedRef());
+	}
+
 	if (StateProperty)
 	{
-		IDetailPropertyRow& Property = StructBuilder.AddProperty(StateProperty.ToSharedRef());
+		StructBuilder.AddProperty(StateProperty.ToSharedRef());
 	}
 
 	if (ConditionsProperty)
@@ -88,7 +96,6 @@ void FStateTreeTransitionDetails::CustomizeChildren(TSharedRef<class IPropertyHa
 	}
 }
 
-
 FText FStateTreeTransitionDetails::GetDescription() const
 {
 	if (StateProperty)
@@ -97,22 +104,29 @@ FText FStateTreeTransitionDetails::GetDescription() const
 		StateProperty->AccessRawData(RawData);
 		if (RawData.Num() == 1)
 		{
-			FStateTreeStateLink* StateLink = static_cast<FStateTreeStateLink*>(RawData[0]);
-			if (StateLink != nullptr)
+			FStateTreeStateLink* State = static_cast<FStateTreeStateLink*>(RawData[0]);
+			if (State != nullptr)
 			{
-				switch (StateLink->Type)
+				switch (State->Type)
 				{
+				case EStateTreeTransitionType::NotSet:
+					return LOCTEXT("TransitionNotSet", "None");
+					break;
 				case EStateTreeTransitionType::Succeeded:
-					return LOCTEXT("TransitionSucceeded", "Succeeded");
+					return LOCTEXT("TransitionTreeSucceeded", "Tree Succeeded");
 					break;
 				case EStateTreeTransitionType::Failed:
-					return LOCTEXT("TransitionFailed", "Failed");
+					return LOCTEXT("TransitionTreeFailed", "Tree Failed");
 					break;
-				case EStateTreeTransitionType::SelectChildState:
-					return LOCTEXT("TransitionSelect", "Select");
+				case EStateTreeTransitionType::NextState:
+					return LOCTEXT("TransitionNextState", "Next State");
 					break;
 				case EStateTreeTransitionType::GotoState:
-					return FText::Join(FText::FromString(TEXT(" ")), LOCTEXT("TransitionGoto", "Go to"), FText::FromName(StateLink->Name));
+					{
+						FFormatNamedArguments Args;
+						Args.Add(TEXT("State"), FText::FromName(State->Name));
+						return FText::Format(LOCTEXT("TransitionActionGotoState", "Go to State {State}"), Args);
+					}
 					break;
 				}
 			}
diff --git a/Engine/Plugins/Runtime/StateTree/Source/StateTreeEditorModule/Private/Customizations/StateTreeTransitionDetails.h b/Engine/Plugins/Runtime/StateTree/Source/StateTreeEditorModule/Private/Customizations/StateTreeTransitionDetails.h
--- a/Engine/Plugins/Runtime/StateTree/Source/StateTreeEditorModule/Private/Customizations/StateTreeTransitionDetails.h
+++ b/Engine/Plugins/Runtime/StateTree/Source/StateTreeEditorModule/Private/Customizations/StateTreeTransitionDetails.h
@@ -28,8 +28,9 @@ private:
 
 	FText GetDescription() const;
 
-	TSharedPtr<IPropertyHandle> TypeProperty;
+	TSharedPtr<IPropertyHandle> EventProperty;
 	TSharedPtr<IPropertyHandle> StateProperty;
+	TSharedPtr<IPropertyHandle> GateDelayProperty;
 	TSharedPtr<IPropertyHandle> ConditionsProperty;
 
 	class IPropertyUtilities* PropUtils;
diff --git a/Engine/Plugins/Runtime/StateTree/Source/StateTreeEditorModule/Private/Customizations/StateTreeVariableDescDetails.cpp b/Engine/Plugins/Runtime/StateTree/Source/StateTreeEditorModule/Private/Customizations/StateTreeVariableDescDetails.cpp
deleted file mode 100644
diff --git a/Engine/Plugins/Runtime/StateTree/Source/StateTreeEditorModule/Private/Customizations/StateTreeVariableDescDetails.h b/Engine/Plugins/Runtime/StateTree/Source/StateTreeEditorModule/Private/Customizations/StateTreeVariableDescDetails.h
deleted file mode 100644
diff --git a/Engine/Plugins/Runtime/StateTree/Source/StateTreeEditorModule/Private/Customizations/StateTreeVariableDetails.cpp b/Engine/Plugins/Runtime/StateTree/Source/StateTreeEditorModule/Private/Customizations/StateTreeVariableDetails.cpp
deleted file mode 100644
diff --git a/Engine/Plugins/Runtime/StateTree/Source/StateTreeEditorModule/Private/Customizations/StateTreeVariableDetails.h b/Engine/Plugins/Runtime/StateTree/Source/StateTreeEditorModule/Private/Customizations/StateTreeVariableDetails.h
deleted file mode 100644
diff --git a/Engine/Plugins/Runtime/StateTree/Source/StateTreeEditorModule/Private/Customizations/StateTreeVariableLayoutDetails.cpp b/Engine/Plugins/Runtime/StateTree/Source/StateTreeEditorModule/Private/Customizations/StateTreeVariableLayoutDetails.cpp
deleted file mode 100644
diff --git a/Engine/Plugins/Runtime/StateTree/Source/StateTreeEditorModule/Private/Customizations/StateTreeVariableLayoutDetails.h b/Engine/Plugins/Runtime/StateTree/Source/StateTreeEditorModule/Private/Customizations/StateTreeVariableLayoutDetails.h
deleted file mode 100644
diff --git a/Engine/Plugins/Runtime/StateTree/Source/StateTreeEditorModule/Private/SStateTreeView.cpp b/Engine/Plugins/Runtime/StateTree/Source/StateTreeEditorModule/Private/SStateTreeView.cpp
--- a/Engine/Plugins/Runtime/StateTree/Source/StateTreeEditorModule/Private/SStateTreeView.cpp
+++ b/Engine/Plugins/Runtime/StateTree/Source/StateTreeEditorModule/Private/SStateTreeView.cpp
@@ -265,18 +265,10 @@ void SStateTreeView::HandleModelStatesChanged(const TSet<UStateTreeState*>& Affe
 	{
 		bArraysChanged = true;
 	}
-	if (PropertyChangedEvent.MemberProperty->GetFName() == GET_MEMBER_NAME_CHECKED(UStateTreeState, Evaluators2))
-	{
-		bArraysChanged = true;
-	}
 	if (PropertyChangedEvent.MemberProperty->GetFName() == GET_MEMBER_NAME_CHECKED(UStateTreeState, Tasks))
 	{
 		bArraysChanged = true;
 	}
-	if (PropertyChangedEvent.MemberProperty->GetFName() == GET_MEMBER_NAME_CHECKED(UStateTreeState, Tasks2))
-	{
-		bArraysChanged = true;
-	}
 		
 	if (bArraysChanged)
 	{
diff --git a/Engine/Plugins/Runtime/StateTree/Source/StateTreeEditorModule/Private/SStateTreeViewRow.cpp b/Engine/Plugins/Runtime/StateTree/Source/StateTreeEditorModule/Private/SStateTreeViewRow.cpp
--- a/Engine/Plugins/Runtime/StateTree/Source/StateTreeEditorModule/Private/SStateTreeViewRow.cpp
+++ b/Engine/Plugins/Runtime/StateTree/Source/StateTreeEditorModule/Private/SStateTreeViewRow.cpp
@@ -400,14 +400,7 @@ EVisibility SStateTreeViewRow::GetConditionVisibility() const
 {
 	if (const UStateTreeState* State = WeakState.Get())
 	{
-		if (IsV2())
-		{
-			return State->EnterConditions2.Num() > 0 ? EVisibility::Visible : EVisibility::Collapsed;
-		}
-		else
-		{
-			return State->EnterConditions.Num() > 0 ? EVisibility::Visible : EVisibility::Collapsed;
-		}
+		return State->EnterConditions.Num() > 0 ? EVisibility::Visible : EVisibility::Collapsed;
 	}
 	return EVisibility::Collapsed;
 }
@@ -426,24 +419,11 @@ EVisibility SStateTreeViewRow::GetEvaluatorsVisibility() const
 	if (const UStateTreeState* State = WeakState.Get())
 	{
 		int32 ValidCount = 0;
-		if (IsV2())
-		{
-			for (int32 i = 0; i < State->Evaluators2.Num(); i++)
-			{
-				if (const FStateTreeEvaluator2Base* Eval = State->Evaluators2[i].Type.GetPtr<FStateTreeEvaluator2Base>())
-				{
-					ValidCount++;
-				}
-			}
-		}
-		else
+		for (int32 i = 0; i < State->Evaluators.Num(); i++)
 		{
-			for (int32 i = 0; i < State->Evaluators.Num(); i++)
+			if (const FStateTreeEvaluator2Base* Eval = State->Evaluators[i].Type.GetPtr<FStateTreeEvaluator2Base>())
 			{
-				if (UStateTreeEvaluatorBase* Eval = State->Evaluators[i])
-				{
-					ValidCount++;
-				}
+				ValidCount++;
 			}
 		}
 		return ValidCount > 0 ? EVisibility::Visible : EVisibility::Collapsed;
@@ -461,24 +441,11 @@ FText SStateTreeViewRow::GetEvaluatorsDesc() const
 
 	TArray<FText> Names;
 
-	if (IsV2())
-	{
-		for (int32 i = 0; i < State->Evaluators2.Num(); i++)
-		{
-			if (const FStateTreeEvaluator2Base* Eval = State->Evaluators2[i].Type.GetPtr<FStateTreeEvaluator2Base>())
-			{
-				Names.Add(FText::FromName(Eval->Name));
-			}
-		}
-	}
-	else
+	for (int32 i = 0; i < State->Evaluators.Num(); i++)
 	{
-		for (int32 i = 0; i < State->Evaluators.Num(); i++)
+		if (const FStateTreeEvaluator2Base* Eval = State->Evaluators[i].Type.GetPtr<FStateTreeEvaluator2Base>())
 		{
-			if (UStateTreeEvaluatorBase* Eval = State->Evaluators[i])
-			{
-				Names.Add(FText::FromName(Eval->GetName()));
-			}
+			Names.Add(FText::FromName(Eval->Name));
 		}
 	}
 
@@ -490,24 +457,11 @@ EVisibility SStateTreeViewRow::GetTasksVisibility() const
 	if (const UStateTreeState* State = WeakState.Get())
 	{
 		int32 ValidCount = 0;
-		if (IsV2())
-		{
-			for (int32 i = 0; i < State->Tasks2.Num(); i++)
-			{
-				if (const FStateTreeTask2Base* Task = State->Tasks2[i].Type.GetPtr<FStateTreeTask2Base>())
-				{
-					ValidCount++;
-				}
-			}
-		}
-		else
+		for (int32 i = 0; i < State->Tasks.Num(); i++)
 		{
-			for (int32 i = 0; i < State->Tasks.Num(); i++)
+			if (const FStateTreeTask2Base* Task = State->Tasks[i].Type.GetPtr<FStateTreeTask2Base>())
 			{
-				if (UStateTreeTaskBase* Task = State->Tasks[i])
-				{
-					ValidCount++;
-				}
+				ValidCount++;
 			}
 		}
 		return ValidCount > 0 ? EVisibility::Visible : EVisibility::Collapsed;
@@ -524,34 +478,15 @@ FText SStateTreeViewRow::GetTasksDesc() const
 	}
 
 	TArray<FText> Names;
-
-	if (IsV2())
+	for (int32 i = 0; i < State->Tasks.Num(); i++)
 	{
-		for (int32 i = 0; i < State->Tasks2.Num(); i++)
+		if (const FStateTreeTask2Base* Task = State->Tasks[i].Type.GetPtr<FStateTreeTask2Base>())
 		{
-			if (const FStateTreeTask2Base* Task = State->Tasks2[i].Type.GetPtr<FStateTreeTask2Base>())
-			{
-				Names.Add(FText::FromName(Task->Name));
-			}
+			Names.Add(FText::FromName(Task->Name));
 		}
 	}
-	else
-	{
-		for (int32 i = 0; i < State->Tasks.Num(); i++)
-		{
-			if (UStateTreeTaskBase* Task = State->Tasks[i])
-			{
-				Names.Add(FText::FromName(Task->Name));
-			}
-		}
-	}
-
-	if (IsV2())
-	{
-		return FText::Join((FText::FromString(TEXT(" & "))), Names);
-	}
 
-	return FText::Join((FText::FromString(TEXT(", "))), Names);
+	return FText::Join((FText::FromString(TEXT(" & "))), Names);
 }
 
 bool SStateTreeViewRow::HasParentTransitionForEvent(const UStateTreeState& State, const EStateTreeTransitionEvent Event) const
@@ -559,7 +494,7 @@ bool SStateTreeViewRow::HasParentTransitionForEvent(const UStateTreeState& State
 	EStateTreeTransitionEvent CombinedEvents = EStateTreeTransitionEvent::None;
 	for (const UStateTreeState* ParentState = State.Parent; ParentState != nullptr; ParentState = ParentState->Parent)
 	{
-		for (const FStateTreeTransition2& Transition : ParentState->Transitions2)
+		for (const FStateTreeTransition& Transition : ParentState->Transitions)
 		{
 			CombinedEvents |= Transition.Event;
 		}
@@ -570,7 +505,7 @@ bool SStateTreeViewRow::HasParentTransitionForEvent(const UStateTreeState& State
 FText SStateTreeViewRow::GetTransitionsDesc(const UStateTreeState& State, const EStateTreeTransitionEvent Event) const
 {
 	TArray<FText> DescItems;
-	for (const FStateTreeTransition2& Transition : State.Transitions2)
+	for (const FStateTreeTransition& Transition : State.Transitions)
 	{
 		if (Transition.Event == Event)
 		{
@@ -627,7 +562,7 @@ FText SStateTreeViewRow::GetTransitionsIcon(const UStateTreeState& State, const
     };
 	uint8 IconType = IconNone;
 	
-	for (const FStateTreeTransition2& Transition : State.Transitions2)
+	for (const FStateTreeTransition& Transition : State.Transitions)
 	{
 		// The icons here depict "transition direction", not the type specifically.
 		if (Transition.Event == Event)
@@ -698,7 +633,7 @@ EVisibility SStateTreeViewRow::GetTransitionsVisibility(const UStateTreeState& S
 	TStaticArray<int32, 5> Counts;
 	Algo::ForEach(Counts, [](int32& Count) { Count = 0; }); 
 
-	for (const FStateTreeTransition2& Transition : State.Transitions2)
+	for (const FStateTreeTransition& Transition : State.Transitions)
 	{
 		Counts[static_cast<uint8>(Transition.Event)]++;
 	}
@@ -724,10 +659,7 @@ EVisibility SStateTreeViewRow::GetCompletedTransitionVisibility() const
 {
 	if (const UStateTreeState* State = WeakState.Get())
 	{
-		if (IsV2())
-		{
-			return GetTransitionsVisibility(*State, EStateTreeTransitionEvent::OnCompleted);
-		}
+		return GetTransitionsVisibility(*State, EStateTreeTransitionEvent::OnCompleted);
 	}
 	return EVisibility::Visible;
 }
@@ -736,31 +668,7 @@ FText SStateTreeViewRow::GetCompletedTransitionsDesc() const
 {
 	if (const UStateTreeState* State = WeakState.Get())
 	{
-		if (IsV2())
-		{
-			return GetTransitionsDesc(*State, EStateTreeTransitionEvent::OnCompleted);
-		}
-		else
-		{
-			switch (State->StateDoneTransition.Type)
-			{
-			case EStateTreeTransitionType::NotSet:
-				return FText::GetEmpty();
-			case EStateTreeTransitionType::Succeeded:
-				return LOCTEXT("TransitionSucceeded", "Succeeded");
-			case EStateTreeTransitionType::Failed:
-				return LOCTEXT("TransitionFailed", "Failed");
-			case EStateTreeTransitionType::NextState:
-				return LOCTEXT("TransitionNext", "Next");
-			case EStateTreeTransitionType::SelectChildState:
-				return LOCTEXT("TransitionSelect", "Select");
-			case EStateTreeTransitionType::GotoState:
-				return FText::FromName(State->StateDoneTransition.Name);
-			default:
-				ensureMsgf(false, TEXT("Unhandled transition type."));
-				break;
-			}
-		}
+		return GetTransitionsDesc(*State, EStateTreeTransitionEvent::OnCompleted);
 	}
 	return LOCTEXT("Invalid", "Invalid");
 }
@@ -769,31 +677,7 @@ FText SStateTreeViewRow::GetCompletedTransitionsIcon() const
 {
 	if (const UStateTreeState* State = WeakState.Get())
 	{
-		if (IsV2())
-		{
-			return GetTransitionsIcon(*State, EStateTreeTransitionEvent::OnCompleted);
-		}
-		else
-		{
-			switch (State->StateDoneTransition.Type)
-			{
-			case EStateTreeTransitionType::NotSet:
-				return FEditorFontGlyphs::Times;
-			case EStateTreeTransitionType::Succeeded:
-				return FEditorFontGlyphs::Check;
-			case EStateTreeTransitionType::Failed:
-				return FEditorFontGlyphs::Times;
-			case EStateTreeTransitionType::NextState:
-				return FEditorFontGlyphs::Long_Arrow_Down;
-			case EStateTreeTransitionType::SelectChildState:
-				return FEditorFontGlyphs::Level_Down;
-			case EStateTreeTransitionType::GotoState:
-				return FEditorFontGlyphs::Long_Arrow_Right;
-			default:
-				ensureMsgf(false, TEXT("Unhandled transition type."));
-				break;
-			}
-		}
+		return GetTransitionsIcon(*State, EStateTreeTransitionEvent::OnCompleted);
 	}
 	return FText::GetEmpty();
 }
@@ -802,10 +686,7 @@ EVisibility SStateTreeViewRow::GetSucceededTransitionVisibility() const
 {
 	if (const UStateTreeState* State = WeakState.Get())
 	{
-		if (IsV2())
-		{
-			return GetTransitionsVisibility(*State, EStateTreeTransitionEvent::OnSucceeded);
-		}
+		return GetTransitionsVisibility(*State, EStateTreeTransitionEvent::OnSucceeded);
 	}
 	return EVisibility::Collapsed;
 }
@@ -814,10 +695,7 @@ FText SStateTreeViewRow::GetSucceededTransitionDesc() const
 {
 	if (const UStateTreeState* State = WeakState.Get())
 	{
-		if (IsV2())
-		{
-			return GetTransitionsDesc(*State, EStateTreeTransitionEvent::OnSucceeded);
-		}
+		return GetTransitionsDesc(*State, EStateTreeTransitionEvent::OnSucceeded);
 	}
 	return FText::GetEmpty();
 }
@@ -826,10 +704,7 @@ FText SStateTreeViewRow::GetSucceededTransitionIcon() const
 {
 	if (const UStateTreeState* State = WeakState.Get())
 	{
-		if (IsV2())
-		{
-			return GetTransitionsIcon(*State, EStateTreeTransitionEvent::OnSucceeded);
-		}
+		return GetTransitionsIcon(*State, EStateTreeTransitionEvent::OnSucceeded);
 	}
 	return FText::GetEmpty();
 }
@@ -838,14 +713,7 @@ EVisibility SStateTreeViewRow::GetFailedTransitionVisibility() const
 {
 	if (const UStateTreeState* State = WeakState.Get())
 	{
-		if (IsV2())
-		{
-			return GetTransitionsVisibility(*State, EStateTreeTransitionEvent::OnFailed);
-		}
-		else
-		{
-			return State->StateFailedTransition.Type != EStateTreeTransitionType::NotSet ? EVisibility::Visible : EVisibility::Collapsed;
-		}
+		return GetTransitionsVisibility(*State, EStateTreeTransitionEvent::OnFailed);
 	}
 	return EVisibility::Collapsed;
 }
@@ -854,31 +722,7 @@ FText SStateTreeViewRow::GetFailedTransitionDesc() const
 {
 	if (const UStateTreeState* State = WeakState.Get())
 	{
-		if (IsV2())
-		{
-			return GetTransitionsDesc(*State, EStateTreeTransitionEvent::OnFailed);
-		}
-		else
-		{
-			switch (State->StateFailedTransition.Type)
-			{
-			case EStateTreeTransitionType::NotSet:
-				return FText::GetEmpty();
-			case EStateTreeTransitionType::Succeeded:
-				return LOCTEXT("TransitionSucceeded", "Succeeded");
-			case EStateTreeTransitionType::Failed:
-				return LOCTEXT("TransitionFailed", "Failed");
-			case EStateTreeTransitionType::NextState:
-				return LOCTEXT("TransitionNext", "Next");
-			case EStateTreeTransitionType::SelectChildState:
-				return LOCTEXT("TransitionSelect", "Select");
-			case EStateTreeTransitionType::GotoState:
-				return FText::FromName(State->StateFailedTransition.Name);
-			default:
-				ensureMsgf(false, TEXT("Unhandled transition type."));
-				break;
-			}
-		}
+		return GetTransitionsDesc(*State, EStateTreeTransitionEvent::OnFailed);
 	}
 	return LOCTEXT("Invalid", "Invalid");
 }
@@ -887,14 +731,7 @@ FText SStateTreeViewRow::GetFailedTransitionIcon() const
 {
 	if (const UStateTreeState* State = WeakState.Get())
 	{
-		if (IsV2())
-		{
-			return GetTransitionsIcon(*State, EStateTreeTransitionEvent::OnFailed);
-		}
-		else
-		{
-			return FEditorFontGlyphs::Ban;
-		}
+		return GetTransitionsIcon(*State, EStateTreeTransitionEvent::OnFailed);
 	}
 	return FEditorFontGlyphs::Ban;
 }
@@ -903,19 +740,7 @@ EVisibility SStateTreeViewRow::GetConditionalTransitionsVisibility() const
 {
 	if (const UStateTreeState* State = WeakState.Get())
 	{
-		if (IsV2())
-		{
-			return GetTransitionsVisibility(*State, EStateTreeTransitionEvent::OnCondition);
-		}
-		else
-		{
-			int32 ValidCount = 0;
-			for (const FStateTreeTransition& Transition : State->Transitions)
-			{
-				ValidCount++;
-			}
-			return ValidCount > 0 ? EVisibility::Visible : EVisibility::Collapsed;
-		}
+		return GetTransitionsVisibility(*State, EStateTreeTransitionEvent::OnCondition);
 	}
 	return EVisibility::Collapsed;
 }
@@ -925,34 +750,7 @@ FText SStateTreeViewRow::GetConditionalTransitionsDesc() const
 	TArray<FText> DescItems;
 	if (const UStateTreeState* State = WeakState.Get())
 	{
-		if (IsV2())
-		{
-			return GetTransitionsDesc(*State, EStateTreeTransitionEvent::OnCondition);
-		}
-		else
-		{
-			for (const FStateTreeTransition& Transition : State->Transitions)
-			{
-				switch (Transition.State.Type)
-				{
-				case EStateTreeTransitionType::Succeeded:
-					DescItems.Add(LOCTEXT("TransitionSucceeded", "Succeeded"));
-					break;
-				case EStateTreeTransitionType::Failed:
-					DescItems.Add(LOCTEXT("TransitionFailed", "Failed"));
-					break;
-				case EStateTreeTransitionType::SelectChildState:
-					DescItems.Add(LOCTEXT("TransitionSelect", "Select"));
-					break;
-				case EStateTreeTransitionType::GotoState:
-					DescItems.Add(FText::FromName(Transition.State.Name));
-					break;
-				default:
-					ensureMsgf(false, TEXT("Unhandled transition type."));
-					break;
-				}
-			}
-		}
+		return GetTransitionsDesc(*State, EStateTreeTransitionEvent::OnCondition);
 	}
 	return FText::Join(FText::FromString(TEXT(", ")), DescItems);
 }
@@ -1043,13 +841,4 @@ FReply SStateTreeViewRow::HandleAcceptDrop(const FDragDropEvent& DragDropEvent,
 	return FReply::Unhandled();
 }
 
-bool SStateTreeViewRow::IsV2() const
-{
-	if (const UStateTreeState* State = WeakState.Get())
-	{
-		return State->IsV2();
-	}
-	return false;	
-}
-
 #undef LOCTEXT_NAMESPACE
diff --git a/Engine/Plugins/Runtime/StateTree/Source/StateTreeEditorModule/Private/SStateTreeViewRow.h b/Engine/Plugins/Runtime/StateTree/Source/StateTreeEditorModule/Private/SStateTreeViewRow.h
--- a/Engine/Plugins/Runtime/StateTree/Source/StateTreeEditorModule/Private/SStateTreeViewRow.h
+++ b/Engine/Plugins/Runtime/StateTree/Source/StateTreeEditorModule/Private/SStateTreeViewRow.h
@@ -67,8 +67,6 @@ private:
 	FReply HandleDragDetected(const FGeometry&, const FPointerEvent&);
 	TOptional<EItemDropZone> HandleCanAcceptDrop(const FDragDropEvent& DragDropEvent, EItemDropZone DropZone, UStateTreeState* TargetState);
 	FReply HandleAcceptDrop(const FDragDropEvent& DragDropEvent, EItemDropZone DropZone, UStateTreeState* TargetState);
-
-	bool IsV2() const;
 	
 	TSharedPtr<FStateTreeViewModel> StateTreeViewModel;
 	TWeakObjectPtr<UStateTreeState> WeakState;
diff --git a/Engine/Plugins/Runtime/StateTree/Source/StateTreeEditorModule/Private/StateTreeBaker.cpp b/Engine/Plugins/Runtime/StateTree/Source/StateTreeEditorModule/Private/StateTreeBaker.cpp
--- a/Engine/Plugins/Runtime/StateTree/Source/StateTreeEditorModule/Private/StateTreeBaker.cpp
+++ b/Engine/Plugins/Runtime/StateTree/Source/StateTreeEditorModule/Private/StateTreeBaker.cpp
@@ -4,173 +4,15 @@
 #include "StateTree.h"
 #include "StateTreeEditorData.h"
 #include "StateTreeTypes.h"
-#include "StateTreeCondition.h"
 #include "Conditions/StateTreeCondition_Common.h"
 #include "StateTreeEvaluatorBase.h"
 #include "StateTreeTaskBase.h"
-#include "StateTreeVariable.h"
-#include "StateTreeVariableDesc.h"
-#include "StateTreeVariableLayout.h"
-#include "StateTreeParameter.h"
-#include "StateTreeParameterLayout.h"
+#include "StateTreeConditionBase.h"
 #include "StateTreeState.h"
 #include "StateTreeExecutionContext.h"
 #include "CoreMinimal.h"
 #include "StateTreePropertyBindingCompiler.h"
 
-bool FStateTreeBaker::Bake(UStateTree& InStateTree)
-{
-	if (InStateTree.IsV2())
-	{
-		return Bake2(InStateTree);
-	}
-
-	StateTree = &InStateTree;
-	TreeData = Cast<UStateTreeEditorData>(StateTree->EditorData);
-	if (!TreeData)
-	{
-		return false;
-	}
-
-	// Cleanup existing state
-	StateTree->ResetBaked();
-
-	DefineParameters();
-
-	if (!CreateStates())
-	{
-		StateTree->ResetBaked();
-		return false;
-	}
-
-	StateTree->Variables.CalculateOffsetsAndMemoryUsage();
-	StateTree->Constants.ConstantBaseOffset = StateTree->Variables.GetMemoryUsage();
-
-	ResolveParameters();
-
-	if (!CreateStateTransitions())
-	{
-		StateTree->ResetBaked();
-		return false;
-	}
-
-	// Create and initialize tasks
-	for (const FSourceTask& Source : SourceTasks)
-	{
-		UStateTreeTaskBase* Task = Source.Task;
-		const UStateTreeState* SourceState = Source.State;
-		if (!Task->ResolveVariables(StateTree->Variables, StateTree->Constants, StateTree))
-		{
-			UE_LOG(LogStateTree, Error, TEXT("%s: %s:%s, Failed to resolve task %s."), ANSI_TO_TCHAR(__FUNCTION__), *GetNameSafe(StateTree), *SourceState->Name.ToString(), *Task->GetName());
-			StateTree->ResetBaked();
-			return false;
-		}
-		StateTree->Tasks.Add(Cast<UStateTreeTaskBase>(StaticDuplicateObject(Task, StateTree)));
-	}
-
-	// Create and initialize evaluators
-	for (const FSourceEvaluator& Source : SourceEvaluators)
-	{
-		UStateTreeEvaluatorBase* Eval = Source.Eval;
-		const UStateTreeState* SourceState = Source.State;
-		if (!Eval->ResolveVariables(StateTree->Variables, StateTree->Constants, StateTree))
-		{
-			UE_LOG(LogStateTree, Error, TEXT("%s: %s:%s, Failed to instantiate evaluator %s."), ANSI_TO_TCHAR(__FUNCTION__), *GetNameSafe(StateTree), *SourceState->Name.ToString(), *Eval->GetName().ToString());
-			StateTree->ResetBaked();
-			return false;
-		}
-		StateTree->Evaluators.Add(Cast<UStateTreeEvaluatorBase>(StaticDuplicateObject(Eval, StateTree)));
-	}
-
-	return true;
-}
-
-void FStateTreeBaker::DefineParameters() const
-{
-	const FStateTreeVariableLayout& InputParameterLayout = StateTree->GetInputParameterLayout();
-	for (const FStateTreeVariableDesc& Desc : InputParameterLayout.Variables)
-	{
-		StateTree->Variables.DefineVariable(Desc);
-	}
-}
-
-void FStateTreeBaker::ResolveParameters() const
-{
-	// Find where to write parameter values
-	StateTree->Parameters.Reset();
-
-	const FStateTreeVariableLayout& InputParameterLayout = StateTree->GetInputParameterLayout();
-	for (const FStateTreeVariableDesc& Desc : InputParameterLayout.Variables)
-	{
-		FStateTreeParameter& Param = StateTree->Parameters.Parameters.AddDefaulted_GetRef();
-		Param.Name = Desc.Name;
-		Param.Variable.Type = Desc.Type;
-		Param.Variable.Handle = StateTree->Variables.GetVariableHandle(Desc.ID);
-	}
-}
-
-bool FStateTreeBaker::CreateStates()
-{
-	for (UStateTreeState* Routine : TreeData->Routines)
-	{
-		if (Routine)
-		{
-			if (!CreateStateRecursive(*Routine, FStateTreeHandle::Invalid))
-			{
-				return false;
-			}
-		}
-	}
-	return true;
-}
-
-bool FStateTreeBaker::CreateStateRecursive(UStateTreeState& State, const FStateTreeHandle Parent)
-{
-	const int32 Idx = StateTree->States.AddDefaulted();
-	FBakedStateTreeState& BakedState = StateTree->States[Idx];
-	BakedState.Name = State.Name;
-	BakedState.Parent = Parent;
-
-	SourceStates.Add(&State);
-	IDToState.Add(State.ID, Idx);
-
-	// Collect tasks
-	BakedState.TasksBegin = uint16(SourceTasks.Num());
-	for (UStateTreeTaskBase* Task : State.Tasks)
-	{
-		if (Task)
-		{
-			SourceTasks.Add(FSourceTask(Task, &State));
-		}
-	}
-	BakedState.TasksNum = uint8(uint16(SourceTasks.Num()) - BakedState.TasksBegin);
-
-	// Collect evaluators and define variables
-	for (UStateTreeEvaluatorBase* Evaluator : State.Evaluators)
-	{
-		if (Evaluator)
-		{
-			Evaluator->DefineOutputVariables(StateTree->Variables);
-			SourceEvaluators.Add(FSourceEvaluator(Evaluator, &State));
-		}
-	}
-
-	// Child states
-	BakedState.ChildrenBegin = uint16(StateTree->States.Num());
-	for (UStateTreeState* Child : State.Children)
-	{
-		if (Child)
-		{
-			if (!CreateStateRecursive(*Child, FStateTreeHandle((uint16)Idx)))
-			{
-				return false;
-			}
-		}
-	}
-	StateTree->States[Idx].ChildrenEnd = uint16(StateTree->States.Num()); // Cannot use BakedState here, it may be invalid due to array resize.
-
-	return true;
-}
 
 bool FStateTreeBaker::ResolveTransitionState(const UStateTreeState& SourceState, const TCHAR* ContextStr, const FStateTreeStateLink& Link, FStateTreeHandle& OutTransitionHandle) const 
 {
@@ -203,112 +45,6 @@ bool FStateTreeBaker::ResolveTransitionState(const UStateTreeState& SourceState,
 	return true;
 }
 
-bool FStateTreeBaker::CreateStateTransitions()
-{
-	for (int32 i = 0; i < StateTree->States.Num(); i++)
-	{
-		FBakedStateTreeState& BakedState = StateTree->States[i];
-		UStateTreeState* SourceState = SourceStates[i];
-		check(SourceState);
-
-		// Resolve default transition
-		BakedState.StateDoneTransitionType = SourceState->StateDoneTransition.Type;
-		BakedState.StateDoneTransitionState = FStateTreeHandle::Invalid;
-		if (!ResolveTransitionState(*SourceState, TEXT("State Done transition"), SourceState->StateDoneTransition, BakedState.StateDoneTransitionState))
-		{
-			return false;
-		}
-		// Default transition must be set.
-		if (BakedState.StateDoneTransitionType == EStateTreeTransitionType::NotSet)
-		{
-			UE_LOG(LogStateTree, Error, TEXT("%s: '%s':%s: State Done transition must be set."), ANSI_TO_TCHAR(__FUNCTION__), *GetNameSafe(StateTree), *SourceState->Name.ToString());
-			return false;
-		}
-
-		// Resolve failed transition
-		BakedState.StateFailedTransitionType = SourceState->StateFailedTransition.Type;
-		BakedState.StateFailedTransitionState = FStateTreeHandle::Invalid;
-		if (!ResolveTransitionState(*SourceState, TEXT("State Failed transition"), SourceState->StateFailedTransition, BakedState.StateFailedTransitionState))
-		{
-			return false;
-		}
-
-		// Enter conditions.
-		BakedState.EnterConditionsBegin = uint16(StateTree->Conditions.Num());
-		for (FStateTreeCondition& Condition : SourceState->EnterConditions)
-		{
-			if (!CreateCondition(Condition))
-			{
-				UE_LOG(LogStateTree, Error, TEXT("%s: '%s':%s: Failed to create state enter condition."), ANSI_TO_TCHAR(__FUNCTION__), *GetNameSafe(StateTree), *SourceState->Name.ToString());
-				return false;
-			}
-		}
-		BakedState.EnterConditionsNum = uint8(uint16(StateTree->Conditions.Num()) - BakedState.EnterConditionsBegin);
-
-		// Transitions
-		BakedState.TransitionsBegin = uint16(StateTree->Transitions.Num());
-		for (FStateTreeTransition& Transition : SourceState->Transitions)
-		{
-			FBakedStateTransition& BakedTransition = StateTree->Transitions.AddDefaulted_GetRef();
-			BakedTransition.Type = Transition.State.Type;
-			BakedTransition.State = FStateTreeHandle::Invalid;
-			if (!ResolveTransitionState(*SourceState, TEXT("transition"), Transition.State, BakedTransition.State))
-			{
-				return false;
-			}
-			// Note: Unset transition is allowed here. It can be used to mask a transition at parent.
-
-			BakedTransition.ConditionsBegin = uint16(StateTree->Conditions.Num());
-			for (FStateTreeCondition& Condition : Transition.Conditions)
-			{
-				if (!CreateCondition(Condition))
-				{
-					UE_LOG(LogStateTree, Error, TEXT("%s: '%s':%s: Failed to create transition condition to %s."), ANSI_TO_TCHAR(__FUNCTION__), *GetNameSafe(StateTree), *SourceState->Name.ToString(), *Transition.State.Name.ToString());
-					return false;
-				}
-			}
-			BakedTransition.ConditionsNum = uint8(uint16(StateTree->Conditions.Num()) - BakedTransition.ConditionsBegin);
-		}
-		BakedState.TransitionsNum = uint8(uint16(StateTree->Transitions.Num()) - BakedState.TransitionsBegin);
-	}
-
-	return true;
-}
-
-
-bool FStateTreeBaker::CreateCondition(const FStateTreeCondition& Condition)
-{
-	FStateTreeCondition& BakedCondition = StateTree->Conditions.Add_GetRef(Condition);
-
-	if (BakedCondition.Left.Type != BakedCondition.Right.Type)
-	{
-		UE_LOG(LogStateTree, Error, TEXT("%s: '%s': Left and right types do not match."), ANSI_TO_TCHAR(__FUNCTION__), *GetNameSafe(StateTree));
-		return false;
-	}
-
-	if (!BakedCondition.Left.IsBound())
-	{
-		UE_LOG(LogStateTree, Error, TEXT("%s: '%s': Left hand variable must be bound."), ANSI_TO_TCHAR(__FUNCTION__), *GetNameSafe(StateTree));
-		return false;
-	}
-
-	BakedCondition.Left.ResolveHandle(StateTree->Variables, StateTree->Constants);
-	if (!BakedCondition.Left.Handle.IsValid())
-	{
-		UE_LOG(LogStateTree, Error, TEXT("%s: '%s': Cannot resolve left hand variable %s."), ANSI_TO_TCHAR(__FUNCTION__), *GetNameSafe(StateTree), *BakedCondition.Left.Name.ToString());
-		return false;
-	}
-
-	BakedCondition.Right.ResolveHandle(StateTree->Variables, StateTree->Constants);
-	if (!BakedCondition.Right.Handle.IsValid())
-	{
-		UE_LOG(LogStateTree, Error, TEXT("%s: '%s': Cannot resolve right hand variable %s."), ANSI_TO_TCHAR(__FUNCTION__), *GetNameSafe(StateTree), *BakedCondition.Right.Name.ToString());
-		return false;
-	}
-
-	return true;
-}
-
 FStateTreeHandle FStateTreeBaker::GetStateHandle(const FGuid& StateID) const
 {
 	const int32* Idx = IDToState.Find(StateID);
@@ -321,7 +57,7 @@ FStateTreeHandle FStateTreeBaker::GetStateHandle(const FGuid& StateID) const
 }
 
 
-bool FStateTreeBaker::Bake2(UStateTree& InStateTree)
+bool FStateTreeBaker::Bake(UStateTree& InStateTree)
 {
 	StateTree = &InStateTree;
 	TreeData = Cast<UStateTreeEditorData>(StateTree->EditorData);
@@ -335,13 +71,13 @@ bool FStateTreeBaker::Bake2(UStateTree& InStateTree)
 
 	BindingsCompiler.Init(StateTree->PropertyBindings);
 
-	if (!CreateStates2())
+	if (!CreateStates())
 	{
 		StateTree->ResetBaked();
 		return false;
 	}
 
-	if (!CreateStateTransitions2())
+	if (!CreateStateTransitions())
 	{
 		StateTree->ResetBaked();
 		return false;
@@ -357,7 +93,7 @@ bool FStateTreeBaker::Bake2(UStateTree& InStateTree)
 	return true;
 }
 
-bool FStateTreeBaker::CreateStates2()
+bool FStateTreeBaker::CreateStates()
 {
 	// Create item for the runtime execution state
 	StateTree->RuntimeStorageItems.Add(FInstancedStruct::Make<FStateTreeExecutionState>());
@@ -366,7 +102,7 @@ bool FStateTreeBaker::CreateStates2()
 	{
 		if (Routine)
 		{
-			if (!CreateStateRecursive2(*Routine, FStateTreeHandle::Invalid))
+			if (!CreateStateRecursive(*Routine, FStateTreeHandle::Invalid))
 			{
 				return false;
 			}
@@ -375,7 +111,7 @@ bool FStateTreeBaker::CreateStates2()
 	return true;
 }
 
-bool FStateTreeBaker::CreateStateTransitions2()
+bool FStateTreeBaker::CreateStateTransitions()
 {
 	for (int32 i = 0; i < StateTree->States.Num(); i++)
 	{
@@ -384,20 +120,20 @@ bool FStateTreeBaker::CreateStateTransitions2()
 		check(SourceState);
 
 		// Enter conditions.
-		BakedState.EnterConditionsBegin = uint16(StateTree->Conditions2.Num());
-		for (FStateTreeConditionItem& ConditionItem : SourceState->EnterConditions2)
+		BakedState.EnterConditionsBegin = uint16(StateTree->Conditions.Num());
+		for (FStateTreeConditionItem& ConditionItem : SourceState->EnterConditions)
 		{
-			if (!CreateCondition2(ConditionItem))
+			if (!CreateCondition(ConditionItem))
 			{
 				UE_LOG(LogStateTree, Error, TEXT("%s: '%s':%s: Failed to create state enter condition."), ANSI_TO_TCHAR(__FUNCTION__), *GetNameSafe(StateTree), *SourceState->Name.ToString());
 				return false;
 			}
 		}
-		BakedState.EnterConditionsNum = uint8(uint16(StateTree->Conditions2.Num()) - BakedState.EnterConditionsBegin);
+		BakedState.EnterConditionsNum = uint8(uint16(StateTree->Conditions.Num()) - BakedState.EnterConditionsBegin);
 
 		// Transitions
 		BakedState.TransitionsBegin = uint16(StateTree->Transitions.Num());
-		for (FStateTreeTransition2& Transition : SourceState->Transitions2)
+		for (FStateTreeTransition& Transition : SourceState->Transitions)
 		{
 			FBakedStateTransition& BakedTransition = StateTree->Transitions.AddDefaulted_GetRef();
 			BakedTransition.Event = Transition.Event;
@@ -410,16 +146,16 @@ bool FStateTreeBaker::CreateStateTransitions2()
 			}
 			// Note: Unset transition is allowed here. It can be used to mask a transition at parent.
 
-			BakedTransition.ConditionsBegin = uint16(StateTree->Conditions2.Num());
+			BakedTransition.ConditionsBegin = uint16(StateTree->Conditions.Num());
 			for (FStateTreeConditionItem& ConditionItem : Transition.Conditions)
 			{
-				if (!CreateCondition2(ConditionItem))
+				if (!CreateCondition(ConditionItem))
 				{
 					UE_LOG(LogStateTree, Error, TEXT("%s: '%s':%s: Failed to create transition condition to %s."), ANSI_TO_TCHAR(__FUNCTION__), *GetNameSafe(StateTree), *SourceState->Name.ToString(), *Transition.State.Name.ToString());
 					return false;
 				}
 			}
-			BakedTransition.ConditionsNum = uint8(uint16(StateTree->Conditions2.Num()) - BakedTransition.ConditionsBegin);
+			BakedTransition.ConditionsNum = uint8(uint16(StateTree->Conditions.Num()) - BakedTransition.ConditionsBegin);
 		}
 		BakedState.TransitionsNum = uint8(uint16(StateTree->Transitions.Num()) - BakedState.TransitionsBegin);
 	}
@@ -427,7 +163,7 @@ bool FStateTreeBaker::CreateStateTransitions2()
 	return true;
 }
 
-bool FStateTreeBaker::CreateCondition2(const FStateTreeConditionItem& CondItem)
+bool FStateTreeBaker::CreateCondition(const FStateTreeConditionItem& CondItem)
 {
 	if (!CondItem.Type.IsValid())
 	{
@@ -436,7 +172,7 @@ bool FStateTreeBaker::CreateCondition2(const FStateTreeConditionItem& CondItem)
 	}
 
 	// Copy the condition
-	FInstancedStruct& CondPtr = StateTree->Conditions2.AddDefaulted_GetRef();
+	FInstancedStruct& CondPtr = StateTree->Conditions.AddDefaulted_GetRef();
 	CondPtr = CondItem.Type;
 	check(CondPtr.IsValid());
 
@@ -654,7 +390,7 @@ bool FStateTreeBaker::CreateExternalItemHandles(FStructView Item)
 	return true;
 }
 
-bool FStateTreeBaker::CreateStateRecursive2(UStateTreeState& State, const FStateTreeHandle Parent)
+bool FStateTreeBaker::CreateStateRecursive(UStateTreeState& State, const FStateTreeHandle Parent)
 {
 	const int32 StateIdx = StateTree->States.AddDefaulted();
 	FBakedStateTreeState& BakedState = StateTree->States[StateIdx];
@@ -668,7 +404,7 @@ bool FStateTreeBaker::CreateStateRecursive2(UStateTreeState& State, const FState
 	check(StateTree->RuntimeStorageItems.Num() <= int32(MAX_uint16));
 	BakedState.EvaluatorsBegin = uint16(StateTree->RuntimeStorageItems.Num());
 
-	for (FStateTreeEvaluatorItem& EvaluatorItem : State.Evaluators2)
+	for (FStateTreeEvaluatorItem& EvaluatorItem : State.Evaluators)
 	{
 		if (EvaluatorItem.Type.IsValid())
 		{
@@ -720,7 +456,7 @@ bool FStateTreeBaker::CreateStateRecursive2(UStateTreeState& State, const FState
 	check(StateTree->RuntimeStorageItems.Num() <= int32(MAX_uint16));
 	BakedState.TasksBegin = uint16(StateTree->RuntimeStorageItems.Num());
 
-	for (FStateTreeTaskItem& TaskItem : State.Tasks2)
+	for (FStateTreeTaskItem& TaskItem : State.Tasks)
 	{
 		if (TaskItem.Type.IsValid())
 		{
@@ -775,7 +511,7 @@ bool FStateTreeBaker::CreateStateRecursive2(UStateTreeState& State, const FState
 	{
 		if (Child)
 		{
-			if (!CreateStateRecursive2(*Child, FStateTreeHandle((uint16)StateIdx)))
+			if (!CreateStateRecursive(*Child, FStateTreeHandle((uint16)StateIdx)))
 			{
 				return false;
 			}
diff --git a/Engine/Plugins/Runtime/StateTree/Source/StateTreeEditorModule/Private/StateTreeEditor.cpp b/Engine/Plugins/Runtime/StateTree/Source/StateTreeEditorModule/Private/StateTreeEditor.cpp
--- a/Engine/Plugins/Runtime/StateTree/Source/StateTreeEditorModule/Private/StateTreeEditor.cpp
+++ b/Engine/Plugins/Runtime/StateTree/Source/StateTreeEditorModule/Private/StateTreeEditor.cpp
@@ -138,7 +138,6 @@ void FStateTreeEditor::InitEditor( const EToolkitMode::Type Mode, const TSharedP
 	UpdateAsset();
 
 	UE::StateTree::Delegates::OnIdentifierChanged.AddSP(this, &FStateTreeEditor::OnIdentifierChanged);
-	UE::StateTree::Delegates::OnParameterLayoutChanged.AddSP(this, &FStateTreeEditor::OnParameterLayoutChanged);
 	UE::StateTree::Delegates::OnSchemaChanged.AddSP(this, &FStateTreeEditor::OnSchemaChanged);
 }
 
@@ -260,16 +259,6 @@ void FStateTreeEditor::OnIdentifierChanged(const UStateTree& InStateTree)
 	}
 }
 
-void FStateTreeEditor::OnParameterLayoutChanged(const UStateTree& InStateTree)
-{
-	// Update when other StateTree assets change.
-	if (StateTree != &InStateTree)
-	{
-		// Update parameter sets. For assets that are not open, this will be done on load.
-		UpdateAsset();
-	}
-}
-
 void FStateTreeEditor::OnSchemaChanged(const UStateTree& InStateTree)
 {
 	if (StateTree == &InStateTree)
@@ -314,7 +303,7 @@ void FStateTreeEditor::OnSelectionFinishedChangingProperties(const FPropertyChan
 	}
 }
 
-namespace FStateTreeEditorHelpers
+namespace UE::StateTree::Editor::Internal
 {
 	bool FixChangedStateLinkName(FStateTreeStateLink& StateLink, const TMap<FGuid, FName>& IDToName)
 	{
@@ -336,44 +325,6 @@ namespace FStateTreeEditorHelpers
 		return false;
 	}
 
-	bool FixChangedVariableName(FStateTreeVariable& VarBinding, const TMap<FGuid, FName>& IDToName)
-	{
-		if (VarBinding.ID.IsValid())
-		{
-			const FName* Name = IDToName.Find(VarBinding.ID);
-			if (Name == nullptr)
-			{
-				// Missing link, we'll show these in the UI
-				return false;
-			}
-			if (VarBinding.Name != *Name)
-			{
-				// Name changed, fix!
-				VarBinding.Name = *Name;
-				return true;
-			}
-		}
-		return false;
-	}
-
-	bool FixChangedVariableName(FStateTreeCondition& Condition, const TMap<FGuid, FName>& IDToName)
-	{
-		bool Result = false;
-		Result |= FixChangedVariableName(Condition.Left, IDToName);
-		Result |= FixChangedVariableName(Condition.Right, IDToName);
-		return Result;
-	}
-
-	bool FixChangedVariableName(TArray<FStateTreeCondition>& Conditions, const TMap<FGuid, FName>& IDToName)
-	{
-		bool Result = false;
-		for (FStateTreeCondition& Cond : Conditions)
-		{
-			Result |= FixChangedVariableName(Cond, IDToName);
-		}
-		return Result;
-	}
-
 	void ValidateLinkedStates(UStateTree& StateTree)
 	{
 		UStateTreeEditorData* TreeData = Cast<UStateTreeEditorData>(StateTree.EditorData);
@@ -420,9 +371,6 @@ namespace FStateTreeEditorHelpers
 				{
 					UStateTreeState* CurState = Stack.Pop();
 
-					FixChangedStateLinkName(CurState->StateDoneTransition, IDToName);
-					FixChangedStateLinkName(CurState->StateFailedTransition, IDToName);
-
 					for (FStateTreeTransition& Transition : CurState->Transitions)
 					{
 						FixChangedStateLinkName(Transition.State, IDToName);
@@ -438,133 +386,6 @@ namespace FStateTreeEditorHelpers
 				}
 			}
 		}
-
-	}
-
-	void ValidateLinkedVariables(UStateTree& StateTree)
-	{
-		UStateTreeEditorData* TreeData = Cast<UStateTreeEditorData>(StateTree.EditorData);
-		if (!TreeData)
-		{
-			return;
-		}
-
-		// Make sure all state links are valid and update the names if needed.
-		TArray<UStateTreeState*> Stack;
-
-		// Create ID to state name map.
-		TMap<FGuid, FName> IDToName;
-
-		for (UStateTreeState* Routine : TreeData->Routines)
-		{
-			if (Routine)
-			{
-				Stack.Reset();
-				Stack.Add(Routine);
-				while (Stack.Num() > 0)
-				{
-					UStateTreeState* CurState = Stack.Pop();
-					for (const UStateTreeEvaluatorBase* Evaluator : CurState->Evaluators)
-					{
-						if (Evaluator)
-						{
-							IDToName.Add(Evaluator->GetID(), Evaluator->GetName());
-						}
-					}
-
-					for (UStateTreeState* ChildState : CurState->Children)
-					{
-						if (ChildState)
-						{
-							Stack.Append(CurState->Children);
-						}
-					}
-				}
-			}
-		}
-
-		for (UStateTreeState* Routine : TreeData->Routines)
-		{
-			if (Routine)
-			{
-				// Fix changed names.
-				Stack.Reset();
-				Stack.Add(Routine);
-				while (Stack.Num() > 0)
-				{
-					UStateTreeState* CurState = Stack.Pop();
-
-					FixChangedVariableName(CurState->EnterConditions, IDToName);
-
-					for (FStateTreeTransition& Transition : CurState->Transitions)
-					{
-						FixChangedVariableName(Transition.Conditions, IDToName);
-					}
-
-					for (UStateTreeState* ChildState : CurState->Children)
-					{
-						if (ChildState)
-						{
-							Stack.Append(CurState->Children);
-						}
-					}
-				}
-			}
-		}
-	}
-
-	void ValidateEvaluatorsAndTracks(UStateTree& StateTree)
-	{
-		UStateTreeEditorData* TreeData = Cast<UStateTreeEditorData>(StateTree.EditorData);
-		if (!TreeData)
-		{
-			return;
-		}
-
-		TArray<UStateTreeState*> Stack;
-		bool bChanged = false;
-
-		for (UStateTreeState* Routine : TreeData->Routines)
-		{
-			if (Routine)
-			{
-				// Make sure all parameterset are valid and update the layout if needed.
-				Stack.Reset();
-				Stack.Add(Routine);
-				while (Stack.Num() > 0)
-				{
-					UStateTreeState* CurState = Stack.Pop();
-					for (UStateTreeEvaluatorBase* Evaluator : CurState->Evaluators)
-					{
-						if (Evaluator)
-						{
-							bChanged |= Evaluator->ValidateParameterLayout();
-						}
-					}
-
-					for (UStateTreeTaskBase* Task : CurState->Tasks)
-					{
-						if (Task)
-						{
-							bChanged |= Task->ValidateParameterLayout();
-						}
-					}
-
-					for (UStateTreeState* ChildState : CurState->Children)
-					{
-						if (ChildState)
-						{
-							Stack.Append(CurState->Children);
-						}
-					}
-				}
-			}
-		}
-
-		if (bChanged)
-		{
-			UE::StateTree::Delegates::OnParametersInvalidated.Broadcast(StateTree);
-		}
 	}
 
 	void UpdateParents(UStateTree& StateTree)
@@ -621,11 +442,9 @@ void FStateTreeEditor::UpdateAsset()
 		return;
 	}
 
-	FStateTreeEditorHelpers::UpdateParents(*StateTree);
-	FStateTreeEditorHelpers::RemoveUnusedBindings(*StateTree);
-	FStateTreeEditorHelpers::ValidateLinkedStates(*StateTree);
-	FStateTreeEditorHelpers::ValidateLinkedVariables(*StateTree);
-	FStateTreeEditorHelpers::ValidateEvaluatorsAndTracks(*StateTree);
+	UE::StateTree::Editor::Internal::UpdateParents(*StateTree);
+	UE::StateTree::Editor::Internal::RemoveUnusedBindings(*StateTree);
+	UE::StateTree::Editor::Internal::ValidateLinkedStates(*StateTree);
 
 	FStateTreeBaker Baker;
 	if (!Baker.Bake(*StateTree))
diff --git a/Engine/Plugins/Runtime/StateTree/Source/StateTreeEditorModule/Private/StateTreeEditor.h b/Engine/Plugins/Runtime/StateTree/Source/StateTreeEditorModule/Private/StateTreeEditor.h
--- a/Engine/Plugins/Runtime/StateTree/Source/StateTreeEditorModule/Private/StateTreeEditor.h
+++ b/Engine/Plugins/Runtime/StateTree/Source/StateTreeEditorModule/Private/StateTreeEditor.h
@@ -62,7 +62,6 @@ private:
 
 	void HandleModelAssetChanged();
 	void OnIdentifierChanged(const UStateTree& StateTree);
-	void OnParameterLayoutChanged(const UStateTree& StateTree);
 	void OnSchemaChanged(const UStateTree& StateTree);
 
 	/* State Tree being edited */
diff --git a/Engine/Plugins/Runtime/StateTree/Source/StateTreeEditorModule/Private/StateTreeEditorData.cpp b/Engine/Plugins/Runtime/StateTree/Source/StateTreeEditorModule/Private/StateTreeEditorData.cpp
--- a/Engine/Plugins/Runtime/StateTree/Source/StateTreeEditorModule/Private/StateTreeEditorData.cpp
+++ b/Engine/Plugins/Runtime/StateTree/Source/StateTreeEditorModule/Private/StateTreeEditorData.cpp
@@ -128,7 +128,7 @@ void UStateTreeEditorData::VisitHierarchy(TFunctionRef<bool(const UStateTreeStat
 			Stack.RemoveAt(0);
 
 			// Evaluators
-			for (const FStateTreeEvaluatorItem& Item : State->Evaluators2)
+			for (const FStateTreeEvaluatorItem& Item : State->Evaluators)
 			{
 				if (const FStateTreeEvaluator2Base* Evaluator = Item.Type.GetPtr<FStateTreeEvaluator2Base>())
 				{
@@ -142,7 +142,7 @@ void UStateTreeEditorData::VisitHierarchy(TFunctionRef<bool(const UStateTreeStat
 			if (bContinue)
 			{
 				// Enter conditions
-				for (const FStateTreeConditionItem& Item : State->EnterConditions2)
+				for (const FStateTreeConditionItem& Item : State->EnterConditions)
 				{
 					if (const FStateTreeConditionBase* Cond = Item.Type.GetPtr<FStateTreeConditionBase>())
 					{
@@ -157,7 +157,7 @@ void UStateTreeEditorData::VisitHierarchy(TFunctionRef<bool(const UStateTreeStat
 			if (bContinue)
 			{
 				// Tasks
-				for (const FStateTreeTaskItem& Item : State->Tasks2)
+				for (const FStateTreeTaskItem& Item : State->Tasks)
 				{
 					if (const FStateTreeTask2Base* Task = Item.Type.GetPtr<FStateTreeTask2Base>())
 					{
@@ -172,7 +172,7 @@ void UStateTreeEditorData::VisitHierarchy(TFunctionRef<bool(const UStateTreeStat
 			if (bContinue)
 			{
 				// Transitions
-				for (const FStateTreeTransition2& Transition : State->Transitions2)
+				for (const FStateTreeTransition& Transition : State->Transitions)
 				{
 					for (const FStateTreeConditionItem& Item : Transition.Conditions)
 					{
diff --git a/Engine/Plugins/Runtime/StateTree/Source/StateTreeEditorModule/Private/StateTreeEditorModule.cpp b/Engine/Plugins/Runtime/StateTree/Source/StateTreeEditorModule/Private/StateTreeEditorModule.cpp
--- a/Engine/Plugins/Runtime/StateTree/Source/StateTreeEditorModule/Private/StateTreeEditorModule.cpp
+++ b/Engine/Plugins/Runtime/StateTree/Source/StateTreeEditorModule/Private/StateTreeEditorModule.cpp
@@ -5,24 +5,14 @@
 #include "Developer/AssetTools/Public/IAssetTools.h"
 #include "Developer/AssetTools/Public/AssetToolsModule.h"
 #include "AssetTypeActions_StateTree.h"
-#include "Customizations/StateTreeVariableDescDetails.h"
-#include "Customizations/StateTreeVariableLayoutDetails.h"
-#include "Customizations/StateTreeVariableDetails.h"
-#include "Customizations/StateTreeParameterLayoutDetails.h"
-#include "Customizations/StateTreeConditionDetails.h"
 #include "Customizations/StateTreeTransitionDetails.h"
 #include "Customizations/StateTreeStateLinkDetails.h"
-#include "Customizations/StateTreeTransition2Details.h"
 #include "Customizations/StateTreeStateDetails.h"
 #include "Customizations/StateTreeBindableItemDetails.h"
-#include "Customizations/StateTreeEditorPropertyPathDetails.h"
 #include "Customizations/StateTreeConditionItemDetails.h"
 #include "Customizations/StateTreeAnyEnumDetails.h"
-#include "AIGraphTypes.h" // Class cache
 #include "StateTreeEditor.h"
 #include "StateTree.h"
-#include "StateTreeState.h"
-#include "StateTreeTaskBase.h"
 #include "StateTreeEditorStyle.h"
 
 
@@ -46,17 +36,10 @@ void FStateTreeEditorModule::StartupModule()
 
 	// Register the details customizer
 	FPropertyEditorModule& PropertyModule = FModuleManager::LoadModuleChecked<FPropertyEditorModule>("PropertyEditor");
-	PropertyModule.RegisterCustomPropertyTypeLayout("StateTreeVariableDesc", FOnGetPropertyTypeCustomizationInstance::CreateStatic(&FStateTreeVariableDescDetails::MakeInstance));
-	PropertyModule.RegisterCustomPropertyTypeLayout("StateTreeVariableLayout", FOnGetPropertyTypeCustomizationInstance::CreateStatic(&FStateTreeVariableLayoutDetails::MakeInstance));
-	PropertyModule.RegisterCustomPropertyTypeLayout("StateTreeVariable", FOnGetPropertyTypeCustomizationInstance::CreateStatic(&FStateTreeVariableDetails::MakeInstance));
-	PropertyModule.RegisterCustomPropertyTypeLayout("StateTreeParameterLayout", FOnGetPropertyTypeCustomizationInstance::CreateStatic(&FStateTreeParameterLayoutDetails::MakeInstance));
-	PropertyModule.RegisterCustomPropertyTypeLayout("StateTreeCondition", FOnGetPropertyTypeCustomizationInstance::CreateStatic(&FStateTreeConditionDetails::MakeInstance));
 	PropertyModule.RegisterCustomPropertyTypeLayout("StateTreeTransition", FOnGetPropertyTypeCustomizationInstance::CreateStatic(&FStateTreeTransitionDetails::MakeInstance));
 	PropertyModule.RegisterCustomPropertyTypeLayout("StateTreeStateLink", FOnGetPropertyTypeCustomizationInstance::CreateStatic(&FStateTreeStateLinkDetails::MakeInstance));
-	PropertyModule.RegisterCustomPropertyTypeLayout("StateTreeTransition2", FOnGetPropertyTypeCustomizationInstance::CreateStatic(&FStateTreeTransition2Details::MakeInstance));
 	PropertyModule.RegisterCustomPropertyTypeLayout("StateTreeEvaluatorItem", FOnGetPropertyTypeCustomizationInstance::CreateStatic(&FStateTreeBindableItemDetails::MakeInstance));
 	PropertyModule.RegisterCustomPropertyTypeLayout("StateTreeTaskItem", FOnGetPropertyTypeCustomizationInstance::CreateStatic(&FStateTreeBindableItemDetails::MakeInstance));
-	PropertyModule.RegisterCustomPropertyTypeLayout("StateTreeEditorPropertyPath", FOnGetPropertyTypeCustomizationInstance::CreateStatic(&FStateTreeEditorPropertyPathDetails::MakeInstance));
 	PropertyModule.RegisterCustomPropertyTypeLayout("StateTreeConditionItem", FOnGetPropertyTypeCustomizationInstance::CreateStatic(&FStateTreeConditionItemDetails::MakeInstance));
 	PropertyModule.RegisterCustomPropertyTypeLayout("StateTreeAnyEnum", FOnGetPropertyTypeCustomizationInstance::CreateStatic(&FStateTreeAnyEnumDetails::MakeInstance));
 	PropertyModule.RegisterCustomClassLayout("StateTreeState", FOnGetDetailCustomizationInstance::CreateStatic(&FStateTreeStateDetails::MakeInstance));
@@ -66,7 +49,6 @@ void FStateTreeEditorModule::StartupModule()
 
 void FStateTreeEditorModule::ShutdownModule()
 {
-	TaskClassCache.Reset();
 	MenuExtensibilityManager.Reset();
 	ToolBarExtensibilityManager.Reset();
 
@@ -90,17 +72,10 @@ void FStateTreeEditorModule::ShutdownModule()
 	if (FModuleManager::Get().IsModuleLoaded("PropertyEditor"))
 	{
 		FPropertyEditorModule& PropertyModule = FModuleManager::LoadModuleChecked<FPropertyEditorModule>("PropertyEditor");
-		PropertyModule.UnregisterCustomPropertyTypeLayout("StateTreeVariableDesc");
-		PropertyModule.UnregisterCustomPropertyTypeLayout("StateTreeVariableLayout");
-		PropertyModule.UnregisterCustomPropertyTypeLayout("StateTreeVariable");
-		PropertyModule.UnregisterCustomPropertyTypeLayout("StateTreeParameterLayout");
-		PropertyModule.UnregisterCustomPropertyTypeLayout("StateTreeCondition");
 		PropertyModule.UnregisterCustomPropertyTypeLayout("StateTreeTransition");
-		PropertyModule.UnregisterCustomPropertyTypeLayout("StateTreeTransition2");
 		PropertyModule.UnregisterCustomPropertyTypeLayout("StateTreeStateLink");
 		PropertyModule.UnregisterCustomPropertyTypeLayout("StateTreeEvaluatorItem");
 		PropertyModule.UnregisterCustomPropertyTypeLayout("StateTreeTaskItem");
-		PropertyModule.UnregisterCustomPropertyTypeLayout("StateTreeEditorPropertyPath");
 		PropertyModule.UnregisterCustomPropertyTypeLayout("StateTreeConditionItem");
 		PropertyModule.UnregisterCustomPropertyTypeLayout("StateTreeAnyEnum");
 		PropertyModule.NotifyCustomizationModuleChanged();
@@ -110,12 +85,6 @@ void FStateTreeEditorModule::ShutdownModule()
 
 TSharedRef<IStateTreeEditor> FStateTreeEditorModule::CreateStateTreeEditor(const EToolkitMode::Type Mode, const TSharedPtr< class IToolkitHost >& InitToolkitHost, UStateTree* StateTree)
 {
-	if (!TaskClassCache.IsValid())
-	{
-		TaskClassCache = MakeShareable(new FGraphNodeClassHelper(UStateTreeTaskBase::StaticClass()));
-		TaskClassCache->UpdateAvailableBlueprintClasses();
-	}
-
 	TSharedRef<FStateTreeEditor> NewEditor(new FStateTreeEditor());
 	NewEditor->InitEditor(Mode, InitToolkitHost, StateTree);
 	return NewEditor;
diff --git a/Engine/Plugins/Runtime/StateTree/Source/StateTreeEditorModule/Private/StateTreePropertyHelpers.cpp b/Engine/Plugins/Runtime/StateTree/Source/StateTreeEditorModule/Private/StateTreePropertyHelpers.cpp
deleted file mode 100644
diff --git a/Engine/Plugins/Runtime/StateTree/Source/StateTreeEditorModule/Private/StateTreePropertyHelpers.h b/Engine/Plugins/Runtime/StateTree/Source/StateTreeEditorModule/Private/StateTreePropertyHelpers.h
--- a/Engine/Plugins/Runtime/StateTree/Source/StateTreeEditorModule/Private/StateTreePropertyHelpers.h
+++ b/Engine/Plugins/Runtime/StateTree/Source/StateTreeEditorModule/Private/StateTreePropertyHelpers.h
@@ -5,105 +5,13 @@
 #include "CoreMinimal.h"
 #include "PropertyEditorModule.h"
 #include "Editor.h"
-#include "StateTreeVariable.h"
 
 struct FGuid;
 struct FSlateBrush;
 
 #define LOCTEXT_NAMESPACE "StateTreeEditor"
 
-namespace UE { namespace StateTree { namespace PropertyHelpers {
-
-/**
- * Gets value from a specific property. Used to deal with union data in FStateTreeVariable.
- * @param ValueProperty Handle to property where value is got from.
- * @return Requested value as optional, in case of multiple values the optional is unset.
- */
-template<typename T>
-TOptional<T> GetVariableValue(const TSharedPtr<IPropertyHandle>& ValueProperty)
-{
-	if (!ValueProperty)
-	{
-		return TOptional<T>();
-	}
-
-	T Value = T();
-	bool bValueSet = false;
-
-	TArray<void*> RawData;
-	ValueProperty->AccessRawData(RawData);
-	for (void* Data : RawData)
-	{
-		if (Data)
-		{
-			T CurValue = *reinterpret_cast<T*>(Data);
-			if (!bValueSet)
-			{
-				bValueSet = true;
-				Value = CurValue;
-			}
-			else if (CurValue != Value)
-			{
-				// Multiple values
-				return TOptional<T>();
-			}
-		}
-	}
-
-	return TOptional<T>(Value);
-}
-
-/**
- * Sets a property to specific value. Used to deal with union data in FStateTreeVariable.
- * @param ValueProperty Handle to property where value is to be stored.
- * @param NewValue Value to set.
- * @param Flags Property set flags.
- */
-template<typename T>
-void SetVariableValue(TSharedPtr<IPropertyHandle>& ValueProperty, T NewValue, EPropertyValueSetFlags::Type Flags = 0)
-{
-	if (!ValueProperty)
-	{
-		return;
-	}
-
-	const bool bTransactable = (Flags & EPropertyValueSetFlags::NotTransactable) == 0;
-	bool bNotifiedPreChange = false;
-	TArray<void*> RawData;
-	ValueProperty->AccessRawData(RawData);
-	for (void* Data : RawData)
-	{
-		if (Data)
-		{
-			if (!bNotifiedPreChange)
-			{
-				if (bTransactable && GEditor)
-				{
-					GEditor->BeginTransaction(FText::Format(LOCTEXT("SetPropertyValue", "Set {0}"), ValueProperty->GetPropertyDisplayName()));
-				}
-				ValueProperty->NotifyPreChange();
-				bNotifiedPreChange = true;
-			}
-
-			// Make sure the value does not end up containing garbage.
-			FMemory::Memzero(Data, FStateTreeVariable::ValueStorageSizeBytes);
-
-			T* Value = reinterpret_cast<T*>(Data);
-			*Value = NewValue;
-		}
-	}
-
-	if (bNotifiedPreChange)
-	{
-		ValueProperty->NotifyPostChange(EPropertyChangeType::ValueSet);
-		if (bTransactable && GEditor)
-		{
-			GEditor->EndTransaction();
-		}
-	}
-
-	ValueProperty->NotifyFinishedChangingProperties();
-}
+namespace UE::StateTree::PropertyHelpers {
 
 /**
  * Gets a struct value from property handle, checks type before access. Expects T is struct.
@@ -210,12 +118,6 @@ FPropertyAccess::Result SetStructValue(const TSharedPtr<IPropertyHandle>& ValueP
 	return FPropertyAccess::Success;
 }
 
-
-/**
- * @return Icon for a specified variable type.
- */
-const FSlateBrush* GetTypeIcon(EStateTreeVariableType Type);
-
-} } } // UE::StateTree::PropertyHelpers
+} // UE::StateTree::PropertyHelpers
 
 #undef LOCTEXT_NAMESPACE
\ No newline at end of file
diff --git a/Engine/Plugins/Runtime/StateTree/Source/StateTreeEditorModule/Private/StateTreeState.cpp b/Engine/Plugins/Runtime/StateTree/Source/StateTreeEditorModule/Private/StateTreeState.cpp
--- a/Engine/Plugins/Runtime/StateTree/Source/StateTreeEditorModule/Private/StateTreeState.cpp
+++ b/Engine/Plugins/Runtime/StateTree/Source/StateTreeEditorModule/Private/StateTreeState.cpp
@@ -1,13 +1,15 @@
 // Copyright Epic Games, Inc. All Rights Reserved.
 
 #include "StateTreeState.h"
-#include "StateTreeEditorData.h"
 #include "StateTree.h"
 #include "StateTreeConditionBase.h"
 #include "StateTreeDelegates.h"
 #include "CoreMinimal.h"
 #include "UObject/Field.h"
 
+//////////////////////////////////////////////////////////////////////////
+// FStateTreeStateLink
+
 void FStateTreeStateLink::Set(const EStateTreeTransitionType InType, const UStateTreeState* InState)
 {
 	Type = InType;
@@ -20,13 +22,19 @@ void FStateTreeStateLink::Set(const EStateTreeTransitionType InType, const UStat
 }
 
 
-FStateTreeTransition2::FStateTreeTransition2(const EStateTreeTransitionEvent InEvent, const EStateTreeTransitionType InType, const UStateTreeState* InState)
+//////////////////////////////////////////////////////////////////////////
+// FStateTreeTransition
+
+FStateTreeTransition::FStateTreeTransition(const EStateTreeTransitionEvent InEvent, const EStateTreeTransitionType InType, const UStateTreeState* InState)
 {
 	Event = InEvent;
 	State.Set(InType, InState);
 }
 
 
+//////////////////////////////////////////////////////////////////////////
+// UStateTreeState
+
 UStateTreeState::UStateTreeState(const FObjectInitializer& ObjectInitializer)
 	: Super(ObjectInitializer)
 	, ID(FGuid::NewGuid())
@@ -58,165 +66,57 @@ void UStateTreeState::PostEditChangeChainProperty(FPropertyChangedChainEvent& Pr
 			}
 		}
 
-		if (PropertyChangedEvent.ChangeType == EPropertyChangeType::Duplicate && MemberProperty)
+		if (MemberProperty)
 		{
 			// Ensure unique ID on duplicated items.
-			if (MemberProperty->GetFName() == GET_MEMBER_NAME_CHECKED(UStateTreeState, Evaluators))
-			{
-				const int32 ArrayIndex = PropertyChangedEvent.GetArrayIndex(MemberProperty->GetFName().ToString());
-				if (Evaluators.IsValidIndex(ArrayIndex))
-				{
-					if (UStateTreeEvaluatorBase* Eval = Evaluators[ArrayIndex])
-					{
-						Eval->SetNewUniqueID();
-						Eval->SetName(FName(Eval->GetName().ToString() + TEXT(" Duplicate")));
-					}
-				}
-			}
-			if (MemberProperty->GetFName() == GET_MEMBER_NAME_CHECKED(UStateTreeState, Tasks))
+			if (PropertyChangedEvent.ChangeType == EPropertyChangeType::Duplicate)
 			{
-				const int32 ArrayIndex = PropertyChangedEvent.GetArrayIndex(MemberProperty->GetFName().ToString());
-				if (Tasks.IsValidIndex(ArrayIndex))
+				if (MemberProperty->GetFName() == GET_MEMBER_NAME_CHECKED(UStateTreeState, Evaluators))
 				{
-					if (UStateTreeTaskBase* Task = Tasks[ArrayIndex])
+					const int32 ArrayIndex = PropertyChangedEvent.GetArrayIndex(MemberProperty->GetFName().ToString());
+					if (Evaluators.IsValidIndex(ArrayIndex))
 					{
-						Task->SetNewUniqueID();
-						Task->Name = FName(Task->Name.ToString() + TEXT(" Duplicate"));
-					}
-				}
-			}
-		}
-
-
-		if (IsV2())
-		{
-			if (MemberProperty)
-			{
-				// Ensure unique ID on duplicated items.
-				if (PropertyChangedEvent.ChangeType == EPropertyChangeType::Duplicate)
-				{
-					if (MemberProperty->GetFName() == GET_MEMBER_NAME_CHECKED(UStateTreeState, Evaluators2))
-					{
-						const int32 ArrayIndex = PropertyChangedEvent.GetArrayIndex(MemberProperty->GetFName().ToString());
-						if (Evaluators2.IsValidIndex(ArrayIndex))
+						if (FStateTreeEvaluator2Base* Eval = Evaluators[ArrayIndex].Type.GetMutablePtr<FStateTreeEvaluator2Base>())
 						{
-							if (FStateTreeEvaluator2Base* Eval = Evaluators2[ArrayIndex].Type.GetMutablePtr<FStateTreeEvaluator2Base>())
-							{
-								Eval->ID = FGuid::NewGuid();
-								Eval->Name = FName(Eval->Name.ToString() + TEXT(" Duplicate"));
-							}
+							Eval->ID = FGuid::NewGuid();
+							Eval->Name = FName(Eval->Name.ToString() + TEXT(" Duplicate"));
 						}
 					}
-					if (MemberProperty->GetFName() == GET_MEMBER_NAME_CHECKED(UStateTreeState, Tasks2))
+				}
+				if (MemberProperty->GetFName() == GET_MEMBER_NAME_CHECKED(UStateTreeState, Tasks))
+				{
+					const int32 ArrayIndex = PropertyChangedEvent.GetArrayIndex(MemberProperty->GetFName().ToString());
+					if (Tasks.IsValidIndex(ArrayIndex))
 					{
-						const int32 ArrayIndex = PropertyChangedEvent.GetArrayIndex(MemberProperty->GetFName().ToString());
-						if (Tasks2.IsValidIndex(ArrayIndex))
+						if (FStateTreeTask2Base* Task = Tasks[ArrayIndex].Type.GetMutablePtr<FStateTreeTask2Base>())
 						{
-							if (FStateTreeTask2Base* Task = Tasks2[ArrayIndex].Type.GetMutablePtr<FStateTreeTask2Base>())
-							{
-								Task->ID = FGuid::NewGuid();
-								Task->Name = FName(Task->Name.ToString() + TEXT(" Duplicate"));
-							}
+							Task->ID = FGuid::NewGuid();
+							Task->Name = FName(Task->Name.ToString() + TEXT(" Duplicate"));
 						}
 					}
-					if (MemberProperty->GetFName() == GET_MEMBER_NAME_CHECKED(UStateTreeState, EnterConditions2))
+				}
+				if (MemberProperty->GetFName() == GET_MEMBER_NAME_CHECKED(UStateTreeState, EnterConditions))
+				{
+					const int32 ArrayIndex = PropertyChangedEvent.GetArrayIndex(MemberProperty->GetFName().ToString());
+					if (EnterConditions.IsValidIndex(ArrayIndex))
 					{
-						const int32 ArrayIndex = PropertyChangedEvent.GetArrayIndex(MemberProperty->GetFName().ToString());
-						if (Tasks2.IsValidIndex(ArrayIndex))
+						if (FStateTreeConditionBase* Cond = EnterConditions[ArrayIndex].Type.GetMutablePtr<FStateTreeConditionBase>())
 						{
-							if (FStateTreeConditionBase* Cond = EnterConditions2[ArrayIndex].Type.GetMutablePtr<FStateTreeConditionBase>())
-							{
-								Cond->ID = FGuid::NewGuid();
-							}
+							Cond->ID = FGuid::NewGuid();
 						}
 					}
 				}
+				// TODO: Transition conditions.
 			}
 		}
 	}
 }
 
-void UStateTreeState::PostLoad()
-{
-	Super::PostLoad();
-	
-	UStateTree* StateTree = GetTypedOuter<UStateTree>();
-	if (StateTree->IsV2())
-	{
-		// Convert existing default transitions.
-		Modify();
-
-		if (StateFailedTransition.Type != EStateTreeTransitionType::NotSet
-			&& StateFailedTransition.Type != EStateTreeTransitionType::SelectChildState)
-		{
-			if (StateFailedTransition.IsValid())
-			{
-				FStateTreeTransition2& Transition = Transitions2.AddDefaulted_GetRef();
-				Transition.Event = EStateTreeTransitionEvent::OnFailed;
-				Transition.State = StateFailedTransition;
-			}
-			// Reset
-			StateFailedTransition = FStateTreeStateLink(EStateTreeTransitionType::NotSet);
-		}
-
-		if (StateDoneTransition.Type != EStateTreeTransitionType::NotSet
-			&& StateDoneTransition.Type != EStateTreeTransitionType::SelectChildState) // Select child state is inferred automatically in V2
-		{
-			if (StateDoneTransition.IsValid())
-			{
-				FStateTreeTransition2& Transition = Transitions2.AddDefaulted_GetRef();
-				Transition.Event = EStateTreeTransitionEvent::OnCompleted;
-				Transition.State = StateDoneTransition;
-			}
-			// Reset
-			StateDoneTransition = FStateTreeStateLink(EStateTreeTransitionType::NotSet);
-		}
-	}
-}
-
 #endif
 
-void UStateTreeState::GetVisibleVariables(FStateTreeVariableLayout& Variables) const
+const FStateTreeTask2Base* UStateTreeState::GetTaskByID(FGuid InID) const
 {
-	// Add input parameters
-	UStateTree* StateTree = GetTypedOuter<UStateTree>();
-	if (ensure(StateTree))
-	{
-		const FStateTreeVariableLayout& ParameterLayout = StateTree->GetInputParameterLayout();
-		for (const FStateTreeVariableDesc& Var : ParameterLayout.Variables)
-		{
-			Variables.DefineVariable(Var);
-		}
-	}
-
-	// Collect states variable descs up to the highest parent
-	for (const UStateTreeState* State = this; State; State = State->Parent)
-	{
-		for (const UStateTreeEvaluatorBase* Evaluator : State->Evaluators)
-		{
-			if (Evaluator)
-			{
-				Evaluator->DefineOutputVariables(Variables);
-			}
-		}
-	}
-}
-
-const UStateTreeTaskBase* UStateTreeState::GetTaskByID(FGuid InID) const
-{
-	const int32 Index = Tasks.IndexOfByPredicate([InID](const UStateTreeTaskBase* Task) -> bool { return Task->ID == InID; });
-	return Index != INDEX_NONE ? Tasks[Index] : nullptr;
-}
-
-UStateTreeTaskBase* UStateTreeState::GetTaskByID(FGuid InID)
-{
-	const int32 Index = Tasks.IndexOfByPredicate([InID](const UStateTreeTaskBase* Task) -> bool { return Task->ID == InID; });
-	return Index != INDEX_NONE ? Tasks[Index] : nullptr;
-}
-
-const FStateTreeTask2Base* UStateTreeState::GetTask2ByID(FGuid InID) const
-{
-	for (const FStateTreeTaskItem& TaskItem : Tasks2)
+	for (const FStateTreeTaskItem& TaskItem : Tasks)
 	{
 		if (const FStateTreeTask2Base* Task = TaskItem.Type.GetPtr<FStateTreeTask2Base>())
 		{
@@ -229,9 +129,9 @@ const FStateTreeTask2Base* UStateTreeState::GetTask2ByID(FGuid InID) const
 	return nullptr;
 }
 
-FStateTreeTask2Base* UStateTreeState::GetTask2ByID(FGuid InID)
+FStateTreeTask2Base* UStateTreeState::GetTaskByID(FGuid InID)
 {
-	for (FStateTreeTaskItem& TaskItem : Tasks2)
+	for (FStateTreeTaskItem& TaskItem : Tasks)
 	{
 		if (FStateTreeTask2Base* Task = TaskItem.Type.GetMutablePtr<FStateTreeTask2Base>())
 		{
@@ -264,12 +164,3 @@ UStateTreeState* UStateTreeState::GetNextSiblingState() const
 	}
 	return nullptr;
 }
-
-bool UStateTreeState::IsV2() const
-{
-	if (const UStateTree* StateTree = GetTypedOuter<UStateTree>())
-	{
-		return StateTree->IsV2();
-	}
-	return false;
-}
\ No newline at end of file
diff --git a/Engine/Plugins/Runtime/StateTree/Source/StateTreeEditorModule/Private/StateTreeVariableProvider.cpp b/Engine/Plugins/Runtime/StateTree/Source/StateTreeEditorModule/Private/StateTreeVariableProvider.cpp
deleted file mode 100644
diff --git a/Engine/Plugins/Runtime/StateTree/Source/StateTreeEditorModule/Private/StateTreeViewModel.cpp b/Engine/Plugins/Runtime/StateTree/Source/StateTreeEditorModule/Private/StateTreeViewModel.cpp
--- a/Engine/Plugins/Runtime/StateTree/Source/StateTreeEditorModule/Private/StateTreeViewModel.cpp
+++ b/Engine/Plugins/Runtime/StateTree/Source/StateTreeEditorModule/Private/StateTreeViewModel.cpp
@@ -158,7 +158,6 @@ FStateTreeViewModel::~FStateTreeViewModel()
 	GEditor->UnregisterForUndo(this);
 
 	UE::StateTree::Delegates::OnIdentifierChanged.RemoveAll(this);
-	UE::StateTree::Delegates::OnParameterLayoutChanged.RemoveAll(this);
 }
 
 void FStateTreeViewModel::Init(UStateTreeEditorData* InTreeData)
@@ -375,12 +374,6 @@ void FStateTreeViewModel::AddChildState(UStateTreeState* ParentState)
 
 	UStateTreeState* NewState = NewObject<UStateTreeState>(TreeData, FName(), RF_Transactional);
 
-	// Make sure the parent state is set to select when it has children.
-/*	if (ParentState->StateDoneTransition.Type != EStateTreeTransitionType::SelectChildState)
-	{
-		ParentState->StateDoneTransition.Type = EStateTreeTransitionType::SelectChildState;
-	}*/
-	
 	ParentState->Children.Add(NewState);
 	NewState->Parent = ParentState;
 
diff --git a/Engine/Plugins/Runtime/StateTree/Source/StateTreeEditorModule/Public/StateTreeBaker.h b/Engine/Plugins/Runtime/StateTree/Source/StateTreeEditorModule/Public/StateTreeBaker.h
--- a/Engine/Plugins/Runtime/StateTree/Source/StateTreeEditorModule/Public/StateTreeBaker.h
+++ b/Engine/Plugins/Runtime/StateTree/Source/StateTreeEditorModule/Public/StateTreeBaker.h
@@ -10,8 +10,6 @@ class UStateTree;
 class UStateTreeState;
 class UStateTreeEditorData;
 struct FStateTreeCondition;
-class UStateTreeEvaluatorBase;
-class UStateTreeTaskBase;
 struct FStateTreeConditionItem;
 struct FStateTreeStateLink;
 
@@ -23,7 +21,6 @@ struct STATETREEEDITORMODULE_API FStateTreeBaker
 {
 public:
 	bool Bake(UStateTree& InStateTree);
-	bool Bake2(UStateTree& InStateTree);
 
 private:
 	enum class EBindingContainerType : uint8
@@ -32,35 +29,13 @@ private:
 		Task,
 	};
 
-	// Helper struct used when collecting evaluators, keeping the state reference so that we can have better errors.
-	struct FSourceEvaluator
-	{
-		FSourceEvaluator(UStateTreeEvaluatorBase* InEval, UStateTreeState* InState) : Eval(InEval), State(InState) {}
-		UStateTreeEvaluatorBase* Eval = nullptr;
-		UStateTreeState* State = nullptr;
-	};
-
-	// Helper struct used when collecting tracks, keeping the state reference so that we can have better errors.
-	struct FSourceTask
-	{
-		FSourceTask(UStateTreeTaskBase* InTask, UStateTreeState* InState) : Task(InTask), State(InState) {}
-		UStateTreeTaskBase* Task = nullptr;
-		UStateTreeState* State = nullptr;
-	};
+	bool ResolveTransitionState(const UStateTreeState& SourceState, const TCHAR* ContextStr, const FStateTreeStateLink& Link, FStateTreeHandle& OutTransitionHandle) const;
+	FStateTreeHandle GetStateHandle(const FGuid& StateID) const;
 
-	void DefineParameters() const;
-	void ResolveParameters() const;
 	bool CreateStates();
 	bool CreateStateRecursive(UStateTreeState& State, const FStateTreeHandle Parent);
-	bool ResolveTransitionState(const UStateTreeState& SourceState, const TCHAR* ContextStr, const FStateTreeStateLink& Link, FStateTreeHandle& OutTransitionHandle) const;
 	bool CreateStateTransitions();
-	bool CreateCondition(const FStateTreeCondition& Cond);
-	FStateTreeHandle GetStateHandle(const FGuid& StateID) const;
-
-	bool CreateStates2();
-	bool CreateStateRecursive2(UStateTreeState& State, const FStateTreeHandle Parent);
-	bool CreateStateTransitions2();
-	bool CreateCondition2(const FStateTreeConditionItem& CondItem);
+	bool CreateCondition(const FStateTreeConditionItem& CondItem);
 	bool GetAndValidateBindings(const FStateTreeBindableStructDesc& TargetStruct, TArray<FStateTreeEditorPropertyBinding>& OutBindings) const;
 	bool IsPropertyAnyEnum(const FStateTreeBindableStructDesc& Struct, FStateTreeEditorPropertyPath Path) const;
 	bool CreateExternalItemHandles(FStructView Item);
@@ -69,7 +44,5 @@ private:
 	UStateTreeEditorData* TreeData = nullptr;
 	TMap<FGuid, int32> IDToState;
 	TArray<UStateTreeState*> SourceStates;
-	TArray<FSourceEvaluator> SourceEvaluators;
-	TArray<FSourceTask> SourceTasks;
 	FStateTreePropertyBindingCompiler BindingsCompiler;
 };
diff --git a/Engine/Plugins/Runtime/StateTree/Source/StateTreeEditorModule/Public/StateTreeEditorModule.h b/Engine/Plugins/Runtime/StateTree/Source/StateTreeEditorModule/Public/StateTreeEditorModule.h
--- a/Engine/Plugins/Runtime/StateTree/Source/StateTreeEditorModule/Public/StateTreeEditorModule.h
+++ b/Engine/Plugins/Runtime/StateTree/Source/StateTreeEditorModule/Public/StateTreeEditorModule.h
@@ -3,8 +3,6 @@
 #pragma  once
 
 #include "CoreMinimal.h"
-#include "Modules/ModuleManager.h"
-#include "AssetTypeCategories.h"
 #include "Toolkits/IToolkitHost.h"
 #include "Toolkits/AssetEditorToolkit.h"
 
@@ -28,11 +26,7 @@ public:
 	virtual TSharedPtr<FExtensibilityManager> GetMenuExtensibilityManager() override { return MenuExtensibilityManager; }
 	virtual TSharedPtr<FExtensibilityManager> GetToolBarExtensibilityManager() override { return ToolBarExtensibilityManager; }
 
-	TSharedPtr<struct FGraphNodeClassHelper> GetTaskClassCache() { return TaskClassCache; }
-
 protected:
-	TSharedPtr<struct FGraphNodeClassHelper> TaskClassCache;
-
 	TSharedPtr<FExtensibilityManager> MenuExtensibilityManager;
 	TSharedPtr<FExtensibilityManager> ToolBarExtensibilityManager;
 
diff --git a/Engine/Plugins/Runtime/StateTree/Source/StateTreeEditorModule/Public/StateTreeState.h b/Engine/Plugins/Runtime/StateTree/Source/StateTreeEditorModule/Public/StateTreeState.h
--- a/Engine/Plugins/Runtime/StateTree/Source/StateTreeEditorModule/Public/StateTreeState.h
+++ b/Engine/Plugins/Runtime/StateTree/Source/StateTreeEditorModule/Public/StateTreeState.h
@@ -2,11 +2,9 @@
 
 #pragma once
 
-#include "StateTreeCondition.h"
 #include "StateTreeEvaluatorBase.h"
 #include "StateTreeTaskBase.h"
 #include "Misc/Guid.h"
-#include "StateTreeVariableProvider.h"
 #include "InstancedStruct.h"
 #include "StateTreeState.generated.h"
 
@@ -37,21 +35,6 @@ public:
 	EStateTreeTransitionType Type;
 };
 
-/**
- * Editor representation of a transition in StateTree
- */
-USTRUCT()
-struct STATETREEEDITORMODULE_API FStateTreeTransition
-{
-	GENERATED_BODY()
-public:
-	UPROPERTY(EditDefaultsOnly, Category = Transition)
-	FStateTreeStateLink State;
-
-	UPROPERTY(EditDefaultsOnly, Category = Transition)
-	TArray<FStateTreeCondition> Conditions;
-};
-
 /**
  * Helper struct for Evaluator details customization (i.e. show summary when collapsed)
  */
@@ -93,12 +76,12 @@ struct STATETREEEDITORMODULE_API FStateTreeConditionItem
  * Editor representation of a transition in StateTree
  */
 USTRUCT()
-struct STATETREEEDITORMODULE_API FStateTreeTransition2
+struct STATETREEEDITORMODULE_API FStateTreeTransition
 {
 	GENERATED_BODY()
 
-	FStateTreeTransition2() = default;
-	FStateTreeTransition2(const EStateTreeTransitionEvent InEvent, const EStateTreeTransitionType InType, const UStateTreeState* InState = nullptr);
+	FStateTreeTransition() = default;
+	FStateTreeTransition(const EStateTreeTransitionEvent InEvent, const EStateTreeTransitionType InType, const UStateTreeState* InState = nullptr);
 
 	template<typename T, typename... TArgs>
 	T& AddCondition(TArgs&&... InArgs)
@@ -129,7 +112,7 @@ struct STATETREEEDITORMODULE_API FStateTreeTransition2
  * Editor representation of a state in StateTree
  */
 UCLASS(BlueprintType, EditInlineNew, CollapseCategories)
-class STATETREEEDITORMODULE_API UStateTreeState : public UObject, public IStateTreeVariableProvider
+class STATETREEEDITORMODULE_API UStateTreeState : public UObject
 {
 	GENERATED_BODY()
 
@@ -138,16 +121,10 @@ public:
 
 #if WITH_EDITOR
 	virtual void PostEditChangeChainProperty(FPropertyChangedChainEvent& PropertyChangedEvent) override;
-	virtual void PostLoad() override;
 #endif
 
-	void GetVisibleVariables(FStateTreeVariableLayout& Variables) const override;
-
-	const UStateTreeTaskBase* GetTaskByID(FGuid ID) const;
-	UStateTreeTaskBase* GetTaskByID(FGuid ID);
-
-	const FStateTreeTask2Base* GetTask2ByID(FGuid ID) const;
-	FStateTreeTask2Base* GetTask2ByID(FGuid ID);
+	const FStateTreeTask2Base* GetTaskByID(FGuid ID) const;
+	FStateTreeTask2Base* GetTaskByID(FGuid ID);
 
 	UStateTreeState* GetNextSiblingState() const;
 
@@ -171,7 +148,7 @@ public:
 	template<typename T, typename... TArgs>
 	T& AddEnterCondition(TArgs&&... InArgs)
 	{
-		FStateTreeConditionItem& CondItem = EnterConditions2.AddDefaulted_GetRef();
+		FStateTreeConditionItem& CondItem = EnterConditions.AddDefaulted_GetRef();
 		CondItem.Type.InitializeAs<T>(Forward<TArgs>(InArgs)...);
 		return CondItem.Type.GetMutable<T>();
 	}
@@ -183,7 +160,7 @@ public:
 	template<typename T, typename... TArgs>
 	T& AddTask(TArgs&&... InArgs)
 	{
-		FStateTreeTaskItem& TaskItem = Tasks2.AddDefaulted_GetRef();
+		FStateTreeTaskItem& TaskItem = Tasks.AddDefaulted_GetRef();
 		TaskItem.Type.InitializeAs<T>(Forward<TArgs>(InArgs)...);
 		T& Task = TaskItem.Type.GetMutable<T>();
 		Task.ID = FGuid::NewGuid();
@@ -197,7 +174,7 @@ public:
 	template<typename T, typename... TArgs>
     T& AddEvaluator(TArgs&&... InArgs)
 	{
-		FStateTreeEvaluatorItem& EvalItem = Evaluators2.AddDefaulted_GetRef();
+		FStateTreeEvaluatorItem& EvalItem = Evaluators.AddDefaulted_GetRef();
 		EvalItem.Type.InitializeAs<T>(Forward<TArgs>(InArgs)...);
 		T& Eval = EvalItem.Type.GetMutable<T>();
 		Eval.ID = FGuid::NewGuid();
@@ -208,52 +185,30 @@ public:
 	 * Adds Transition.
 	 * @return reference to the new Transition. 
 	 */
-	FStateTreeTransition2& AddTransition(const EStateTreeTransitionEvent InEvent, const EStateTreeTransitionType InType, const UStateTreeState* InState = nullptr)
+	FStateTreeTransition& AddTransition(const EStateTreeTransitionEvent InEvent, const EStateTreeTransitionType InType, const UStateTreeState* InState = nullptr)
 	{
-		return Transitions2.Emplace_GetRef(InEvent, InType, InState);
+		return Transitions.Emplace_GetRef(InEvent, InType, InState);
 	}
 
 	// ~StateTree Builder API
 
-	bool IsV2() const;
-	
 	UPROPERTY(EditDefaultsOnly, Category = State)
 	FName Name;
 
 	UPROPERTY()
 	FGuid ID;
 
-	UPROPERTY(EditDefaultsOnly, Category = "State", Meta = (ToolTip = "State to transition to when the State execution is done. If no State Failed transition is set, this transition will be used."))
-	FStateTreeStateLink StateDoneTransition;
-
-	UPROPERTY(EditDefaultsOnly, Category = "State", Meta = (ToolTip = "State to transition to if the State execution fails."))
-	FStateTreeStateLink StateFailedTransition = FStateTreeStateLink(EStateTreeTransitionType::NotSet);
-
-	UPROPERTY(EditDefaultsOnly, Category = "Enter Conditions")
-	TArray<FStateTreeCondition> EnterConditions;
-
-	UPROPERTY(EditDefaultsOnly, Category = "Evaluators", Instanced)
-	TArray<UStateTreeEvaluatorBase*> Evaluators;
-
-	UPROPERTY(EditDefaultsOnly, Category = "Tasks", Instanced)
-	TArray<UStateTreeTaskBase*> Tasks;
-
-	UPROPERTY(EditDefaultsOnly, Category = "Transitions")
-	TArray<FStateTreeTransition> Transitions;
-
-// STATETREE_V2
 	UPROPERTY(EditDefaultsOnly, Category = "Enter Conditions")
-	TArray<FStateTreeConditionItem> EnterConditions2;
+	TArray<FStateTreeConditionItem> EnterConditions;
 
 	UPROPERTY(EditDefaultsOnly, Category = "Evaluators")
-	TArray<FStateTreeEvaluatorItem> Evaluators2;
+	TArray<FStateTreeEvaluatorItem> Evaluators;
 
 	UPROPERTY(EditDefaultsOnly, Category = "Tasks")
-	TArray<FStateTreeTaskItem> Tasks2;
+	TArray<FStateTreeTaskItem> Tasks;
 
 	UPROPERTY(EditDefaultsOnly, Category = "Transitions")
-	TArray<FStateTreeTransition2> Transitions2;
-// ~STATETREE_V2
+	TArray<FStateTreeTransition> Transitions;
 
 	UPROPERTY()
 	TArray<UStateTreeState*> Children;
diff --git a/Engine/Plugins/Runtime/StateTree/Source/StateTreeEditorModule/Public/StateTreeVariableProvider.h b/Engine/Plugins/Runtime/StateTree/Source/StateTreeEditorModule/Public/StateTreeVariableProvider.h
deleted file mode 100644
diff --git a/Engine/Plugins/Runtime/StateTree/Source/StateTreeModule/Private/Components/StateTreeComponent.cpp b/Engine/Plugins/Runtime/StateTree/Source/StateTreeModule/Private/Components/StateTreeComponent.cpp
--- a/Engine/Plugins/Runtime/StateTree/Source/StateTreeModule/Private/Components/StateTreeComponent.cpp
+++ b/Engine/Plugins/Runtime/StateTree/Source/StateTreeModule/Private/Components/StateTreeComponent.cpp
@@ -6,12 +6,28 @@
 #include "VisualLogger/VisualLoggerTypes.h"
 #include "VisualLogger/VisualLogger.h"
 #include "StateTree.h"
-#include "StateTreeInstance.h"
+#include "StateTreeEvaluatorBase.h"
+#include "StateTreeTaskBase.h"
 #include "AIController.h"
+#include "GameFramework/CharacterMovementComponent.h"
 
 #define STATETREE_LOG(Verbosity, Format, ...) UE_VLOG(GetOwner(), LogStateTree, Verbosity, Format, ##__VA_ARGS__)
 #define STATETREE_CLOG(Condition, Verbosity, Format, ...) UE_CVLOG((Condition), GetOwner(), LogStateTree, Verbosity, Format, ##__VA_ARGS__)
 
+
+//////////////////////////////////////////////////////////////////////////
+// UBrainComponentStateTreeSchema
+
+bool UBrainComponentStateTreeSchema::IsStructAllowed(const UScriptStruct* InScriptStruct) const
+{
+	// @todo: Proper base class for actor nodes.
+	return InScriptStruct->IsChildOf(FStateTreeEvaluator2Base::StaticStruct()) || InScriptStruct->IsChildOf(FStateTreeTask2Base::StaticStruct());
+}
+
+
+//////////////////////////////////////////////////////////////////////////
+// UStateTreeComponent
+
 UStateTreeComponent::UStateTreeComponent(const FObjectInitializer& ObjectInitializer)
 	: Super(ObjectInitializer)
 {
@@ -28,9 +44,9 @@ void UStateTreeComponent::InitializeComponent()
 		return;
 	}
 
-	if (!StateTreeInstance.Init(*this, *StateTree))
+	if (!StateTreeContext.Init(*this, *StateTree, EStateTreeStorage::Internal))
 	{
-		STATETREE_LOG(Error, TEXT("%s: Failed to init StateTreeInstance."), ANSI_TO_TCHAR(__FUNCTION__));
+		STATETREE_LOG(Error, TEXT("%s: Failed to init StateTreeContext."), ANSI_TO_TCHAR(__FUNCTION__));
 	}
 }
 
@@ -47,7 +63,7 @@ void UStateTreeComponent::TickComponent(float DeltaTime, enum ELevelTick TickTyp
 		return;
 	}
 
-	StateTreeInstance.Tick(DeltaTime);
+	StateTreeContext.Tick(DeltaTime);
 }
 
 void UStateTreeComponent::StartLogic()
@@ -57,7 +73,7 @@ void UStateTreeComponent::StartLogic()
 	// TODO: Move this to UStateTree
 	StateTree->ResolvePropertyPaths();
 
-	StateTreeInstance.Start();
+	StateTreeContext.Start();
 
 	bIsRunning = true;
 }
@@ -65,7 +81,7 @@ void UStateTreeComponent::StartLogic()
 void UStateTreeComponent::RestartLogic()
 {
 	STATETREE_LOG(Log, TEXT("%s: Restart Logic"), ANSI_TO_TCHAR(__FUNCTION__));
-	StateTreeInstance.Start();
+	StateTreeContext.Start();
 
 	bIsRunning = true;
 }
@@ -74,7 +90,7 @@ void UStateTreeComponent::StopLogic(const FString& Reason)
 {
 	STATETREE_LOG(Log, TEXT("%s: Stopping, reason: \'%s\'"), ANSI_TO_TCHAR(__FUNCTION__), *Reason);
 
-	StateTreeInstance.Stop();
+	StateTreeContext.Stop();
 
 	bIsRunning = false;
 }
@@ -179,27 +195,10 @@ void UStateTreeComponent::OnGameplayTaskInitialized(UGameplayTask& Task)
 	}
 }
 
-#if ENABLE_VISUAL_LOG
-void UStateTreeComponent::DescribeSelfToVisLog(struct FVisualLogEntry* Snapshot) const
-{
-	if (!IsValid(this))
-	{
-		return;
-	}
-
-	Super::DescribeSelfToVisLog(Snapshot);
-
-	if (Snapshot)
-	{
-		StateTreeInstance.DescribeSelfToVisLog(*Snapshot);
-	}
-}
-#endif // ENABLE_VISUAL_LOG
-
 #if WITH_GAMEPLAY_DEBUGGER
 FString UStateTreeComponent::GetDebugInfoString() const
 {
-	return StateTreeInstance.GetDebugInfoString();
+	return StateTreeContext.GetDebugInfoString();
 }
 #endif // WITH_GAMEPLAY_DEBUGGER
 
diff --git a/Engine/Plugins/Runtime/StateTree/Source/StateTreeModule/Private/Evaluators/StateTreeEvaluator_Test.cpp b/Engine/Plugins/Runtime/StateTree/Source/StateTreeModule/Private/Evaluators/StateTreeEvaluator_Test.cpp
--- a/Engine/Plugins/Runtime/StateTree/Source/StateTreeModule/Private/Evaluators/StateTreeEvaluator_Test.cpp
+++ b/Engine/Plugins/Runtime/StateTree/Source/StateTreeModule/Private/Evaluators/StateTreeEvaluator_Test.cpp
@@ -2,98 +2,8 @@
 
 #include "Evaluators/StateTreeEvaluator_Test.h"
 #include "CoreMinimal.h"
-#include "GameFramework/Actor.h"
-#include "StateTreeInstance.h"
-#include "StateTreeVariableLayout.h"
-#include "StateTreeConstantStorage.h"
 #include "Engine/Engine.h"
 
-UStateTreeEvaluator_Test::UStateTreeEvaluator_Test(const FObjectInitializer& ObjectInitializer)
-	:Super(ObjectInitializer)
-	, Timer(0.0f)
-	, SomeFloat(EStateTreeVariableBindingMode::Typed, EStateTreeVariableType::Float)
-	, BoolVar(EStateTreeVariableBindingMode::Definition, FName(TEXT("Bool")), EStateTreeVariableType::Bool)
-	, FloatVar(EStateTreeVariableBindingMode::Definition, FName(TEXT("Float")), EStateTreeVariableType::Float)
-	, IntVar(EStateTreeVariableBindingMode::Definition, FName(TEXT("Int")), EStateTreeVariableType::Int)
-	, VectorVar(EStateTreeVariableBindingMode::Definition, FName(TEXT("Vector")), EStateTreeVariableType::Vector)
-	, Object1Var(EStateTreeVariableBindingMode::Definition, FName(TEXT("Object1")), EStateTreeVariableType::Object, AActor::StaticClass())
-	, Object2Var(EStateTreeVariableBindingMode::Definition, FName(TEXT("Object2")), EStateTreeVariableType::Object, AActor::StaticClass())
-	, ObjectRVar(EStateTreeVariableBindingMode::Definition, FName(TEXT("ObjectR")), EStateTreeVariableType::Object, UStateTreeEvaluatorBase::StaticClass())
-{
-}
-
-
-bool UStateTreeEvaluator_Test::Initialize(FStateTreeInstance& StateTreeInstance)
-{
-	return true;
-}
-
-void UStateTreeEvaluator_Test::Activate(FStateTreeInstance& StateTreeInstance)
-{
-	// Clear variables
-	StateTreeInstance.SetValueBool(BoolVar.Handle, false);
-	StateTreeInstance.SetValueFloat(FloatVar.Handle, 0.0f);
-	StateTreeInstance.SetValueInt(IntVar.Handle, 0);
-	StateTreeInstance.SetValueVector(VectorVar.Handle, FVector(1.0f, 20.0f, 300.0f));
-	StateTreeInstance.SetValueObject(Object1Var.Handle, nullptr);
-	StateTreeInstance.SetValueObject(Object2Var.Handle, nullptr);
-	StateTreeInstance.SetValueObject(ObjectRVar.Handle, nullptr);
-
-	Timer = 0.0f;
-}
-
-void UStateTreeEvaluator_Test::Deactivate(FStateTreeInstance& StateTreeInstance)
-{
-}
-
-void UStateTreeEvaluator_Test::Tick(FStateTreeInstance& StateTreeInstance, const float DeltaTime)
-{
-	Timer += DeltaTime;
-
-	const float FloatVal = StateTreeInstance.GetValueFloat(SomeFloat.Handle, 0.0f);
-
-	StateTreeInstance.SetValueFloat(FloatVar.Handle, Timer);
-	StateTreeInstance.SetValueInt(IntVar.Handle, int32(Timer / 4.0f));
-	StateTreeInstance.SetValueBool(BoolVar.Handle, FMath::Frac(Timer / 3.0f) > (2.5f / 3.0f));
-}
-
-#if WITH_GAMEPLAY_DEBUGGER
-void UStateTreeEvaluator_Test::AppendDebugInfoString(FString& DebugString, const FStateTreeInstance& StateTreeInstance) const
-{
-	Super::AppendDebugInfoString(DebugString, StateTreeInstance);
-}
-#endif // WITH_GAMEPLAY_DEBUGGER
-
-
-#if WITH_EDITOR
-void UStateTreeEvaluator_Test::DefineOutputVariables(FStateTreeVariableLayout& Variables) const
-{
-	Variables.DefineVariable(Name, BoolVar);
-	Variables.DefineVariable(Name, FloatVar);
-	Variables.DefineVariable(Name, IntVar);
-	Variables.DefineVariable(Name, VectorVar);
-	Variables.DefineVariable(Name, Object1Var);
-	Variables.DefineVariable(Name, Object2Var);
-	Variables.DefineVariable(Name, ObjectRVar);
-}
-
-bool UStateTreeEvaluator_Test::ResolveVariables(const FStateTreeVariableLayout& Variables, FStateTreeConstantStorage& Constants, UObject* Outer)
-{
-	SomeFloat.ResolveHandle(Variables, Constants);
-
-	BoolVar.ResolveHandle(Variables, Constants);
-	FloatVar.ResolveHandle(Variables, Constants);
-	IntVar.ResolveHandle(Variables, Constants);
-	VectorVar.ResolveHandle(Variables, Constants);
-	Object1Var.ResolveHandle(Variables, Constants);
-	Object2Var.ResolveHandle(Variables, Constants);
-	ObjectRVar.ResolveHandle(Variables, Constants);
-
-	return true;
-}
-#endif
-
-
 FStateTreeEvaluator2_Test::FStateTreeEvaluator2_Test()
 {
 }
diff --git a/Engine/Plugins/Runtime/StateTree/Source/StateTreeModule/Private/StateTree.cpp b/Engine/Plugins/Runtime/StateTree/Source/StateTreeModule/Private/StateTree.cpp
--- a/Engine/Plugins/Runtime/StateTree/Source/StateTreeModule/Private/StateTree.cpp
+++ b/Engine/Plugins/Runtime/StateTree/Source/StateTreeModule/Private/StateTree.cpp
@@ -20,15 +20,9 @@ bool UStateTree::IsValidStateTree() const
 #if WITH_EDITOR
 void UStateTree::ResetBaked()
 {
-	Variables.Reset();
-	Constants.Reset();
-	Tasks.Reset();
-	Evaluators.Reset();
 	States.Reset();
 	Conditions.Reset();
-	Conditions2.Reset();
 	Transitions.Reset();
-	Parameters.Reset();
 
 	RuntimeStorageItems.Reset();
 	RuntimeStorageStruct = nullptr;
@@ -42,20 +36,9 @@ void UStateTree::PostEditChangeProperty(struct FPropertyChangedEvent& PropertyCh
 {
 	if (PropertyChangedEvent.MemberProperty && PropertyChangedEvent.Property)
 	{
-		if (PropertyChangedEvent.MemberProperty->GetFName() == GET_MEMBER_NAME_CHECKED(UStateTree, Parameters))
-		{
-			UE::StateTree::Delegates::OnParameterLayoutChanged.Broadcast(*this);
-		}
-
 		if (PropertyChangedEvent.MemberProperty->GetFName() == GET_MEMBER_NAME_CHECKED(UStateTree, Schema))
 		{
-			UE::StateTree::Delegates::OnParameterLayoutChanged.Broadcast(*this);
-		}
-
-		if (PropertyChangedEvent.Property->GetOwnerStruct() == FStateTreeVariableDesc::StaticStruct() &&
-			PropertyChangedEvent.Property->GetFName() == GET_MEMBER_NAME_CHECKED(FStateTreeVariableDesc, Name))
-		{
-			UE::StateTree::Delegates::OnIdentifierChanged.Broadcast(*this);
+			UE::StateTree::Delegates::OnSchemaChanged.Broadcast(*this);
 		}
 	}
 }
diff --git a/Engine/Plugins/Runtime/StateTree/Source/StateTreeModule/Private/StateTreeCondition.cpp b/Engine/Plugins/Runtime/StateTree/Source/StateTreeModule/Private/StateTreeCondition.cpp
deleted file mode 100644
diff --git a/Engine/Plugins/Runtime/StateTree/Source/StateTreeModule/Private/StateTreeDelegates.cpp b/Engine/Plugins/Runtime/StateTree/Source/StateTreeModule/Private/StateTreeDelegates.cpp
--- a/Engine/Plugins/Runtime/StateTree/Source/StateTreeModule/Private/StateTreeDelegates.cpp
+++ b/Engine/Plugins/Runtime/StateTree/Source/StateTreeModule/Private/StateTreeDelegates.cpp
@@ -8,8 +8,6 @@ namespace UE::StateTree::Delegates
 	
 #if WITH_EDITOR
 FOnIdentifierChanged OnIdentifierChanged;
-FOnParametersInvalidated OnParametersInvalidated;
-FOnParameterLayoutChanged OnParameterLayoutChanged;
 FOnSchemaChanged OnSchemaChanged;
 #endif
 	
diff --git a/Engine/Plugins/Runtime/StateTree/Source/StateTreeModule/Private/StateTreeEvaluatorBase.cpp b/Engine/Plugins/Runtime/StateTree/Source/StateTreeModule/Private/StateTreeEvaluatorBase.cpp
--- a/Engine/Plugins/Runtime/StateTree/Source/StateTreeModule/Private/StateTreeEvaluatorBase.cpp
+++ b/Engine/Plugins/Runtime/StateTree/Source/StateTreeModule/Private/StateTreeEvaluatorBase.cpp
@@ -2,41 +2,6 @@
 
 #include "StateTreeEvaluatorBase.h"
 #include "CoreMinimal.h"
-#include "StateTree.h"
-#include "StateTreeDelegates.h"
-
-// Base class for all evaluators
-UStateTreeEvaluatorBase::UStateTreeEvaluatorBase(const FObjectInitializer& ObjectInitializer)
-	: Super(ObjectInitializer)
-{
-#if WITH_EDITORONLY_DATA
-	ID = FGuid::NewGuid();
-#endif // WITH_EDITORONLY_DATA
-}
-
-#if WITH_GAMEPLAY_DEBUGGER
-void UStateTreeEvaluatorBase::AppendDebugInfoString(FString& DebugString, const FStateTreeInstance& StateTreeInstance) const
-{
-	DebugString += FString::Printf(TEXT("[%s]\n"), *Name.ToString());
-}
-#endif // WITH_GAMEPLAY_DEBUGGER
-
-#if WITH_EDITOR
-void UStateTreeEvaluatorBase::PostEditChangeProperty(struct FPropertyChangedEvent& PropertyChangedEvent)
-{
-	if (PropertyChangedEvent.Property)
-	{
-		if (PropertyChangedEvent.Property->GetFName() == GET_MEMBER_NAME_CHECKED(UStateTreeEvaluatorBase, Name))
-		{
-			UStateTree* StateTree = GetTypedOuter<UStateTree>();
-			if (ensure(StateTree))
-			{
-				UE::StateTree::Delegates::OnIdentifierChanged.Broadcast(*StateTree);
-			}
-		}
-	}
-}
-#endif
 
 #if WITH_GAMEPLAY_DEBUGGER
 void FStateTreeEvaluator2Base::AppendDebugInfoString(FString& DebugString, const FStateTreeExecutionContext& Context) const
diff --git a/Engine/Plugins/Runtime/StateTree/Source/StateTreeModule/Private/StateTreeExecutionContext.cpp b/Engine/Plugins/Runtime/StateTree/Source/StateTreeModule/Private/StateTreeExecutionContext.cpp
--- a/Engine/Plugins/Runtime/StateTree/Source/StateTreeModule/Private/StateTreeExecutionContext.cpp
+++ b/Engine/Plugins/Runtime/StateTree/Source/StateTreeModule/Private/StateTreeExecutionContext.cpp
@@ -40,7 +40,6 @@ bool FStateTreeExecutionContext::Init(UObject& InOwner, const UStateTree& InStat
 		return false;
 	}
 	StateTree = &InStateTree;
-	check(StateTree->IsV2());
 
 	StorageType = InStorageType;
 	if (StorageType == EStateTreeStorage::Internal)
@@ -578,7 +577,7 @@ bool FStateTreeExecutionContext::TestAllConditions(const uint32 ConditionsOffset
 	for (uint32 i = 0; i < ConditionsNum; i++)
 	{
 		// Copy struct values from the StateTree asset to a stack copy.
-		const FInstancedStruct& SourceCondPtr = StateTree->Conditions2[ConditionsOffset + i];
+		const FInstancedStruct& SourceCondPtr = StateTree->Conditions[ConditionsOffset + i];
 		const FStateTreeConditionBase& SourceCond = SourceCondPtr.Get<FStateTreeConditionBase>();
 		const UScriptStruct* CondStruct = SourceCondPtr.GetScriptStruct();
 		check(CondStruct);
@@ -825,13 +824,13 @@ EStateTreeRunStatus FStateTreeExecutionContext::GetLastTickStatus(FStateTreeItem
 
 #if WITH_GAMEPLAY_DEBUGGER
 
-FString FStateTreeExecutionContext::GetDebugInfoString(FStateTreeItemView ExternalStorage)
+FString FStateTreeExecutionContext::GetDebugInfoString(FStateTreeItemView ExternalStorage) const
 {
 	if (!StateTree)
 	{
 		return FString(TEXT("No StateTree asset."));
 	}
-	FStateTreeItemView Storage = SelectMutableStorage(ExternalStorage);
+	FStateTreeItemView Storage = SelectStorage(ExternalStorage);
 	FStateTreeExecutionState& Exec = GetExecState(Storage);
 
 	FString DebugString = FString::Printf(TEXT("StateTree (asset: '%s')\n"), *GetNameSafe(StateTree));
@@ -920,8 +919,8 @@ void FStateTreeExecutionContext::DebugPrintInternalLayout(FStateTreeItemView Ext
 	}
 
 	// Conditions
-	DebugString += FString::Printf(TEXT("\nConditions(%d)\n"), StateTree->Conditions2.Num());
-	for (const FInstancedStruct& Condition : StateTree->Conditions2)
+	DebugString += FString::Printf(TEXT("\nConditions(%d)\n"), StateTree->Conditions.Num());
+	for (const FInstancedStruct& Condition : StateTree->Conditions)
 	{
 		DebugString += FString::Printf(TEXT("  %s\n"), Condition.IsValid() ? *Condition.GetScriptStruct()->GetName() : TEXT("null"));
 	}
diff --git a/Engine/Plugins/Runtime/StateTree/Source/StateTreeModule/Private/StateTreeInstance.cpp b/Engine/Plugins/Runtime/StateTree/Source/StateTreeModule/Private/StateTreeInstance.cpp
deleted file mode 100644
diff --git a/Engine/Plugins/Runtime/StateTree/Source/StateTreeModule/Private/StateTreeParameter.cpp b/Engine/Plugins/Runtime/StateTree/Source/StateTreeModule/Private/StateTreeParameter.cpp
deleted file mode 100644
diff --git a/Engine/Plugins/Runtime/StateTree/Source/StateTreeModule/Private/StateTreeParameterLayout.cpp b/Engine/Plugins/Runtime/StateTree/Source/StateTreeModule/Private/StateTreeParameterLayout.cpp
deleted file mode 100644
diff --git a/Engine/Plugins/Runtime/StateTree/Source/StateTreeModule/Private/StateTreeParameterStorage.cpp b/Engine/Plugins/Runtime/StateTree/Source/StateTreeModule/Private/StateTreeParameterStorage.cpp
deleted file mode 100644
diff --git a/Engine/Plugins/Runtime/StateTree/Source/StateTreeModule/Private/StateTreeTaskBase.cpp b/Engine/Plugins/Runtime/StateTree/Source/StateTreeModule/Private/StateTreeTaskBase.cpp
--- a/Engine/Plugins/Runtime/StateTree/Source/StateTreeModule/Private/StateTreeTaskBase.cpp
+++ b/Engine/Plugins/Runtime/StateTree/Source/StateTreeModule/Private/StateTreeTaskBase.cpp
@@ -3,29 +3,6 @@
 #include "StateTreeTaskBase.h"
 #include "CoreMinimal.h"
 
-// Base class for all evaluators
-UStateTreeTaskBase::UStateTreeTaskBase(const FObjectInitializer& ObjectInitializer)
-	: Super(ObjectInitializer)
-	, ID(FGuid::NewGuid())
-{
-}
-
-#if WITH_EDITOR
-void UStateTreeTaskBase::PostLoad()
-{
-	Super::PostLoad();
-	if (!ID.IsValid())
-	{
-		ID = FGuid::NewGuid();
-	}
-}
-#endif
-
-#if WITH_GAMEPLAY_DEBUGGER
-void UStateTreeTaskBase::AppendDebugInfoString(FString& DebugString, const FStateTreeInstance& StateTreeInstance) const
-{
-}
-#endif
 
 #if WITH_GAMEPLAY_DEBUGGER
 void FStateTreeTask2Base::AppendDebugInfoString(FString& DebugString, const FStateTreeExecutionContext& Context) const
diff --git a/Engine/Plugins/Runtime/StateTree/Source/StateTreeModule/Private/StateTreeVariable.cpp b/Engine/Plugins/Runtime/StateTree/Source/StateTreeModule/Private/StateTreeVariable.cpp
deleted file mode 100644
diff --git a/Engine/Plugins/Runtime/StateTree/Source/StateTreeModule/Private/StateTreeVariableDesc.cpp b/Engine/Plugins/Runtime/StateTree/Source/StateTreeModule/Private/StateTreeVariableDesc.cpp
deleted file mode 100644
diff --git a/Engine/Plugins/Runtime/StateTree/Source/StateTreeModule/Private/Tasks/StateTreeTask_SubStateTree.cpp b/Engine/Plugins/Runtime/StateTree/Source/StateTreeModule/Private/Tasks/StateTreeTask_SubStateTree.cpp
deleted file mode 100644
diff --git a/Engine/Plugins/Runtime/StateTree/Source/StateTreeModule/Private/Tasks/StateTreeTask_Test.cpp b/Engine/Plugins/Runtime/StateTree/Source/StateTreeModule/Private/Tasks/StateTreeTask_Test.cpp
--- a/Engine/Plugins/Runtime/StateTree/Source/StateTreeModule/Private/Tasks/StateTreeTask_Test.cpp
+++ b/Engine/Plugins/Runtime/StateTree/Source/StateTreeModule/Private/Tasks/StateTreeTask_Test.cpp
@@ -1,53 +1,6 @@
 // Copyright Epic Games, Inc. All Rights Reserved.
 
 #include "Tasks/StateTreeTask_Test.h"
-#include "Engine/Engine.h"
-
-
-UStateTreeTask_Test::UStateTreeTask_Test(const FObjectInitializer& ObjectInitializer)
-	:Super(ObjectInitializer)
-{
-}
-
-bool UStateTreeTask_Test::Initialize(FStateTreeInstance& StateTreeInstance)
-{
-	return true;
-}
-
-void UStateTreeTask_Test::Activate(FStateTreeInstance& StateTreeInstance)
-{
-	Time = 0.0f;
-}
-
-void UStateTreeTask_Test::Deactivate(FStateTreeInstance& StateTreeInstance)
-{
-}
-
-EStateTreeRunStatus UStateTreeTask_Test::Tick(FStateTreeInstance& StateTreeInstance, const float DeltaTime)
-{
-	Time += DeltaTime;
-	if (Time > Duration)
-	{
-		return bFail ? EStateTreeRunStatus::Failed : EStateTreeRunStatus::Succeeded;		
-	}
-	return EStateTreeRunStatus::Running;
-}
-
-#if WITH_GAMEPLAY_DEBUGGER
-void UStateTreeTask_Test::AppendDebugInfoString(FString& DebugString, const FStateTreeInstance& StateTreeInstance) const
-{
-	Super::AppendDebugInfoString(DebugString, StateTreeInstance);
-	DebugString += FString::Printf(TEXT("Time= %f/%f\n"), Time, Duration);
-}
-#endif
-
-#if WITH_EDITOR
-bool UStateTreeTask_Test::ResolveVariables(const FStateTreeVariableLayout& Variables, FStateTreeConstantStorage& Constants, UObject* Outer)
-{
-	return true;
-}
-#endif
-
 
 
 FStateTreeTask2_Test::FStateTreeTask2_Test()
diff --git a/Engine/Plugins/Runtime/StateTree/Source/StateTreeModule/Public/Components/StateTreeComponent.h b/Engine/Plugins/Runtime/StateTree/Source/StateTreeModule/Public/Components/StateTreeComponent.h
--- a/Engine/Plugins/Runtime/StateTree/Source/StateTreeModule/Public/Components/StateTreeComponent.h
+++ b/Engine/Plugins/Runtime/StateTree/Source/StateTreeModule/Public/Components/StateTreeComponent.h
@@ -5,13 +5,24 @@
 #include "Components/ActorComponent.h"
 #include "GameplayTagAssetInterface.h"
 #include "GameplayTagContainer.h"
+#include "StateTreeExecutionContext.h"
 #include "BrainComponent.h"
-#include "StateTreeInstance.h"
 #include "Tasks/AITask.h"
 #include "StateTreeComponent.generated.h"
 
 class UStateTree;
 
+
+UCLASS(Abstract)
+class STATETREEMODULE_API UBrainComponentStateTreeSchema : public UStateTreeSchema
+{
+	GENERATED_BODY()
+
+public:
+	virtual bool IsStructAllowed(const UScriptStruct* InScriptStruct) const override;
+};
+
+
 UCLASS(ClassGroup = AI, HideCategories = (Activation, Collision), meta = (BlueprintSpawnableComponent))
 class STATETREEMODULE_API UStateTreeComponent : public UBrainComponent, public IGameplayTaskOwnerInterface
 {
@@ -44,10 +55,6 @@ public:
 	virtual void OnGameplayTaskInitialized(UGameplayTask& Task) override;
 	// END IGameplayTaskOwnerInterface
 
-#if ENABLE_VISUAL_LOG
-	virtual void DescribeSelfToVisLog(struct FVisualLogEntry* Snapshot) const;
-#endif // ENABLE_VISUAL_LOG
-
 #if WITH_GAMEPLAY_DEBUGGER
 	virtual FString GetDebugInfoString() const override;
 #endif // WITH_GAMEPLAY_DEBUGGER
@@ -55,12 +62,12 @@ public:
 private:
 
 protected:
-
-	UPROPERTY(EditDefaultsOnly, Category = AI)
+	
+	UPROPERTY(EditDefaultsOnly, Category = AI, meta=(RequiredAssetDataTags="Schema=BrainComponentStateTreeSchema"))
 	UStateTree* StateTree;
 
 	UPROPERTY()
-	FStateTreeInstance StateTreeInstance;
+	FStateTreeExecutionContext StateTreeContext;
 
 	/** if set, state tree execution is allowed */
 	uint8 bIsRunning : 1;
diff --git a/Engine/Plugins/Runtime/StateTree/Source/StateTreeModule/Public/Evaluators/StateTreeEvaluator_Test.h b/Engine/Plugins/Runtime/StateTree/Source/StateTreeModule/Public/Evaluators/StateTreeEvaluator_Test.h
--- a/Engine/Plugins/Runtime/StateTree/Source/StateTreeModule/Public/Evaluators/StateTreeEvaluator_Test.h
+++ b/Engine/Plugins/Runtime/StateTree/Source/StateTreeModule/Public/Evaluators/StateTreeEvaluator_Test.h
@@ -3,78 +3,10 @@
 #pragma once
 
 #include "CoreMinimal.h"
-#include "Misc/Guid.h"
 #include "StateTreeTypes.h"
-#include "StateTreeVariable.h"
 #include "StateTreeEvaluatorBase.h"
 #include "StateTreeEvaluator_Test.generated.h"
 
-struct FStateTreeVariableLayout;
-struct FStateTreeConstantStorage;
-class UStateTreeEvaluatorInstanceBase;
-
-/**
- * Test evaluator
- */
-
-UCLASS(BlueprintType, EditInlineNew, CollapseCategories)
-class STATETREEMODULE_API UStateTreeEvaluator_Test : public UStateTreeEvaluatorBase
-{
-	GENERATED_BODY()
-
-public:
-	UStateTreeEvaluator_Test(const FObjectInitializer& ObjectInitializer = FObjectInitializer::Get());
-
-	// Called when instantiated first time.
-	virtual bool Initialize(FStateTreeInstance& StateTreeInstance) override;
-	// Called when evaluator becomes active/ticking.
-	virtual void Activate(FStateTreeInstance& StateTreeInstance) override;
-	// Called when evaluator becomes inactive.
-	virtual void Deactivate(FStateTreeInstance& StateTreeInstance) override;
-
-	// Called on each tick.
-	virtual void Tick(FStateTreeInstance& StateTreeInstance, const float DeltaTime) override;
-
-#if WITH_GAMEPLAY_DEBUGGER
-	virtual void AppendDebugInfoString(FString& DebugString, const FStateTreeInstance& StateTreeInstance) const override;
-#endif // WITH_GAMEPLAY_DEBUGGER
-
-#if WITH_EDITOR
-	virtual void DefineOutputVariables(FStateTreeVariableLayout& Variables) const override;
-	virtual bool ResolveVariables(const FStateTreeVariableLayout& Variables, FStateTreeConstantStorage& Constants, UObject* Outer) override;
-#endif
-
-protected:
-	// State
-	float Timer;
-
-	// UI properties
-	UPROPERTY(EditDefaultsOnly, Category = Test)
-	float PlainNumber;
-
-	UPROPERTY(EditDefaultsOnly, Category = Test)
-	FStateTreeVariable SomeFloat;
-
-	// Output variables
-	UPROPERTY()
-	FStateTreeVariable BoolVar;
-	UPROPERTY()
-	FStateTreeVariable FloatVar;
-	UPROPERTY()
-	FStateTreeVariable IntVar;
-	UPROPERTY()
-	FStateTreeVariable VectorVar;
-	UPROPERTY()
-	FStateTreeVariable Object1Var;
-	UPROPERTY()
-	FStateTreeVariable Object2Var;
-	UPROPERTY()
-	FStateTreeVariable ObjectRVar;
-};
-
-
-// STATETREE_V2
-
 UENUM()
 namespace EStateTreeOldEnumTest
 {
@@ -138,5 +70,3 @@ protected:
 	
 	FDummyStateTreeResult DummyResult;
 };
-
-// ~STATETREE_V2
diff --git a/Engine/Plugins/Runtime/StateTree/Source/StateTreeModule/Public/StateTree.h b/Engine/Plugins/Runtime/StateTree/Source/StateTreeModule/Public/StateTree.h
--- a/Engine/Plugins/Runtime/StateTree/Source/StateTreeModule/Public/StateTree.h
+++ b/Engine/Plugins/Runtime/StateTree/Source/StateTreeModule/Public/StateTree.h
@@ -6,17 +6,10 @@
 #include "Engine/DataAsset.h"
 #include "StateTreeTypes.h"
 #include "StateTreeSchema.h"
-#include "StateTreeCondition.h"
-#include "StateTreeVariableLayout.h"
-#include "StateTreeParameterLayout.h"
-#include "StateTreeConstantStorage.h"
 #include "InstancedStruct.h"
 #include "StateTreePropertyBindings.h"
 #include "StateTree.generated.h"
 
-class UStateTreeEvaluatorBase;
-class UStateTreeTaskBase;
-
 /**
  * StateTree asset. Contains the StateTree definition in both editor and runtime (baked) formats.
  */
@@ -29,8 +22,6 @@ public:
 
 	UStateTree(const FObjectInitializer& ObjectInitializer = FObjectInitializer::Get());
 
-	const FStateTreeVariableLayout& GetInputParameterLayout() const { return InputParameter; }
-
 	/** @return Script Struct that can be used to instantiate the runtime storage */
 	const UScriptStruct* GetRuntimeStorageStruct() const { return RuntimeStorageStruct; }
 
@@ -47,10 +38,7 @@ public:
 	/** @return true is the tree asset is considered valid (e.g. at least one state) */
 	bool IsValidStateTree() const;
 
-	// STATETREE_V2
 	void ResolvePropertyPaths();
-	bool IsV2() const { return Schema ? Schema->IsV2() : false; }
-	// ~STATETREE_V2
 
 #if WITH_EDITOR
 	/** Resets the baked data to empty. */
@@ -85,26 +73,6 @@ private:
 	UPROPERTY(EditDefaultsOnly, Category = Common, Instanced)
 	UStateTreeSchema* Schema = nullptr;
 
-	// TODO: find a nice way to combine this with Parameters. The data is almost the same, but the semantics different:
-	// InputParameter defines layout with offsets, while Parameters defines how they connect to Variables.
-	// This could be editor only, maybe.
-	UPROPERTY(EditDefaultsOnly, Category = InputParameters)
-	FStateTreeVariableLayout InputParameter;
-
-	// Baked data
-	UPROPERTY()
-	FStateTreeVariableLayout Variables;
-
-	UPROPERTY()
-	FStateTreeConstantStorage Constants;
-
-	UPROPERTY()
-	TArray<UStateTreeTaskBase*> Tasks;
-
-	UPROPERTY()
-	TArray<UStateTreeEvaluatorBase*> Evaluators;
-
-// STATETREE_V2
 	/** Evaluators and Tasks that require runtime state */
 	UPROPERTY()
 	TArray<FInstancedStruct> RuntimeStorageItems;
@@ -121,7 +89,7 @@ private:
 	FInstancedStruct RuntimeStorageDefaultValue;
 
 	UPROPERTY(meta = (BaseStruct = "StateTreeConditionBase", ExcludeBaseStruct))
-	TArray<FInstancedStruct> Conditions2;
+	TArray<FInstancedStruct> Conditions;
 
 	/** List of external items required by the state tree */
 	UPROPERTY()
@@ -129,21 +97,13 @@ private:
 
 	UPROPERTY()
 	FStateTreePropertyBindings PropertyBindings;
-	
-// ~STATETREE_V2
 
 	UPROPERTY()
 	TArray<FBakedStateTreeState> States;
 
-	UPROPERTY()
-	TArray<FStateTreeCondition> Conditions;
-
 	UPROPERTY()
 	TArray<FBakedStateTransition> Transitions;
 
-	UPROPERTY()
-	FStateTreeParameterLayout Parameters;
-
 	friend struct FStateTreeInstance;
 	friend struct FStateTreeExecutionContext;
 #if WITH_EDITORONLY_DATA
diff --git a/Engine/Plugins/Runtime/StateTree/Source/StateTreeModule/Public/StateTreeCondition.h b/Engine/Plugins/Runtime/StateTree/Source/StateTreeModule/Public/StateTreeCondition.h
deleted file mode 100644
diff --git a/Engine/Plugins/Runtime/StateTree/Source/StateTreeModule/Public/StateTreeConstantStorage.h b/Engine/Plugins/Runtime/StateTree/Source/StateTreeModule/Public/StateTreeConstantStorage.h
deleted file mode 100644
diff --git a/Engine/Plugins/Runtime/StateTree/Source/StateTreeModule/Public/StateTreeDelegates.h b/Engine/Plugins/Runtime/StateTree/Source/StateTreeModule/Public/StateTreeDelegates.h
--- a/Engine/Plugins/Runtime/StateTree/Source/StateTreeModule/Public/StateTreeDelegates.h
+++ b/Engine/Plugins/Runtime/StateTree/Source/StateTreeModule/Public/StateTreeDelegates.h
@@ -14,14 +14,6 @@ namespace UE::StateTree::Delegates
 	DECLARE_MULTICAST_DELEGATE_OneParam(FOnIdentifierChanged, const UStateTree& /*StateTree*/);
 	extern STATETREEMODULE_API FOnIdentifierChanged OnIdentifierChanged;
 
-	// Called when parameters are invalidated, i.e. linked asset's parameter layout changed.
-	DECLARE_MULTICAST_DELEGATE_OneParam(FOnParametersInvalidated, const UStateTree& /*StateTree*/);
-	extern STATETREEMODULE_API FOnParametersInvalidated OnParametersInvalidated;
-
-	// Called when parameter layout for a StateTree has changed.
-	DECLARE_MULTICAST_DELEGATE_OneParam(FOnParameterLayoutChanged, const UStateTree& /*StateTree*/);
-	extern STATETREEMODULE_API FOnParameterLayoutChanged OnParameterLayoutChanged;
-	
 	// Called when schema of the StateTree has changed.
 	DECLARE_MULTICAST_DELEGATE_OneParam(FOnSchemaChanged, const UStateTree& /*StateTree*/);
 	extern STATETREEMODULE_API FOnSchemaChanged OnSchemaChanged;
diff --git a/Engine/Plugins/Runtime/StateTree/Source/StateTreeModule/Public/StateTreeEvaluatorBase.h b/Engine/Plugins/Runtime/StateTree/Source/StateTreeModule/Public/StateTreeEvaluatorBase.h
--- a/Engine/Plugins/Runtime/StateTree/Source/StateTreeModule/Public/StateTreeEvaluatorBase.h
+++ b/Engine/Plugins/Runtime/StateTree/Source/StateTreeModule/Public/StateTreeEvaluatorBase.h
@@ -8,68 +8,7 @@
 #include "StateTreePropertyBindings.h"
 #include "StateTreeEvaluatorBase.generated.h"
 
-struct FStateTreeInstance;
 struct FStateTreeExecutionContext;
-struct FStateTreeVariableLayout;
-struct FStateTreeConstantStorage;
-
-/**
- * Base class of StateTree Evaluators.
- * Evaluators calculate and expose data to be used for decision making in a StateTree.
- */
-UCLASS(Abstract)
-class STATETREEMODULE_API UStateTreeEvaluatorBase : public UObject
-{
-	GENERATED_BODY()
-
-public:
-	UStateTreeEvaluatorBase(const FObjectInitializer& ObjectInitializer);
-
-	// Called when instantiated first time.
-	virtual bool Initialize(FStateTreeInstance& StateTreeInstance) PURE_VIRTUAL(UStateTreeEvaluatorBase::Initialize, return false; );
-	// Called when evaluator becomes active/ticking.
-	virtual void Activate(FStateTreeInstance & StateTreeInstance) PURE_VIRTUAL(UStateTreeEvaluatorBase::Activate, return; );
-	// Called when evaluator becomes inactive.
-	virtual void Deactivate(FStateTreeInstance & StateTreeInstance) PURE_VIRTUAL(UStateTreeEvaluatorBase::Deactivate, return; );
-
-	// Called on each tick.
-	virtual void Tick(FStateTreeInstance & StateTreeInstance, const float DeltaTime) PURE_VIRTUAL(UStateTreeEvaluatorBase::Tick, return; );
-
-	// Returns ID of the Evaluator
-	const FName& GetName() const { return Name; }
-
-#if WITH_GAMEPLAY_DEBUGGER
-	virtual void AppendDebugInfoString(FString & DebugString, const FStateTreeInstance & StateTreeInstance) const;
-#endif // WITH_GAMEPLAY_DEBUGGER
-
-#if WITH_EDITOR
-	// Defines output variables in variable map.
-	virtual void DefineOutputVariables(FStateTreeVariableLayout& Variables) const {};
-	// Create Evaluator instance template, set default variables, and resolve variable handles.
-	virtual bool ResolveVariables(const FStateTreeVariableLayout& Variables, FStateTreeConstantStorage& Constants, UObject* Outer) PURE_VIRTUAL(UStateTreeEvaluatorBase::Resolve, return false; );
-	// Returns ture if changed.
-	virtual bool ValidateParameterLayout() { return false; }
-
-	const FGuid GetID() const { return ID; }
-
-	void SetNewUniqueID() { ID = FGuid::NewGuid(); }
-	void SetName(const FName& InName) { Name = InName; }
-
-	virtual void PostEditChangeProperty(struct FPropertyChangedEvent& PropertyChangedEvent) override;
-#endif
-
-protected:
-	UPROPERTY(EditDefaultsOnly, Category = Evaluator)
-	FName Name;
-
-#if WITH_EDITORONLY_DATA
-	UPROPERTY(meta = (IgnoreForMemberInitializationTest))
-	FGuid ID;
-#endif
-};
-
-
-// STATETREE_V2
 
 /**
  * Base struct of StateTree Evaluators.
@@ -139,6 +78,5 @@ struct STATETREEMODULE_API FStateTreeEvaluator2Base // TODO: change to FStateTre
 	UPROPERTY()
 	uint16 SourceStructIndex = 0;									// Property binding Source Struct index of the evaluator.
 };
-template<> struct TStructOpsTypeTraits<FStateTreeEvaluator2Base> : public TStructOpsTypeTraitsBase2<FStateTreeEvaluator2Base> { enum { WithPureVirtual = true, }; };
 
-// ~STATETREE_V2
+template<> struct TStructOpsTypeTraits<FStateTreeEvaluator2Base> : public TStructOpsTypeTraitsBase2<FStateTreeEvaluator2Base> { enum { WithPureVirtual = true, }; };
diff --git a/Engine/Plugins/Runtime/StateTree/Source/StateTreeModule/Public/StateTreeExecutionContext.h b/Engine/Plugins/Runtime/StateTree/Source/StateTreeModule/Public/StateTreeExecutionContext.h
--- a/Engine/Plugins/Runtime/StateTree/Source/StateTreeModule/Public/StateTreeExecutionContext.h
+++ b/Engine/Plugins/Runtime/StateTree/Source/StateTreeModule/Public/StateTreeExecutionContext.h
@@ -129,7 +129,7 @@ public:
 
 #if WITH_GAMEPLAY_DEBUGGER
 	/** @return Debug string describing the current state of the execution */
-	FString GetDebugInfoString(FStateTreeItemView ExternalStorage = FStateTreeItemView());
+	FString GetDebugInfoString(FStateTreeItemView ExternalStorage = FStateTreeItemView()) const;
 #endif // WITH_GAMEPLAY_DEBUGGER
 
 #if WITH_STATETREE_DEBUG
@@ -228,7 +228,7 @@ protected:
 	}
 
 	/** @return View to an Evaluator or a Task. */
-	FStateTreeItemView GetItem(FStateTreeItemView Storage, const int32 Index)
+	FStateTreeItemView GetItem(FStateTreeItemView Storage, const int32 Index) const
 	{
 		const FStateTreeRuntimeStorageItemOffset& ItemOffset = StateTree->RuntimeStorageOffsets[Index];
 		return FStateTreeItemView(ItemOffset.Struct, Storage.GetMutableMemory() + ItemOffset.Offset);
diff --git a/Engine/Plugins/Runtime/StateTree/Source/StateTreeModule/Public/StateTreeInstance.h b/Engine/Plugins/Runtime/StateTree/Source/StateTreeModule/Public/StateTreeInstance.h
deleted file mode 100644
diff --git a/Engine/Plugins/Runtime/StateTree/Source/StateTreeModule/Public/StateTreeParameter.h b/Engine/Plugins/Runtime/StateTree/Source/StateTreeModule/Public/StateTreeParameter.h
deleted file mode 100644
diff --git a/Engine/Plugins/Runtime/StateTree/Source/StateTreeModule/Public/StateTreeParameterLayout.h b/Engine/Plugins/Runtime/StateTree/Source/StateTreeModule/Public/StateTreeParameterLayout.h
deleted file mode 100644
diff --git a/Engine/Plugins/Runtime/StateTree/Source/StateTreeModule/Public/StateTreeParameterStorage.h b/Engine/Plugins/Runtime/StateTree/Source/StateTreeModule/Public/StateTreeParameterStorage.h
deleted file mode 100644
diff --git a/Engine/Plugins/Runtime/StateTree/Source/StateTreeModule/Public/StateTreeSchema.h b/Engine/Plugins/Runtime/StateTree/Source/StateTreeModule/Public/StateTreeSchema.h
--- a/Engine/Plugins/Runtime/StateTree/Source/StateTreeModule/Public/StateTreeSchema.h
+++ b/Engine/Plugins/Runtime/StateTree/Source/StateTreeModule/Public/StateTreeSchema.h
@@ -3,8 +3,6 @@
 #pragma once
 
 #include "CoreMinimal.h"
-#include "Engine/DataAsset.h"
-#include "StateTreeTypes.h"
 #include "StateTreeSchema.generated.h"
 
 /**
@@ -31,7 +29,4 @@ public:
 
 	/** @return True if specified struct/class is supported as external item */
 	virtual bool IsExternalItemAllowed(const UStruct& InStruct) const { return false; };
-
-	/** @return True if we should use StateTree V2 */
-	virtual bool IsV2() const { return false; }
 };
diff --git a/Engine/Plugins/Runtime/StateTree/Source/StateTreeModule/Public/StateTreeTaskBase.h b/Engine/Plugins/Runtime/StateTree/Source/StateTreeModule/Public/StateTreeTaskBase.h
--- a/Engine/Plugins/Runtime/StateTree/Source/StateTreeModule/Public/StateTreeTaskBase.h
+++ b/Engine/Plugins/Runtime/StateTree/Source/StateTreeModule/Public/StateTreeTaskBase.h
@@ -7,57 +7,7 @@
 #include "StateTreeTypes.h"
 #include "StateTreeTaskBase.generated.h"
 
-struct FStateTreeInstance;
 struct FStateTreeExecutionContext;
-struct FStateTreeVariableLayout;
-struct FStateTreeConstantStorage;
-
-/**
- * Base class for StateTree Tasks.
- * Tasks are logic executed at an active state.
- */
-UCLASS(Abstract)
-class STATETREEMODULE_API UStateTreeTaskBase : public UObject
-{
-	GENERATED_BODY()
-
-public:
-	UStateTreeTaskBase(const FObjectInitializer& ObjectInitializer);
-
-	// Called when instantiated first time.
-	virtual bool Initialize(FStateTreeInstance& StateTreeInstance) PURE_VIRTUAL(UStateTreeTaskBase::Initialize, return false; );
-	// Called when evaluator becomes active/ticking.
-	virtual void Activate(FStateTreeInstance& StateTreeInstance) PURE_VIRTUAL(UStateTreeTaskBase::Activate, return; );
-	// Called when evaluator becomes inactive.
-	virtual void Deactivate(FStateTreeInstance& StateTreeInstance) PURE_VIRTUAL(UStateTreeTaskBase::Deactivate, return; );
-
-	// Called on each tick.
-	virtual EStateTreeRunStatus Tick(FStateTreeInstance& StateTreeInstance, const float DeltaTime) PURE_VIRTUAL(UStateTreeTaskBase::Tick, return EStateTreeRunStatus::Failed; );
-
-#if WITH_GAMEPLAY_DEBUGGER
-	virtual void AppendDebugInfoString(FString& DebugString, const FStateTreeInstance& StateTreeInstance) const;
-#endif
-
-#if WITH_EDITOR
-	// Create Task instance template, set default variables, and resolve variable handles.
-	virtual bool ResolveVariables(const FStateTreeVariableLayout& Variables, FStateTreeConstantStorage& Constants, UObject* Outer) PURE_VIRTUAL(UStateTreeTaskBase::Resolve, return false; );
-	// Returns ture if changed.
-	virtual bool ValidateParameterLayout() { return false; }
-
-	virtual void PostLoad() override;
-#endif
-
-	void SetNewUniqueID() { ID = FGuid::NewGuid(); }
-
-	UPROPERTY(EditAnywhere, Category = Task)
-	FName Name;
-
-	UPROPERTY(meta = (IgnoreForMemberInitializationTest))
-	FGuid ID;
-};
-
-
-// STATETREE_V2
 
 /**
  * Base stuct for StateTree Tasks.
@@ -126,6 +76,5 @@ struct STATETREEMODULE_API FStateTreeTask2Base // TODO: change to FStateTreeTask
 	UPROPERTY()
 	uint16 SourceStructIndex = 0;								// Property binding Source Struct index of the task.
 };
-template<> struct TStructOpsTypeTraits<FStateTreeTask2Base> : public TStructOpsTypeTraitsBase2<FStateTreeTask2Base> { enum { WithPureVirtual = true, }; };
 
-// ~STATETREE_V2
+template<> struct TStructOpsTypeTraits<FStateTreeTask2Base> : public TStructOpsTypeTraitsBase2<FStateTreeTask2Base> { enum { WithPureVirtual = true, }; };
diff --git a/Engine/Plugins/Runtime/StateTree/Source/StateTreeModule/Public/StateTreeTypes.h b/Engine/Plugins/Runtime/StateTree/Source/StateTreeModule/Public/StateTreeTypes.h
--- a/Engine/Plugins/Runtime/StateTree/Source/StateTreeModule/Public/StateTreeTypes.h
+++ b/Engine/Plugins/Runtime/StateTree/Source/StateTreeModule/Public/StateTreeTypes.h
@@ -58,7 +58,6 @@ enum class EStateTreeTransitionType : uint8
 {
 	Succeeded,			// Signal StateTree execution succeeded.
 	Failed,				// Signal StateTree execution failed.
-	SelectChildState,	// Select child state.
 	GotoState,			// Transition to specified state.
 	NotSet,				// No transition.
 	NextState,			// Goto next sibling state.
diff --git a/Engine/Plugins/Runtime/StateTree/Source/StateTreeModule/Public/StateTreeVariable.h b/Engine/Plugins/Runtime/StateTree/Source/StateTreeModule/Public/StateTreeVariable.h
deleted file mode 100644
diff --git a/Engine/Plugins/Runtime/StateTree/Source/StateTreeModule/Public/StateTreeVariableDesc.h b/Engine/Plugins/Runtime/StateTree/Source/StateTreeModule/Public/StateTreeVariableDesc.h
deleted file mode 100644
diff --git a/Engine/Plugins/Runtime/StateTree/Source/StateTreeModule/Public/StateTreeVariableLayout.h b/Engine/Plugins/Runtime/StateTree/Source/StateTreeModule/Public/StateTreeVariableLayout.h
deleted file mode 100644
diff --git a/Engine/Plugins/Runtime/StateTree/Source/StateTreeModule/Public/Tasks/StateTreeTask_SubStateTree.h b/Engine/Plugins/Runtime/StateTree/Source/StateTreeModule/Public/Tasks/StateTreeTask_SubStateTree.h
deleted file mode 100644
diff --git a/Engine/Plugins/Runtime/StateTree/Source/StateTreeModule/Public/Tasks/StateTreeTask_Test.h b/Engine/Plugins/Runtime/StateTree/Source/StateTreeModule/Public/Tasks/StateTreeTask_Test.h
--- a/Engine/Plugins/Runtime/StateTree/Source/StateTreeModule/Public/Tasks/StateTreeTask_Test.h
+++ b/Engine/Plugins/Runtime/StateTree/Source/StateTreeModule/Public/Tasks/StateTreeTask_Test.h
@@ -8,42 +8,6 @@
 #include "StateTreeTask_Test.generated.h"
 
 
-UCLASS(BlueprintType, EditInlineNew, CollapseCategories)
-class STATETREEMODULE_API UStateTreeTask_Test : public UStateTreeTaskBase
-{
-	GENERATED_BODY()
-
-protected:
-	UStateTreeTask_Test(const FObjectInitializer& ObjectInitializer);
-
-	/** Called when instantiated first time. */
-	virtual bool Initialize(FStateTreeInstance& StateTreeInstance) override;
-	/** Called when evaluator becomes active/ticking. */
-	virtual void Activate(FStateTreeInstance& StateTreeInstance)  override;
-	/** Called when evaluator becomes inactive. */
-	virtual void Deactivate(FStateTreeInstance& StateTreeInstance)  override;
-
-	/** Called on each tick. */
-	virtual EStateTreeRunStatus Tick(FStateTreeInstance& StateTreeInstance, const float DeltaTime) override;
-
-#if WITH_GAMEPLAY_DEBUGGER
-	virtual void AppendDebugInfoString(FString& DebugString, const FStateTreeInstance& StateTreeInstance) const;
-#endif
-
-#if WITH_EDITOR
-	/** Create Task instance template, set default variables, and resolve variable handles. */
-	virtual bool ResolveVariables(const FStateTreeVariableLayout& Variables, FStateTreeConstantStorage& Constants, UObject* Outer) override;
-#endif
-
-	float Time = 0.0f;
-
-	UPROPERTY(EditAnywhere, Category = Task)
-	float Duration = 0.0f;
-
-	UPROPERTY(EditAnywhere, Category = Task)
-	bool bFail = false;
-};
-
 USTRUCT(meta = (DisplayName = "Task Test2"))
 struct STATETREEMODULE_API FStateTreeTask2_Test : public FStateTreeTask2Base
 {
diff --git a/Engine/Plugins/Runtime/StateTree/Source/StateTreeTestSuite/Private/StateTreeTest.cpp b/Engine/Plugins/Runtime/StateTree/Source/StateTreeTestSuite/Private/StateTreeTest.cpp
--- a/Engine/Plugins/Runtime/StateTree/Source/StateTreeTestSuite/Private/StateTreeTest.cpp
+++ b/Engine/Plugins/Runtime/StateTree/Source/StateTreeTestSuite/Private/StateTreeTest.cpp
@@ -46,8 +46,7 @@ struct FStateTreeTest_MakeAndBakeStateTree : FAITestBase
 
 		// Root
 		FTestEval_A& EvalA = Root.AddEvaluator<FTestEval_A>();
-		Root.StateDoneTransition.Set(EStateTreeTransitionType::SelectChildState);
-
+		
 		// State A
 		FTestTask_B& TaskB1 = StateA.AddTask<FTestTask_B>();
 		EditorData.AddPropertyBinding(STATETREE_PROPPATH_CHECKED(EvalA, IntA), STATETREE_PROPPATH_CHECKED(TaskB1, IntB));
@@ -55,20 +54,20 @@ struct FStateTreeTest_MakeAndBakeStateTree : FAITestBase
 		FStateTreeCondition_CompareInt& IntCond = StateA.AddEnterCondition<FStateTreeCondition_CompareInt>(0, EGenericAICheck::Less, 2);
 		EditorData.AddPropertyBinding(STATETREE_PROPPATH_CHECKED(EvalA, IntA), STATETREE_PROPPATH_CHECKED(IntCond, Left));
 
-		StateA.StateDoneTransition.Set(EStateTreeTransitionType::GotoState, &StateB);
+		StateA.AddTransition(EStateTreeTransitionEvent::OnCompleted, EStateTreeTransitionType::GotoState, &StateB);
 
 		// State B
 		FTestTask_B& TaskB2 = StateB.AddTask<FTestTask_B>();
 		EditorData.AddPropertyBinding(STATETREE_PROPPATH_CHECKED(EvalA, bBoolA), STATETREE_PROPPATH_CHECKED(TaskB2, bBoolB));
 
-		FStateTreeTransition2& Trans = StateB.AddTransition(EStateTreeTransitionEvent::OnCondition, EStateTreeTransitionType::GotoState, &Root);
+		FStateTreeTransition& Trans = StateB.AddTransition(EStateTreeTransitionEvent::OnCondition, EStateTreeTransitionType::GotoState, &Root);
 		FStateTreeCondition_CompareFloat& TransFloatCond = Trans.AddCondition<FStateTreeCondition_CompareFloat>(0.0f, EGenericAICheck::Less, 13.0f);
 		EditorData.AddPropertyBinding(STATETREE_PROPPATH_CHECKED(EvalA, FloatA), STATETREE_PROPPATH_CHECKED(TransFloatCond, Left));
 
-		StateB.StateDoneTransition.Set(EStateTreeTransitionType::Succeeded);
+		StateB.AddTransition(EStateTreeTransitionEvent::OnCompleted, EStateTreeTransitionType::Succeeded);
 
 		FStateTreeBaker Baker;
-		bool bResult = Baker.Bake2(StateTree);
+		bool bResult = Baker.Bake(StateTree);
 
 		AITEST_TRUE("StateTree should get baked", bResult);
 
@@ -96,7 +95,6 @@ struct FStateTreeTest_WanderLoop : FAITestBase
 		UStateTreeState& UseSO = UseSmartObjectOnLane.AddChildState(FName(TEXT("UseSO")));
 
 		// - Root
-		Root.StateDoneTransition.Set(EStateTreeTransitionType::SelectChildState);
 
 		//   \- Wander
 		FTestEval_Wander& WanderEval = Wander.AddEvaluator<FTestEval_Wander>(FName(TEXT("WanderEval")));
@@ -135,7 +133,7 @@ struct FStateTreeTest_WanderLoop : FAITestBase
 		StandOnLane.AddTransition(EStateTreeTransitionEvent::OnCompleted, EStateTreeTransitionType::GotoState, &Wander);
 
 		FStateTreeBaker Baker;
-		bool bResult = Baker.Bake2(StateTree);
+		bool bResult = Baker.Bake(StateTree);
 
 		AITEST_TRUE("StateTree should get baked", bResult);
 
@@ -227,7 +225,7 @@ struct FStateTreeTest_Sequence : FAITestBase
 		State2.AddTransition(EStateTreeTransitionEvent::OnCompleted, EStateTreeTransitionType::Succeeded);
 
 		FStateTreeBaker Baker;
-		const bool bResult = Baker.Bake2(StateTree);
+		const bool bResult = Baker.Bake(StateTree);
 		AITEST_TRUE("StateTree should get baked", bResult);
 
 		EStateTreeRunStatus Status = EStateTreeRunStatus::Unset;
@@ -287,7 +285,7 @@ struct FStateTreeTest_Select : FAITestBase
 		State1A.AddTransition(EStateTreeTransitionEvent::OnCompleted, EStateTreeTransitionType::GotoState, &State1);
 
 		FStateTreeBaker Baker;
-		const bool bResult = Baker.Bake2(StateTree);
+		const bool bResult = Baker.Bake(StateTree);
 		AITEST_TRUE("StateTree should get baked", bResult);
 
 		EStateTreeRunStatus Status = EStateTreeRunStatus::Unset;
diff --git a/Engine/Plugins/Runtime/StateTree/Source/StateTreeTestSuite/Private/StateTreeTest.h b/Engine/Plugins/Runtime/StateTree/Source/StateTreeTestSuite/Private/StateTreeTest.h
--- a/Engine/Plugins/Runtime/StateTree/Source/StateTreeTestSuite/Private/StateTreeTest.h
+++ b/Engine/Plugins/Runtime/StateTree/Source/StateTreeTestSuite/Private/StateTreeTest.h
@@ -3,7 +3,6 @@
 #pragma once
 
 #include "StateTreeSchema.h"
-
 #include "StateTreeTest.generated.h"
 
 UCLASS()
@@ -11,11 +10,6 @@ class UStateTreeTestSchema : public UStateTreeSchema
 {
 	GENERATED_BODY()
 
-public:
-	/** @return True if we should use StateTree V2 */
-	virtual bool IsV2() const { return true; }
-
-	/** @return True if specified struct is supported */
-	virtual bool IsStructAllowed(const UScriptStruct* InScriptStruct) const { return false; }
+	virtual bool IsStructAllowed(const UScriptStruct* InScriptStruct) const override { return true; }
 };
 
