StateTree: renamed task2 and evaluator2

diff --git a/Engine/Plugins/AI/MassAI/Source/MassAIBehavior/Public/MassStateTreeTypes.h b/Engine/Plugins/AI/MassAI/Source/MassAIBehavior/Public/MassStateTreeTypes.h
--- a/Engine/Plugins/AI/MassAI/Source/MassAIBehavior/Public/MassStateTreeTypes.h
+++ b/Engine/Plugins/AI/MassAI/Source/MassAIBehavior/Public/MassStateTreeTypes.h
@@ -25,7 +25,7 @@ namespace UE::Mass::Signals
  * Base struct for all Mass StateTree Evaluators.
  */
 USTRUCT(meta = (DisplayName = "Mass Evaluator Base"))
-struct MASSAIBEHAVIOR_API FMassStateTreeEvaluatorBase : public FStateTreeEvaluator2Base
+struct MASSAIBEHAVIOR_API FMassStateTreeEvaluatorBase : public FStateTreeEvaluatorBase
 {
 	GENERATED_BODY()
 
@@ -38,7 +38,7 @@ template<> struct TStructOpsTypeTraits<FMassStateTreeEvaluatorBase> : public TSt
  * Base struct for all Mass StateTree Tasks.
  */
 USTRUCT(meta = (DisplayName = "Mass Task Base"))
-struct MASSAIBEHAVIOR_API FMassStateTreeTaskBase : public FStateTreeTask2Base
+struct MASSAIBEHAVIOR_API FMassStateTreeTaskBase : public FStateTreeTaskBase
 {
 	GENERATED_BODY()
 
diff --git a/Engine/Plugins/Runtime/StateTree/Config/BaseStateTree.ini b/Engine/Plugins/Runtime/StateTree/Config/BaseStateTree.ini
--- a/Engine/Plugins/Runtime/StateTree/Config/BaseStateTree.ini
+++ b/Engine/Plugins/Runtime/StateTree/Config/BaseStateTree.ini
@@ -6,6 +6,8 @@
 +PropertyRedirects=(OldName="StateTreeBindableStructDesc.ScriptStruct",NewName="StateTreeBindableStructDesc.Struct")
 
 +StructRedirects=(OldName="StateTreeTransition2",NewName="StateTreeTransition")
++StructRedirects=(OldName="StateTreeTask2Base",NewName="StateTreeTaskBase")
++StructRedirects=(OldName="StateTreeEvaluator2Base",NewName="StateTreeEvaluatorBase")
 
 +PropertyRedirects=(OldName="StateTreeState.Transitions2",NewName="StateTreeState.Transitions")
 +PropertyRedirects=(OldName="StateTreeState.EnterConditions2",NewName="StateTreeState.EnterConditions")
diff --git a/Engine/Plugins/Runtime/StateTree/Source/StateTreeEditorModule/Private/Customizations/StateTreeBindableItemDetails.cpp b/Engine/Plugins/Runtime/StateTree/Source/StateTreeEditorModule/Private/Customizations/StateTreeBindableItemDetails.cpp
--- a/Engine/Plugins/Runtime/StateTree/Source/StateTreeEditorModule/Private/Customizations/StateTreeBindableItemDetails.cpp
+++ b/Engine/Plugins/Runtime/StateTree/Source/StateTreeEditorModule/Private/Customizations/StateTreeBindableItemDetails.cpp
@@ -441,16 +441,16 @@ void FStateTreeBindableItemDetails::OnStructPicked(const UScriptStruct* InStruct
 					Struct->InitializeAs(InStruct);
 
 					// Generate new ID and name.
-					if (InStruct->IsChildOf(FStateTreeTask2Base::StaticStruct()))
+					if (InStruct->IsChildOf(FStateTreeTaskBase::StaticStruct()))
 					{
-						FStateTreeTask2Base* Task = Struct->GetMutablePtr<FStateTreeTask2Base>();
+						FStateTreeTaskBase* Task = Struct->GetMutablePtr<FStateTreeTaskBase>();
 						check(Task);
 						Task->ID = FGuid::NewGuid();
 						Task->Name = FName(InStruct->GetDisplayNameText().ToString());
 					}
-					else if (InStruct->IsChildOf(FStateTreeEvaluator2Base::StaticStruct()))
+					else if (InStruct->IsChildOf(FStateTreeEvaluatorBase::StaticStruct()))
 					{
-						FStateTreeEvaluator2Base* Eval = Struct->GetMutablePtr<FStateTreeEvaluator2Base>();
+						FStateTreeEvaluatorBase* Eval = Struct->GetMutablePtr<FStateTreeEvaluatorBase>();
 						check(Eval);
 						Eval->ID = FGuid::NewGuid();
 						Eval->Name = FName(InStruct->GetDisplayNameText().ToString());
diff --git a/Engine/Plugins/Runtime/StateTree/Source/StateTreeEditorModule/Private/SStateTreeViewRow.cpp b/Engine/Plugins/Runtime/StateTree/Source/StateTreeEditorModule/Private/SStateTreeViewRow.cpp
--- a/Engine/Plugins/Runtime/StateTree/Source/StateTreeEditorModule/Private/SStateTreeViewRow.cpp
+++ b/Engine/Plugins/Runtime/StateTree/Source/StateTreeEditorModule/Private/SStateTreeViewRow.cpp
@@ -421,7 +421,7 @@ EVisibility SStateTreeViewRow::GetEvaluatorsVisibility() const
 		int32 ValidCount = 0;
 		for (int32 i = 0; i < State->Evaluators.Num(); i++)
 		{
-			if (const FStateTreeEvaluator2Base* Eval = State->Evaluators[i].Type.GetPtr<FStateTreeEvaluator2Base>())
+			if (const FStateTreeEvaluatorBase* Eval = State->Evaluators[i].Type.GetPtr<FStateTreeEvaluatorBase>())
 			{
 				ValidCount++;
 			}
@@ -443,7 +443,7 @@ FText SStateTreeViewRow::GetEvaluatorsDesc() const
 
 	for (int32 i = 0; i < State->Evaluators.Num(); i++)
 	{
-		if (const FStateTreeEvaluator2Base* Eval = State->Evaluators[i].Type.GetPtr<FStateTreeEvaluator2Base>())
+		if (const FStateTreeEvaluatorBase* Eval = State->Evaluators[i].Type.GetPtr<FStateTreeEvaluatorBase>())
 		{
 			Names.Add(FText::FromName(Eval->Name));
 		}
@@ -459,7 +459,7 @@ EVisibility SStateTreeViewRow::GetTasksVisibility() const
 		int32 ValidCount = 0;
 		for (int32 i = 0; i < State->Tasks.Num(); i++)
 		{
-			if (const FStateTreeTask2Base* Task = State->Tasks[i].Type.GetPtr<FStateTreeTask2Base>())
+			if (const FStateTreeTaskBase* Task = State->Tasks[i].Type.GetPtr<FStateTreeTaskBase>())
 			{
 				ValidCount++;
 			}
@@ -480,7 +480,7 @@ FText SStateTreeViewRow::GetTasksDesc() const
 	TArray<FText> Names;
 	for (int32 i = 0; i < State->Tasks.Num(); i++)
 	{
-		if (const FStateTreeTask2Base* Task = State->Tasks[i].Type.GetPtr<FStateTreeTask2Base>())
+		if (const FStateTreeTaskBase* Task = State->Tasks[i].Type.GetPtr<FStateTreeTaskBase>())
 		{
 			Names.Add(FText::FromName(Task->Name));
 		}
diff --git a/Engine/Plugins/Runtime/StateTree/Source/StateTreeEditorModule/Private/StateTreeBaker.cpp b/Engine/Plugins/Runtime/StateTree/Source/StateTreeEditorModule/Private/StateTreeBaker.cpp
--- a/Engine/Plugins/Runtime/StateTree/Source/StateTreeEditorModule/Private/StateTreeBaker.cpp
+++ b/Engine/Plugins/Runtime/StateTree/Source/StateTreeEditorModule/Private/StateTreeBaker.cpp
@@ -412,7 +412,7 @@ bool FStateTreeBaker::CreateStateRecursive(UStateTreeState& State, const FStateT
 			FInstancedStruct& EvalPtr = StateTree->RuntimeStorageItems.AddDefaulted_GetRef();
 			EvalPtr = EvaluatorItem.Type;
 			check(EvalPtr.IsValid());
-			FStateTreeEvaluator2Base& Eval = EvalPtr.GetMutable<FStateTreeEvaluator2Base>();
+			FStateTreeEvaluatorBase& Eval = EvalPtr.GetMutable<FStateTreeEvaluatorBase>();
 
 			// Create binding source struct descriptor.
 			FStateTreeBindableStructDesc StructDesc;
@@ -464,7 +464,7 @@ bool FStateTreeBaker::CreateStateRecursive(UStateTreeState& State, const FStateT
 			FInstancedStruct& TaskPtr = StateTree->RuntimeStorageItems.AddDefaulted_GetRef();
 			TaskPtr = TaskItem.Type;
 			check(TaskPtr.IsValid());
-			FStateTreeTask2Base& Task = TaskPtr.GetMutable<FStateTreeTask2Base>();
+			FStateTreeTaskBase& Task = TaskPtr.GetMutable<FStateTreeTaskBase>();
 
 			// Create binding source struct descriptor.
 			FStateTreeBindableStructDesc StructDesc;
diff --git a/Engine/Plugins/Runtime/StateTree/Source/StateTreeEditorModule/Private/StateTreeEditorData.cpp b/Engine/Plugins/Runtime/StateTree/Source/StateTreeEditorModule/Private/StateTreeEditorData.cpp
--- a/Engine/Plugins/Runtime/StateTree/Source/StateTreeEditorModule/Private/StateTreeEditorData.cpp
+++ b/Engine/Plugins/Runtime/StateTree/Source/StateTreeEditorModule/Private/StateTreeEditorData.cpp
@@ -28,7 +28,7 @@ void UStateTreeEditorData::GetAccessibleStructs(const FGuid TargetStructID, TArr
 						OutStructDescs.Append(EvalDescs);
 
 						// Only tasks can see other tasks too.
-						if (Struct->IsChildOf(FStateTreeTask2Base::StaticStruct()))
+						if (Struct->IsChildOf(FStateTreeTaskBase::StaticStruct()))
 						{
 							OutStructDescs.Append(TaskDescs);
 						}
@@ -36,7 +36,7 @@ void UStateTreeEditorData::GetAccessibleStructs(const FGuid TargetStructID, TArr
 						return false; // Stop visit
 					}
 
-					if (Struct->IsChildOf(FStateTreeEvaluator2Base::StaticStruct()))
+					if (Struct->IsChildOf(FStateTreeEvaluatorBase::StaticStruct()))
 					{
 						// All evaluators up to the target struct are accessible.
 						FStateTreeBindableStructDesc& Desc = EvalDescs.AddDefaulted_GetRef();
@@ -44,7 +44,7 @@ void UStateTreeEditorData::GetAccessibleStructs(const FGuid TargetStructID, TArr
 						Desc.Name = Name;
 						Desc.ID = ID;
 					}
-					else if (Struct->IsChildOf(FStateTreeTask2Base::StaticStruct()))
+					else if (Struct->IsChildOf(FStateTreeTaskBase::StaticStruct()))
 					{
 						// All tasks up to the target struct are accessible (for other tasks).
 						FStateTreeBindableStructDesc& Desc = TaskDescs.AddDefaulted_GetRef();
@@ -130,7 +130,7 @@ void UStateTreeEditorData::VisitHierarchy(TFunctionRef<bool(const UStateTreeStat
 			// Evaluators
 			for (const FStateTreeEvaluatorItem& Item : State->Evaluators)
 			{
-				if (const FStateTreeEvaluator2Base* Evaluator = Item.Type.GetPtr<FStateTreeEvaluator2Base>())
+				if (const FStateTreeEvaluatorBase* Evaluator = Item.Type.GetPtr<FStateTreeEvaluatorBase>())
 				{
 					if (!InFunc(*State, Evaluator->ID, Evaluator->Name, Item.Type.GetScriptStruct()))
 					{
@@ -159,7 +159,7 @@ void UStateTreeEditorData::VisitHierarchy(TFunctionRef<bool(const UStateTreeStat
 				// Tasks
 				for (const FStateTreeTaskItem& Item : State->Tasks)
 				{
-					if (const FStateTreeTask2Base* Task = Item.Type.GetPtr<FStateTreeTask2Base>())
+					if (const FStateTreeTaskBase* Task = Item.Type.GetPtr<FStateTreeTaskBase>())
 					{
 						if (!InFunc(*State, Task->ID, Task->Name, Item.Type.GetScriptStruct()))
 						{
diff --git a/Engine/Plugins/Runtime/StateTree/Source/StateTreeEditorModule/Private/StateTreeState.cpp b/Engine/Plugins/Runtime/StateTree/Source/StateTreeEditorModule/Private/StateTreeState.cpp
--- a/Engine/Plugins/Runtime/StateTree/Source/StateTreeEditorModule/Private/StateTreeState.cpp
+++ b/Engine/Plugins/Runtime/StateTree/Source/StateTreeEditorModule/Private/StateTreeState.cpp
@@ -76,7 +76,7 @@ void UStateTreeState::PostEditChangeChainProperty(FPropertyChangedChainEvent& Pr
 					const int32 ArrayIndex = PropertyChangedEvent.GetArrayIndex(MemberProperty->GetFName().ToString());
 					if (Evaluators.IsValidIndex(ArrayIndex))
 					{
-						if (FStateTreeEvaluator2Base* Eval = Evaluators[ArrayIndex].Type.GetMutablePtr<FStateTreeEvaluator2Base>())
+						if (FStateTreeEvaluatorBase* Eval = Evaluators[ArrayIndex].Type.GetMutablePtr<FStateTreeEvaluatorBase>())
 						{
 							Eval->ID = FGuid::NewGuid();
 							Eval->Name = FName(Eval->Name.ToString() + TEXT(" Duplicate"));
@@ -88,7 +88,7 @@ void UStateTreeState::PostEditChangeChainProperty(FPropertyChangedChainEvent& Pr
 					const int32 ArrayIndex = PropertyChangedEvent.GetArrayIndex(MemberProperty->GetFName().ToString());
 					if (Tasks.IsValidIndex(ArrayIndex))
 					{
-						if (FStateTreeTask2Base* Task = Tasks[ArrayIndex].Type.GetMutablePtr<FStateTreeTask2Base>())
+						if (FStateTreeTaskBase* Task = Tasks[ArrayIndex].Type.GetMutablePtr<FStateTreeTaskBase>())
 						{
 							Task->ID = FGuid::NewGuid();
 							Task->Name = FName(Task->Name.ToString() + TEXT(" Duplicate"));
@@ -114,11 +114,11 @@ void UStateTreeState::PostEditChangeChainProperty(FPropertyChangedChainEvent& Pr
 
 #endif
 
-const FStateTreeTask2Base* UStateTreeState::GetTaskByID(FGuid InID) const
+const FStateTreeTaskBase* UStateTreeState::GetTaskByID(FGuid InID) const
 {
 	for (const FStateTreeTaskItem& TaskItem : Tasks)
 	{
-		if (const FStateTreeTask2Base* Task = TaskItem.Type.GetPtr<FStateTreeTask2Base>())
+		if (const FStateTreeTaskBase* Task = TaskItem.Type.GetPtr<FStateTreeTaskBase>())
 		{
 			if (Task->ID == InID)
 			{
@@ -129,11 +129,11 @@ const FStateTreeTask2Base* UStateTreeState::GetTaskByID(FGuid InID) const
 	return nullptr;
 }
 
-FStateTreeTask2Base* UStateTreeState::GetTaskByID(FGuid InID)
+FStateTreeTaskBase* UStateTreeState::GetTaskByID(FGuid InID)
 {
 	for (FStateTreeTaskItem& TaskItem : Tasks)
 	{
-		if (FStateTreeTask2Base* Task = TaskItem.Type.GetMutablePtr<FStateTreeTask2Base>())
+		if (FStateTreeTaskBase* Task = TaskItem.Type.GetMutablePtr<FStateTreeTaskBase>())
 		{
 			if (Task->ID == InID)
 			{
diff --git a/Engine/Plugins/Runtime/StateTree/Source/StateTreeEditorModule/Private/StateTreeViewModel.h b/Engine/Plugins/Runtime/StateTree/Source/StateTreeEditorModule/Private/StateTreeViewModel.h
--- a/Engine/Plugins/Runtime/StateTree/Source/StateTreeEditorModule/Private/StateTreeViewModel.h
+++ b/Engine/Plugins/Runtime/StateTree/Source/StateTreeEditorModule/Private/StateTreeViewModel.h
@@ -9,7 +9,6 @@
 
 class UStateTreeEditorData;
 class UStateTreeState;
-class UStateTreeTaskBase;
 class UStateTree;
 class FMenuBuilder;
 
diff --git a/Engine/Plugins/Runtime/StateTree/Source/StateTreeEditorModule/Public/StateTreeState.h b/Engine/Plugins/Runtime/StateTree/Source/StateTreeEditorModule/Public/StateTreeState.h
--- a/Engine/Plugins/Runtime/StateTree/Source/StateTreeEditorModule/Public/StateTreeState.h
+++ b/Engine/Plugins/Runtime/StateTree/Source/StateTreeEditorModule/Public/StateTreeState.h
@@ -123,8 +123,8 @@ public:
 	virtual void PostEditChangeChainProperty(FPropertyChangedChainEvent& PropertyChangedEvent) override;
 #endif
 
-	const FStateTreeTask2Base* GetTaskByID(FGuid ID) const;
-	FStateTreeTask2Base* GetTaskByID(FGuid ID);
+	const FStateTreeTaskBase* GetTaskByID(FGuid ID) const;
+	FStateTreeTaskBase* GetTaskByID(FGuid ID);
 
 	UStateTreeState* GetNextSiblingState() const;
 
diff --git a/Engine/Plugins/Runtime/StateTree/Source/StateTreeModule/Private/Components/StateTreeComponent.cpp b/Engine/Plugins/Runtime/StateTree/Source/StateTreeModule/Private/Components/StateTreeComponent.cpp
--- a/Engine/Plugins/Runtime/StateTree/Source/StateTreeModule/Private/Components/StateTreeComponent.cpp
+++ b/Engine/Plugins/Runtime/StateTree/Source/StateTreeModule/Private/Components/StateTreeComponent.cpp
@@ -21,7 +21,7 @@
 bool UBrainComponentStateTreeSchema::IsStructAllowed(const UScriptStruct* InScriptStruct) const
 {
 	// @todo: Proper base class for actor nodes.
-	return InScriptStruct->IsChildOf(FStateTreeEvaluator2Base::StaticStruct()) || InScriptStruct->IsChildOf(FStateTreeTask2Base::StaticStruct());
+	return InScriptStruct->IsChildOf(FStateTreeEvaluatorBase::StaticStruct()) || InScriptStruct->IsChildOf(FStateTreeTaskBase::StaticStruct());
 }
 
 
diff --git a/Engine/Plugins/Runtime/StateTree/Source/StateTreeModule/Private/Evaluators/StateTreeEvaluator_Test.cpp b/Engine/Plugins/Runtime/StateTree/Source/StateTreeModule/Private/Evaluators/StateTreeEvaluator_Test.cpp
--- a/Engine/Plugins/Runtime/StateTree/Source/StateTreeModule/Private/Evaluators/StateTreeEvaluator_Test.cpp
+++ b/Engine/Plugins/Runtime/StateTree/Source/StateTreeModule/Private/Evaluators/StateTreeEvaluator_Test.cpp
@@ -4,15 +4,15 @@
 #include "CoreMinimal.h"
 #include "Engine/Engine.h"
 
-FStateTreeEvaluator2_Test::FStateTreeEvaluator2_Test()
+FStateTreeEvaluator_Test::FStateTreeEvaluator_Test()
 {
 }
 
-FStateTreeEvaluator2_Test::~FStateTreeEvaluator2_Test()
+FStateTreeEvaluator_Test::~FStateTreeEvaluator_Test()
 {
 }
 
-void FStateTreeEvaluator2_Test::EnterState(FStateTreeExecutionContext& Context, const EStateTreeStateChangeType ChangeType, const FStateTreeTransitionResult& Transition)
+void FStateTreeEvaluator_Test::EnterState(FStateTreeExecutionContext& Context, const EStateTreeStateChangeType ChangeType, const FStateTreeTransitionResult& Transition)
 {
 	Foo = 0.0f;
 	IntVal = 0;
@@ -21,12 +21,12 @@ void FStateTreeEvaluator2_Test::EnterState(FStateTreeExecutionContext& Context,
 	DummyResult.Count = 0;
 }
 
-void FStateTreeEvaluator2_Test::ExitState(FStateTreeExecutionContext& Context, const EStateTreeStateChangeType ChangeType, const FStateTreeTransitionResult& Transition)
+void FStateTreeEvaluator_Test::ExitState(FStateTreeExecutionContext& Context, const EStateTreeStateChangeType ChangeType, const FStateTreeTransitionResult& Transition)
 {
 	OutResult = &DummyResult;
 }
 
-void FStateTreeEvaluator2_Test::Evaluate(FStateTreeExecutionContext& Context, const EStateTreeEvaluationType EvalType, const float DeltaTime)
+void FStateTreeEvaluator_Test::Evaluate(FStateTreeExecutionContext& Context, const EStateTreeEvaluationType EvalType, const float DeltaTime)
 {
 
 	Foo += DeltaTime;
diff --git a/Engine/Plugins/Runtime/StateTree/Source/StateTreeModule/Private/StateTreeEvaluatorBase.cpp b/Engine/Plugins/Runtime/StateTree/Source/StateTreeModule/Private/StateTreeEvaluatorBase.cpp
--- a/Engine/Plugins/Runtime/StateTree/Source/StateTreeModule/Private/StateTreeEvaluatorBase.cpp
+++ b/Engine/Plugins/Runtime/StateTree/Source/StateTreeModule/Private/StateTreeEvaluatorBase.cpp
@@ -4,7 +4,7 @@
 #include "CoreMinimal.h"
 
 #if WITH_GAMEPLAY_DEBUGGER
-void FStateTreeEvaluator2Base::AppendDebugInfoString(FString& DebugString, const FStateTreeExecutionContext& Context) const
+void FStateTreeEvaluatorBase::AppendDebugInfoString(FString& DebugString, const FStateTreeExecutionContext& Context) const
 {
 	DebugString += FString::Printf(TEXT("[%s]\n"), *Name.ToString());
 }
diff --git a/Engine/Plugins/Runtime/StateTree/Source/StateTreeModule/Private/StateTreeExecutionContext.cpp b/Engine/Plugins/Runtime/StateTree/Source/StateTreeModule/Private/StateTreeExecutionContext.cpp
--- a/Engine/Plugins/Runtime/StateTree/Source/StateTreeModule/Private/StateTreeExecutionContext.cpp
+++ b/Engine/Plugins/Runtime/StateTree/Source/StateTreeModule/Private/StateTreeExecutionContext.cpp
@@ -257,13 +257,13 @@ EStateTreeRunStatus FStateTreeExecutionContext::EnterState(FStateTreeItemView St
 			for (int32 EvalIndex = 0; EvalIndex < int32(State.EvaluatorsNum); EvalIndex++)
 			{
 				FStateTreeItemView EvalView = GetItem(Storage, int32(State.EvaluatorsBegin) + EvalIndex);
-				const FStateTreeEvaluator2Base& Eval = EvalView.Get<FStateTreeEvaluator2Base>();
+				const FStateTreeEvaluatorBase& Eval = EvalView.Get<FStateTreeEvaluatorBase>();
 				ItemViews[Eval.SourceStructIndex] = EvalView;
 			}
 			for (int32 TaskIndex = 0; TaskIndex < int32(State.TasksNum); TaskIndex++)
 			{
 				FStateTreeItemView TaskView = GetItem(Storage, int32(State.TasksBegin) + TaskIndex);
-				const FStateTreeTask2Base& Task = TaskView.Get<FStateTreeTask2Base>();
+				const FStateTreeTaskBase& Task = TaskView.Get<FStateTreeTaskBase>();
 				ItemViews[Task.SourceStructIndex] = TaskView;
 			}
 			continue;
@@ -278,7 +278,7 @@ EStateTreeRunStatus FStateTreeExecutionContext::EnterState(FStateTreeItemView St
 		for (int32 EvalIndex = 0; EvalIndex < int32(State.EvaluatorsNum); EvalIndex++)
 		{
 			FStateTreeItemView EvalView = GetItem(Storage, int32(State.EvaluatorsBegin) + EvalIndex);
-			FStateTreeEvaluator2Base& Eval = EvalView.GetMutable<FStateTreeEvaluator2Base>();
+			FStateTreeEvaluatorBase& Eval = EvalView.GetMutable<FStateTreeEvaluatorBase>();
 
 			// Copy bound properties.
 			if (Eval.BindingsBatch.IsValid())
@@ -295,7 +295,7 @@ EStateTreeRunStatus FStateTreeExecutionContext::EnterState(FStateTreeItemView St
 		for (int32 TaskIndex = 0; TaskIndex < int32(State.TasksNum); TaskIndex++)
 		{
 			FStateTreeItemView TaskView = GetItem(Storage, int32(State.TasksBegin) + TaskIndex);
-			FStateTreeTask2Base& Task = TaskView.GetMutable<FStateTreeTask2Base>();
+			FStateTreeTaskBase& Task = TaskView.GetMutable<FStateTreeTaskBase>();
 
 			// Copy bound properties.
 			if (Task.BindingsBatch.IsValid())
@@ -368,13 +368,13 @@ void FStateTreeExecutionContext::ExitState(FStateTreeItemView Storage, const FSt
 			for (int32 EvalIndex = 0; EvalIndex < int32(State.EvaluatorsNum); EvalIndex++)
 			{
 				FStateTreeItemView EvalView = GetItem(Storage, int32(State.EvaluatorsBegin) + EvalIndex);
-				const FStateTreeEvaluator2Base& Eval = EvalView.Get<FStateTreeEvaluator2Base>();
+				const FStateTreeEvaluatorBase& Eval = EvalView.Get<FStateTreeEvaluatorBase>();
 				ItemViews[Eval.SourceStructIndex] = EvalView;
 			}
 			for (int32 TaskIndex = 0; TaskIndex < int32(State.TasksNum); TaskIndex++)
 			{
 				FStateTreeItemView TaskView = GetItem(Storage, int32(State.TasksBegin) + TaskIndex);
-				const FStateTreeTask2Base& Task = TaskView.Get<FStateTreeTask2Base>();
+				const FStateTreeTaskBase& Task = TaskView.Get<FStateTreeTaskBase>();
 				ItemViews[Task.SourceStructIndex] = TaskView;
 			}
 			continue;
@@ -389,7 +389,7 @@ void FStateTreeExecutionContext::ExitState(FStateTreeItemView Storage, const FSt
 		for (int32 EvalIndex = 0; EvalIndex < int32(State.EvaluatorsNum); EvalIndex++)
 		{
 			FStateTreeItemView EvalView = GetItem(Storage, int32(State.EvaluatorsBegin) + EvalIndex);
-			FStateTreeEvaluator2Base& Eval = EvalView.GetMutable<FStateTreeEvaluator2Base>();
+			FStateTreeEvaluatorBase& Eval = EvalView.GetMutable<FStateTreeEvaluatorBase>();
 
 			// Copy bound properties.
 			if (Eval.BindingsBatch.IsValid())
@@ -406,7 +406,7 @@ void FStateTreeExecutionContext::ExitState(FStateTreeItemView Storage, const FSt
 		for (int32 TaskIndex = 0; TaskIndex < int32(State.TasksNum); TaskIndex++)
 		{
 			FStateTreeItemView TaskView = GetItem(Storage, int32(State.TasksBegin) + TaskIndex);
-			FStateTreeTask2Base& Task = TaskView.GetMutable<FStateTreeTask2Base>();
+			FStateTreeTaskBase& Task = TaskView.GetMutable<FStateTreeTaskBase>();
 
 			// Copy bound properties.
 			if (Task.BindingsBatch.IsValid())
@@ -445,7 +445,7 @@ void FStateTreeExecutionContext::StateCompleted(FStateTreeItemView Storage, cons
 		for (int32 TaskIndex = int32(State.TasksNum) - 1; TaskIndex >= 0; TaskIndex--)
 		{
 			FStateTreeItemView TaskView = GetItem(Storage, int32(State.TasksBegin) + TaskIndex);
-			FStateTreeTask2Base& Task = TaskView.GetMutable<FStateTreeTask2Base>();
+			FStateTreeTaskBase& Task = TaskView.GetMutable<FStateTreeTaskBase>();
 
 			STATETREE_LOG(Verbose, TEXT("%*s  Notify Task '%s'"), Index*UE::StateTree::DebugIndentSize, TEXT(""), *Task.Name.ToString());
 			Task.StateCompleted(*this, CompletionStatus, CurrentState);
@@ -455,7 +455,7 @@ void FStateTreeExecutionContext::StateCompleted(FStateTreeItemView Storage, cons
 		for (int32 EvalIndex = int32(State.EvaluatorsNum) - 1; EvalIndex >= 0; EvalIndex--)
 		{
 			FStateTreeItemView EvalView = GetItem(Storage, int32(State.EvaluatorsBegin) + EvalIndex);
-			FStateTreeEvaluator2Base& Eval = EvalView.GetMutable<FStateTreeEvaluator2Base>();
+			FStateTreeEvaluatorBase& Eval = EvalView.GetMutable<FStateTreeEvaluatorBase>();
 
 			STATETREE_LOG(Verbose, TEXT("%*s  Notify Evaluator '%s'"), Index*UE::StateTree::DebugIndentSize, TEXT(""), *Eval.Name.ToString());
 			Eval.StateCompleted(*this, CompletionStatus, CurrentState);
@@ -492,7 +492,7 @@ void FStateTreeExecutionContext::TickEvaluators(FStateTreeItemView Storage, cons
 		for (int32 EvalIndex = 0; EvalIndex < int32(State.EvaluatorsNum); EvalIndex++)
 		{
 			FStateTreeItemView EvalView = GetItem(Storage, int32(State.EvaluatorsBegin) + EvalIndex);
-			FStateTreeEvaluator2Base& Eval = EvalView.GetMutable<FStateTreeEvaluator2Base>();
+			FStateTreeEvaluatorBase& Eval = EvalView.GetMutable<FStateTreeEvaluatorBase>();
 
 			// Copy bound properties.
 			if (Eval.BindingsBatch.IsValid())
@@ -534,7 +534,7 @@ EStateTreeRunStatus FStateTreeExecutionContext::TickTasks(FStateTreeItemView Sto
 		for (int32 TaskIndex = 0; TaskIndex < int32(State.TasksNum); TaskIndex++)
 		{
 			FStateTreeItemView TaskView = GetItem(Storage, int32(State.TasksBegin) + TaskIndex);
-			FStateTreeTask2Base& Task = TaskView.GetMutable<FStateTreeTask2Base>();
+			FStateTreeTaskBase& Task = TaskView.GetMutable<FStateTreeTaskBase>();
 
 			// Copy bound properties.
 			if (Task.BindingsBatch.IsValid())
@@ -870,7 +870,7 @@ FString FStateTreeExecutionContext::GetDebugInfoString(FStateTreeItemView Extern
 				for (int32 j = 0; j < int32(State.TasksNum); j++)
 				{
 					FStateTreeItemView TaskView = GetItem(Storage, int32(State.TasksBegin) + j);
-					const FStateTreeTask2Base& Task = TaskView.Get<FStateTreeTask2Base>();
+					const FStateTreeTaskBase& Task = TaskView.Get<FStateTreeTaskBase>();
 					Task.AppendDebugInfoString(DebugString, *this);
 				}
 			}
@@ -880,7 +880,7 @@ FString FStateTreeExecutionContext::GetDebugInfoString(FStateTreeItemView Extern
 				for (int32 EvalIndex = 0; EvalIndex < int32(State.EvaluatorsNum); EvalIndex++)
 				{
 					FStateTreeItemView EvalView = GetItem(Storage, int32(State.EvaluatorsBegin) + EvalIndex);
-					const FStateTreeEvaluator2Base& Eval = EvalView.Get<FStateTreeEvaluator2Base>();
+					const FStateTreeEvaluatorBase& Eval = EvalView.Get<FStateTreeEvaluatorBase>();
 					Eval.AppendDebugInfoString(DebugString, *this);
 				}
 			}
@@ -983,7 +983,7 @@ void FStateTreeExecutionContext::DebugPrintInternalLayout(FStateTreeItemView Ext
 		{
 			for (int32 j = 0; j < State.TasksNum; j++)
 			{
-				const FStateTreeTask2Base& Task = GetItem(Storage, State.TasksBegin + j).Get<FStateTreeTask2Base>();
+				const FStateTreeTaskBase& Task = GetItem(Storage, State.TasksBegin + j).Get<FStateTreeTaskBase>();
 				DebugString += FString::Printf(TEXT("  | %-30s | %-12s | %-30s | %15s | %10d |\n"), *State.Name.ToString(),
 					TEXT("  Task"), *Task.Name.ToString(), *Task.BindingsBatch.Describe(), Task.SourceStructIndex);
 			}
@@ -994,7 +994,7 @@ void FStateTreeExecutionContext::DebugPrintInternalLayout(FStateTreeItemView Ext
 		{
 			for (int32 EvalIndex = 0; EvalIndex < State.EvaluatorsNum; EvalIndex++)
 			{
-				const FStateTreeEvaluator2Base& Eval = GetItem(Storage, State.EvaluatorsBegin + EvalIndex).Get<FStateTreeEvaluator2Base>();
+				const FStateTreeEvaluatorBase& Eval = GetItem(Storage, State.EvaluatorsBegin + EvalIndex).Get<FStateTreeEvaluatorBase>();
 				DebugString += FString::Printf(TEXT("  | %-30s | %-12s | %-30s | %15s | %10d |\n"), *State.Name.ToString(),
 					TEXT("  Evaluator"), *Eval.Name.ToString(), *Eval.BindingsBatch.Describe(), Eval.SourceStructIndex);
 			}
diff --git a/Engine/Plugins/Runtime/StateTree/Source/StateTreeModule/Private/StateTreeTaskBase.cpp b/Engine/Plugins/Runtime/StateTree/Source/StateTreeModule/Private/StateTreeTaskBase.cpp
--- a/Engine/Plugins/Runtime/StateTree/Source/StateTreeModule/Private/StateTreeTaskBase.cpp
+++ b/Engine/Plugins/Runtime/StateTree/Source/StateTreeModule/Private/StateTreeTaskBase.cpp
@@ -5,7 +5,7 @@
 
 
 #if WITH_GAMEPLAY_DEBUGGER
-void FStateTreeTask2Base::AppendDebugInfoString(FString& DebugString, const FStateTreeExecutionContext& Context) const
+void FStateTreeTaskBase::AppendDebugInfoString(FString& DebugString, const FStateTreeExecutionContext& Context) const
 {
 	DebugString += FString::Printf(TEXT("[%s]\n"), *Name.ToString());
 }
diff --git a/Engine/Plugins/Runtime/StateTree/Source/StateTreeModule/Public/Evaluators/StateTreeEvaluator_Test.h b/Engine/Plugins/Runtime/StateTree/Source/StateTreeModule/Public/Evaluators/StateTreeEvaluator_Test.h
--- a/Engine/Plugins/Runtime/StateTree/Source/StateTreeModule/Public/Evaluators/StateTreeEvaluator_Test.h
+++ b/Engine/Plugins/Runtime/StateTree/Source/StateTreeModule/Public/Evaluators/StateTreeEvaluator_Test.h
@@ -32,13 +32,13 @@ struct FDummyStateTreeResult : public FStateTreeResult
 	int32 Count = 0;
 };
 
-USTRUCT(meta=(DisplayName="Evaluator Test2"))
-struct STATETREEMODULE_API FStateTreeEvaluator2_Test : public FStateTreeEvaluator2Base
+USTRUCT(meta=(DisplayName="Evaluator Test"))
+struct STATETREEMODULE_API FStateTreeEvaluator_Test : public FStateTreeEvaluatorBase
 {
 	GENERATED_BODY()
 
-	FStateTreeEvaluator2_Test();
-	virtual ~FStateTreeEvaluator2_Test();
+	FStateTreeEvaluator_Test();
+	virtual ~FStateTreeEvaluator_Test();
 	virtual void EnterState(FStateTreeExecutionContext& Context, const EStateTreeStateChangeType ChangeType, const FStateTreeTransitionResult& Transition) override;
 	virtual void ExitState(FStateTreeExecutionContext& Context, const EStateTreeStateChangeType ChangeType, const FStateTreeTransitionResult& Transition) override;
 	virtual void Evaluate(FStateTreeExecutionContext& Context, const EStateTreeEvaluationType EvalType, const float DeltaTime) override;
diff --git a/Engine/Plugins/Runtime/StateTree/Source/StateTreeModule/Public/StateTreeEvaluatorBase.h b/Engine/Plugins/Runtime/StateTree/Source/StateTreeModule/Public/StateTreeEvaluatorBase.h
--- a/Engine/Plugins/Runtime/StateTree/Source/StateTreeModule/Public/StateTreeEvaluatorBase.h
+++ b/Engine/Plugins/Runtime/StateTree/Source/StateTreeModule/Public/StateTreeEvaluatorBase.h
@@ -15,13 +15,13 @@ struct FStateTreeExecutionContext;
  * Evaluators calculate and expose data to be used for decision making in a StateTree.
  */
 USTRUCT()
-struct STATETREEMODULE_API FStateTreeEvaluator2Base // TODO: change to FStateTreeEvaluatorBase once UStateTreeEvaluatorBase is removed
+struct STATETREEMODULE_API FStateTreeEvaluatorBase // TODO: change to FStateTreeEvaluatorBase once UStateTreeEvaluatorBase is removed
 {
 	GENERATED_BODY()
 
-	FStateTreeEvaluator2Base() = default;
+	FStateTreeEvaluatorBase() = default;
 
-	virtual ~FStateTreeEvaluator2Base() {}
+	virtual ~FStateTreeEvaluatorBase() {}
 
 	/**
 	 * Called when a new state is entered and evaluator is part of active states. The change type parameter describes if the evaluator's state
@@ -79,4 +79,4 @@ struct STATETREEMODULE_API FStateTreeEvaluator2Base // TODO: change to FStateTre
 	uint16 SourceStructIndex = 0;									// Property binding Source Struct index of the evaluator.
 };
 
-template<> struct TStructOpsTypeTraits<FStateTreeEvaluator2Base> : public TStructOpsTypeTraitsBase2<FStateTreeEvaluator2Base> { enum { WithPureVirtual = true, }; };
+template<> struct TStructOpsTypeTraits<FStateTreeEvaluatorBase> : public TStructOpsTypeTraitsBase2<FStateTreeEvaluatorBase> { enum { WithPureVirtual = true, }; };
diff --git a/Engine/Plugins/Runtime/StateTree/Source/StateTreeModule/Public/StateTreeTaskBase.h b/Engine/Plugins/Runtime/StateTree/Source/StateTreeModule/Public/StateTreeTaskBase.h
--- a/Engine/Plugins/Runtime/StateTree/Source/StateTreeModule/Public/StateTreeTaskBase.h
+++ b/Engine/Plugins/Runtime/StateTree/Source/StateTreeModule/Public/StateTreeTaskBase.h
@@ -10,17 +10,17 @@
 struct FStateTreeExecutionContext;
 
 /**
- * Base stuct for StateTree Tasks.
+ * Base struct for StateTree Tasks.
  * Tasks are logic executed in an active state.
  */
 USTRUCT()
-struct STATETREEMODULE_API FStateTreeTask2Base // TODO: change to FStateTreeTaskBase once UStateTreeTaskBase is removed
+struct STATETREEMODULE_API FStateTreeTaskBase
 {
 	GENERATED_BODY()
 
-	FStateTreeTask2Base() = default;
+	FStateTreeTaskBase() = default;
 
-	virtual ~FStateTreeTask2Base() {}
+	virtual ~FStateTreeTaskBase() {}
 
 	/**
 	 * Called when a new state is entered and task is part of active states. The change type parameter describes if the task's state
@@ -77,4 +77,4 @@ struct STATETREEMODULE_API FStateTreeTask2Base // TODO: change to FStateTreeTask
 	uint16 SourceStructIndex = 0;								// Property binding Source Struct index of the task.
 };
 
-template<> struct TStructOpsTypeTraits<FStateTreeTask2Base> : public TStructOpsTypeTraitsBase2<FStateTreeTask2Base> { enum { WithPureVirtual = true, }; };
+template<> struct TStructOpsTypeTraits<FStateTreeTaskBase> : public TStructOpsTypeTraitsBase2<FStateTreeTaskBase> { enum { WithPureVirtual = true, }; };
diff --git a/Engine/Plugins/Runtime/StateTree/Source/StateTreeModule/Public/Tasks/StateTreeTask_Test.h b/Engine/Plugins/Runtime/StateTree/Source/StateTreeModule/Public/Tasks/StateTreeTask_Test.h
--- a/Engine/Plugins/Runtime/StateTree/Source/StateTreeModule/Public/Tasks/StateTreeTask_Test.h
+++ b/Engine/Plugins/Runtime/StateTree/Source/StateTreeModule/Public/Tasks/StateTreeTask_Test.h
@@ -9,7 +9,7 @@
 
 
 USTRUCT(meta = (DisplayName = "Task Test2"))
-struct STATETREEMODULE_API FStateTreeTask2_Test : public FStateTreeTask2Base
+struct STATETREEMODULE_API FStateTreeTask2_Test : public FStateTreeTaskBase
 {
 	GENERATED_BODY()
 
diff --git a/Engine/Plugins/Runtime/StateTree/Source/StateTreeTestSuite/Private/StateTreeTestTypes.h b/Engine/Plugins/Runtime/StateTree/Source/StateTreeTestSuite/Private/StateTreeTestTypes.h
--- a/Engine/Plugins/Runtime/StateTree/Source/StateTreeTestSuite/Private/StateTreeTestTypes.h
+++ b/Engine/Plugins/Runtime/StateTree/Source/StateTreeTestSuite/Private/StateTreeTestTypes.h
@@ -76,7 +76,7 @@ struct FTestStateTreeExecutionContext : public FStateTreeExecutionContext
 
 
 USTRUCT()
-struct FTestEval_A : public FStateTreeEvaluator2Base
+struct FTestEval_A : public FStateTreeEvaluatorBase
 {
 	GENERATED_BODY()
 
@@ -94,7 +94,7 @@ struct FTestEval_A : public FStateTreeEvaluator2Base
 };
 
 USTRUCT()
-struct FTestTask_B : public FStateTreeTask2Base
+struct FTestTask_B : public FStateTreeTaskBase
 {
 	GENERATED_BODY()
 
@@ -204,7 +204,7 @@ struct FTestSmartObjectResult : public FStateTreeResult
 
 
 USTRUCT()
-struct FTestEval_Wander : public FStateTreeEvaluator2Base
+struct FTestEval_Wander : public FStateTreeEvaluatorBase
 {
 	GENERATED_BODY()
 
@@ -247,7 +247,7 @@ struct FTestEval_Wander : public FStateTreeEvaluator2Base
 };
 
 USTRUCT()
-struct FTestEval_SmartObjectSensor : public FStateTreeEvaluator2Base
+struct FTestEval_SmartObjectSensor : public FStateTreeEvaluatorBase
 {
 	GENERATED_BODY()
 
@@ -328,7 +328,7 @@ struct FTestEval_SmartObjectSensor : public FStateTreeEvaluator2Base
 
 
 USTRUCT()
-struct FTestTask_ReserveSmartObject : public FStateTreeTask2Base
+struct FTestTask_ReserveSmartObject : public FStateTreeTaskBase
 {
 	GENERATED_BODY()
 
@@ -469,7 +469,7 @@ struct FTestTask_ReserveSmartObject : public FStateTreeTask2Base
 
 
 USTRUCT()
-struct FTestTask_MoveTo : public FStateTreeTask2Base
+struct FTestTask_MoveTo : public FStateTreeTaskBase
 {
 	GENERATED_BODY()
 
@@ -554,7 +554,7 @@ struct FTestTask_MoveTo : public FStateTreeTask2Base
 };
 
 USTRUCT()
-struct FTestTask_Stand : public FStateTreeTask2Base
+struct FTestTask_Stand : public FStateTreeTaskBase
 {
 	GENERATED_BODY()
 
@@ -602,7 +602,7 @@ struct FTestTask_Stand : public FStateTreeTask2Base
 };
 
 USTRUCT()
-struct FTestTask_UseSmartObject : public FStateTreeTask2Base
+struct FTestTask_UseSmartObject : public FStateTreeTaskBase
 {
 	GENERATED_BODY()
 
