Instanced Struct and StateTree: Fixed reset to default buttons (was previously crashing)

diff --git a/Engine/Plugins/Experimental/StructUtils/Source/StructUtilsEditor/Private/Customizations/InstancedStructDetails.cpp b/Engine/Plugins/Experimental/StructUtils/Source/StructUtilsEditor/Private/Customizations/InstancedStructDetails.cpp
--- a/Engine/Plugins/Experimental/StructUtils/Source/StructUtilsEditor/Private/Customizations/InstancedStructDetails.cpp
+++ b/Engine/Plugins/Experimental/StructUtils/Source/StructUtilsEditor/Private/Customizations/InstancedStructDetails.cpp
@@ -206,6 +206,10 @@ void FInstancedStructDetails::CustomizeHeader(TSharedRef<class IPropertyHandle>
 	static const FName StructTypeConst(TEXT("StructTypeConst"));
 	const bool bEnableStructSelection = !StructProperty->HasMetaData(StructTypeConst);
 
+	const FIsResetToDefaultVisible IsResetVisible = FIsResetToDefaultVisible::CreateSP(this, &FInstancedStructDetails::ShouldResetToDefault);
+	const FResetToDefaultHandler ResetHandler = FResetToDefaultHandler::CreateSP(this, &FInstancedStructDetails::ResetToDefault);
+	const FResetToDefaultOverride ResetOverride = FResetToDefaultOverride::Create(IsResetVisible, ResetHandler);
+
 	HeaderRow
 		.NameContent()
 		[
@@ -238,7 +242,62 @@ void FInstancedStructDetails::CustomizeHeader(TSharedRef<class IPropertyHandle>
 					.Font(IDetailLayoutBuilder::GetDetailFont())
 				]
 			]
-		];
+		]
+		.OverrideResetToDefault(ResetOverride);
+}
+
+
+bool FInstancedStructDetails::ShouldResetToDefault(TSharedPtr<IPropertyHandle> PropertyHandle) const
+{
+	check(StructProperty);
+	
+	bool bAnyValid = false;
+	
+	TArray<void*> RawData;
+	StructProperty->AccessRawData(RawData);
+	for (void* Data : RawData)
+	{
+		if (FInstancedStruct* Struct = static_cast<FInstancedStruct*>(Data))
+		{
+			if (Struct->IsValid())
+			{
+				bAnyValid = true;
+				break;
+			}
+		}
+	}
+	
+	// Assume that the default value is empty. Any valid means that some can be reset to empty.
+	return bAnyValid;
+}
+
+void FInstancedStructDetails::ResetToDefault(TSharedPtr<IPropertyHandle> PropertyHandle)
+{
+	check(StructProperty);
+	
+	GEditor->BeginTransaction(LOCTEXT("OnResetToDefault", "Reset to default"));
+
+	StructProperty->NotifyPreChange();
+
+	TArray<void*> RawData;
+	StructProperty->AccessRawData(RawData);
+	for (void* Data : RawData)
+	{
+		if (FInstancedStruct* Struct = static_cast<FInstancedStruct*>(Data))
+		{
+			// Assume that the default value is empty.
+			Struct->Reset();
+		}
+	}
+
+	StructProperty->NotifyPostChange(EPropertyChangeType::ValueSet);
+	GEditor->EndTransaction();
+	StructProperty->NotifyFinishedChangingProperties();
+
+	if (PropUtils)
+	{
+		PropUtils->ForceRefresh();
+	}
 }
 
 void FInstancedStructDetails::CustomizeChildren(TSharedRef<class IPropertyHandle> StructPropertyHandle, class IDetailChildrenBuilder& StructBuilder, IPropertyTypeCustomizationUtils& StructCustomizationUtils)
diff --git a/Engine/Plugins/Experimental/StructUtils/Source/StructUtilsEditor/Public/InstancedStructDetails.h b/Engine/Plugins/Experimental/StructUtils/Source/StructUtilsEditor/Public/InstancedStructDetails.h
--- a/Engine/Plugins/Experimental/StructUtils/Source/StructUtilsEditor/Public/InstancedStructDetails.h
+++ b/Engine/Plugins/Experimental/StructUtils/Source/StructUtilsEditor/Public/InstancedStructDetails.h
@@ -26,6 +26,8 @@ public:
 	virtual void CustomizeChildren(TSharedRef<IPropertyHandle> StructPropertyHandle, class IDetailChildrenBuilder& StructBuilder, IPropertyTypeCustomizationUtils& StructCustomizationUtils) override;
 
 private:
+	bool ShouldResetToDefault(TSharedPtr<IPropertyHandle> PropertyHandle) const;
+	void ResetToDefault(TSharedPtr<IPropertyHandle> PropertyHandle);
 	const UScriptStruct* GetCommonScriptStruct() const;
 	FText GetDisplayValueString() const;
 	const FSlateBrush* GetDisplayValueIcon() const;
diff --git a/Engine/Plugins/Runtime/StateTree/Source/StateTreeEditorModule/Private/Customizations/StateTreeBindableItemDetails.cpp b/Engine/Plugins/Runtime/StateTree/Source/StateTreeEditorModule/Private/Customizations/StateTreeBindableItemDetails.cpp
--- a/Engine/Plugins/Runtime/StateTree/Source/StateTreeEditorModule/Private/Customizations/StateTreeBindableItemDetails.cpp
+++ b/Engine/Plugins/Runtime/StateTree/Source/StateTreeEditorModule/Private/Customizations/StateTreeBindableItemDetails.cpp
@@ -154,6 +154,10 @@ void FStateTreeBindableItemDetails::CustomizeHeader(TSharedRef<class IPropertyHa
 	const FString BaseClassName = TypeProperty->GetMetaData(BaseClassMetaName);
 	BaseScriptStruct = FindObject<UScriptStruct>(ANY_PACKAGE, *BaseClassName);
 
+	const FIsResetToDefaultVisible IsResetVisible = FIsResetToDefaultVisible::CreateSP(this, &FStateTreeBindableItemDetails::ShouldResetToDefault);
+	const FResetToDefaultHandler ResetHandler = FResetToDefaultHandler::CreateSP(this, &FStateTreeBindableItemDetails::ResetToDefault);
+	const FResetToDefaultOverride ResetOverride = FResetToDefaultOverride::Create(IsResetVisible, ResetHandler);
+
 	UE::StateTree::Delegates::OnIdentifierChanged.AddSP(this, &FStateTreeBindableItemDetails::OnIdentifierChanged);
 
 	FindOuterObjects();
@@ -210,7 +214,61 @@ void FStateTreeBindableItemDetails::CustomizeHeader(TSharedRef<class IPropertyHa
 			[
 				StructPropertyHandle->CreateDefaultPropertyButtonWidgets()
 			]
-		];
+		]
+		.OverrideResetToDefault(ResetOverride);
+}
+
+bool FStateTreeBindableItemDetails::ShouldResetToDefault(TSharedPtr<IPropertyHandle> PropertyHandle) const
+{
+	check(TypeProperty);
+	
+	bool bAnyValid = false;
+	
+	TArray<void*> RawData;
+	TypeProperty->AccessRawData(RawData);
+	for (void* Data : RawData)
+	{
+		if (FInstancedStruct* Struct = static_cast<FInstancedStruct*>(Data))
+		{
+			if (Struct->IsValid())
+			{
+				bAnyValid = true;
+				break;
+			}
+		}
+	}
+	
+	// Assume that the default value is empty. Any valid means that some can be reset to empty.
+	return bAnyValid;
+}
+
+void FStateTreeBindableItemDetails::ResetToDefault(TSharedPtr<IPropertyHandle> PropertyHandle)
+{
+	check(TypeProperty);
+	
+	GEditor->BeginTransaction(LOCTEXT("OnResetToDefault", "Reset to default"));
+
+	TypeProperty->NotifyPreChange();
+
+	TArray<void*> RawData;
+	TypeProperty->AccessRawData(RawData);
+	for (void* Data : RawData)
+	{
+		if (FInstancedStruct* Struct = static_cast<FInstancedStruct*>(Data))
+		{
+			// Assume that the default value is empty.
+			Struct->Reset();
+		}
+	}
+
+	TypeProperty->NotifyPostChange(EPropertyChangeType::ValueSet);
+	GEditor->EndTransaction();
+	TypeProperty->NotifyFinishedChangingProperties();
+
+	if (PropUtils)
+	{
+		PropUtils->ForceRefresh();
+	}
 }
 
 void FStateTreeBindableItemDetails::CustomizeChildren(TSharedRef<class IPropertyHandle> StructPropertyHandle, class IDetailChildrenBuilder& StructBuilder, IPropertyTypeCustomizationUtils& StructCustomizationUtils)
@@ -466,8 +524,6 @@ void FStateTreeBindableItemDetails::OnStructPicked(const UScriptStruct* InStruct
 		TypeProperty->NotifyPostChange(EPropertyChangeType::ValueSet);
 		GEditor->EndTransaction();
 		TypeProperty->NotifyFinishedChangingProperties();
-
-		GEditor->EndTransaction();
 	}
 
 	ComboButton->SetIsOpen(false);
diff --git a/Engine/Plugins/Runtime/StateTree/Source/StateTreeEditorModule/Private/Customizations/StateTreeBindableItemDetails.h b/Engine/Plugins/Runtime/StateTree/Source/StateTreeEditorModule/Private/Customizations/StateTreeBindableItemDetails.h
--- a/Engine/Plugins/Runtime/StateTree/Source/StateTreeEditorModule/Private/Customizations/StateTreeBindableItemDetails.h
+++ b/Engine/Plugins/Runtime/StateTree/Source/StateTreeEditorModule/Private/Customizations/StateTreeBindableItemDetails.h
@@ -26,6 +26,9 @@ public:
 
 private:
 
+	bool ShouldResetToDefault(TSharedPtr<IPropertyHandle> PropertyHandle) const;
+	void ResetToDefault(TSharedPtr<IPropertyHandle> PropertyHandle);
+
 	FText GetName() const;
 	void OnNameCommitted(const FText& NewText, ETextCommit::Type InTextCommit);
 	bool IsNameEnabled() const;
diff --git a/Engine/Plugins/Runtime/StateTree/Source/StateTreeEditorModule/Private/Customizations/StateTreeConditionItemDetails.cpp b/Engine/Plugins/Runtime/StateTree/Source/StateTreeEditorModule/Private/Customizations/StateTreeConditionItemDetails.cpp
--- a/Engine/Plugins/Runtime/StateTree/Source/StateTreeEditorModule/Private/Customizations/StateTreeConditionItemDetails.cpp
+++ b/Engine/Plugins/Runtime/StateTree/Source/StateTreeEditorModule/Private/Customizations/StateTreeConditionItemDetails.cpp
@@ -143,6 +143,10 @@ void FStateTreeConditionItemDetails::CustomizeHeader(TSharedRef<class IPropertyH
 	const FString BaseClassName = TypeProperty->GetMetaData(BaseClassMetaName);
 	BaseScriptStruct = FindObject<UScriptStruct>(ANY_PACKAGE, *BaseClassName);
 
+	const FIsResetToDefaultVisible IsResetVisible = FIsResetToDefaultVisible::CreateSP(this, &FStateTreeConditionItemDetails::ShouldResetToDefault);
+	const FResetToDefaultHandler ResetHandler = FResetToDefaultHandler::CreateSP(this, &FStateTreeConditionItemDetails::ResetToDefault);
+	const FResetToDefaultOverride ResetOverride = FResetToDefaultOverride::Create(IsResetVisible, ResetHandler);
+
 	OnBindingChangedHandle = UE::StateTree::PropertyBinding::OnStateTreeBindingChanged.AddRaw(this, &FStateTreeConditionItemDetails::OnBindingChanged);
 
 	FindOuterObjects();
@@ -188,7 +192,61 @@ void FStateTreeConditionItemDetails::CustomizeHeader(TSharedRef<class IPropertyH
 			[
 				StructPropertyHandle->CreateDefaultPropertyButtonWidgets()
 			]
-		];
+		]
+		.OverrideResetToDefault(ResetOverride);
+}
+
+bool FStateTreeConditionItemDetails::ShouldResetToDefault(TSharedPtr<IPropertyHandle> PropertyHandle) const
+{
+	check(TypeProperty);
+	
+	bool bAnyValid = false;
+	
+	TArray<void*> RawData;
+	TypeProperty->AccessRawData(RawData);
+	for (void* Data : RawData)
+	{
+		if (FInstancedStruct* Struct = static_cast<FInstancedStruct*>(Data))
+		{
+			if (Struct->IsValid())
+			{
+				bAnyValid = true;
+				break;
+			}
+		}
+	}
+	
+	// Assume that the default value is empty. Any valid means that some can be reset to empty.
+	return bAnyValid;
+}
+
+void FStateTreeConditionItemDetails::ResetToDefault(TSharedPtr<IPropertyHandle> PropertyHandle)
+{
+	check(TypeProperty);
+	
+	GEditor->BeginTransaction(LOCTEXT("OnResetToDefault", "Reset to default"));
+
+	TypeProperty->NotifyPreChange();
+
+	TArray<void*> RawData;
+	TypeProperty->AccessRawData(RawData);
+	for (void* Data : RawData)
+	{
+		if (FInstancedStruct* Struct = static_cast<FInstancedStruct*>(Data))
+		{
+			// Assume that the default value is empty.
+			Struct->Reset();
+		}
+	}
+
+	TypeProperty->NotifyPostChange(EPropertyChangeType::ValueSet);
+	GEditor->EndTransaction();
+	TypeProperty->NotifyFinishedChangingProperties();
+
+	if (PropUtils)
+	{
+		PropUtils->ForceRefresh();
+	}
 }
 
 void FStateTreeConditionItemDetails::CustomizeChildren(TSharedRef<class IPropertyHandle> StructPropertyHandle, class IDetailChildrenBuilder& StructBuilder, IPropertyTypeCustomizationUtils& StructCustomizationUtils)
diff --git a/Engine/Plugins/Runtime/StateTree/Source/StateTreeEditorModule/Private/Customizations/StateTreeConditionItemDetails.h b/Engine/Plugins/Runtime/StateTree/Source/StateTreeEditorModule/Private/Customizations/StateTreeConditionItemDetails.h
--- a/Engine/Plugins/Runtime/StateTree/Source/StateTreeEditorModule/Private/Customizations/StateTreeConditionItemDetails.h
+++ b/Engine/Plugins/Runtime/StateTree/Source/StateTreeEditorModule/Private/Customizations/StateTreeConditionItemDetails.h
@@ -30,6 +30,9 @@ public:
 
 private:
 
+	bool ShouldResetToDefault(TSharedPtr<IPropertyHandle> PropertyHandle) const;
+	void ResetToDefault(TSharedPtr<IPropertyHandle> PropertyHandle);
+
 	FText GetDescription() const;
 	void OnChildChanged();
 	void OnBindingChanged(const FStateTreeEditorPropertyPath& SourcePath, const FStateTreeEditorPropertyPath& TargetPath);
