StateTree: Mac compile fix.

diff --git a/Engine/Plugins/Runtime/StateTree/Source/StateTreeModule/Public/StateTreeExecutionContext.h b/Engine/Plugins/Runtime/StateTree/Source/StateTreeModule/Public/StateTreeExecutionContext.h
--- a/Engine/Plugins/Runtime/StateTree/Source/StateTreeModule/Public/StateTreeExecutionContext.h
+++ b/Engine/Plugins/Runtime/StateTree/Source/StateTreeModule/Public/StateTreeExecutionContext.h
@@ -142,7 +142,7 @@ public:
 		check(StateTree);
 		check(ItemHandle.IsValid());
 		checkSlow(StateTree->ExternalItems[ItemHandle.ItemIndex].Requirement != EStateTreeItemRequirement::Optional); // Optionals should query pointer instead.
-		return ItemViews[ItemHandle.ItemIndex].GetMutable<T::ItemType>();
+		return ItemViews[ItemHandle.ItemIndex].template GetMutable<typename T::ItemType>();
 	}
 
 	/**
@@ -155,7 +155,7 @@ public:
 	{
 		check(StateTree);
 		check(ItemHandle.IsValid());
-		return ItemViews[ItemHandle.ItemIndex].GetMutablePtr<T::ItemType>();
+		return ItemViews[ItemHandle.ItemIndex].template GetMutablePtr<typename T::ItemType>();
 	}
 	
 	EStateTreeRunStatus GetLastTickStatus(FStateTreeItemView ExternalStorage = FStateTreeItemView()) const;
