StateTree: humble beginnings of compiler message log.
- added message log to StateTree editor
- added statistics window to StateTree editor
- converted baker and property copy to use compiler log instead of UE_LOG

diff --git a/Engine/Plugins/Runtime/StateTree/Source/StateTreeEditorModule/Private/StateTreeBaker.cpp b/Engine/Plugins/Runtime/StateTree/Source/StateTreeEditorModule/Private/StateTreeBaker.cpp
--- a/Engine/Plugins/Runtime/StateTree/Source/StateTreeEditorModule/Private/StateTreeBaker.cpp
+++ b/Engine/Plugins/Runtime/StateTree/Source/StateTreeEditorModule/Private/StateTreeBaker.cpp
@@ -14,49 +14,6 @@
 #include "StateTreePropertyBindingCompiler.h"
 
 
-bool FStateTreeBaker::ResolveTransitionState(const UStateTreeState& SourceState, const TCHAR* ContextStr, const FStateTreeStateLink& Link, FStateTreeHandle& OutTransitionHandle) const 
-{
-	if (Link.Type == EStateTreeTransitionType::GotoState)
-	{
-		OutTransitionHandle = GetStateHandle(Link.ID);
-		if (!OutTransitionHandle.IsValid())
-		{
-			UE_LOG(LogStateTree, Error, TEXT("%s: '%s':%s: Failed to resolve %s to %s."), ANSI_TO_TCHAR(__FUNCTION__), *GetNameSafe(StateTree), *SourceState.Name.ToString(), ContextStr, *Link.Name.ToString());
-			return false;
-		}
-	}
-	else if (Link.Type == EStateTreeTransitionType::NextState)
-	{
-		// Find next state.
-		const UStateTreeState* NextState = SourceState.GetNextSiblingState();
-		if (NextState == nullptr)
-		{
-			UE_LOG(LogStateTree, Error, TEXT("%s: '%s':%s: Failed to resolve default transition, there's no next State after %s."), ANSI_TO_TCHAR(__FUNCTION__), *GetNameSafe(StateTree), *SourceState.Name.ToString(), ContextStr, *SourceState.Name.ToString());
-			return false;
-		}
-		OutTransitionHandle = GetStateHandle(NextState->ID);
-		if (!OutTransitionHandle.IsValid())
-		{
-			UE_LOG(LogStateTree, Error, TEXT("%s: '%s':%s: Failed to resolve default transition next state %s."), ANSI_TO_TCHAR(__FUNCTION__), *GetNameSafe(StateTree), *SourceState.Name.ToString(), ContextStr, *Link.Name.ToString());
-			return false;
-		}
-	}
-	
-	return true;
-}
-
-FStateTreeHandle FStateTreeBaker::GetStateHandle(const FGuid& StateID) const
-{
-	const int32* Idx = IDToState.Find(StateID);
-	if (Idx == nullptr)
-	{
-		return FStateTreeHandle::Invalid;
-	}
-
-	return FStateTreeHandle(uint16(*Idx));
-}
-
-
 bool FStateTreeBaker::Bake(UStateTree& InStateTree)
 {
 	StateTree = &InStateTree;
@@ -69,7 +26,7 @@ bool FStateTreeBaker::Bake(UStateTree& InStateTree)
 	// Cleanup existing state
 	StateTree->ResetBaked();
 
-	BindingsCompiler.Init(StateTree->PropertyBindings);
+	BindingsCompiler.Init(StateTree->PropertyBindings, Log);
 
 	if (!CreateStates())
 	{
@@ -95,6 +52,17 @@ bool FStateTreeBaker::Bake(UStateTree& InStateTree)
 	return true;
 }
 
+FStateTreeHandle FStateTreeBaker::GetStateHandle(const FGuid& StateID) const
+{
+	const int32* Idx = IDToState.Find(StateID);
+	if (Idx == nullptr)
+	{
+		return FStateTreeHandle::Invalid;
+	}
+
+	return FStateTreeHandle(uint16(*Idx));
+}
+
 bool FStateTreeBaker::CreateStates()
 {
 	// Create item for the runtime execution state
@@ -121,13 +89,16 @@ bool FStateTreeBaker::CreateStateTransitions()
 		UStateTreeState* SourceState = SourceStates[i];
 		check(SourceState);
 
+		FStateTreeCompilerLogStateScope LogStateScope(SourceState, Log);
+		
 		// Enter conditions.
 		BakedState.EnterConditionsBegin = uint16(StateTree->Conditions.Num());
 		for (FStateTreeConditionItem& ConditionItem : SourceState->EnterConditions)
 		{
 			if (!CreateCondition(ConditionItem))
 			{
-				UE_LOG(LogStateTree, Error, TEXT("%s: '%s':%s: Failed to create state enter condition."), ANSI_TO_TCHAR(__FUNCTION__), *GetNameSafe(StateTree), *SourceState->Name.ToString());
+				Log.Reportf(EMessageSeverity::Error,
+					TEXT("Failed to create state enter condition."));
 				return false;
 			}
 		}
@@ -142,7 +113,7 @@ bool FStateTreeBaker::CreateStateTransitions()
 			BakedTransition.Type = Transition.State.Type;
 			BakedTransition.GateDelay = (uint8)FMath::Clamp(FMath::CeilToInt(Transition.GateDelay * 10.0f), 0, 255);
 			BakedTransition.State = FStateTreeHandle::Invalid;
-			if (!ResolveTransitionState(*SourceState, TEXT("transition"), Transition.State, BakedTransition.State))
+			if (!ResolveTransitionState(*SourceState, Transition.State, BakedTransition.State))
 			{
 				return false;
 			}
@@ -153,7 +124,9 @@ bool FStateTreeBaker::CreateStateTransitions()
 			{
 				if (!CreateCondition(ConditionItem))
 				{
-					UE_LOG(LogStateTree, Error, TEXT("%s: '%s':%s: Failed to create transition condition to %s."), ANSI_TO_TCHAR(__FUNCTION__), *GetNameSafe(StateTree), *SourceState->Name.ToString(), *Transition.State.Name.ToString());
+					Log.Reportf(EMessageSeverity::Error,
+						TEXT("Failed to create condition for transition to %s."),
+						*Transition.State.Name.ToString());
 					return false;
 				}
 			}
@@ -162,6 +135,44 @@ bool FStateTreeBaker::CreateStateTransitions()
 		BakedState.TransitionsNum = uint8(uint16(StateTree->Transitions.Num()) - BakedState.TransitionsBegin);
 	}
 
+	// @todo: Add test to check that all success/failure transition is possible (see editor).
+	
+	return true;
+}
+
+bool FStateTreeBaker::ResolveTransitionState(const UStateTreeState& SourceState, const FStateTreeStateLink& Link, FStateTreeHandle& OutTransitionHandle) const 
+{
+	if (Link.Type == EStateTreeTransitionType::GotoState)
+	{
+		OutTransitionHandle = GetStateHandle(Link.ID);
+		if (!OutTransitionHandle.IsValid())
+		{
+			Log.Reportf(EMessageSeverity::Error,
+				TEXT("Failed to resolve transition to state %s."),
+				*Link.Name.ToString());
+			return false;
+		}
+	}
+	else if (Link.Type == EStateTreeTransitionType::NextState)
+	{
+		// Find next state.
+		const UStateTreeState* NextState = SourceState.GetNextSiblingState();
+		if (NextState == nullptr)
+		{
+			Log.Reportf(EMessageSeverity::Error,
+				TEXT("Failed to resolve transition, there's no next state."));
+			return false;
+		}
+		OutTransitionHandle = GetStateHandle(NextState->ID);
+		if (!OutTransitionHandle.IsValid())
+		{
+			Log.Reportf(EMessageSeverity::Error,
+				TEXT("Failed to resolve transition next state, no handle found for %s."),
+				*NextState->Name.ToString());
+			return false;
+		}
+	}
+	
 	return true;
 }
 
@@ -209,12 +220,12 @@ bool FStateTreeBaker::IsPropertyAnyEnum(const FStateTreeBindableStructDesc& Stru
 {
 	bool bIsAnyEnum = false;
 	TArray<FStateTreePropertySegment> Segments;
-	FProperty* LeafProperty = nullptr;
+	const FProperty* LeafProperty = nullptr;
 	int32 LeafArrayIndex = INDEX_NONE;
 	const bool bResolved = FStateTreePropertyBindingCompiler::ResolvePropertyPath(Struct, Path, Segments, LeafProperty, LeafArrayIndex);
 	if (bResolved && LeafProperty)
 	{
-		if (FProperty* OwnerProperty = LeafProperty->GetOwnerProperty())
+		if (const FProperty* OwnerProperty = LeafProperty->GetOwnerProperty())
 		{
 			if (const FStructProperty* OwnerStructProperty = CastField<FStructProperty>(OwnerProperty))
 			{
@@ -241,7 +252,9 @@ bool FStateTreeBaker::GetAndValidateBindings(const FStateTreeBindableStructDesc&
 		const int32 SourceStructIdx = BindingsCompiler.GetSourceStructIndexByID(SourceStructID);
 		if (SourceStructIdx == INDEX_NONE)
 		{
-			UE_LOG(LogStateTree, Error, TEXT("%s: '%s': Failed to find container for path %s in %s."), ANSI_TO_TCHAR(__FUNCTION__), *GetNameSafe(StateTree), *Binding.SourcePath.ToString(), *TargetStruct.Name.ToString());
+			Log.Reportf(EMessageSeverity::Error, TargetStruct,
+						TEXT("Failed to find binding source %s:%s."),
+						*TargetStruct.Name.ToString(), *Binding.SourcePath.ToString());
 			return false;
 		}
 		const FStateTreeBindableStructDesc& SourceStruct = BindingsCompiler.GetSourceStructDesc(SourceStructIdx);
@@ -255,8 +268,9 @@ bool FStateTreeBaker::GetAndValidateBindings(const FStateTreeBindableStructDesc&
 			});
 		if (!SourceAccessible)
 		{
-			UE_LOG(LogStateTree, Error, TEXT("%s: '%s': %s:%s is not accessible to %s:%s."), ANSI_TO_TCHAR(__FUNCTION__), *GetNameSafe(StateTree),
-				*SourceStruct.Name.ToString(), *Binding.SourcePath.ToString(), *TargetStruct.Name.ToString(), *Binding.TargetPath.ToString());
+			Log.Reportf(EMessageSeverity::Error, TargetStruct,
+        				TEXT("%s:%s is not accessible to %s:%s."),
+        				*SourceStruct.Name.ToString(), *Binding.SourcePath.ToString(), *TargetStruct.Name.ToString(), *Binding.TargetPath.ToString());
 			return false;
 		}
 
@@ -289,6 +303,8 @@ bool FStateTreeBaker::GetAndValidateBindings(const FStateTreeBindableStructDesc&
 
 bool FStateTreeBaker::CreateStateRecursive(UStateTreeState& State, const FStateTreeHandle Parent)
 {
+	FStateTreeCompilerLogStateScope LogStateScope(&State, Log);
+
 	const int32 StateIdx = StateTree->States.AddDefaulted();
 	FBakedStateTreeState& BakedState = StateTree->States[StateIdx];
 	BakedState.Name = State.Name;
diff --git a/Engine/Plugins/Runtime/StateTree/Source/StateTreeEditorModule/Private/StateTreeCompilerLog.cpp b/Engine/Plugins/Runtime/StateTree/Source/StateTreeEditorModule/Private/StateTreeCompilerLog.cpp
new file mode 100644
--- /dev/null
+++ b/Engine/Plugins/Runtime/StateTree/Source/StateTreeEditorModule/Private/StateTreeCompilerLog.cpp
@@ -0,0 +1,37 @@
+﻿// Copyright Epic Games, Inc. All Rights Reserved.
+
+#include "StateTreeCompilerLog.h"
+
+#include "StateTreeState.h"
+#include "Developer/MessageLog/Public/IMessageLogListing.h"
+#include "Misc/UObjectToken.h"
+
+#define LOCTEXT_NAMESPACE "StateTreeEditor"
+
+void FStateTreeCompilerLog::AppendToLog(IMessageLogListing* LogListing) const
+{
+	for (const FStateTreeCompilerLogMessage& StateTreeMessage : Messages)
+	{
+		TSharedRef<FTokenizedMessage> Message = FTokenizedMessage::Create((EMessageSeverity::Type)StateTreeMessage.Severity);
+
+		if (StateTreeMessage.State != nullptr)
+		{
+			Message->AddToken(FUObjectToken::Create(StateTreeMessage.State, FText::FromName(StateTreeMessage.State->Name)));
+		}
+
+		if (StateTreeMessage.Item.ID.IsValid())
+		{
+			Message->AddToken(FTextToken::Create(FText::Format(LOCTEXT("LogMessageItem", " {0}"), FText::FromName(StateTreeMessage.Item.Name))));
+		}
+
+		if (!StateTreeMessage.Message.IsEmpty())
+		{
+			Message->AddToken(FTextToken::Create(FText::FromString(StateTreeMessage.Message)));
+		}
+
+		LogListing->AddMessage(Message);
+	}
+}
+
+
+#undef LOCTEXT_NAMESPACE
diff --git a/Engine/Plugins/Runtime/StateTree/Source/StateTreeEditorModule/Private/StateTreeEditor.cpp b/Engine/Plugins/Runtime/StateTree/Source/StateTreeEditorModule/Private/StateTreeEditor.cpp
--- a/Engine/Plugins/Runtime/StateTree/Source/StateTreeEditorModule/Private/StateTreeEditor.cpp
+++ b/Engine/Plugins/Runtime/StateTree/Source/StateTreeEditorModule/Private/StateTreeEditor.cpp
@@ -15,9 +15,14 @@
 #include "StateTreeBaker.h"
 #include "IDetailsView.h"
 #include "Widgets/Docking/SDockTab.h"
+#include "Widgets/Input/SMultiLineEditableTextBox.h"
 #include "Customizations/StateTreeBindingExtension.h"
 #include "Logging/MessageLog.h"
+#include "Misc/UObjectToken.h"
 
+#include "Developer/MessageLog/Public/IMessageLogListing.h"
+#include "Developer/MessageLog/Public/MessageLogInitializationOptions.h"
+#include "Developer/MessageLog/Public/MessageLogModule.h"
 
 #define LOCTEXT_NAMESPACE "StateTreeEditor"
 
@@ -26,6 +31,8 @@ const FName StateTreeEditorAppName(TEXT("StateTreeEditorApp"));
 const FName FStateTreeEditor::SelectionDetailsTabId(TEXT("StateTreeEditor_SelectionDetails"));
 const FName FStateTreeEditor::AssetDetailsTabId(TEXT("StateTreeEditor_AssetDetails"));
 const FName FStateTreeEditor::StateTreeViewTabId(TEXT("StateTreeEditor_StateTreeView"));
+const FName FStateTreeEditor::StateTreeStatisticsTabId(TEXT("StateTreeEditor_StateTreeStatistics"));
+const FName FStateTreeEditor::CompilerResultsTabId(TEXT("StateTreeEditor_CompilerResults"));
 
 void FStateTreeEditor::PostUndo(bool bSuccess)
 {
@@ -64,8 +71,17 @@ void FStateTreeEditor::RegisterTabSpawners(const TSharedRef<class FTabManager>&
 		.SetDisplayName(NSLOCTEXT("StateTreeEditor", "StateTreeViewTab", "StateTree"))
 		.SetGroup(WorkspaceMenuCategoryRef)
 		.SetIcon(FSlateIcon(FEditorStyle::GetStyleSetName(), "LevelEditor.Tabs.Outliner"));
+	InTabManager->RegisterTabSpawner(StateTreeStatisticsTabId, FOnSpawnTab::CreateSP(this, &FStateTreeEditor::SpawnTab_StateTreeStatistics))
+		.SetDisplayName(NSLOCTEXT("StateTreeEditor", "StatisticsTab", "StateTree Statistics"))
+		.SetGroup(WorkspaceMenuCategoryRef)
+		.SetIcon(FSlateIcon(FEditorStyle::GetStyleSetName(), "LevelEditor.Tabs.Outliner"));
+	InTabManager->RegisterTabSpawner(CompilerResultsTabId, FOnSpawnTab::CreateSP(this, &FStateTreeEditor::SpawnTab_CompilerResults))
+		.SetDisplayName(NSLOCTEXT("StateTreeEditor", "CompilerResultsTab", "Compiler Results"))
+		.SetGroup(WorkspaceMenuCategoryRef)
+		.SetIcon(FSlateIcon(FEditorStyle::GetStyleSetName(), "LevelEditor.Tabs.Outliner"));
 }
 
+
 void FStateTreeEditor::UnregisterTabSpawners(const TSharedRef<class FTabManager>& InTabManager)
 {
 	FAssetEditorToolkit::UnregisterTabSpawners(InTabManager);
@@ -73,6 +89,8 @@ void FStateTreeEditor::UnregisterTabSpawners(const TSharedRef<class FTabManager>
 	InTabManager->UnregisterTabSpawner(SelectionDetailsTabId);
 	InTabManager->UnregisterTabSpawner(AssetDetailsTabId);
 	InTabManager->UnregisterTabSpawner(StateTreeViewTabId);
+	InTabManager->UnregisterTabSpawner(StateTreeStatisticsTabId);
+	InTabManager->UnregisterTabSpawner(CompilerResultsTabId);
 }
 
 void FStateTreeEditor::InitEditor( const EToolkitMode::Type Mode, const TSharedPtr< class IToolkitHost >& InitToolkitHost, UStateTree* InStateTree)
@@ -93,8 +111,20 @@ void FStateTreeEditor::InitEditor( const EToolkitMode::Type Mode, const TSharedP
 	StateTreeViewModel->GetOnAssetChanged().AddSP(this, &FStateTreeEditor::HandleModelAssetChanged);
 	StateTreeViewModel->GetOnSelectionChanged().AddSP(this, &FStateTreeEditor::HandleModelSelectionChanged);
 
+	FMessageLogModule& MessageLogModule = FModuleManager::LoadModuleChecked<FMessageLogModule>("MessageLog");
+	FMessageLogInitializationOptions LogOptions;
+	// Show Pages so that user is never allowed to clear log messages
+	LogOptions.bShowPages = false;
+	LogOptions.bShowFilters = false;
+	LogOptions.bAllowClear = false;
+	LogOptions.MaxPageCount = 1;
+	CompilerResultsListing = MessageLogModule.CreateLogListing("StateTreeCompiler", LogOptions);
+	CompilerResults = MessageLogModule.CreateLogListingWidget(CompilerResultsListing.ToSharedRef());
+
+	CompilerResultsListing->OnMessageTokenClicked().AddSP(this, &FStateTreeEditor::HandleMessageTokenClicked);
 
-	TSharedRef<FTabManager::FLayout> StandaloneDefaultLayout = FTabManager::NewLayout("Standalone_StateTree_Layout_v2")
+	
+	TSharedRef<FTabManager::FLayout> StandaloneDefaultLayout = FTabManager::NewLayout("Standalone_StateTree_Layout_v3")
 	->AddArea
 	(
 		FTabManager::NewPrimaryArea() ->SetOrientation(Orient_Vertical)
@@ -103,16 +133,37 @@ void FStateTreeEditor::InitEditor( const EToolkitMode::Type Mode, const TSharedP
 			FTabManager::NewSplitter() ->SetOrientation(Orient_Horizontal)
 			->Split
 			(
-				FTabManager::NewStack()
+				FTabManager::NewSplitter()->SetOrientation(Orient_Vertical)
 				->SetSizeCoefficient(0.2f)
-				->AddTab(AssetDetailsTabId, ETabState::OpenedTab)
-				->SetForegroundTab(AssetDetailsTabId)
+				->Split
+				(
+					FTabManager::NewStack()
+					->SetSizeCoefficient(0.65f)
+					->AddTab(AssetDetailsTabId, ETabState::OpenedTab)
+				)
+				->Split
+				(
+					FTabManager::NewStack()
+					->SetSizeCoefficient(0.35f)
+					->AddTab(StateTreeStatisticsTabId, ETabState::OpenedTab)
+				)
 			)
 			->Split
 			(
-				FTabManager::NewStack()
+				FTabManager::NewSplitter()->SetOrientation(Orient_Vertical)
 				->SetSizeCoefficient(0.5f)
-				->AddTab(StateTreeViewTabId, ETabState::OpenedTab)
+				->Split
+				(
+					FTabManager::NewStack()
+					->SetSizeCoefficient(0.75f)
+					->AddTab(StateTreeViewTabId, ETabState::OpenedTab)
+				)
+				->Split
+				(
+					FTabManager::NewStack()
+					->SetSizeCoefficient(0.25f)
+					->AddTab(CompilerResultsTabId, ETabState::ClosedTab)
+				)
 			)
 			->Split
 			(
@@ -156,12 +207,24 @@ FString FStateTreeEditor::GetWorldCentricTabPrefix() const
 	return NSLOCTEXT("StateTreeEditor", "WorldCentricTabPrefix", "State Tree").ToString();
 }
 
-
 FLinearColor FStateTreeEditor::GetWorldCentricTabColorScale() const
 {
 	return FLinearColor( 0.0f, 0.0f, 0.2f, 0.5f );
 }
 
+void FStateTreeEditor::HandleMessageTokenClicked(const TSharedRef<IMessageToken>& InMessageToken)
+{
+	if (InMessageToken->GetType() == EMessageToken::Object)
+	{
+		const TSharedRef<FUObjectToken> ObjectToken = StaticCastSharedRef<FUObjectToken>(InMessageToken);
+		UStateTreeState* State = Cast<UStateTreeState>(ObjectToken->GetObject().Get());
+		if (State)
+		{
+			StateTreeViewModel->SetSelection(State);
+		}
+	}
+}
+
 TSharedRef<SDockTab> FStateTreeEditor::SpawnTab_StateTreeView(const FSpawnTabArgs& Args)
 {
 	check(Args.GetTabId() == StateTreeViewTabId);
@@ -219,6 +282,55 @@ TSharedRef<SDockTab> FStateTreeEditor::SpawnTab_AssetDetails(const FSpawnTabArgs
 	return SpawnedTab;
 }
 
+TSharedRef<SDockTab> FStateTreeEditor::SpawnTab_StateTreeStatistics(const FSpawnTabArgs& Args)
+{
+	check(Args.GetTabId() == StateTreeStatisticsTabId);
+	TSharedRef<SDockTab> SpawnedTab = SNew(SDockTab)
+		.Label(LOCTEXT("StatisticsTitle", "StateTree Statistics"))
+		[
+			SNew(SMultiLineEditableTextBox)
+			.Padding(10)
+			.Style(FEditorStyle::Get(), "Log.TextBox")
+			.Font(FCoreStyle::GetDefaultFontStyle("Mono", 9))
+			.ForegroundColor(FLinearColor::Gray)
+			.IsReadOnly(true)
+			.Text(this, &FStateTreeEditor::GetStatisticsText)
+		];
+	return SpawnedTab;
+}
+
+TSharedRef<SDockTab> FStateTreeEditor::SpawnTab_CompilerResults(const FSpawnTabArgs& Args)
+{
+	check(Args.GetTabId() == CompilerResultsTabId);
+	TSharedRef<SDockTab> SpawnedTab = SNew(SDockTab)
+		.Label(LOCTEXT("CompilerResultsTitle", "Compiler Results"))
+		[
+			SNew(SBox)
+			[
+				CompilerResults.ToSharedRef()
+			]
+		];
+	return SpawnedTab;
+}
+
+FText FStateTreeEditor::GetStatisticsText() const
+{
+	if (!StateTree)
+	{
+		return FText::GetEmpty();
+	}
+
+	if (const UScriptStruct* RuntimeStruct = StateTree->GetRuntimeStorageStruct())
+	{
+		const FText SizeText = FText::AsMemory((uint64)RuntimeStruct->GetStructureSize());
+		const FText NumItemsText = FText::AsNumber(StateTree->GetRuntimeStorageItemCount());
+
+		return FText::Format(LOCTEXT("RuntimeSize", "Runtime size: {0}, {1} items"), SizeText, NumItemsText);
+	}
+	
+	return FText::GetEmpty();
+}
+
 void FStateTreeEditor::HandleModelAssetChanged()
 {
 	UpdateAsset();
@@ -446,15 +558,28 @@ void FStateTreeEditor::UpdateAsset()
 	UE::StateTree::Editor::Internal::RemoveUnusedBindings(*StateTree);
 	UE::StateTree::Editor::Internal::ValidateLinkedStates(*StateTree);
 
-	FStateTreeBaker Baker;
-	if (!Baker.Bake(*StateTree))
+	if (CompilerResultsListing.IsValid())
+	{
+		CompilerResultsListing->ClearMessages();
+	}
+
+	FStateTreeCompilerLog Log;
+	FStateTreeBaker Baker(Log);
+
+	const bool bSuccess = Baker.Bake(*StateTree);
+
+	if (CompilerResultsListing.IsValid())
 	{
-		FMessageLog EditorErrors("StateTree");
-		EditorErrors.Error(FText::Format(LOCTEXT("StateTreeBakerFailed", "Failed to bake state tree asset {0}, see the log for details"), FText::FromString(StateTree->GetName())));
-		EditorErrors.Notify(FText::Format(LOCTEXT("StateTreeBakerFailed", "Failed to bake state tree asset {0}, see the log for details"), FText::FromString(StateTree->GetName())));
+		Log.AppendToLog(CompilerResultsListing.Get());
+	}
 
+	if (!bSuccess)
+	{
 		// Make sure not to leave stale data on failed bake.
 		StateTree->ResetBaked();
+
+		// Show log
+		TabManager->TryInvokeTab(CompilerResultsTabId);
 	}
 }
 
diff --git a/Engine/Plugins/Runtime/StateTree/Source/StateTreeEditorModule/Private/StateTreeEditor.h b/Engine/Plugins/Runtime/StateTree/Source/StateTreeEditorModule/Private/StateTreeEditor.h
--- a/Engine/Plugins/Runtime/StateTree/Source/StateTreeEditorModule/Private/StateTreeEditor.h
+++ b/Engine/Plugins/Runtime/StateTree/Source/StateTreeEditorModule/Private/StateTreeEditor.h
@@ -49,14 +49,15 @@ protected:
 
 private:
 
+	void HandleMessageTokenClicked(const TSharedRef<IMessageToken>& InMessageToken);
 	void HandleModelSelectionChanged(const TArray<UStateTreeState*>& SelectedStates);
 
 	/** Spawns the tab with the update graph inside */
 	TSharedRef<SDockTab> SpawnTab_StateTreeView(const FSpawnTabArgs& Args);
 	TSharedRef<SDockTab> SpawnTab_SelectionDetails(const FSpawnTabArgs& Args);
 	TSharedRef<SDockTab> SpawnTab_AssetDetails(const FSpawnTabArgs& Args);
-
-private:
+	TSharedRef<SDockTab> SpawnTab_StateTreeStatistics(const FSpawnTabArgs& Args);
+	TSharedRef<SDockTab> SpawnTab_CompilerResults(const FSpawnTabArgs& Args);
 
 	void UpdateAsset();
 
@@ -64,6 +65,8 @@ private:
 	void OnIdentifierChanged(const UStateTree& StateTree);
 	void OnSchemaChanged(const UStateTree& StateTree);
 
+	FText GetStatisticsText() const;
+
 	/* State Tree being edited */
 	UStateTree* StateTree;
 
@@ -76,9 +79,15 @@ private:
 	/** Tree View */
 	TSharedPtr<class SStateTreeView> StateTreeView;
 
+	/** Compiler Results log */
+	TSharedPtr<class SWidget> CompilerResults;
+	TSharedPtr<class IMessageLogListing> CompilerResultsListing;
+	
 	TSharedPtr<FStateTreeViewModel> StateTreeViewModel;
 
 	static const FName StateTreeViewTabId;
 	static const FName SelectionDetailsTabId;
 	static const FName AssetDetailsTabId;
+	static const FName StateTreeStatisticsTabId;
+	static const FName CompilerResultsTabId;
 };
diff --git a/Engine/Plugins/Runtime/StateTree/Source/StateTreeEditorModule/Private/StateTreeEditorPropertyBindings.cpp b/Engine/Plugins/Runtime/StateTree/Source/StateTreeEditorModule/Private/StateTreeEditorPropertyBindings.cpp
--- a/Engine/Plugins/Runtime/StateTree/Source/StateTreeEditorModule/Private/StateTreeEditorPropertyBindings.cpp
+++ b/Engine/Plugins/Runtime/StateTree/Source/StateTreeEditorModule/Private/StateTreeEditorPropertyBindings.cpp
@@ -70,7 +70,7 @@ void FStateTreeEditorPropertyBindings::RemoveUnusedBindings(const TMap<FGuid, co
 			StructDesc.ID = Binding.TargetPath.StructID;
 		
 			TArray<FStateTreePropertySegment> Segments;
-			FProperty* LeafProperty = nullptr;
+			const FProperty* LeafProperty = nullptr;
 			int32 LeafArrayIndex = INDEX_NONE;
 			if (!FStateTreePropertyBindingCompiler::ResolvePropertyPath(StructDesc, Binding.TargetPath, Segments, LeafProperty, LeafArrayIndex))
 			{
@@ -125,13 +125,13 @@ const FProperty* FStateTreeBindingLookup::GetPropertyPathLeafProperty(const FSta
 	FStateTreeEditorPropertyBindings* EditorBindings = BindingOwner->GetPropertyEditorBindings();
 	check(EditorBindings);
 
-	FProperty* Result = nullptr;
+	const FProperty* Result = nullptr;
 	FStateTreeBindableStructDesc Struct;
 	if (BindingOwner->GetStructByID(InPath.StructID, Struct))
 	{
 		// TODO: Could use inline allocator here, since there are usually only few segments. Needs API change in ResolvePropertyPath().
 		TArray<FStateTreePropertySegment> Segments;
-		FProperty* LeafProperty = nullptr;
+		const FProperty* LeafProperty = nullptr;
 		int32 LeafArrayIndex = INDEX_NONE;
 		if (FStateTreePropertyBindingCompiler::ResolvePropertyPath(Struct, InPath, Segments, LeafProperty, LeafArrayIndex))
 		{
diff --git a/Engine/Plugins/Runtime/StateTree/Source/StateTreeEditorModule/Private/StateTreePropertyBindingCompiler.cpp b/Engine/Plugins/Runtime/StateTree/Source/StateTreeEditorModule/Private/StateTreePropertyBindingCompiler.cpp
--- a/Engine/Plugins/Runtime/StateTree/Source/StateTreeEditorModule/Private/StateTreePropertyBindingCompiler.cpp
+++ b/Engine/Plugins/Runtime/StateTree/Source/StateTreeEditorModule/Private/StateTreePropertyBindingCompiler.cpp
@@ -5,8 +5,9 @@
 #include "StateTreeTypes.h"
 #include "StateTreePropertyBindings.h"
 
-bool FStateTreePropertyBindingCompiler::Init(FStateTreePropertyBindings& InPropertyBindings)
+bool FStateTreePropertyBindingCompiler::Init(FStateTreePropertyBindings& InPropertyBindings, FStateTreeCompilerLog& InLog)
 {
+	Log = &InLog;
 	PropertyBindings = &InPropertyBindings;
 	PropertyBindings->Reset();
 	SourceStructs.Reset();
@@ -15,6 +16,7 @@ bool FStateTreePropertyBindingCompiler::Init(FStateTreePropertyBindings& InPrope
 
 bool FStateTreePropertyBindingCompiler::CompileBatch(const FStateTreeBindableStructDesc& TargetStruct, TConstArrayView<FStateTreeEditorPropertyBinding> EditorPropertyBindings, int32& OutBatchIndex)
 {
+	check(Log);
 	check(PropertyBindings);
 	OutBatchIndex = INDEX_NONE;
 
@@ -36,7 +38,8 @@ bool FStateTreePropertyBindingCompiler::CompileBatch(const FStateTreeBindableStr
 			});
 		if (SourceStructIdx == INDEX_NONE)
 		{
-			UE_LOG(LogStateTree, Error, TEXT("%s: Could not find a binding source for '%s'."), ANSI_TO_TCHAR(__FUNCTION__), *TargetStruct.Name.ToString());
+			Log->Reportf(EMessageSeverity::Error, TargetStruct,
+				TEXT("Could not find a binding source."));
 			return false;
 		}
 		const FStateTreeBindableStructDesc& SourceStruct = SourceStructs[SourceStructIdx];
@@ -45,17 +48,21 @@ bool FStateTreePropertyBindingCompiler::CompileBatch(const FStateTreeBindableStr
 
 		// Resolve paths
 		FResolvedPathResult SourceResult;
-		if (!ResolvePropertyPath(SourceStruct, EditorBinding.SourcePath, SourceResult))
+		if (!ResolvePropertyPath(TargetStruct, SourceStruct, EditorBinding.SourcePath, SourceResult))
 		{
-			UE_LOG(LogStateTree, Error, TEXT("%s: Could not resolve path to '%s:%s'."), ANSI_TO_TCHAR(__FUNCTION__), *SourceStruct.Name.ToString(), *EditorBinding.SourcePath.ToString());
+			Log->Reportf(EMessageSeverity::Error, TargetStruct,
+				TEXT("Could not resolve path to '%s:%s'."),
+				*SourceStruct.Name.ToString(), *EditorBinding.SourcePath.ToString());
 			return false;
 		}
 
 		// Destination container is set to 0, it is assumed to be passed in when doing the batch copy.
 		FResolvedPathResult TargetResult;
-		if (!ResolvePropertyPath(TargetStruct, EditorBinding.TargetPath, TargetResult))
+		if (!ResolvePropertyPath(TargetStruct, TargetStruct, EditorBinding.TargetPath, TargetResult))
 		{
-			UE_LOG(LogStateTree, Error, TEXT("%s: Could not resolve path to '%s:%s'."), ANSI_TO_TCHAR(__FUNCTION__), *TargetStruct.Name.ToString(), *EditorBinding.TargetPath.ToString());
+			Log->Reportf(EMessageSeverity::Error, TargetStruct,
+				TEXT("Could not resolve path to '%s:%s'."),
+				*TargetStruct.Name.ToString(), *EditorBinding.TargetPath.ToString());
 			return false;
 		}
 
@@ -94,9 +101,9 @@ int32 FStateTreePropertyBindingCompiler::GetSourceStructIndexByID(const FGuid& I
 	return SourceStructs.IndexOfByPredicate([ID](const FStateTreeBindableStructDesc& Structs) { return (Structs.ID == ID); });
 }
 
-EStateTreePropertyCopyType FStateTreePropertyBindingCompiler::GetCopyType(FProperty* SourceProperty, int32 SourceArrayIndex, FProperty* TargetProperty, int32 TargetArrayIndex)
+EStateTreePropertyCopyType FStateTreePropertyBindingCompiler::GetCopyType(const FProperty* SourceProperty, const int32 SourceArrayIndex, const FProperty* TargetProperty, const int32 TargetArrayIndex)
 {
-	if (FArrayProperty* SourceArrayProperty = CastField<FArrayProperty>(SourceProperty))
+	if (const FArrayProperty* SourceArrayProperty = CastField<FArrayProperty>(SourceProperty))
 	{
 		// use the array's inner property if we are not trying to copy the whole array
 		if (SourceArrayIndex != INDEX_NONE)
@@ -105,7 +112,7 @@ EStateTreePropertyCopyType FStateTreePropertyBindingCompiler::GetCopyType(FPrope
 		}
 	}
 
-	if (FArrayProperty* TargetArrayProperty = CastField<FArrayProperty>(TargetProperty))
+	if (const FArrayProperty* TargetArrayProperty = CastField<FArrayProperty>(TargetProperty))
 	{
 		// use the array's inner property if we are not trying to copy the whole array
 		if (TargetArrayIndex != INDEX_NONE)
@@ -114,7 +121,7 @@ EStateTreePropertyCopyType FStateTreePropertyBindingCompiler::GetCopyType(FPrope
 		}
 	}
 
-	EPropertyAccessCompatibility Compatibility = GetPropertyCompatibility(SourceProperty, TargetProperty);
+	const EPropertyAccessCompatibility Compatibility = GetPropertyCompatibility(SourceProperty, TargetProperty);
 
 	// Extract underlying types for enums
 	if (const FEnumProperty* EnumPropertyA = CastField<const FEnumProperty>(SourceProperty))
@@ -258,13 +265,13 @@ EStateTreePropertyCopyType FStateTreePropertyBindingCompiler::GetCopyType(FPrope
 	return EStateTreePropertyCopyType::None;
 }
 
-bool FStateTreePropertyBindingCompiler::ResolvePropertyPath(const FStateTreeBindableStructDesc& InStructDesc, const FStateTreeEditorPropertyPath& InPath, FResolvedPathResult& OutResult)
+bool FStateTreePropertyBindingCompiler::ResolvePropertyPath(const FStateTreeBindableStructDesc& InOwnerStructDesc, const FStateTreeBindableStructDesc& InStructDesc, const FStateTreeEditorPropertyPath& InPath, FResolvedPathResult& OutResult)
 {
 	const int32 PathIndex = PropertyBindings->PropertyPaths.Num();
 	FStateTreePropertyPath& Path = PropertyBindings->PropertyPaths.AddDefaulted_GetRef();
 	Path.SegmentsBegin = PropertyBindings->PropertySegments.Num();
 
-	const bool bResult = ResolvePropertyPath(InStructDesc, InPath, PropertyBindings->PropertySegments, OutResult.LeafProperty, OutResult.LeafArrayIndex, /*bLogErrors*/true);
+	const bool bResult = ResolvePropertyPath(InStructDesc, InPath, PropertyBindings->PropertySegments, OutResult.LeafProperty, OutResult.LeafArrayIndex, Log, &InOwnerStructDesc);
 
 	Path.SegmentsEnd = PropertyBindings->PropertySegments.Num();
 	OutResult.PathIndex = PathIndex;
@@ -273,17 +280,22 @@ bool FStateTreePropertyBindingCompiler::ResolvePropertyPath(const FStateTreeBind
 }
 
 bool FStateTreePropertyBindingCompiler::ResolvePropertyPath(const FStateTreeBindableStructDesc& InStructDesc, const FStateTreeEditorPropertyPath& InPath,
-															TArray<FStateTreePropertySegment>& OutSegments, FProperty*& OutLeafProperty, int32& OutLeafArrayIndex,
-															const bool bLogErrors)
+															TArray<FStateTreePropertySegment>& OutSegments, const FProperty*& OutLeafProperty, int32& OutLeafArrayIndex,
+															FStateTreeCompilerLog* InLog, const FStateTreeBindableStructDesc* InLogContextStruct)
 {
 	if (!InPath.IsValid())
 	{
-		UE_CLOG(bLogErrors, LogStateTree, Error, TEXT("%s: Invalid path '%s:%s'."), ANSI_TO_TCHAR(__FUNCTION__), *InStructDesc.Name.ToString(), *InPath.ToString());
+		if (InLog != nullptr && InLogContextStruct != nullptr)
+		{
+			InLog->Reportf(EMessageSeverity::Error, *InLogContextStruct,
+					TEXT("Invalid path '%s:%s'."),
+					*InStructDesc.Name.ToString(), *InPath.ToString());
+		}
 		return false;
 	}
 
 	const UStruct* CurrentStruct = InStructDesc.Struct;
-	FProperty* LeafProperty = nullptr;
+	const FProperty* LeafProperty = nullptr;
 	int32 LeafArrayIndex = INDEX_NONE;
 	bool bResult = InPath.Path.Num() > 0;
 
@@ -296,24 +308,33 @@ bool FStateTreePropertyBindingCompiler::ResolvePropertyPath(const FStateTreeBind
 		PropertyPathHelpers::FindFieldNameAndArrayIndex(SegmentString.Len(), *SegmentString, PropertyNameLength, &PropertyNamePtr, ArrayIndex);
 		ensure(PropertyNamePtr != nullptr);
 		FString PropertyNameString(PropertyNameLength, PropertyNamePtr);
-		FName PropertyName = FName(*PropertyNameString, FNAME_Find);
+		const FName PropertyName = FName(*PropertyNameString, FNAME_Find);
 
 		const bool bFinalSegment = SegmentIndex == (InPath.Path.Num() - 1);
 
 		if (CurrentStruct == nullptr)
 		{
-			UE_CLOG(bLogErrors, LogStateTree, Error, TEXT("%s: Malformed path '%s:%s'."), ANSI_TO_TCHAR(__FUNCTION__), *InStructDesc.Name.ToString(), *InPath.ToString(SegmentIndex, TEXT("<"), TEXT(">")));
+			if (InLog != nullptr && InLogContextStruct != nullptr)
+			{
+				InLog->Reportf(EMessageSeverity::Error, *InLogContextStruct,
+						TEXT("Malformed path '%s:%s'."),
+						*InStructDesc.Name.ToString(), *InPath.ToString(SegmentIndex, TEXT("<"), TEXT(">")));
+			}
 			bResult = false;
 			break;
 		}
 
-		FProperty* Property = CurrentStruct->FindPropertyByName(PropertyName);
+		const FProperty* Property = CurrentStruct->FindPropertyByName(PropertyName);
 		if (Property == nullptr)
 		{
 			// TODO: use core redirects to fix up the name.
-			UE_CLOG(bLogErrors, LogStateTree, Error, TEXT("%s: Malformed path '%s:%s', could not find property '%s%s.%s'."),
-				ANSI_TO_TCHAR(__FUNCTION__), *InStructDesc.Name.ToString(), *InPath.ToString(SegmentIndex, TEXT("<"), TEXT(">")),
-				CurrentStruct->GetPrefixCPP(), *CurrentStruct->GetName(), *PropertyName.ToString());
+			if (InLog != nullptr && InLogContextStruct != nullptr)
+			{
+				InLog->Reportf(EMessageSeverity::Error, *InLogContextStruct,
+						TEXT("Malformed path '%s:%s', could not find property '%s%s.%s'."),
+						*InStructDesc.Name.ToString(), *InPath.ToString(SegmentIndex, TEXT("<"), TEXT(">")),
+						CurrentStruct->GetPrefixCPP(), *CurrentStruct->GetName(), *PropertyName.ToString());
+			}
 			bResult = false;
 			break;
 		}
@@ -323,17 +344,17 @@ bool FStateTreePropertyBindingCompiler::ResolvePropertyPath(const FStateTreeBind
 		Segment.ArrayIndex = ArrayIndex;
 
 		// Check to see if it is an array access first
-		FArrayProperty* ArrayProperty = CastField<FArrayProperty>(Property);
+		const FArrayProperty* ArrayProperty = CastField<FArrayProperty>(Property);
 		if (ArrayProperty != nullptr && ArrayIndex != INDEX_NONE)
 		{
 			// It is an array, now check to see if this is an array of structures
-			if (FStructProperty* ArrayOfStructsProperty = CastField<FStructProperty>(ArrayProperty->Inner))
+			if (const FStructProperty* ArrayOfStructsProperty = CastField<FStructProperty>(ArrayProperty->Inner))
 			{
 				Segment.Type = EStateTreePropertyAccessType::IndexArray;
 				CurrentStruct = ArrayOfStructsProperty->Struct;
 			}
 			// if it's not an array of structs, maybe it's an array of objects
-			else if (FObjectPropertyBase* ArrayOfObjectsProperty = CastField<FObjectPropertyBase>(ArrayProperty->Inner))
+			else if (const FObjectPropertyBase* ArrayOfObjectsProperty = CastField<FObjectPropertyBase>(ArrayProperty->Inner))
 			{
 				Segment.Type = EStateTreePropertyAccessType::IndexArray;
 				CurrentStruct = ArrayOfObjectsProperty->PropertyClass;
@@ -344,16 +365,16 @@ bool FStateTreePropertyBindingCompiler::ResolvePropertyPath(const FStateTreeBind
 					FStateTreePropertySegment& ExtraSegment = OutSegments.AddDefaulted_GetRef();
 					ExtraSegment.ArrayIndex = 0;
 					ExtraSegment.Type = EStateTreePropertyAccessType::Object;
-					FProperty* InnerProperty = ArrayProperty->Inner;
-					if (FObjectProperty* ObjectProperty = CastField<FObjectProperty>(InnerProperty))
+					const FProperty* InnerProperty = ArrayProperty->Inner;
+					if (const FObjectProperty* ObjectProperty = CastField<FObjectProperty>(InnerProperty))
 					{
 						ExtraSegment.Type = EStateTreePropertyAccessType::Object;
 					}
-					else if (FWeakObjectProperty* WeakObjectProperty = CastField<FWeakObjectProperty>(InnerProperty))
+					else if (const FWeakObjectProperty* WeakObjectProperty = CastField<FWeakObjectProperty>(InnerProperty))
 					{
 						ExtraSegment.Type = EStateTreePropertyAccessType::WeakObject;
 					}
-					else if (FSoftObjectProperty* SoftObjectProperty = CastField<FSoftObjectProperty>(InnerProperty))
+					else if (const FSoftObjectProperty* SoftObjectProperty = CastField<FSoftObjectProperty>(InnerProperty))
 					{
 						ExtraSegment.Type = EStateTreePropertyAccessType::SoftObject;
 					}
@@ -373,33 +394,38 @@ bool FStateTreePropertyBindingCompiler::ResolvePropertyPath(const FStateTreeBind
 			CurrentStruct = nullptr;
 		}
 		// Check to see if this is a simple structure (eg. not an array of structures)
-		else if (FStructProperty* StructProperty = CastField<FStructProperty>(Property))
+		else if (const FStructProperty* StructProperty = CastField<FStructProperty>(Property))
 		{
 			Segment.Type = EStateTreePropertyAccessType::Offset;
 			CurrentStruct = StructProperty->Struct;
 		}
 		// Check to see if this is a simple object (eg. not an array of objects)
-		else if (FObjectProperty* ObjectProperty = CastField<FObjectProperty>(Property))
+		else if (const FObjectProperty* ObjectProperty = CastField<FObjectProperty>(Property))
 		{
 			Segment.Type = EStateTreePropertyAccessType::Object;
 			CurrentStruct = ObjectProperty->PropertyClass;
 		}
 		// Check to see if this is a simple weak object property (eg. not an array of weak objects).
-		else if (FWeakObjectProperty* WeakObjectProperty = CastField<FWeakObjectProperty>(Property))
+		else if (const FWeakObjectProperty* WeakObjectProperty = CastField<FWeakObjectProperty>(Property))
 		{
 			Segment.Type = EStateTreePropertyAccessType::WeakObject;
 			CurrentStruct = WeakObjectProperty->PropertyClass;
 		}
 		// Check to see if this is a simple soft object property (eg. not an array of soft objects).
-		else if (FSoftObjectProperty* SoftObjectProperty = CastField<FSoftObjectProperty>(Property))
+		else if (const FSoftObjectProperty* SoftObjectProperty = CastField<FSoftObjectProperty>(Property))
 		{
 			Segment.Type = EStateTreePropertyAccessType::SoftObject;
 			CurrentStruct = SoftObjectProperty->PropertyClass;
 		}
 		else
 		{
-			UE_CLOG(bLogErrors, LogStateTree, Error, TEXT("%s: Unsupported segment %s in path '%s:%s'."),
-				ANSI_TO_TCHAR(__FUNCTION__), *Property->GetCPPType(), *InStructDesc.Name.ToString(), *InPath.ToString(SegmentIndex, TEXT("<"), TEXT(">")));
+			if (InLog != nullptr && InLogContextStruct != nullptr)
+			{
+				InLog->Reportf(EMessageSeverity::Error, *InLogContextStruct,
+						TEXT("Unsupported segment %s in path '%s:%s'."),
+						*InStructDesc.Name.ToString(), *InPath.ToString(SegmentIndex, TEXT("<"), TEXT(">")),
+						*Property->GetCPPType(), *InStructDesc.Name.ToString(), *InPath.ToString(SegmentIndex, TEXT("<"), TEXT(">")));
+			}
 			bResult = false;
 			break;
 		}
diff --git a/Engine/Plugins/Runtime/StateTree/Source/StateTreeEditorModule/Public/StateTreeBaker.h b/Engine/Plugins/Runtime/StateTree/Source/StateTreeEditorModule/Public/StateTreeBaker.h
--- a/Engine/Plugins/Runtime/StateTree/Source/StateTreeEditorModule/Public/StateTreeBaker.h
+++ b/Engine/Plugins/Runtime/StateTree/Source/StateTreeEditorModule/Public/StateTreeBaker.h
@@ -5,6 +5,7 @@
 #include "CoreMinimal.h"
 #include "StateTreeTypes.h"
 #include "StateTreePropertyBindingCompiler.h"
+#include "StateTreeCompilerLog.h"
 
 class UStateTree;
 class UStateTreeState;
@@ -20,16 +21,17 @@ struct FStateTreeStateLink;
 struct STATETREEEDITORMODULE_API FStateTreeBaker
 {
 public:
+
+	FStateTreeBaker(FStateTreeCompilerLog& InLog)
+		: Log(InLog)
+	{
+	}
+	
 	bool Bake(UStateTree& InStateTree);
 
 private:
-	enum class EBindingContainerType : uint8
-	{
-		Evaluator,
-		Task,
-	};
 
-	bool ResolveTransitionState(const UStateTreeState& SourceState, const TCHAR* ContextStr, const FStateTreeStateLink& Link, FStateTreeHandle& OutTransitionHandle) const;
+	bool ResolveTransitionState(const UStateTreeState& SourceState, const FStateTreeStateLink& Link, FStateTreeHandle& OutTransitionHandle) const;
 	FStateTreeHandle GetStateHandle(const FGuid& StateID) const;
 
 	bool CreateStates();
@@ -39,6 +41,7 @@ private:
 	bool GetAndValidateBindings(const FStateTreeBindableStructDesc& TargetStruct, TArray<FStateTreeEditorPropertyBinding>& OutBindings) const;
 	bool IsPropertyAnyEnum(const FStateTreeBindableStructDesc& Struct, FStateTreeEditorPropertyPath Path) const;
 
+	FStateTreeCompilerLog& Log;
 	UStateTree* StateTree = nullptr;
 	UStateTreeEditorData* TreeData = nullptr;
 	TMap<FGuid, int32> IDToState;
diff --git a/Engine/Plugins/Runtime/StateTree/Source/StateTreeEditorModule/Public/StateTreeCompilerLog.h b/Engine/Plugins/Runtime/StateTree/Source/StateTreeEditorModule/Public/StateTreeCompilerLog.h
new file mode 100644
--- /dev/null
+++ b/Engine/Plugins/Runtime/StateTree/Source/StateTreeEditorModule/Public/StateTreeCompilerLog.h
@@ -0,0 +1,120 @@
+﻿// Copyright Epic Games, Inc. All Rights Reserved.
+
+#pragma once
+
+#include "CoreMinimal.h"
+#include "Logging/TokenizedMessage.h"
+#include "StateTreePropertyBindings.h"
+#include "StateTreeCompilerLog.generated.h"
+
+class UStateTreeState;
+
+/** StateTree compiler log message */
+USTRUCT()
+struct STATETREEEDITORMODULE_API FStateTreeCompilerLogMessage
+{
+	GENERATED_BODY()
+
+	FStateTreeCompilerLogMessage() = default;
+	FStateTreeCompilerLogMessage(const EMessageSeverity::Type InSeverity, const UStateTreeState* InState, const FStateTreeBindableStructDesc& InItem, const FString& InMessage)
+		: Severity(InSeverity)
+		, State(InState)
+		, Item(InItem)
+		, Message(InMessage)
+	{
+	}
+
+	/** Severity of the message. */
+	UPROPERTY()
+	int32 Severity = EMessageSeverity::Error;
+
+	/** (optional) The StateTree state the message refers to. */
+	UPROPERTY()
+	const UStateTreeState* State = nullptr;
+
+	/** (optional) The State tee item (condition/evaluator/task) the message refers to. */
+	UPROPERTY()
+	FStateTreeBindableStructDesc Item;
+
+	/** The message */
+	UPROPERTY()
+	FString Message;
+};
+
+/** Message log for StateTree compilation */
+USTRUCT()
+struct STATETREEEDITORMODULE_API FStateTreeCompilerLog
+{
+	GENERATED_BODY()
+
+	/** Pushes State to be reported along with the message. */
+	void PushState(const UStateTreeState* InState)
+	{
+		StateStack.Add(InState);
+	}
+	
+	/** Pops State to be reported along with the message. @see FStateTreeCompilerLogStateScope */
+	void PopState(const UStateTreeState* InState)
+	{
+		// Check for potentially miss matching push/pop
+		check(StateStack.Num() > 0 && StateStack.Last() == InState);
+		StateStack.Pop();
+	}
+
+	/** Returns current state context. */
+	const UStateTreeState* CurrentState() const { return StateStack.Num() > 0 ? StateStack.Last() : nullptr;  }
+
+	/**
+	 * Reports a message.
+	 * @param InSeverity Severity of the message.
+	 * @param InItem StateTree item (condition/evaluator/task) the message affects.
+	 * @param InMessage Message to display.
+	 */
+	void Report(EMessageSeverity::Type InSeverity, const FStateTreeBindableStructDesc& InItem, const FString& InMessage)
+	{
+		Messages.Emplace(InSeverity, CurrentState(), InItem, InMessage);
+	}
+
+	/** Formatted version of the Report(). */
+	template <typename FmtType, typename... Types>
+	void Reportf(EMessageSeverity::Type InSeverity, const FStateTreeBindableStructDesc& InItem, const FmtType& Fmt, Types... Args)
+	{
+		Report(InSeverity, InItem, FString::Printf(Fmt, Args...));
+	}
+
+	/** Formatted version of the Report(), omits Item for convenience. */
+	template <typename FmtType, typename... Types>
+	void Reportf(EMessageSeverity::Type InSeverity, const FmtType& Fmt, Types... Args)
+	{
+		Report(InSeverity, FStateTreeBindableStructDesc(), FString::Printf(Fmt, Args...));
+	}
+
+	/** Appends StateTree log to log listing. */
+	void AppendToLog(class IMessageLogListing* LogListing) const;
+	
+protected:
+	UPROPERTY()
+	TArray<const UStateTreeState*> StateStack;
+	
+	UPROPERTY()
+	TArray<FStateTreeCompilerLogMessage> Messages;
+};
+
+/** Helper struct to manage reported state within a scope. */
+struct FStateTreeCompilerLogStateScope
+{
+	FStateTreeCompilerLogStateScope(const UStateTreeState* InState, FStateTreeCompilerLog& InLog)
+		: Log(InLog)
+		, State(InState)
+	{
+		Log.PushState(State);
+	}
+
+	~FStateTreeCompilerLogStateScope()
+	{
+		Log.PopState(State);
+	}
+
+	FStateTreeCompilerLog& Log; 
+	const UStateTreeState* State = nullptr;
+};
diff --git a/Engine/Plugins/Runtime/StateTree/Source/StateTreeEditorModule/Public/StateTreePropertyBindingCompiler.h b/Engine/Plugins/Runtime/StateTree/Source/StateTreeEditorModule/Public/StateTreePropertyBindingCompiler.h
--- a/Engine/Plugins/Runtime/StateTree/Source/StateTreeEditorModule/Public/StateTreePropertyBindingCompiler.h
+++ b/Engine/Plugins/Runtime/StateTree/Source/StateTreeEditorModule/Public/StateTreePropertyBindingCompiler.h
@@ -9,6 +9,7 @@
 #include "IPropertyAccessEditor.h"
 #include "StateTreeEditorPropertyBindings.h"
 #include "StateTreePropertyBindings.h"
+#include "StateTreeCompilerLog.h"
 #include "StateTreePropertyBindingCompiler.generated.h"
 
 /**
@@ -25,7 +26,7 @@ struct STATETREEEDITORMODULE_API FStateTreePropertyBindingCompiler
 	  * @param PropertyBindings - Reference to the Property Bindings where all the batches will be stored.
 	  * @return true on success.
 	  */
-	bool Init(FStateTreePropertyBindings& InPropertyBindings);
+	bool Init(FStateTreePropertyBindings& InPropertyBindings, FStateTreeCompilerLog& InLog);
 
 	/**
 	  * Compiles a batch of property copies.
@@ -58,16 +59,19 @@ struct STATETREEEDITORMODULE_API FStateTreePropertyBindingCompiler
 
 	/**
 	 * Resolves a string based property path in specified struct into segments of property names and access types.
+	 * If logging is required both, Log and LogContextStruct needs to be non-null.
 	 * @param InStructDesc Description of the struct in which the property path is valid.
 	 * @param InPath The property path in string format.
 	 * @param OutSegments The resolved property access path as segments.
 	 * @param OutLeafProperty The leaf property of the resolved path.
 	 * @param OutLeafArrayIndex The left array index (or INDEX_NONE if not applicable) of the resolved path.
-	 * @param bLogErrors Set to true to log errors during resolve into the console.
+	 * @param Log Pointer to compiler log, or null if no logging needed.
+	 * @param LogContextStruct Pointer to bindable struct desc where the property path belongs to.
 	 * @return True of the property was solved successfully.
 	 */
 	static bool ResolvePropertyPath(const FStateTreeBindableStructDesc& InStructDesc, const FStateTreeEditorPropertyPath& InPath,
-									TArray<FStateTreePropertySegment>& OutSegments, FProperty*& OutLeafProperty, int32& OutLeafArrayIndex, const bool bLogErrors = false);
+									TArray<FStateTreePropertySegment>& OutSegments, const FProperty*& OutLeafProperty, int32& OutLeafArrayIndex,
+									FStateTreeCompilerLog* Log = nullptr, const FStateTreeBindableStructDesc* LogContextStruct = nullptr);
 
 	/**
 	 * Checks if two property types can are compatible for copying.
@@ -86,13 +90,15 @@ protected:
 
 	FStateTreePropertyBindings* PropertyBindings = nullptr;
 
+	FStateTreeCompilerLog* Log = nullptr;
+	
 	struct FResolvedPathResult
 	{
 		int32 PathIndex = INDEX_NONE;
-		FProperty* LeafProperty = nullptr;
+		const FProperty* LeafProperty = nullptr;
 		int32 LeafArrayIndex = INDEX_NONE;
 	};
 
-	EStateTreePropertyCopyType GetCopyType(FProperty* SourceProperty, int32 SourceArrayIndex, FProperty* TargetProperty, int32 TargetArrayIndex);
-	bool ResolvePropertyPath(const FStateTreeBindableStructDesc& InStructDesc, const FStateTreeEditorPropertyPath& InPath, FResolvedPathResult& OutResult);
+	EStateTreePropertyCopyType GetCopyType(const FProperty* SourceProperty, const int32 SourceArrayIndex, const FProperty* TargetProperty, const int32 TargetArrayIndex);
+	bool ResolvePropertyPath(const FStateTreeBindableStructDesc& InOwnerStructDesc, const FStateTreeBindableStructDesc& InStructDesc, const FStateTreeEditorPropertyPath& InPath, FResolvedPathResult& OutResult);
 };
diff --git a/Engine/Plugins/Runtime/StateTree/Source/StateTreeTestSuite/Private/StateTreeTest.cpp b/Engine/Plugins/Runtime/StateTree/Source/StateTreeTestSuite/Private/StateTreeTest.cpp
--- a/Engine/Plugins/Runtime/StateTree/Source/StateTreeTestSuite/Private/StateTreeTest.cpp
+++ b/Engine/Plugins/Runtime/StateTree/Source/StateTreeTestSuite/Private/StateTreeTest.cpp
@@ -66,7 +66,8 @@ struct FStateTreeTest_MakeAndBakeStateTree : FAITestBase
 
 		StateB.AddTransition(EStateTreeTransitionEvent::OnCompleted, EStateTreeTransitionType::Succeeded);
 
-		FStateTreeBaker Baker;
+		FStateTreeCompilerLog Log;
+		FStateTreeBaker Baker(Log);
 		bool bResult = Baker.Bake(StateTree);
 
 		AITEST_TRUE("StateTree should get baked", bResult);
@@ -132,7 +133,8 @@ struct FStateTreeTest_WanderLoop : FAITestBase
 		StandTask.TicksToCompletion = 2;
 		StandOnLane.AddTransition(EStateTreeTransitionEvent::OnCompleted, EStateTreeTransitionType::GotoState, &Wander);
 
-		FStateTreeBaker Baker;
+		FStateTreeCompilerLog Log;
+		FStateTreeBaker Baker(Log);
 		bool bResult = Baker.Bake(StateTree);
 
 		AITEST_TRUE("StateTree should get baked", bResult);
@@ -224,7 +226,8 @@ struct FStateTreeTest_Sequence : FAITestBase
 		FTestTask_Stand& Task2 = State2.AddTask<FTestTask_Stand>(FName(TEXT("Task2")));
 		State2.AddTransition(EStateTreeTransitionEvent::OnCompleted, EStateTreeTransitionType::Succeeded);
 
-		FStateTreeBaker Baker;
+		FStateTreeCompilerLog Log;
+		FStateTreeBaker Baker(Log);
 		const bool bResult = Baker.Bake(StateTree);
 		AITEST_TRUE("StateTree should get baked", bResult);
 
@@ -284,7 +287,8 @@ struct FStateTreeTest_Select : FAITestBase
 		Task1A.TicksToCompletion = 2;
 		State1A.AddTransition(EStateTreeTransitionEvent::OnCompleted, EStateTreeTransitionType::GotoState, &State1);
 
-		FStateTreeBaker Baker;
+		FStateTreeCompilerLog Log;
+		FStateTreeBaker Baker(Log);
 		const bool bResult = Baker.Bake(StateTree);
 		AITEST_TRUE("StateTree should get baked", bResult);
 
