StateTree: calculate external item indices correctly

diff --git a/Engine/Plugins/Runtime/StateTree/Source/StateTreeModule/Private/StateTree.cpp b/Engine/Plugins/Runtime/StateTree/Source/StateTreeModule/Private/StateTree.cpp
--- a/Engine/Plugins/Runtime/StateTree/Source/StateTreeModule/Private/StateTree.cpp
+++ b/Engine/Plugins/Runtime/StateTree/Source/StateTreeModule/Private/StateTree.cpp
@@ -31,6 +31,8 @@ void UStateTree::ResetBaked()
 	RuntimeStorageDefaultValue.Reset();
 	ExternalItems.Reset();
 	PropertyBindings.Reset();
+
+	NumLinkedItems = 0;
 }
 
 void UStateTree::PostEditChangeProperty(struct FPropertyChangedEvent& PropertyChangedEvent)
@@ -96,6 +98,9 @@ void UStateTree::Link()
 {
 	FStateTreeLinker Linker;
 
+	const int32 BaseIndex = PropertyBindings.GetSourceStructNum();
+	Linker.SetItemBaseIndex(BaseIndex);
+	
 	for (FInstancedStruct& RuntimeItem : RuntimeStorageItems)
 	{
 		if (FStateTreeEvaluatorBase* Eval = RuntimeItem.GetMutablePtr<FStateTreeEvaluatorBase>())
@@ -109,6 +114,8 @@ void UStateTree::Link()
 	}
 
 	ExternalItems = Linker.GetItemDescs();
+	
+	NumLinkedItems = BaseIndex + ExternalItems.Num();
 }
 
 void UStateTree::InitRuntimeStorage()
diff --git a/Engine/Plugins/Runtime/StateTree/Source/StateTreeModule/Private/StateTreeExecutionContext.cpp b/Engine/Plugins/Runtime/StateTree/Source/StateTreeModule/Private/StateTreeExecutionContext.cpp
--- a/Engine/Plugins/Runtime/StateTree/Source/StateTreeModule/Private/StateTreeExecutionContext.cpp
+++ b/Engine/Plugins/Runtime/StateTree/Source/StateTreeModule/Private/StateTreeExecutionContext.cpp
@@ -48,8 +48,8 @@ bool FStateTreeExecutionContext::Init(UObject& InOwner, const UStateTree& InStat
 		StorageInstance = StateTree->RuntimeStorageDefaultValue;
 	}
 
-	// Initialize struct view for all possible source structs.
-	ItemViews.SetNum(StateTree->PropertyBindings.GetSourceStructNum());
+	// Initialize struct view for all possible source items.
+	ItemViews.SetNum(StateTree->GetLinkedItemCount());
 
 	// Initialize array to keep track of which states have been updated.
 	VisitedStates.Init(false, StateTree->States.Num());
diff --git a/Engine/Plugins/Runtime/StateTree/Source/StateTreeModule/Public/StateTree.h b/Engine/Plugins/Runtime/StateTree/Source/StateTreeModule/Public/StateTree.h
--- a/Engine/Plugins/Runtime/StateTree/Source/StateTreeModule/Public/StateTree.h
+++ b/Engine/Plugins/Runtime/StateTree/Source/StateTreeModule/Public/StateTree.h
@@ -31,6 +31,9 @@ public:
 	/** @return Number of items (Evaluators & Tasks) in the runtime storage. */
 	int32 GetRuntimeStorageItemCount() const { return RuntimeStorageOffsets.Num(); }
 
+	/** @return Number of linked items (Evaluators, Tasks, external items). */
+	int32 GetLinkedItemCount() const { return NumLinkedItems; }
+
 	/** @return List of external items required by the state tree */
 	TConstArrayView<FStateTreeExternalItemDesc> GetExternalItems() const { return ExternalItems; }
 
@@ -98,6 +101,9 @@ private:
 	UPROPERTY(Transient)
 	TArray<FStateTreeExternalItemDesc> ExternalItems;
 
+	UPROPERTY(Transient)
+	int32 NumLinkedItems = 0;
+
 	UPROPERTY()
 	FStateTreePropertyBindings PropertyBindings;
 
diff --git a/Engine/Plugins/Runtime/StateTree/Source/StateTreeModule/Public/StateTreeTypes.h b/Engine/Plugins/Runtime/StateTree/Source/StateTreeModule/Public/StateTreeTypes.h
--- a/Engine/Plugins/Runtime/StateTree/Source/StateTreeModule/Public/StateTreeTypes.h
+++ b/Engine/Plugins/Runtime/StateTree/Source/StateTreeModule/Public/StateTreeTypes.h
@@ -457,6 +457,9 @@ struct STATETREEMODULE_API FStateTreeExternalItemDesc
  */
 struct FStateTreeLinker
 {
+	/** Sets base index for all external item handles. */
+	void SetItemBaseIndex(const int32 InItemBaseIndex) { ItemBaseIndex = InItemBaseIndex; }
+	
 	/**
 	 * Links reference to an external UObject.
 	 * @param Handle Reference to TStateTreeItemHandle<> with UOBJECT type to link to.
@@ -488,11 +491,12 @@ protected:
 		if (Index == INDEX_NONE)
 		{
 			Index = ItemDescs.Add(Desc);
-			ItemDescs[Index].Handle.ItemIndex = (uint8)Index;
+			check(FStateTreeItemHandle::IsValidIndex(Index + ItemBaseIndex));
+			ItemDescs[Index].Handle.ItemIndex = (uint8)(Index + ItemBaseIndex);
 		}
-		check(FStateTreeItemHandle::IsValidIndex(Index));
-		Handle.ItemIndex = (uint8)Index;
+		Handle.ItemIndex = (uint8)(Index + ItemBaseIndex);
 	}
 
+	int32 ItemBaseIndex = 0;
 	TArray<FStateTreeExternalItemDesc> ItemDescs;
 };
