StateTree: Fixed validation crash in GetExternalItem()

diff --git a/Engine/Plugins/Runtime/StateTree/Source/StateTreeModule/Private/StateTree.cpp b/Engine/Plugins/Runtime/StateTree/Source/StateTreeModule/Private/StateTree.cpp
--- a/Engine/Plugins/Runtime/StateTree/Source/StateTreeModule/Private/StateTree.cpp
+++ b/Engine/Plugins/Runtime/StateTree/Source/StateTreeModule/Private/StateTree.cpp
@@ -33,6 +33,7 @@ void UStateTree::ResetBaked()
 	PropertyBindings.Reset();
 
 	NumLinkedItems = 0;
+	ExternalItemBaseIndex = 0;
 }
 
 void UStateTree::PostEditChangeProperty(struct FPropertyChangedEvent& PropertyChangedEvent)
@@ -98,8 +99,8 @@ void UStateTree::Link()
 {
 	FStateTreeLinker Linker;
 
-	const int32 BaseIndex = PropertyBindings.GetSourceStructNum();
-	Linker.SetItemBaseIndex(BaseIndex);
+	ExternalItemBaseIndex = PropertyBindings.GetSourceStructNum();
+	Linker.SetItemBaseIndex(ExternalItemBaseIndex);
 	
 	for (FInstancedStruct& RuntimeItem : RuntimeStorageItems)
 	{
@@ -115,7 +116,7 @@ void UStateTree::Link()
 
 	ExternalItems = Linker.GetItemDescs();
 	
-	NumLinkedItems = BaseIndex + ExternalItems.Num();
+	NumLinkedItems = ExternalItemBaseIndex + ExternalItems.Num();
 }
 
 void UStateTree::InitRuntimeStorage()
diff --git a/Engine/Plugins/Runtime/StateTree/Source/StateTreeModule/Public/StateTree.h b/Engine/Plugins/Runtime/StateTree/Source/StateTreeModule/Public/StateTree.h
--- a/Engine/Plugins/Runtime/StateTree/Source/StateTreeModule/Public/StateTree.h
+++ b/Engine/Plugins/Runtime/StateTree/Source/StateTreeModule/Public/StateTree.h
@@ -34,6 +34,9 @@ public:
 	/** @return Number of linked items (Evaluators, Tasks, external items). */
 	int32 GetLinkedItemCount() const { return NumLinkedItems; }
 
+	/** @return Base index for external items. */
+	int32 GetExternalItemBaseIndex() const { return ExternalItemBaseIndex; }
+
 	/** @return List of external items required by the state tree */
 	TConstArrayView<FStateTreeExternalItemDesc> GetExternalItems() const { return ExternalItems; }
 
@@ -104,6 +107,9 @@ private:
 	UPROPERTY(Transient)
 	int32 NumLinkedItems = 0;
 
+	UPROPERTY(Transient)
+	int32 ExternalItemBaseIndex = 0;
+
 	UPROPERTY()
 	FStateTreePropertyBindings PropertyBindings;
 
diff --git a/Engine/Plugins/Runtime/StateTree/Source/StateTreeModule/Public/StateTreeExecutionContext.h b/Engine/Plugins/Runtime/StateTree/Source/StateTreeModule/Public/StateTreeExecutionContext.h
--- a/Engine/Plugins/Runtime/StateTree/Source/StateTreeModule/Public/StateTreeExecutionContext.h
+++ b/Engine/Plugins/Runtime/StateTree/Source/StateTreeModule/Public/StateTreeExecutionContext.h
@@ -141,7 +141,7 @@ public:
 	{
 		check(StateTree);
 		check(ItemHandle.IsValid());
-		checkSlow(StateTree->ExternalItems[ItemHandle.ItemIndex].Requirement != EStateTreeItemRequirement::Optional); // Optionals should query pointer instead.
+		checkSlow(StateTree->ExternalItems[ItemHandle.ItemIndex - StateTree->GetExternalItemBaseIndex()].Requirement != EStateTreeItemRequirement::Optional); // Optionals should query pointer instead.
 		return ItemViews[ItemHandle.ItemIndex].template GetMutable<typename T::ItemType>();
 	}
 
