StateTree: allow task/evaluator/condition to have separate data for isntance and item in the tree.

diff --git a/Engine/Plugins/AI/MassAI/Source/MassAIBehavior/Public/Tasks/MassZoneGraphFindEscapeTarget.h b/Engine/Plugins/AI/MassAI/Source/MassAIBehavior/Public/Tasks/MassZoneGraphFindEscapeTarget.h
--- a/Engine/Plugins/AI/MassAI/Source/MassAIBehavior/Public/Tasks/MassZoneGraphFindEscapeTarget.h
+++ b/Engine/Plugins/AI/MassAI/Source/MassAIBehavior/Public/Tasks/MassZoneGraphFindEscapeTarget.h
@@ -14,6 +14,15 @@ class UZoneGraphAnnotationSubsystem;
 /**
  * Updates TargetLocation to a escape target based on the agents current location on ZoneGraph, and disturbance annotation.
  */
+USTRUCT()
+struct MASSAIBEHAVIOR_API FMassZoneGraphFindEscapeTargetInstangeData
+{
+	GENERATED_BODY()
+
+	UPROPERTY(VisibleAnywhere, Category = Output)
+	FMassZoneGraphTargetLocation EscapeTargetLocation;
+};
+
 USTRUCT(meta = (DisplayName = "ZG Find Escape Target"))
 struct MASSAIBEHAVIOR_API FMassZoneGraphFindEscapeTarget : public FMassStateTreeTaskBase
 {
@@ -23,19 +32,15 @@ struct MASSAIBEHAVIOR_API FMassZoneGraphFindEscapeTarget : public FMassStateTree
 
 protected:
 	virtual bool Link(FStateTreeLinker& Linker) override;
-	virtual EStateTreeRunStatus EnterState(FStateTreeExecutionContext& Context, const EStateTreeStateChangeType ChangeType, const FStateTreeTransitionResult& Transition) override;
-	virtual void ExitState(FStateTreeExecutionContext& Context, const EStateTreeStateChangeType ChangeType, const FStateTreeTransitionResult& Transition) override;
-	virtual EStateTreeRunStatus Tick(FStateTreeExecutionContext& Context, const float DeltaTime) override;
-
-	TStateTreeItemHandle<FMassZoneGraphLaneLocationFragment> LocationHandle;
-	TStateTreeItemHandle<UZoneGraphSubsystem> ZoneGraphSubsystemHandle;
-	TStateTreeItemHandle<UZoneGraphAnnotationSubsystem> ZoneGraphAnnotationSubsystemHandle;
+	virtual const UStruct* GetInstanceDataType() const override { return FMassZoneGraphFindEscapeTargetInstangeData::StaticStruct(); }
+	virtual EStateTreeRunStatus EnterState(FStateTreeExecutionContext& Context, const EStateTreeStateChangeType ChangeType, const FStateTreeTransitionResult& Transition) const override;
 
-	UPROPERTY(EditAnywhere, Category = Parameters)
-	FZoneGraphTag DisturbanceAnnotationTag = FZoneGraphTag::None;
+	TStateTreeExternalDataHandle<FMassZoneGraphLaneLocationFragment> LocationHandle;
+	TStateTreeExternalDataHandle<UZoneGraphSubsystem> ZoneGraphSubsystemHandle;
+	TStateTreeExternalDataHandle<UZoneGraphAnnotationSubsystem> ZoneGraphAnnotationSubsystemHandle;
 
-	FMassZoneGraphTargetLocation TargetLocation;
+	TStateTreeInstanceDataPropertyHandle<FMassZoneGraphTargetLocation> EscapeTargetLocationHandle;
 
-	UPROPERTY(EditAnywhere, Category = Parameters, meta=(Struct="MassZoneGraphTargetLocation"))
-	FStateTreeResultRef EscapeTargetLocation;
+	UPROPERTY(EditAnywhere, Category = Parameter)
+	FZoneGraphTag DisturbanceAnnotationTag = FZoneGraphTag::None;
 };
diff --git a/Engine/Plugins/AI/MassAI/Source/MassAIBehavior/Public/Tasks/MassZoneGraphFindSmartObjectTarget.h b/Engine/Plugins/AI/MassAI/Source/MassAIBehavior/Public/Tasks/MassZoneGraphFindSmartObjectTarget.h
--- a/Engine/Plugins/AI/MassAI/Source/MassAIBehavior/Public/Tasks/MassZoneGraphFindSmartObjectTarget.h
+++ b/Engine/Plugins/AI/MassAI/Source/MassAIBehavior/Public/Tasks/MassZoneGraphFindSmartObjectTarget.h
@@ -12,6 +12,15 @@ class UZoneGraphAnnotationSubsystem;
 /**
 * Computes move target to a smart object based on current location on ZoneGraph.
 */
+USTRUCT()
+struct MASSAIBEHAVIOR_API FMassZoneGraphFindSmartObjectTargetInstanceData
+{
+	GENERATED_BODY()
+
+	UPROPERTY(EditAnywhere, Category = Output)
+	FMassZoneGraphTargetLocation SmartObjectLocation;
+};
+
 USTRUCT(meta = (DisplayName = "ZG Find Smart Object Target"))
 struct MASSAIBEHAVIOR_API FMassZoneGraphFindSmartObjectTarget : public FMassStateTreeTaskBase
 {
@@ -19,16 +28,12 @@ struct MASSAIBEHAVIOR_API FMassZoneGraphFindSmartObjectTarget : public FMassStat
 
 protected:
 	virtual bool Link(FStateTreeLinker& Linker) override;
-	virtual EStateTreeRunStatus EnterState(FStateTreeExecutionContext& Context, const EStateTreeStateChangeType ChangeType, const FStateTreeTransitionResult& Transition) override;
-	virtual void ExitState(FStateTreeExecutionContext& Context, const EStateTreeStateChangeType ChangeType, const FStateTreeTransitionResult& Transition) override;
-	virtual EStateTreeRunStatus Tick(FStateTreeExecutionContext& Context, const float DeltaTime) override;
-
-	TStateTreeItemHandle<FMassSmartObjectUserFragment> SmartObjectUserHandle;
-	TStateTreeItemHandle<FMassZoneGraphLaneLocationFragment> LocationHandle;
-	TStateTreeItemHandle<UZoneGraphAnnotationSubsystem> AnnotationSubsystemHandle;
+	virtual const UStruct* GetInstanceDataType() const override { return FMassZoneGraphFindSmartObjectTargetInstanceData::StaticStruct(); }
+	virtual EStateTreeRunStatus EnterState(FStateTreeExecutionContext& Context, const EStateTreeStateChangeType ChangeType, const FStateTreeTransitionResult& Transition) const override;
 
-	FMassZoneGraphTargetLocation TargetLocation;
+	TStateTreeExternalDataHandle<FMassSmartObjectUserFragment> SmartObjectUserHandle;
+	TStateTreeExternalDataHandle<FMassZoneGraphLaneLocationFragment> LocationHandle;
+	TStateTreeExternalDataHandle<UZoneGraphAnnotationSubsystem> AnnotationSubsystemHandle;
 
-	UPROPERTY(EditAnywhere, Category = Parameters, meta=(Struct="MassZoneGraphTargetLocation"))
-	FStateTreeResultRef TargetLocationRef;
+	TStateTreeInstanceDataPropertyHandle<FMassZoneGraphTargetLocation> SmartObjectLocationHandle;
 };
\ No newline at end of file
diff --git a/Engine/Plugins/AI/MassAI/Source/MassAIBehavior/Public/Tasks/MassZoneGraphPathFollowTask.h b/Engine/Plugins/AI/MassAI/Source/MassAIBehavior/Public/Tasks/MassZoneGraphPathFollowTask.h
--- a/Engine/Plugins/AI/MassAI/Source/MassAIBehavior/Public/Tasks/MassZoneGraphPathFollowTask.h
+++ b/Engine/Plugins/AI/MassAI/Source/MassAIBehavior/Public/Tasks/MassZoneGraphPathFollowTask.h
@@ -19,12 +19,10 @@ struct FDataFragment_AgentRadius;
 class UZoneGraphSubsystem;
 
 USTRUCT()
-struct MASSAIBEHAVIOR_API FMassZoneGraphTargetLocation : public FStateTreeResult
+struct MASSAIBEHAVIOR_API FMassZoneGraphTargetLocation
 {
 	GENERATED_BODY()
 
-	virtual const UScriptStruct& GetStruct() const override { return *StaticStruct(); }
-
 	void Reset()
 	{
 		LaneHandle.Reset();
@@ -68,6 +66,21 @@ struct MASSAIBEHAVIOR_API FMassZoneGraphTargetLocation : public FStateTreeResult
 /**
  * Follows a path long the current lane to a specified point.
  */
+USTRUCT()
+struct MASSAIBEHAVIOR_API FMassZoneGraphPathFollowTaskInstanceData
+{
+	GENERATED_BODY()
+
+	UPROPERTY(EditAnywhere, Category = Input)
+	FMassZoneGraphTargetLocation TargetLocation;
+
+	UPROPERTY(EditAnywhere, Category = Parameter)
+	FMassMovementStyleRef MovementStyle;
+
+	UPROPERTY(EditAnywhere, Category = Parameter)
+	float SpeedScale = 1.0f;
+};
+
 USTRUCT(meta = (DisplayName = "ZG Path Follow"))
 struct MASSAIBEHAVIOR_API FMassZoneGraphPathFollowTask : public FMassStateTreeTaskBase
 {
@@ -75,26 +88,22 @@ struct MASSAIBEHAVIOR_API FMassZoneGraphPathFollowTask : public FMassStateTreeTa
 
 protected:
 	virtual bool Link(FStateTreeLinker& Linker) override;
-	virtual EStateTreeRunStatus EnterState(FStateTreeExecutionContext& Context, const EStateTreeStateChangeType ChangeType, const FStateTreeTransitionResult& Transition) override;
-	virtual EStateTreeRunStatus Tick(FStateTreeExecutionContext& Context, const float DeltaTime) override;
+	virtual const UStruct* GetInstanceDataType() const override { return FMassZoneGraphPathFollowTaskInstanceData::StaticStruct(); };
+	virtual EStateTreeRunStatus EnterState(FStateTreeExecutionContext& Context, const EStateTreeStateChangeType ChangeType, const FStateTreeTransitionResult& Transition) const override;
+	virtual EStateTreeRunStatus Tick(FStateTreeExecutionContext& Context, const float DeltaTime) const override;
 
 	bool RequestPath(FMassStateTreeExecutionContext& Context, const FMassZoneGraphTargetLocation& TargetLocation) const;
 
-	TStateTreeItemHandle<FMassZoneGraphLaneLocationFragment> LocationHandle;
-	TStateTreeItemHandle<FMassMoveTargetFragment> MoveTargetHandle;
-	TStateTreeItemHandle<FMassMovementConfigFragment> MovementConfigHandle;
-	TStateTreeItemHandle<FMassZoneGraphPathRequestFragment> PathRequestHandle;
-	TStateTreeItemHandle<FMassZoneGraphShortPathFragment> ShortPathHandle;
-	TStateTreeItemHandle<FMassZoneGraphCachedLaneFragment> CachedLaneHandle;
-	TStateTreeItemHandle<FDataFragment_AgentRadius> AgentRadiusHandle;
-	TStateTreeItemHandle<UZoneGraphSubsystem> ZoneGraphSubsystemHandle;
-
-	UPROPERTY(EditAnywhere, Category = Parameters, meta=(Struct="MassZoneGraphTargetLocation", Bindable))
-	FStateTreeResultRef TargetLocation;
-
-	UPROPERTY(EditAnywhere, Category = Parameters)
-	FMassMovementStyleRef MovementStyle;
-
-	UPROPERTY(EditAnywhere, Category = Parameters)
-	float SpeedScale = 1.0f;
+	TStateTreeExternalDataHandle<FMassZoneGraphLaneLocationFragment> LocationHandle;
+	TStateTreeExternalDataHandle<FMassMoveTargetFragment> MoveTargetHandle;
+	TStateTreeExternalDataHandle<FMassMovementConfigFragment> MovementConfigHandle;
+	TStateTreeExternalDataHandle<FMassZoneGraphPathRequestFragment> PathRequestHandle;
+	TStateTreeExternalDataHandle<FMassZoneGraphShortPathFragment> ShortPathHandle;
+	TStateTreeExternalDataHandle<FMassZoneGraphCachedLaneFragment> CachedLaneHandle;
+	TStateTreeExternalDataHandle<FDataFragment_AgentRadius> AgentRadiusHandle;
+	TStateTreeExternalDataHandle<UZoneGraphSubsystem> ZoneGraphSubsystemHandle;
+
+	TStateTreeInstanceDataPropertyHandle<FMassZoneGraphTargetLocation> TargetLocationHandle;
+	TStateTreeInstanceDataPropertyHandle<FMassMovementStyleRef> MovementStyleHandle;
+	TStateTreeInstanceDataPropertyHandle<float> SpeedScaleHandle;
 };
diff --git a/Engine/Plugins/AI/MassAI/Source/MassAIBehavior/Public/Tasks/MassZoneGraphStandTask.h b/Engine/Plugins/AI/MassAI/Source/MassAIBehavior/Public/Tasks/MassZoneGraphStandTask.h
--- a/Engine/Plugins/AI/MassAI/Source/MassAIBehavior/Public/Tasks/MassZoneGraphStandTask.h
+++ b/Engine/Plugins/AI/MassAI/Source/MassAIBehavior/Public/Tasks/MassZoneGraphStandTask.h
@@ -18,28 +18,39 @@ struct FMassMovementConfigFragment;
 /**
  * Stop, and stand on current ZoneGraph location
  */
-USTRUCT(meta = (DisplayName = "ZG Stand"))
-struct MASSAIBEHAVIOR_API FMassZoneGraphStandTask : public FMassStateTreeTaskBase
+
+USTRUCT()
+struct MASSAIBEHAVIOR_API FMassZoneGraphStandTaskInstanceData
 {
 	GENERATED_BODY()
 
-protected:
-	virtual bool Link(FStateTreeLinker& Linker) override;
-	virtual EStateTreeRunStatus EnterState(FStateTreeExecutionContext& Context, const EStateTreeStateChangeType ChangeType, const FStateTreeTransitionResult& Transition) override;
-	virtual EStateTreeRunStatus Tick(FStateTreeExecutionContext& Context, const float DeltaTime) override;
-
-	TStateTreeItemHandle<FMassZoneGraphLaneLocationFragment> LocationHandle;
-	TStateTreeItemHandle<FMassMoveTargetFragment> MoveTargetHandle;
-	TStateTreeItemHandle<FMassZoneGraphShortPathFragment> ShortPathHandle;
-	TStateTreeItemHandle<FMassZoneGraphCachedLaneFragment> CachedLaneHandle;
-	TStateTreeItemHandle<UZoneGraphSubsystem> ZoneGraphSubsystemHandle;
-	TStateTreeItemHandle<UMassSignalSubsystem> MassSignalSubsystemHandle;
-	TStateTreeItemHandle<FMassMovementConfigFragment> MovementConfigHandle;
+	/** Delay before the task ends. Default (0 or any negative) will run indefinitely so it requires a transition in the state tree to stop it. */
+	UPROPERTY(EditAnywhere, Category = Parameter)
+	float Duration = 0.0f;
 
 	UPROPERTY()
 	float Time = 0.0f;
+};
 
-	/** Delay before the task ends. Default (0 or any negative) will run indefinitely so it requires a transition in the state tree to stop it. */
-	UPROPERTY(EditAnywhere, Category = Parameters, meta = (Bindable))
-	float Duration = 0.0f;
+USTRUCT(meta = (DisplayName = "ZG Stand"))
+struct MASSAIBEHAVIOR_API FMassZoneGraphStandTask : public FMassStateTreeTaskBase
+{
+	GENERATED_BODY()
+
+protected:
+	virtual bool Link(FStateTreeLinker& Linker) override;
+	virtual const UStruct* GetInstanceDataType() const override { return FMassZoneGraphStandTaskInstanceData::StaticStruct(); }
+	virtual EStateTreeRunStatus EnterState(FStateTreeExecutionContext& Context, const EStateTreeStateChangeType ChangeType, const FStateTreeTransitionResult& Transition) const override;
+	virtual EStateTreeRunStatus Tick(FStateTreeExecutionContext& Context, const float DeltaTime) const override;
+
+	TStateTreeExternalDataHandle<FMassZoneGraphLaneLocationFragment> LocationHandle;
+	TStateTreeExternalDataHandle<FMassMoveTargetFragment> MoveTargetHandle;
+	TStateTreeExternalDataHandle<FMassZoneGraphShortPathFragment> ShortPathHandle;
+	TStateTreeExternalDataHandle<FMassZoneGraphCachedLaneFragment> CachedLaneHandle;
+	TStateTreeExternalDataHandle<UZoneGraphSubsystem> ZoneGraphSubsystemHandle;
+	TStateTreeExternalDataHandle<UMassSignalSubsystem> MassSignalSubsystemHandle;
+	TStateTreeExternalDataHandle<FMassMovementConfigFragment> MovementConfigHandle;
+
+	TStateTreeInstanceDataPropertyHandle<float> DurationHandle;
+	TStateTreeInstanceDataPropertyHandle<float> TimeHandle;
 };
diff --git a/Engine/Plugins/AI/MassAI/Source/MassAIDebug/Private/GameplayDebuggerCategory_Mass.cpp b/Engine/Plugins/AI/MassAI/Source/MassAIDebug/Private/GameplayDebuggerCategory_Mass.cpp
--- a/Engine/Plugins/AI/MassAI/Source/MassAIDebug/Private/GameplayDebuggerCategory_Mass.cpp
+++ b/Engine/Plugins/AI/MassAI/Source/MassAIDebug/Private/GameplayDebuggerCategory_Mass.cpp
@@ -491,7 +491,7 @@ void FGameplayDebuggerCategory_Mass::CollectData(APlayerController* OwnerPC, AAc
 							StateTreeContext.Init(*OwnerPC, *StateTree, EStateTreeStorage::External);
 							StateTreeContext.SetEntity(Entity);
 							
-							FStructView Storage = EntitySystem->GetFragmentDataStruct(Entity, StateTree->GetRuntimeStorageStruct());
+							FStructView Storage = EntitySystem->GetFragmentDataStruct(Entity, StateTree->GetInstanceStorageStruct());
 							
 							Status += StateTreeContext.GetActiveStateName(Storage);
 							Status += TEXT("\n");
diff --git a/Engine/Plugins/AI/MassAI/Source/MassAIDebug/Private/MassDebugStateTreeProcessor.cpp b/Engine/Plugins/AI/MassAI/Source/MassAIDebug/Private/MassDebugStateTreeProcessor.cpp
--- a/Engine/Plugins/AI/MassAI/Source/MassAIDebug/Private/MassDebugStateTreeProcessor.cpp
+++ b/Engine/Plugins/AI/MassAI/Source/MassAIDebug/Private/MassDebugStateTreeProcessor.cpp
@@ -74,7 +74,7 @@ void UMassDebugStateTreeProcessor::Execute(UMassEntitySubsystem& EntitySubsystem
 					StateTreeContext.SetEntity(Entity);
 					
 #if WITH_GAMEPLAY_DEBUGGER
-					const FStructView Storage = EntitySubsystem.GetFragmentDataStruct(Entity, StateTree->GetRuntimeStorageStruct());
+					const FStructView Storage = EntitySubsystem.GetFragmentDataStruct(Entity, StateTree->GetInstanceStorageStruct());
 					Debugger->AppendSelectedEntityInfo(StateTreeContext.GetDebugInfoString(Storage));
 #endif // WITH_GAMEPLAY_DEBUGGER
 				}
@@ -92,7 +92,7 @@ void UMassDebugStateTreeProcessor::Execute(UMassEntitySubsystem& EntitySubsystem
 					StateTreeContext.Init(*this, *StateTree, EStateTreeStorage::External);
 					StateTreeContext.SetEntity(Entity);
 
-					const FStructView Storage = EntitySubsystem.GetFragmentDataStruct(Entity, StateTree->GetRuntimeStorageStruct());
+					const FStructView Storage = EntitySubsystem.GetFragmentDataStruct(Entity, StateTree->GetInstanceStorageStruct());
 
 					// State
 					UE_VLOG_SEGMENT_THICK(this, LogStateTree, Log, Position, Position + FVector(0,0,50), EntityColor, /*Thickness*/ 2, TEXT("%s %s"),
diff --git a/Engine/Plugins/AI/MassCrowd/Source/MassCrowd/Private/Tasks/MassCrowdClaimWaitSlotTask.cpp b/Engine/Plugins/AI/MassCrowd/Source/MassCrowd/Private/Tasks/MassCrowdClaimWaitSlotTask.cpp
--- a/Engine/Plugins/AI/MassCrowd/Source/MassCrowd/Private/Tasks/MassCrowdClaimWaitSlotTask.cpp
+++ b/Engine/Plugins/AI/MassCrowd/Source/MassCrowd/Private/Tasks/MassCrowdClaimWaitSlotTask.cpp
@@ -9,28 +9,34 @@
 
 bool FMassCrowdClaimWaitSlotTask::Link(FStateTreeLinker& Linker)
 {
-	Linker.LinkExternalItem(LocationHandle);
-	Linker.LinkExternalItem(MoveTargetHandle);
-	Linker.LinkExternalItem(CrowdSubsystemHandle);
+	Linker.LinkExternalData(LocationHandle);
+	Linker.LinkExternalData(MoveTargetHandle);
+	Linker.LinkExternalData(CrowdSubsystemHandle);
+
+	Linker.LinkInstanceDataProperty(WaitSlotLocationHandle, STATETREE_INSTANCEDATA_PROPERTY(FMassCrowdClaimWaitSlotTaskInstanceData, WaitSlotLocation));
+	Linker.LinkInstanceDataProperty(WaitingSlotIndexHandle, STATETREE_INSTANCEDATA_PROPERTY(FMassCrowdClaimWaitSlotTaskInstanceData, WaitingSlotIndex));
+	Linker.LinkInstanceDataProperty(AcquiredLaneHandle, STATETREE_INSTANCEDATA_PROPERTY(FMassCrowdClaimWaitSlotTaskInstanceData, AcquiredLane));
 
 	return true;
 }
 
-EStateTreeRunStatus FMassCrowdClaimWaitSlotTask::EnterState(FStateTreeExecutionContext& Context, const EStateTreeStateChangeType ChangeType, const FStateTreeTransitionResult& Transition)
+EStateTreeRunStatus FMassCrowdClaimWaitSlotTask::EnterState(FStateTreeExecutionContext& Context, const EStateTreeStateChangeType ChangeType, const FStateTreeTransitionResult& Transition) const
 {
 	if (ChangeType != EStateTreeStateChangeType::Changed)
 	{
 		return EStateTreeRunStatus::Running;
 	}
 
-	WaitSlotLocation = nullptr;
+	FMassZoneGraphTargetLocation& WaitSlotLocation = Context.GetInstanceData(WaitSlotLocationHandle);
+	int32& WaitingSlotIndex = Context.GetInstanceData(WaitingSlotIndexHandle);
+	FZoneGraphLaneHandle& AcquiredLane = Context.GetInstanceData(AcquiredLaneHandle);
 
 	const FMassStateTreeExecutionContext& MassContext = static_cast<FMassStateTreeExecutionContext&>(Context);
 	const FMassEntityHandle Entity = MassContext.GetEntity();
 	
-	const FMassZoneGraphLaneLocationFragment& LaneLocation = Context.GetExternalItem(LocationHandle);
-	const FMassMoveTargetFragment& MoveTarget = Context.GetExternalItem(MoveTargetHandle);
-	UMassCrowdSubsystem& CrowdSubsystem = Context.GetExternalItem(CrowdSubsystemHandle);
+	const FMassZoneGraphLaneLocationFragment& LaneLocation = Context.GetExternalData(LocationHandle);
+	const FMassMoveTargetFragment& MoveTarget = Context.GetExternalData(MoveTargetHandle);
+	UMassCrowdSubsystem& CrowdSubsystem = Context.GetExternalData(CrowdSubsystemHandle);
 
 	FVector SlotPosition = FVector::ZeroVector;
 	FVector SlotDirection = FVector::ForwardVector;
@@ -41,28 +47,24 @@ EStateTreeRunStatus FMassCrowdClaimWaitSlotTask::EnterState(FStateTreeExecutionC
 		return EStateTreeRunStatus::Failed;
 	}
 	
-	AcquiredLaneHandle = LaneLocation.LaneHandle;
+	AcquiredLane = LaneLocation.LaneHandle;
 
-	TargetLocation.LaneHandle = LaneLocation.LaneHandle;
-	TargetLocation.NextExitLinkType = EZoneLaneLinkType::None;
-	TargetLocation.NextLaneHandle.Reset();
-	TargetLocation.bMoveReverse = false;
-	TargetLocation.EndOfPathIntent = EMassMovementAction::Stand;
-	TargetLocation.EndOfPathPosition = SlotPosition;
-	TargetLocation.EndOfPathDirection = SlotDirection;
-	TargetLocation.TargetDistance = LaneLocation.LaneLength; // Go to end of lane
+	WaitSlotLocation.LaneHandle = LaneLocation.LaneHandle;
+	WaitSlotLocation.NextExitLinkType = EZoneLaneLinkType::None;
+	WaitSlotLocation.NextLaneHandle.Reset();
+	WaitSlotLocation.bMoveReverse = false;
+	WaitSlotLocation.EndOfPathIntent = EMassMovementAction::Stand;
+	WaitSlotLocation.EndOfPathPosition = SlotPosition;
+	WaitSlotLocation.EndOfPathDirection = SlotDirection;
+	WaitSlotLocation.TargetDistance = LaneLocation.LaneLength; // Go to end of lane
 	// Let's start moving toward the interaction a bit before the entry point.
-	TargetLocation.AnticipationDistance.Set(100.f);
-
-	WaitSlotLocation = &TargetLocation;
+	WaitSlotLocation.AnticipationDistance.Set(100.f);
 	
 	return EStateTreeRunStatus::Running;
 }
 
-void FMassCrowdClaimWaitSlotTask::ExitState(FStateTreeExecutionContext& Context, const EStateTreeStateChangeType ChangeType, const FStateTreeTransitionResult& Transition)
+void FMassCrowdClaimWaitSlotTask::ExitState(FStateTreeExecutionContext& Context, const EStateTreeStateChangeType ChangeType, const FStateTreeTransitionResult& Transition) const
 {
-	WaitSlotLocation = &TargetLocation;
-	
 	if (ChangeType != EStateTreeStateChangeType::Changed)
 	{
 		return;
@@ -70,20 +72,19 @@ void FMassCrowdClaimWaitSlotTask::ExitState(FStateTreeExecutionContext& Context,
 
 	const FMassStateTreeExecutionContext& MassContext = static_cast<FMassStateTreeExecutionContext&>(Context);
 	const FMassEntityHandle Entity = MassContext.GetEntity();
-	
-	UMassCrowdSubsystem& CrowdSubsystem = Context.GetExternalItem(CrowdSubsystemHandle);
 
+	UMassCrowdSubsystem& CrowdSubsystem = Context.GetExternalData(CrowdSubsystemHandle);
+
+	FMassZoneGraphTargetLocation& WaitSlotLocation = Context.GetInstanceData(WaitSlotLocationHandle);
+	int32& WaitingSlotIndex = Context.GetInstanceData(WaitingSlotIndexHandle);
+	FZoneGraphLaneHandle& AcquiredLane = Context.GetInstanceData(AcquiredLaneHandle);
+	
 	if (WaitingSlotIndex != INDEX_NONE)
 	{
-		CrowdSubsystem.ReleaseWaitingSlot(Entity, AcquiredLaneHandle, WaitingSlotIndex);
+		CrowdSubsystem.ReleaseWaitingSlot(Entity, AcquiredLane, WaitingSlotIndex);
 	}
 	
-	TargetLocation.Reset();
-}
-
-EStateTreeRunStatus FMassCrowdClaimWaitSlotTask::Tick(FStateTreeExecutionContext& Context, const float DeltaTime)
-{
-	WaitSlotLocation = &TargetLocation;
-
-	return EStateTreeRunStatus::Running;
+	WaitingSlotIndex = INDEX_NONE;
+	AcquiredLane.Reset();
+	WaitSlotLocation.Reset();
 }
diff --git a/Engine/Plugins/AI/MassCrowd/Source/MassCrowd/Private/Tasks/MassZoneGraphFindWanderTarget.cpp b/Engine/Plugins/AI/MassCrowd/Source/MassCrowd/Private/Tasks/MassZoneGraphFindWanderTarget.cpp
--- a/Engine/Plugins/AI/MassCrowd/Source/MassCrowd/Private/Tasks/MassZoneGraphFindWanderTarget.cpp
+++ b/Engine/Plugins/AI/MassCrowd/Source/MassCrowd/Private/Tasks/MassZoneGraphFindWanderTarget.cpp
@@ -19,22 +19,24 @@ FMassZoneGraphFindWanderTarget::FMassZoneGraphFindWanderTarget()
 
 bool FMassZoneGraphFindWanderTarget::Link(FStateTreeLinker& Linker)
 {
-	Linker.LinkExternalItem(LocationHandle);
-	Linker.LinkExternalItem(ZoneGraphSubsystemHandle);
-	Linker.LinkExternalItem(ZoneGraphAnnotationSubsystemHandle);
-	Linker.LinkExternalItem(MassCrowdSubsystemHandle);
+	Linker.LinkExternalData(LocationHandle);
+	Linker.LinkExternalData(ZoneGraphSubsystemHandle);
+	Linker.LinkExternalData(ZoneGraphAnnotationSubsystemHandle);
+	Linker.LinkExternalData(MassCrowdSubsystemHandle);
+
+	Linker.LinkInstanceDataProperty(WanderTargetLocationHandle, STATETREE_INSTANCEDATA_PROPERTY(FMassZoneGraphFindWanderTargetInstanceData, WanderTargetLocation));
 
 	return true;
 }
 
-EStateTreeRunStatus FMassZoneGraphFindWanderTarget::EnterState(FStateTreeExecutionContext& Context, const EStateTreeStateChangeType ChangeType, const FStateTreeTransitionResult& Transition)
+EStateTreeRunStatus FMassZoneGraphFindWanderTarget::EnterState(FStateTreeExecutionContext& Context, const EStateTreeStateChangeType ChangeType, const FStateTreeTransitionResult& Transition) const
 {
 	const FMassStateTreeExecutionContext& MassContext = static_cast<FMassStateTreeExecutionContext&>(Context);
 
-	const FMassZoneGraphLaneLocationFragment& LaneLocation = Context.GetExternalItem(LocationHandle);
-	const UZoneGraphSubsystem& ZoneGraphSubsystem = Context.GetExternalItem(ZoneGraphSubsystemHandle);
-	UZoneGraphAnnotationSubsystem& ZoneGraphAnnotationSubsystem = Context.GetExternalItem(ZoneGraphAnnotationSubsystemHandle);
-	const UMassCrowdSubsystem& MassCrowdSubsystem = Context.GetExternalItem(MassCrowdSubsystemHandle);
+	const FMassZoneGraphLaneLocationFragment& LaneLocation = Context.GetExternalData(LocationHandle);
+	const UZoneGraphSubsystem& ZoneGraphSubsystem = Context.GetExternalData(ZoneGraphSubsystemHandle);
+	UZoneGraphAnnotationSubsystem& ZoneGraphAnnotationSubsystem = Context.GetExternalData(ZoneGraphAnnotationSubsystemHandle);
+	const UMassCrowdSubsystem& MassCrowdSubsystem = Context.GetExternalData(MassCrowdSubsystemHandle);
 
 	bool bDisplayDebug = false;
 #if WITH_MASSGAMEPLAY_DEBUG
@@ -56,12 +58,14 @@ EStateTreeRunStatus FMassZoneGraphFindWanderTarget::EnterState(FStateTreeExecuti
 
 	const float MoveDistance = GetDefault<UMassCrowdSettings>()->GetMoveDistance();
 
-	TargetLocation.LaneHandle = LaneLocation.LaneHandle;
-	TargetLocation.TargetDistance = LaneLocation.DistanceAlongLane + MoveDistance;
-	TargetLocation.NextExitLinkType = EZoneLaneLinkType::None;
-	TargetLocation.NextLaneHandle.Reset();
-	TargetLocation.bMoveReverse = false;
-	TargetLocation.EndOfPathIntent = EMassMovementAction::Move;
+	FMassZoneGraphTargetLocation& WanderTargetLocation = Context.GetInstanceData(WanderTargetLocationHandle);
+
+	WanderTargetLocation.LaneHandle = LaneLocation.LaneHandle;
+	WanderTargetLocation.TargetDistance = LaneLocation.DistanceAlongLane + MoveDistance;
+	WanderTargetLocation.NextExitLinkType = EZoneLaneLinkType::None;
+	WanderTargetLocation.NextLaneHandle.Reset();
+	WanderTargetLocation.bMoveReverse = false;
+	WanderTargetLocation.EndOfPathIntent = EMassMovementAction::Move;
 
 	if (bDisplayDebug)
 	{
@@ -71,9 +75,9 @@ EStateTreeRunStatus FMassZoneGraphFindWanderTarget::EnterState(FStateTreeExecuti
 	EStateTreeRunStatus Status = EStateTreeRunStatus::Running;
 	
 	// When close to end of a lane, choose next lane.
-	if (TargetLocation.TargetDistance > LaneLocation.LaneLength)
+	if (WanderTargetLocation.TargetDistance > LaneLocation.LaneLength)
 	{
-		TargetLocation.TargetDistance = FMath::Min(TargetLocation.TargetDistance, LaneLocation.LaneLength);
+		WanderTargetLocation.TargetDistance = FMath::Min(WanderTargetLocation.TargetDistance, LaneLocation.LaneLength);
 
 		typedef TTuple<const FZoneGraphLinkedLane, const float> FBranchingCandidate;
 		TArray<FBranchingCandidate, TInlineAllocator<8>> Candidates;
@@ -88,7 +92,7 @@ EStateTreeRunStatus FMassZoneGraphFindWanderTarget::EnterState(FStateTreeExecuti
 			{
 				// Apply tag filter
 				const FZoneGraphTagMask BehaviorTags = ZoneGraphAnnotationSubsystem.GetAnnotationTags(LinkedLane.DestLane);
-				if (!AllowedBehaviorTags.Pass(BehaviorTags))
+				if (!AllowedAnnotationTags.Pass(BehaviorTags))
 				{
 					continue;
 				}
@@ -105,28 +109,28 @@ EStateTreeRunStatus FMassZoneGraphFindWanderTarget::EnterState(FStateTreeExecuti
 
 		if (FindCandidates(EZoneLaneLinkType::Outgoing))
 		{
-			TargetLocation.NextExitLinkType = EZoneLaneLinkType::Outgoing;
+			WanderTargetLocation.NextExitLinkType = EZoneLaneLinkType::Outgoing;
 		}
 		else
 		{
 			// Could not continue, try to switch to an adjacent lane.
 			// @todo: we could try to do something smarter here so that agents do not clump up. May need to have some heuristic,
 			//		  i.e. at intersections it looks better to switch lane immediately, with flee, it looks better to vary the location randomly.
-			TargetLocation.TargetDistance = LaneLocation.DistanceAlongLane;
+			WanderTargetLocation.TargetDistance = LaneLocation.DistanceAlongLane;
 
 			// Try adjacent lanes
 			if (FindCandidates(EZoneLaneLinkType::Adjacent))
 			{
 				// Found adjacent lane, choose it once followed the short path. Keeping the random offset from above,
 				// so that all agents dont follow until the end of the path to turn.
-				TargetLocation.NextExitLinkType = EZoneLaneLinkType::Adjacent;
+				WanderTargetLocation.NextExitLinkType = EZoneLaneLinkType::Adjacent;
 			}
 		}
 
 		if (Candidates.IsEmpty())
 		{
 			// Could not find next lane, fail.
-			TargetLocation.Reset();
+			WanderTargetLocation.Reset();
 			Status = EStateTreeRunStatus::Failed;
 		}
 		else
@@ -139,28 +143,12 @@ EStateTreeRunStatus FMassZoneGraphFindWanderTarget::EnterState(FStateTreeExecuti
 				if (Rand < CandidateCombinedWeight)
 				{
 					const FZoneGraphLinkedLane& LinkedLane = Candidate.Get<0>();
-					TargetLocation.NextLaneHandle = LinkedLane.DestLane;
+					WanderTargetLocation.NextLaneHandle = LinkedLane.DestLane;
 					break;
 				}
 			}
 		}
 	}
 
-	WanderTargetLocation = &TargetLocation;
-	
 	return Status;
 }
-
-void FMassZoneGraphFindWanderTarget::ExitState(FStateTreeExecutionContext& Context, const EStateTreeStateChangeType ChangeType, const FStateTreeTransitionResult& Transition)
-{
-	// Temp solution to make sure the target location is always up to date (will be replaced with automatic update).
-	WanderTargetLocation = &TargetLocation;
-}
-
-EStateTreeRunStatus FMassZoneGraphFindWanderTarget::Tick(FStateTreeExecutionContext& Context, const float DeltaTime)
-{
-	// Temp solution to make sure the target location is always up to date (will be replaced with automatic update).
-	WanderTargetLocation = &TargetLocation;
-	
-	return EStateTreeRunStatus::Running;
-}
diff --git a/Engine/Plugins/AI/MassCrowd/Source/MassCrowd/Public/Tasks/MassCrowdClaimWaitSlotTask.h b/Engine/Plugins/AI/MassCrowd/Source/MassCrowd/Public/Tasks/MassCrowdClaimWaitSlotTask.h
--- a/Engine/Plugins/AI/MassCrowd/Source/MassCrowd/Public/Tasks/MassCrowdClaimWaitSlotTask.h
+++ b/Engine/Plugins/AI/MassCrowd/Source/MassCrowd/Public/Tasks/MassCrowdClaimWaitSlotTask.h
@@ -13,6 +13,21 @@ class UMassCrowdSubsystem;
 /**
 * Claim wait slot and expose slot position for path follow.
 */
+USTRUCT()
+struct MASSCROWD_API FMassCrowdClaimWaitSlotTaskInstanceData
+{
+	GENERATED_BODY()
+
+	UPROPERTY(EditAnywhere, Category = Output)
+	FMassZoneGraphTargetLocation WaitSlotLocation;
+
+	UPROPERTY()
+	int32 WaitingSlotIndex = INDEX_NONE;
+	
+	UPROPERTY()
+	FZoneGraphLaneHandle AcquiredLane;
+};
+
 USTRUCT(meta = (DisplayName = "Crowd Claim Wait Slot"))
 struct MASSCROWD_API FMassCrowdClaimWaitSlotTask : public FMassStateTreeTaskBase
 {
@@ -20,20 +35,15 @@ struct MASSCROWD_API FMassCrowdClaimWaitSlotTask : public FMassStateTreeTaskBase
 
 protected:
 	virtual bool Link(FStateTreeLinker& Linker) override;
-	virtual EStateTreeRunStatus EnterState(FStateTreeExecutionContext& Context, const EStateTreeStateChangeType ChangeType, const FStateTreeTransitionResult& Transition) override;
-	virtual void ExitState(FStateTreeExecutionContext& Context, const EStateTreeStateChangeType ChangeType, const FStateTreeTransitionResult& Transition) override;
-	virtual EStateTreeRunStatus Tick(FStateTreeExecutionContext& Context, const float DeltaTime) override;
-
-	TStateTreeItemHandle<FMassZoneGraphLaneLocationFragment> LocationHandle;
-	TStateTreeItemHandle<FMassMoveTargetFragment> MoveTargetHandle;
-	TStateTreeItemHandle<UMassCrowdSubsystem> CrowdSubsystemHandle;
+	virtual const UStruct* GetInstanceDataType() const override { return FMassCrowdClaimWaitSlotTaskInstanceData::StaticStruct(); }
+	virtual EStateTreeRunStatus EnterState(FStateTreeExecutionContext& Context, const EStateTreeStateChangeType ChangeType, const FStateTreeTransitionResult& Transition) const override;
+	virtual void ExitState(FStateTreeExecutionContext& Context, const EStateTreeStateChangeType ChangeType, const FStateTreeTransitionResult& Transition) const override;
 
-	UPROPERTY(EditAnywhere, Category = Parameters, meta=(Struct="MassZoneGraphTargetLocation"))
-	FStateTreeResultRef WaitSlotLocation;
+	TStateTreeExternalDataHandle<FMassZoneGraphLaneLocationFragment> LocationHandle;
+	TStateTreeExternalDataHandle<FMassMoveTargetFragment> MoveTargetHandle;
+	TStateTreeExternalDataHandle<UMassCrowdSubsystem> CrowdSubsystemHandle;
 
-	FMassZoneGraphTargetLocation TargetLocation;
-
-	int32 WaitingSlotIndex = INDEX_NONE;
-	
-	FZoneGraphLaneHandle AcquiredLaneHandle;
+	TStateTreeInstanceDataPropertyHandle<FMassZoneGraphTargetLocation> WaitSlotLocationHandle;
+	TStateTreeInstanceDataPropertyHandle<int32> WaitingSlotIndexHandle;
+	TStateTreeInstanceDataPropertyHandle<FZoneGraphLaneHandle> AcquiredLaneHandle;
 };
diff --git a/Engine/Plugins/AI/MassCrowd/Source/MassCrowd/Public/Tasks/MassZoneGraphFindWanderTarget.h b/Engine/Plugins/AI/MassCrowd/Source/MassCrowd/Public/Tasks/MassZoneGraphFindWanderTarget.h
--- a/Engine/Plugins/AI/MassCrowd/Source/MassCrowd/Public/Tasks/MassZoneGraphFindWanderTarget.h
+++ b/Engine/Plugins/AI/MassCrowd/Source/MassCrowd/Public/Tasks/MassZoneGraphFindWanderTarget.h
@@ -15,6 +15,15 @@ class UMassCrowdSubsystem;
 /**
  * Updates TargetLocation to a wander target based on the agents current location on ZoneGraph.
  */
+USTRUCT()
+struct MASSCROWD_API FMassZoneGraphFindWanderTargetInstanceData
+{
+	GENERATED_BODY()
+
+	UPROPERTY(EditAnywhere, Category = Output)
+	FMassZoneGraphTargetLocation WanderTargetLocation;
+};
+
 USTRUCT(meta = (DisplayName = "ZG Find Wander Target"))
 struct MASSCROWD_API FMassZoneGraphFindWanderTarget : public FMassStateTreeTaskBase
 {
@@ -24,20 +33,16 @@ struct MASSCROWD_API FMassZoneGraphFindWanderTarget : public FMassStateTreeTaskB
 
 protected:
 	virtual bool Link(FStateTreeLinker& Linker) override;
-	virtual EStateTreeRunStatus EnterState(FStateTreeExecutionContext& Context, const EStateTreeStateChangeType ChangeType, const FStateTreeTransitionResult& Transition) override;
-	virtual void ExitState(FStateTreeExecutionContext& Context, const EStateTreeStateChangeType ChangeType, const FStateTreeTransitionResult& Transition) override;
-	virtual EStateTreeRunStatus Tick(FStateTreeExecutionContext& Context, const float DeltaTime) override;
-
-	TStateTreeItemHandle<FMassZoneGraphLaneLocationFragment> LocationHandle;
-	TStateTreeItemHandle<UZoneGraphSubsystem> ZoneGraphSubsystemHandle;
-	TStateTreeItemHandle<UZoneGraphAnnotationSubsystem> ZoneGraphAnnotationSubsystemHandle;
-	TStateTreeItemHandle<UMassCrowdSubsystem> MassCrowdSubsystemHandle;
+	virtual const UStruct* GetInstanceDataType() const override { return FMassZoneGraphFindWanderTargetInstanceData::StaticStruct(); }
+	virtual EStateTreeRunStatus EnterState(FStateTreeExecutionContext& Context, const EStateTreeStateChangeType ChangeType, const FStateTreeTransitionResult& Transition) const override;
 
-	UPROPERTY(EditAnywhere, Category = Parameters)
-	FZoneGraphTagFilter AllowedBehaviorTags;
+	TStateTreeExternalDataHandle<FMassZoneGraphLaneLocationFragment> LocationHandle;
+	TStateTreeExternalDataHandle<UZoneGraphSubsystem> ZoneGraphSubsystemHandle;
+	TStateTreeExternalDataHandle<UZoneGraphAnnotationSubsystem> ZoneGraphAnnotationSubsystemHandle;
+	TStateTreeExternalDataHandle<UMassCrowdSubsystem> MassCrowdSubsystemHandle;
 
-	FMassZoneGraphTargetLocation TargetLocation;
+	TStateTreeInstanceDataPropertyHandle<FMassZoneGraphTargetLocation> WanderTargetLocationHandle;
 
-	UPROPERTY(EditAnywhere, Category = Parameters, meta=(Struct="MassZoneGraphTargetLocation"))
-	FStateTreeResultRef WanderTargetLocation;
+	UPROPERTY(EditAnywhere, Category = Parameter)
+	FZoneGraphTagFilter AllowedAnnotationTags;
 };
diff --git a/Engine/Plugins/Runtime/StateTree/Source/StateTreeEditorModule/Private/Customizations/StateTreeBindableItemDetails.cpp b/Engine/Plugins/Runtime/StateTree/Source/StateTreeEditorModule/Private/Customizations/StateTreeBindableItemDetails.cpp
--- a/Engine/Plugins/Runtime/StateTree/Source/StateTreeEditorModule/Private/Customizations/StateTreeBindableItemDetails.cpp
+++ b/Engine/Plugins/Runtime/StateTree/Source/StateTreeEditorModule/Private/Customizations/StateTreeBindableItemDetails.cpp
@@ -12,6 +12,7 @@
 #include "Widgets/Input/SComboButton.h"
 #include "Widgets/Input/SEditableTextBox.h"
 #include "Widgets/Images/SImage.h"
+#include "Widgets/Text/SRichTextBlock.h"
 #include "Framework/MultiBox/MultiBoxBuilder.h"
 #include "InstancedStruct.h"
 #include "Styling/SlateIconFinder.h"
@@ -25,6 +26,12 @@
 #include "StateTreeDelegates.h"
 #include "StateTreeTaskBase.h"
 #include "StateTreeEvaluatorBase.h"
+#include "StateTreeConditionBase.h"
+#include "StateTreePropertyHelpers.h"
+#include "StateTreeEditorStyle.h"
+
+#include "EdGraphSchema_K2.h"
+#include "Kismet2/BlueprintEditorUtils.h"
 
 #define LOCTEXT_NAMESPACE "StateTreeEditor"
 
@@ -32,31 +39,33 @@
 class FBindableItemDataDetails : public FInstancedStructDataDetails
 {
 public:
-	FBindableItemDataDetails(TSharedPtr<IPropertyHandle> InStructProperty, UStateTreeEditorData* InEditorData)
+
+	FBindableItemDataDetails(TSharedPtr<IPropertyHandle> InStructProperty, FGuid InID, UStateTreeEditorData* InEditorData)
 		: FInstancedStructDataDetails(InStructProperty)
 		, EditorData(InEditorData)
 	{
 		EditorPropBindings = EditorData ? EditorData->GetPropertyEditorBindings() : nullptr;
+		ID = InID;
 	}
 
 	virtual void OnChildRowAdded(IDetailPropertyRow& ChildRow)
 	{
-		static const FName BindableMetaName(TEXT("Bindable"));
-
 		if (!EditorPropBindings)
 		{
 			return;
 		}
 		
 		TSharedPtr<IPropertyHandle> ChildPropHandle = ChildRow.GetPropertyHandle();
-
+		const EStateTreePropertyUsage Category = UE::StateTree::PropertyBinding::ParsePropertyUsage(ChildPropHandle);
+		const FProperty* Property = ChildPropHandle->GetProperty();
+		
 		// Conditionally control visibility of the value field of bound properties.
-		if (ChildPropHandle->HasMetaData(BindableMetaName))
+		if (Category != EStateTreePropertyUsage::Invalid && ID.IsValid())
 		{
-			// Find property path relative to closest outer bindable struct.
-			FStateTreeEditorPropertyPath Path;
-			UE::StateTree::PropertyBinding::GetOuterStructPropertyPath(ChildPropHandle, Path);
+			// Pass the item ID to binding extension. Since the properties are added using AddChildStructure(), we break the hierarchy and cannot access parent.
+			ChildPropHandle->SetInstanceMetaData(FName(TEXT("StateTreeItemID")), LexToString(ID));
 
+			FStateTreeEditorPropertyPath Path(ID, *Property->GetFName().ToString());
 			TSharedPtr<SWidget> NameWidget;
 			TSharedPtr<SWidget> ValueWidget;
 			FDetailWidgetRow Row;
@@ -67,25 +76,68 @@ public:
 					return EditorPropBindings->HasPropertyBinding(Path) ? EVisibility::Collapsed : EVisibility::Visible;
 				});
 
-			ChildRow
-				.CustomWidget(/*bShowChildren*/true)
-				.NameContent()
-				[
-					NameWidget.ToSharedRef()
-				]
-				.ValueContent()
-				[
-					SNew(SBox)
-					.Visibility(IsValueVisible)
+			if (Category == EStateTreePropertyUsage::Input || Category == EStateTreePropertyUsage::Output)
+			{
+				FEdGraphPinType PinType;
+				const UEdGraphSchema_K2* Schema = GetDefault<UEdGraphSchema_K2>();
+				Schema->ConvertPropertyToPinType(Property, PinType);
+				const FSlateBrush* Icon = FBlueprintEditorUtils::GetIconFromPin(PinType, true);
+				FText Label = Category == EStateTreePropertyUsage::Input ? LOCTEXT("Input", "Input") : LOCTEXT("Output", "Output");
+				FText Tooltip = FText::FromString(Property->GetCPPType());
+				const FLinearColor Color = Schema->GetPinTypeColor(PinType);
+
+				ChildRow
+					.CustomWidget(/*bShowChildren*/false)
+					.NameContent()
+					[
+						NameWidget.ToSharedRef()
+					]
+					.ValueContent()
 					[
-						ValueWidget.ToSharedRef()
+						SNew(SHorizontalBox)
+						.Visibility(IsValueVisible)
+						+SHorizontalBox::Slot()
+						.AutoWidth()
+						.VAlign(VAlign_Center)
+						.Padding(4.0f, 0.0f)
+						[
+							SNew(SImage)
+							.Image(Icon)
+							.ColorAndOpacity(Color)
+						]
+						+SHorizontalBox::Slot()
+						.VAlign(VAlign_Center)
+						[
+							SNew(STextBlock)
+							.Font(IDetailLayoutBuilder::GetDetailFont())
+							.Text(Label)
+							.ToolTipText(Tooltip)
+						]
+					];
+			}
+			else
+			{
+				ChildRow
+					.CustomWidget(/*bShowChildren*/true)
+					.NameContent()
+					[
+						NameWidget.ToSharedRef()
 					]
-				];
+					.ValueContent()
+					[
+						SNew(SBox)
+						.Visibility(IsValueVisible)
+						[
+							ValueWidget.ToSharedRef()
+						]
+					];
+			}
 		}
 	}
 
 	UStateTreeEditorData* EditorData;
 	FStateTreeEditorPropertyBindings* EditorPropBindings;
+	FGuid ID;
 };
 
 
@@ -108,11 +160,16 @@ public:
 
 	virtual bool IsStructAllowed(const FStructViewerInitializationOptions& InInitOptions, const UScriptStruct* InStruct, TSharedRef<FStructViewerFilterFuncs> InFilterFuncs) override
 	{
-		if (Schema && !Schema->IsStructAllowed(InStruct))
+		if (!Schema)
 		{
 			return false;
 		}
 
+		if (!Schema->IsStructAllowed(InStruct))
+		{
+			return false;
+		}
+		
 		if (InStruct->IsA<UUserDefinedStruct>())
 		{
 			return bAllowUserDefinedStructs;
@@ -142,16 +199,27 @@ TSharedRef<IPropertyTypeCustomization> FStateTreeBindableItemDetails::MakeInstan
 	return MakeShareable(new FStateTreeBindableItemDetails);
 }
 
+FStateTreeBindableItemDetails::~FStateTreeBindableItemDetails()
+{
+	UE::StateTree::PropertyBinding::OnStateTreeBindingChanged.Remove(OnBindingChangedHandle);
+}
+
 void FStateTreeBindableItemDetails::CustomizeHeader(TSharedRef<class IPropertyHandle> StructPropertyHandle, class FDetailWidgetRow& HeaderRow, IPropertyTypeCustomizationUtils& StructCustomizationUtils)
 {
 	StructProperty = StructPropertyHandle;
 	PropUtils = StructCustomizationUtils.GetPropertyUtilities().Get();
 
-	TypeProperty = StructProperty->GetChildHandle(TEXT("Type"));
+	ItemProperty = StructProperty->GetChildHandle(TEXT("Item"));
+	InstanceProperty = StructProperty->GetChildHandle(TEXT("Instance"));
+	IDProperty = StructProperty->GetChildHandle(TEXT("ID"));
 
+	check(ItemProperty.IsValid());
+	check(InstanceProperty.IsValid());
+	check(IDProperty.IsValid());
+	
 	// Find base class from meta data.
 	static const FName BaseClassMetaName(TEXT("BaseStruct")); // TODO: move these names into one central place.
-	const FString BaseClassName = TypeProperty->GetMetaData(BaseClassMetaName);
+	const FString BaseClassName = ItemProperty->GetMetaData(BaseClassMetaName);
 	BaseScriptStruct = FindObject<UScriptStruct>(ANY_PACKAGE, *BaseClassName);
 
 	const FIsResetToDefaultVisible IsResetVisible = FIsResetToDefaultVisible::CreateSP(this, &FStateTreeBindableItemDetails::ShouldResetToDefault);
@@ -159,6 +227,7 @@ void FStateTreeBindableItemDetails::CustomizeHeader(TSharedRef<class IPropertyHa
 	const FResetToDefaultOverride ResetOverride = FResetToDefaultOverride::Create(IsResetVisible, ResetHandler);
 
 	UE::StateTree::Delegates::OnIdentifierChanged.AddSP(this, &FStateTreeBindableItemDetails::OnIdentifierChanged);
+	OnBindingChangedHandle = UE::StateTree::PropertyBinding::OnStateTreeBindingChanged.AddRaw(this, &FStateTreeBindableItemDetails::OnBindingChanged);
 
 	FindOuterObjects();
 
@@ -167,11 +236,21 @@ void FStateTreeBindableItemDetails::CustomizeHeader(TSharedRef<class IPropertyHa
 		.VAlign(VAlign_Center)
 		[
 			SNew(SHorizontalBox)
-			// Name
+			// Description (Condition)
+			+ SHorizontalBox::Slot()
+			.FillWidth(3.0f)
+			.VAlign(VAlign_Center)
+			[
+				SNew(SRichTextBlock)
+				.DecoratorStyleSet(FStateTreeEditorStyle::Get().Get())
+				.TextStyle(FStateTreeEditorStyle::Get(), "Details.Normal")
+				.Text(this, &FStateTreeBindableItemDetails::GetDescription)
+				.Visibility(this, &FStateTreeBindableItemDetails::IsDescriptionVisible)
+			]
+			// Name (Eval/Task)
 			+ SHorizontalBox::Slot()
 			.FillWidth(1.5f)
 			.VAlign(VAlign_Center)
-			.Padding(FMargin(FMargin(0.0f, 0.0f, 4.0f, 0.0f)))
 			[
 				SNew(SEditableTextBox)
 				.IsEnabled(TAttribute<bool>(this, &FStateTreeBindableItemDetails::IsNameEnabled))
@@ -180,10 +259,12 @@ void FStateTreeBindableItemDetails::CustomizeHeader(TSharedRef<class IPropertyHa
 				.SelectAllTextWhenFocused(true)
 				.RevertTextOnEscape(true)
 				.Font(IDetailLayoutBuilder::GetDetailFont())
+				.Visibility(this, &FStateTreeBindableItemDetails::IsNameVisible)
 			]
 			// Class picker
 			+ SHorizontalBox::Slot()
 			.FillWidth(1.0f)
+			.Padding(FMargin(FMargin(4.0f, 0.0f, 0.0f, 0.0f)))
 			.VAlign(VAlign_Center)
 			[
 				SAssignNew(ComboButton, SComboButton)
@@ -220,15 +301,15 @@ void FStateTreeBindableItemDetails::CustomizeHeader(TSharedRef<class IPropertyHa
 
 bool FStateTreeBindableItemDetails::ShouldResetToDefault(TSharedPtr<IPropertyHandle> PropertyHandle) const
 {
-	check(TypeProperty);
+	check(ItemProperty);
 	
 	bool bAnyValid = false;
 	
 	TArray<void*> RawData;
-	TypeProperty->AccessRawData(RawData);
+	ItemProperty->AccessRawData(RawData);
 	for (void* Data : RawData)
 	{
-		if (FInstancedStruct* Struct = static_cast<FInstancedStruct*>(Data))
+		if (const FInstancedStruct* Struct = static_cast<FInstancedStruct*>(Data))
 		{
 			if (Struct->IsValid())
 			{
@@ -244,15 +325,30 @@ bool FStateTreeBindableItemDetails::ShouldResetToDefault(TSharedPtr<IPropertyHan
 
 void FStateTreeBindableItemDetails::ResetToDefault(TSharedPtr<IPropertyHandle> PropertyHandle)
 {
-	check(TypeProperty);
+	check(ItemProperty);
+	check(InstanceProperty);
 	
 	GEditor->BeginTransaction(LOCTEXT("OnResetToDefault", "Reset to default"));
 
-	TypeProperty->NotifyPreChange();
+	ItemProperty->NotifyPreChange();
+	TArray<void*> RawItemData;
+	ItemProperty->AccessRawData(RawItemData);
+	for (void* Data : RawItemData)
+	{
+		if (FInstancedStruct* Struct = static_cast<FInstancedStruct*>(Data))
+		{
+			// Assume that the default value is empty.
+			Struct->Reset();
+		}
+	}
+
+	InstanceProperty->NotifyPostChange(EPropertyChangeType::ValueSet);
+	InstanceProperty->NotifyFinishedChangingProperties();
 
-	TArray<void*> RawData;
-	TypeProperty->AccessRawData(RawData);
-	for (void* Data : RawData)
+	InstanceProperty->NotifyPreChange();
+	TArray<void*> RawInstanceData;
+	InstanceProperty->AccessRawData(RawInstanceData);
+	for (void* Data : RawInstanceData)
 	{
 		if (FInstancedStruct* Struct = static_cast<FInstancedStruct*>(Data))
 		{
@@ -261,9 +357,10 @@ void FStateTreeBindableItemDetails::ResetToDefault(TSharedPtr<IPropertyHandle> P
 		}
 	}
 
-	TypeProperty->NotifyPostChange(EPropertyChangeType::ValueSet);
+	InstanceProperty->NotifyPostChange(EPropertyChangeType::ValueSet);
+	InstanceProperty->NotifyFinishedChangingProperties();
+
 	GEditor->EndTransaction();
-	TypeProperty->NotifyFinishedChangingProperties();
 
 	if (PropUtils)
 	{
@@ -273,8 +370,14 @@ void FStateTreeBindableItemDetails::ResetToDefault(TSharedPtr<IPropertyHandle> P
 
 void FStateTreeBindableItemDetails::CustomizeChildren(TSharedRef<class IPropertyHandle> StructPropertyHandle, class IDetailChildrenBuilder& StructBuilder, IPropertyTypeCustomizationUtils& StructCustomizationUtils)
 {
-	TSharedRef<FBindableItemDataDetails> DataDetails = MakeShareable(new FBindableItemDataDetails(TypeProperty, EditorData));
-	StructBuilder.AddCustomBuilder(DataDetails);
+	FGuid ID;
+	UE::StateTree::PropertyHelpers::GetStructValue<FGuid>(IDProperty, ID);
+
+	TSharedRef<FBindableItemDataDetails> ItemDetails = MakeShareable(new FBindableItemDataDetails(ItemProperty, FGuid(), EditorData));
+	StructBuilder.AddCustomBuilder(ItemDetails);
+
+	TSharedRef<FBindableItemDataDetails> InstanceDetails = MakeShareable(new FBindableItemDataDetails(InstanceProperty, ID, EditorData));
+	StructBuilder.AddCustomBuilder(InstanceDetails);
 }
 
 void FStateTreeBindableItemDetails::OnIdentifierChanged(const UStateTree& InStateTree)
@@ -285,133 +388,228 @@ void FStateTreeBindableItemDetails::OnIdentifierChanged(const UStateTree& InStat
 	}
 }
 
+void FStateTreeBindableItemDetails::OnBindingChanged(const FStateTreeEditorPropertyPath& SourcePath, const FStateTreeEditorPropertyPath& TargetPath)
+{
+	check(ItemProperty);
+
+	FGuid ID;
+	if (UE::StateTree::PropertyHelpers::GetStructValue(IDProperty, ID) != FPropertyAccess::Success)
+	{
+		return;
+	}
+
+	if (ID != TargetPath.StructID)
+	{
+		return;
+	}
+
+	TArray<UObject*> OuterObjects;
+	ItemProperty->GetOuterObjects(OuterObjects);
+
+	TArray<void*> RawItemData;
+	ItemProperty->AccessRawData(RawItemData);
+
+	TArray<void*> RawInstanceData;
+	InstanceProperty->AccessRawData(RawInstanceData);
+
+	if (OuterObjects.Num() != RawItemData.Num() || OuterObjects.Num() != RawInstanceData.Num())
+	{
+		return;
+	}
+
+	for (int32 i = 0; i < OuterObjects.Num(); i++)
+	{
+		// Dig out name from the struct without knowing the type.
+		const FInstancedStruct* ItemPtr = static_cast<FInstancedStruct*>(RawItemData[i]);
+		const FInstancedStruct* InstancePtr = static_cast<FInstancedStruct*>(RawInstanceData[i]);
+		UObject* OuterObject = OuterObjects[i]; // Immediate outer, i.e StateTreeState
+
+		if (InstancePtr->IsValid() && EditorData != nullptr)
+		{
+			if (FStateTreeConditionBase* Condition = ItemPtr->GetMutablePtr<FStateTreeConditionBase>())
+			{
+				const FStateTreeBindingLookup BindingLookup(EditorData);
+
+				OuterObject->Modify();
+				Condition->OnBindingChanged(ID, *InstancePtr, SourcePath, TargetPath, BindingLookup);
+			}
+		}
+	}
+}
+
 void FStateTreeBindableItemDetails::FindOuterObjects()
 {
+	check(ItemProperty);
+	
 	EditorData = nullptr;
 	StateTree = nullptr;
 
-	if (TypeProperty)
+	TArray<UObject*> OuterObjects;
+	ItemProperty->GetOuterObjects(OuterObjects);
+	for (UObject* Outer : OuterObjects)
 	{
-		TArray<UObject*> OuterObjects;
-		TypeProperty->GetOuterObjects(OuterObjects);
-		for (UObject* Outer : OuterObjects)
+		UStateTreeEditorData* OuterEditorData = Outer->GetTypedOuter<UStateTreeEditorData>();
+		UStateTree* OuterStateTree = OuterEditorData ? OuterEditorData->GetTypedOuter<UStateTree>() : nullptr;
+		if (OuterEditorData && OuterStateTree)
 		{
-			UStateTreeEditorData* OuterEditorData = Outer->GetTypedOuter<UStateTreeEditorData>();
-			UStateTree* OuterStateTree = OuterEditorData ? OuterEditorData->GetTypedOuter<UStateTree>() : nullptr;
-			if (OuterEditorData && OuterStateTree)
-			{
-				StateTree = OuterStateTree;
-				EditorData = OuterEditorData;
-				break;
-			}
+			StateTree = OuterStateTree;
+			EditorData = OuterEditorData;
+			break;
 		}
 	}
 }
 
-FText FStateTreeBindableItemDetails::GetName() const
+FText FStateTreeBindableItemDetails::GetDescription() const
 {
-	if (TypeProperty)
+	check(ItemProperty);
+	check(InstanceProperty);
+
+	// Multiple descriptions do not make sense, just if only one item is selected.
+	if (ItemProperty->GetNumOuterObjects() != 1)
+	{
+		return LOCTEXT("MultipleSelected", "<Details.Subdued>Multiple selected</>");
+	}
+
+	FGuid ID;
+	if (UE::StateTree::PropertyHelpers::GetStructValue(IDProperty, ID) != FPropertyAccess::Success)
 	{
-		TArray<void*> RawData;
-		TypeProperty->AccessRawData(RawData);
-		if (RawData.Num() == 1)
+		return LOCTEXT("MultipleSelected", "<Details.Subdued>Multiple selected</>");
+	}
+
+	TArray<void*> RawItemData;
+	ItemProperty->AccessRawData(RawItemData);
+	check(RawItemData.Num() == 1);
+
+	TArray<void*> RawInstanceData;
+	InstanceProperty->AccessRawData(RawInstanceData);
+	check(RawInstanceData.Num() == 1);
+
+	// Dig out name from the struct without knowing the type.
+	const FInstancedStruct* ItemPtr = static_cast<FInstancedStruct*>(RawItemData[0]);
+	const FInstancedStruct* InstancePtr = static_cast<FInstancedStruct*>(RawInstanceData[0]);
+
+	if (InstancePtr->IsValid() && EditorData != nullptr)
+	{
+		if (const FStateTreeConditionBase* Condition = ItemPtr->GetPtr<const FStateTreeConditionBase>())
 		{
-			// Dig out name from the struct without knowing the type.
-			FInstancedStruct* StructPtr = static_cast<FInstancedStruct*>(RawData[0]);
+			const FStateTreeBindingLookup BindingLookup(EditorData);
+			return Condition->GetDescription(ID, *InstancePtr, BindingLookup);
+		}
+	}
+	
+	return LOCTEXT("ConditionNotSet", "<Details.Subdued>Condition Not Set</>");
+}
+
+EVisibility FStateTreeBindableItemDetails::IsDescriptionVisible() const
+{
+	const UScriptStruct* ScriptStruct = GetCommonItemScriptStruct();
+	return ScriptStruct != nullptr && ScriptStruct->IsChildOf(FStateTreeConditionBase::StaticStruct()) ? EVisibility::Visible : EVisibility::Collapsed;
+}
+
+FText FStateTreeBindableItemDetails::GetName() const
+{
+	check(ItemProperty);
 
-			// Make sure the associated property has valid data (e.g. while moving an item in array)
-			const UScriptStruct* ScriptStruct = StructPtr != nullptr ? StructPtr->GetScriptStruct() : nullptr;
-			if (ScriptStruct != nullptr)
+	// Multiple names do not make sense, just if only one item is selected.
+	TArray<void*> RawItemData;
+	ItemProperty->AccessRawData(RawItemData);
+	if (RawItemData.Num() == 1)
+	{
+		// Dig out name from the struct without knowing the type.
+		FInstancedStruct* StructPtr = static_cast<FInstancedStruct*>(RawItemData[0]);
+
+		// Make sure the associated property has valid data (e.g. while moving an item in array)
+		const UScriptStruct* ScriptStruct = StructPtr != nullptr ? StructPtr->GetScriptStruct() : nullptr;
+		if (ScriptStruct != nullptr)
+		{
+			if (FProperty* NameProperty = ScriptStruct->FindPropertyByName(TEXT("Name")))
 			{
-				if (FProperty* NameProperty = ScriptStruct->FindPropertyByName(TEXT("Name")))
+				if (NameProperty->IsA(FNameProperty::StaticClass()))
 				{
-					if (NameProperty->IsA(FNameProperty::StaticClass()))
-					{
-						void* Ptr = const_cast<void*>(static_cast<const void*>(StructPtr->GetMemory()));
-						FName NameValue = *NameProperty->ContainerPtrToValuePtr<FName>(Ptr);
-						return FText::FromName(NameValue);
-					}
+					void* Ptr = const_cast<void*>(static_cast<const void*>(StructPtr->GetMemory()));
+					FName NameValue = *NameProperty->ContainerPtrToValuePtr<FName>(Ptr);
+					return FText::FromName(NameValue);
 				}
 			}
-			return LOCTEXT("Empty", "Empty");
-		}
-		else
-		{
-			return LOCTEXT("MultipleSelected", "Multiple selected");
 		}
+		return LOCTEXT("Empty", "Empty");
 	}
-	return FText::GetEmpty();
+
+	return LOCTEXT("MultipleSelected", "Multiple selected");
+}
+
+EVisibility FStateTreeBindableItemDetails::IsNameVisible() const
+{
+	const UScriptStruct* ScriptStruct = GetCommonItemScriptStruct();
+	return ScriptStruct != nullptr && (ScriptStruct->IsChildOf(FStateTreeTaskBase::StaticStruct()) || ScriptStruct->IsChildOf(FStateTreeEvaluatorBase::StaticStruct())) ? EVisibility::Visible : EVisibility::Collapsed;
 }
 
 void FStateTreeBindableItemDetails::OnNameCommitted(const FText& NewText, ETextCommit::Type InTextCommit)
 {
+	check(ItemProperty);
+
 	if (InTextCommit == ETextCommit::OnEnter || InTextCommit == ETextCommit::OnUserMovedFocus)
 	{
 		// Remove excess whitespace and prevent categories with just spaces
 		FString NewName = FText::TrimPrecedingAndTrailing(NewText).ToString();
 
-		if (TypeProperty)
+		if (GEditor)
 		{
-			TArray<void*> RawData;
-			TypeProperty->AccessRawData(RawData);
-			if (RawData.Num() == 1)
+			GEditor->BeginTransaction(LOCTEXT("SetName", "Set Name"));
+		}
+		ItemProperty->NotifyPreChange();
+
+		TArray<void*> RawItemData;
+		ItemProperty->AccessRawData(RawItemData);
+		for (void* Data : RawItemData)
+		{
+			// Set Name
+			FInstancedStruct* ItemPtr = static_cast<FInstancedStruct*>(Data);
+			if (const UScriptStruct* ScriptStruct = ItemPtr->GetScriptStruct())
 			{
-				if (GEditor)
-				{
-					GEditor->BeginTransaction(LOCTEXT("SetName", "Set Name"));
-				}
-				TypeProperty->NotifyPreChange();
-					
-				// Set Name
-				FInstancedStruct* StructPtr = static_cast<FInstancedStruct*>(RawData[0]);
-				if (const UScriptStruct* ScriptStruct = StructPtr->GetScriptStruct())
+				if (FProperty* NameProperty = ScriptStruct->FindPropertyByName(TEXT("Name")))
 				{
-					if (FProperty* NameProperty = ScriptStruct->FindPropertyByName(TEXT("Name")))
+					if (NameProperty->IsA(FNameProperty::StaticClass()))
 					{
-						if (NameProperty->IsA(FNameProperty::StaticClass()))
-						{
-							void* Ptr = const_cast<void*>(static_cast<const void*>(StructPtr->GetMemory()));
-							FName& NameValue = *NameProperty->ContainerPtrToValuePtr<FName>(Ptr);
-							NameValue = FName(NewName);
-						}
+						void* Ptr = const_cast<void*>(static_cast<const void*>(ItemPtr->GetMemory()));
+						FName& NameValue = *NameProperty->ContainerPtrToValuePtr<FName>(Ptr);
+						NameValue = FName(NewName);
 					}
 				}
+			}
+		}
 
-				TypeProperty->NotifyPostChange(EPropertyChangeType::ValueSet);
-
-				if (StateTree)
-				{
-					UE::StateTree::Delegates::OnIdentifierChanged.Broadcast(*StateTree);
-				}
+		ItemProperty->NotifyPostChange(EPropertyChangeType::ValueSet);
 
-				if (GEditor)
-				{
-					GEditor->EndTransaction();
-				}
+		if (StateTree)
+		{
+			UE::StateTree::Delegates::OnIdentifierChanged.Broadcast(*StateTree);
+		}
 
-				TypeProperty->NotifyFinishedChangingProperties();
-			}
+		if (GEditor)
+		{
+			GEditor->EndTransaction();
 		}
+
+		ItemProperty->NotifyFinishedChangingProperties();
 	}
 }
 
 bool FStateTreeBindableItemDetails::IsNameEnabled() const
 {
 	// Can only edit if we have valid instantiated type.
-	return GetCommonScriptStruct() != nullptr;
+	return GetCommonItemScriptStruct() != nullptr;
 }
 
-const UScriptStruct* FStateTreeBindableItemDetails::GetCommonScriptStruct() const
+const UScriptStruct* FStateTreeBindableItemDetails::GetCommonItemScriptStruct() const
 {
-	if (!TypeProperty)
-	{
-		return nullptr;
-	}
+	check(ItemProperty);
+
 	const UScriptStruct* CommonScriptStruct = nullptr;
 	bool bMultipleValues = false;
-	TArray<void*> RawData;
-	TypeProperty->AccessRawData(RawData);
-	for (void* Data : RawData)
+	TArray<void*> RawItemData;
+	ItemProperty->AccessRawData(RawItemData);
+	for (void* Data : RawItemData)
 	{
 		if (FInstancedStruct* Struct = static_cast<FInstancedStruct*>(Data))
 		{
@@ -433,7 +631,7 @@ const UScriptStruct* FStateTreeBindableItemDetails::GetCommonScriptStruct() cons
 
 FText FStateTreeBindableItemDetails::GetDisplayValueString() const
 {
-	const UScriptStruct* ScriptStruct = GetCommonScriptStruct();
+	const UScriptStruct* ScriptStruct = GetCommonItemScriptStruct();
 	if (ScriptStruct)
 	{
 		return ScriptStruct->GetDisplayNameText();
@@ -449,7 +647,7 @@ const FSlateBrush* FStateTreeBindableItemDetails::GetDisplayValueIcon() const
 TSharedRef<SWidget> FStateTreeBindableItemDetails::GenerateStructPicker()
 {
 	static const FName ExcludeBaseStruct(TEXT("ExcludeBaseStruct")); // TODO: move these names into one central place.
-	const bool bExcludeBaseStruct = TypeProperty->HasMetaData(ExcludeBaseStruct);
+	const bool bExcludeBaseStruct = ItemProperty->HasMetaData(ExcludeBaseStruct);
 
 	TSharedRef<FBindableItemStructFilter> StructFilter = MakeShared<FBindableItemStructFilter>();
 	StructFilter->BaseStruct = BaseScriptStruct;
@@ -482,50 +680,112 @@ TSharedRef<SWidget> FStateTreeBindableItemDetails::GenerateStructPicker()
 
 void FStateTreeBindableItemDetails::OnStructPicked(const UScriptStruct* InStruct)
 {
-	if (TypeProperty)
+	check(ItemProperty);
+	check(InstanceProperty);
+	check(IDProperty);
+
+	TArray<void*> RawItemData;
+	ItemProperty->AccessRawData(RawItemData);
+
+	TArray<void*> RawInstanceData;
+	InstanceProperty->AccessRawData(RawInstanceData);
+
+	TArray<void*> RawIDData;
+	IDProperty->AccessRawData(RawIDData);
+
+	if (RawItemData.Num() != RawInstanceData.Num() || RawItemData.Num() != RawIDData.Num())
 	{
-		GEditor->BeginTransaction(LOCTEXT("OnStructPicked", "Set Struct"));
+		return;
+	}
+
+	GEditor->BeginTransaction(LOCTEXT("OnStructPicked", "Set Struct"));
 
-		TypeProperty->NotifyPreChange();
+	ItemProperty->NotifyPreChange();
+	InstanceProperty->NotifyPreChange();
+	IDProperty->NotifyPreChange();
 
-		TArray<void*> RawData;
-		TypeProperty->AccessRawData(RawData);
-		for (void* Data : RawData)
+	for (int32 Index = 0; Index < RawItemData.Num(); Index++)
+	{
+		FInstancedStruct* ItemPtr = static_cast<FInstancedStruct*>(RawItemData[Index]);
+		FInstancedStruct* InstancePtr = static_cast<FInstancedStruct*>(RawInstanceData[Index]);
+		FGuid* ID = static_cast<FGuid*>(RawIDData[Index]);
+		
+		if (ItemPtr != nullptr && InstancePtr != nullptr && ID != nullptr)
 		{
-			if (FInstancedStruct* Struct = static_cast<FInstancedStruct*>(Data))
+			if (InStruct)
 			{
-				if (InStruct)
+				// Generate new ID and name.
+				*ID = FGuid::NewGuid();
+
+				// Intialize item
+				ItemPtr->InitializeAs(InStruct);
+				
+				// Generate new name and instantiate instance data.
+				if (InStruct->IsChildOf(FStateTreeTaskBase::StaticStruct()))
 				{
-					Struct->InitializeAs(InStruct);
+					FStateTreeTaskBase* Task = ItemPtr->GetMutablePtr<FStateTreeTaskBase>();
+					check(Task);
+					Task->Name = FName(InStruct->GetDisplayNameText().ToString());
 
-					// Generate new ID and name.
-					if (InStruct->IsChildOf(FStateTreeTaskBase::StaticStruct()))
+					if (const UScriptStruct* InstanceType = Cast<const UScriptStruct>(Task->GetInstanceDataType()))
 					{
-						FStateTreeTaskBase* Task = Struct->GetMutablePtr<FStateTreeTaskBase>();
-						check(Task);
-						Task->ID = FGuid::NewGuid();
-						Task->Name = FName(InStruct->GetDisplayNameText().ToString());
+						InstancePtr->InitializeAs(InstanceType);
 					}
-					else if (InStruct->IsChildOf(FStateTreeEvaluatorBase::StaticStruct()))
+					else
 					{
-						FStateTreeEvaluatorBase* Eval = Struct->GetMutablePtr<FStateTreeEvaluatorBase>();
-						check(Eval);
-						Eval->ID = FGuid::NewGuid();
-						Eval->Name = FName(InStruct->GetDisplayNameText().ToString());
+						InstancePtr->Reset();
 					}
 				}
-				else
+				else if (InStruct->IsChildOf(FStateTreeEvaluatorBase::StaticStruct()))
 				{
-					Struct->Reset();
+					FStateTreeEvaluatorBase* Eval = ItemPtr->GetMutablePtr<FStateTreeEvaluatorBase>();
+					check(Eval);
+					Eval->Name = FName(InStruct->GetDisplayNameText().ToString());
+
+					if (const UScriptStruct* InstanceType = Cast<const UScriptStruct>(Eval->GetInstanceDataType()))
+					{
+						InstancePtr->InitializeAs(InstanceType);
+					}
+					else
+					{
+						InstancePtr->Reset();
+					}
 				}
+				else if (InStruct->IsChildOf(FStateTreeConditionBase::StaticStruct()))
+				{
+					FStateTreeConditionBase* Cond = ItemPtr->GetMutablePtr<FStateTreeConditionBase>();
+					check(Cond);
+
+					if (const UScriptStruct* InstanceType = Cast<const UScriptStruct>(Cond->GetInstanceDataType()))
+					{
+						InstancePtr->InitializeAs(InstanceType);
+					}
+					else
+					{
+						InstancePtr->Reset();
+					}
+				}
+			}
+			else
+			{
+				ItemPtr->Reset();
+				InstancePtr->Reset();
+				ID->Invalidate();
 			}
 		}
-
-		TypeProperty->NotifyPostChange(EPropertyChangeType::ValueSet);
-		GEditor->EndTransaction();
-		TypeProperty->NotifyFinishedChangingProperties();
 	}
 
+	ItemProperty->NotifyPostChange(EPropertyChangeType::ValueSet);
+	ItemProperty->NotifyFinishedChangingProperties();
+
+	InstanceProperty->NotifyPostChange(EPropertyChangeType::ValueSet);
+	InstanceProperty->NotifyFinishedChangingProperties();
+
+	IDProperty->NotifyPostChange(EPropertyChangeType::ValueSet);
+	IDProperty->NotifyFinishedChangingProperties();
+
+	GEditor->EndTransaction();
+
 	ComboButton->SetIsOpen(false);
 
 	if (PropUtils)
diff --git a/Engine/Plugins/Runtime/StateTree/Source/StateTreeEditorModule/Private/Customizations/StateTreeBindableItemDetails.h b/Engine/Plugins/Runtime/StateTree/Source/StateTreeEditorModule/Private/Customizations/StateTreeBindableItemDetails.h
--- a/Engine/Plugins/Runtime/StateTree/Source/StateTreeEditorModule/Private/Customizations/StateTreeBindableItemDetails.h
+++ b/Engine/Plugins/Runtime/StateTree/Source/StateTreeEditorModule/Private/Customizations/StateTreeBindableItemDetails.h
@@ -10,6 +10,7 @@ class IPropertyHandle;
 class UStateTree;
 class UStateTreeState;
 class UStateTreeEditorData;
+struct FStateTreeEditorPropertyPath;
 
 /**
  * Semi-generic type customization for bindable items (Evaluators and Tasks) in StateTreeState.
@@ -20,6 +21,8 @@ public:
 	/** Makes a new instance of this detail layout class for a specific detail view requesting it */
 	static TSharedRef<IPropertyTypeCustomization> MakeInstance();
 
+	virtual ~FStateTreeBindableItemDetails();
+	
 	/** IPropertyTypeCustomization interface */
 	virtual void CustomizeHeader(TSharedRef<IPropertyHandle> StructPropertyHandle, class FDetailWidgetRow& HeaderRow, IPropertyTypeCustomizationUtils& StructCustomizationUtils) override;
 	virtual void CustomizeChildren(TSharedRef<IPropertyHandle> StructPropertyHandle, class IDetailChildrenBuilder& StructBuilder, IPropertyTypeCustomizationUtils& StructCustomizationUtils) override;
@@ -29,17 +32,22 @@ private:
 	bool ShouldResetToDefault(TSharedPtr<IPropertyHandle> PropertyHandle) const;
 	void ResetToDefault(TSharedPtr<IPropertyHandle> PropertyHandle);
 
+	FText GetDescription() const;
+	EVisibility IsDescriptionVisible() const;
+	
 	FText GetName() const;
+	EVisibility IsNameVisible() const;
 	void OnNameCommitted(const FText& NewText, ETextCommit::Type InTextCommit);
 	bool IsNameEnabled() const;
 
-	const UScriptStruct* GetCommonScriptStruct() const;
+	const UScriptStruct* GetCommonItemScriptStruct() const;
 	FText GetDisplayValueString() const;
 	const FSlateBrush* GetDisplayValueIcon() const;
 	TSharedRef<SWidget> GenerateStructPicker();
 	void OnStructPicked(const UScriptStruct* InStruct);
 
 	void OnIdentifierChanged(const UStateTree& StateTree);
+	void OnBindingChanged(const FStateTreeEditorPropertyPath& SourcePath, const FStateTreeEditorPropertyPath& TargetPath);
 	void FindOuterObjects();
 
 	UScriptStruct* BaseScriptStruct = nullptr;
@@ -50,5 +58,9 @@ private:
 
 	class IPropertyUtilities* PropUtils = nullptr;
 	TSharedPtr<IPropertyHandle> StructProperty;
-	TSharedPtr<IPropertyHandle> TypeProperty;
+	TSharedPtr<IPropertyHandle> ItemProperty;
+	TSharedPtr<IPropertyHandle> InstanceProperty;
+	TSharedPtr<IPropertyHandle> IDProperty;
+
+	FDelegateHandle OnBindingChangedHandle;
 };
diff --git a/Engine/Plugins/Runtime/StateTree/Source/StateTreeEditorModule/Private/Customizations/StateTreeBindingExtension.cpp b/Engine/Plugins/Runtime/StateTree/Source/StateTreeEditorModule/Private/Customizations/StateTreeBindingExtension.cpp
--- a/Engine/Plugins/Runtime/StateTree/Source/StateTreeEditorModule/Private/Customizations/StateTreeBindingExtension.cpp
+++ b/Engine/Plugins/Runtime/StateTree/Source/StateTreeEditorModule/Private/Customizations/StateTreeBindingExtension.cpp
@@ -14,11 +14,12 @@
 #include "StateTreePropertyHelpers.h"
 #include "StateTreeAnyEnum.h"
 #include "StateTreePropertyBindingCompiler.h"
+#include "PropertyEditor/Private/PropertyNode.h"
 
 #define LOCTEXT_NAMESPACE "StateTreeEditor"
 
-
-namespace UE { namespace StateTree { namespace PropertyBinding {
+namespace UE::StateTree::PropertyBinding
+{
 
 UObject* FindEditorBindingsOwner(UObject* InObject)
 {
@@ -36,57 +37,74 @@ UObject* FindEditorBindingsOwner(UObject* InObject)
 	return Result;
 }
 
-void GetOuterStructPropertyPath(TSharedPtr<IPropertyHandle> InPropertyHandle, FStateTreeEditorPropertyPath& OutPath)
+void GetStructPropertyPath(TSharedPtr<IPropertyHandle> InPropertyHandle, FStateTreeEditorPropertyPath& OutPath)
 {
 	FGuid StructID;
 	TArray<FBindingChainElement> BindingChain;
-	TSharedPtr<IPropertyHandle> CurrentPropertyHandle = InPropertyHandle;
 
-	while (CurrentPropertyHandle.IsValid())
+	if (FProperty* CurrentProperty = InPropertyHandle->GetProperty())
 	{
-		if (FProperty* CurrentProperty = CurrentPropertyHandle->GetProperty())
-		{
-			// Keep track of the path.
-			BindingChain.Insert(FBindingChainElement(CurrentProperty, CurrentPropertyHandle->GetIndexInArray()), 0);
-
-			// If the owner struct of this property has a property called ID and it's FGuid, we'll use that as struct ID.
-			TSharedPtr<IPropertyHandle> OwnerHandle = CurrentPropertyHandle->GetParentHandle();
-			if (OwnerHandle.IsValid())
-			{
-				TSharedPtr<IPropertyHandle> IdHandle = OwnerHandle->GetChildHandle(TEXT("ID"), /*bRecurse*/false); // Omitting recursive, since we expect that the ID belongs to owner.
-				FGuid ID;
-				if (UE::StateTree::PropertyHelpers::GetStructValue<FGuid>(IdHandle, ID) == FPropertyAccess::Success)
-				{
-					StructID = ID;
-					break;
-				}
-			}
-
-			CurrentPropertyHandle = CurrentPropertyHandle->GetParentHandle();
-		}
-		else
+		// Keep track of the path.
+		BindingChain.Insert(FBindingChainElement(CurrentProperty, InPropertyHandle->GetIndexInArray()), 0);
+		
+		if (const FString* IDString = InPropertyHandle->GetInstanceMetaData(FName(TEXT("StateTreeItemID"))))
 		{
-			BindingChain.Reset();
-			CurrentPropertyHandle.Reset();
-			break;
+			LexFromString(StructID, **IDString);
 		}
 	}
+	
+	if (StructID.IsValid())
+	{
+		OutPath.StructID = StructID;
+		IPropertyAccessEditor& PropertyAccessEditor = IModularFeatures::Get().GetModularFeature<IPropertyAccessEditor>("PropertyAccessEditor");
+		PropertyAccessEditor.MakeStringPath(BindingChain, OutPath.Path);
+	}
+	else
+	{
+		OutPath = FStateTreeEditorPropertyPath();
+	}
+}
 
-	OutPath.StructID = StructID;
+	
+EStateTreePropertyUsage ParsePropertyUsage(TSharedPtr<const IPropertyHandle> InPropertyHandle)
+{
+	static const FName CategoryName(TEXT("Category"));
+	
+	const FString Category = InPropertyHandle->GetMetaData(CategoryName);
 
-	IPropertyAccessEditor& PropertyAccessEditor = IModularFeatures::Get().GetModularFeature<IPropertyAccessEditor>("PropertyAccessEditor");
-	PropertyAccessEditor.MakeStringPath(BindingChain, OutPath.Path);
+	if (Category == TEXT("Input"))
+	{
+		return EStateTreePropertyUsage::Input;
+	}
+	if (Category == TEXT("Output"))
+	{
+		return EStateTreePropertyUsage::Output;
+	}
+	if (Category == TEXT("Parameter"))
+	{
+		return EStateTreePropertyUsage::Parameter;
+	}
+	
+	return EStateTreePropertyUsage::Invalid;
 }
 
+
 FOnStateTreeBindingChanged STATETREEEDITORMODULE_API OnStateTreeBindingChanged;
 
-}}} // UE::StateTree::PropertyBinding
+} // UE::StateTree::PropertyBinding
 
 bool FStateTreeBindingExtension::IsPropertyExtendable(const UClass* InObjectClass, const IPropertyHandle& PropertyHandle) const
 {
-	static const FName BindableMetaName(TEXT("Bindable"));
+	// Bindable property must have item ID
+	static const FName StateTreeItemIDName(TEXT("StateTreeItemID"));
+	if (PropertyHandle.GetInstanceMetaData(StateTreeItemIDName) == nullptr)
+	{
+		return false;
+	}
 
-	return PropertyHandle.HasMetaData(BindableMetaName);
+	// Only inputs and parameters are bindable.
+	const EStateTreePropertyUsage Usage = UE::StateTree::PropertyBinding::ParsePropertyUsage(PropertyHandle.AsShared());
+	return Usage == EStateTreePropertyUsage::Input || Usage == EStateTreePropertyUsage::Parameter;
 }
 
 
@@ -117,7 +135,7 @@ void FStateTreeBindingExtension::ExtendWidgetRow(FDetailWidgetRow& InWidgetRow,
 		OwnerObject = UE::StateTree::PropertyBinding::FindEditorBindingsOwner(OuterObjects[0]);
 
 		// Figure out the structs we're editing, and property path relative to current property.
-		UE::StateTree::PropertyBinding::GetOuterStructPropertyPath(InPropertyHandle, TargetPath);
+		UE::StateTree::PropertyBinding::GetStructPropertyPath(InPropertyHandle, TargetPath);
 
 		if (IStateTreeEditorPropertyBindingsOwner* BindingOwner = Cast<IStateTreeEditorPropertyBindingsOwner>(OwnerObject))
 		{
diff --git a/Engine/Plugins/Runtime/StateTree/Source/StateTreeEditorModule/Private/Customizations/StateTreeBindingExtension.h b/Engine/Plugins/Runtime/StateTree/Source/StateTreeEditorModule/Private/Customizations/StateTreeBindingExtension.h
--- a/Engine/Plugins/Runtime/StateTree/Source/StateTreeEditorModule/Private/Customizations/StateTreeBindingExtension.h
+++ b/Engine/Plugins/Runtime/StateTree/Source/StateTreeEditorModule/Private/Customizations/StateTreeBindingExtension.h
@@ -3,7 +3,7 @@
 #pragma once
 
 #include "CoreMinimal.h"
-#include "Widgets/SNullWidget.h"
+#include "StateTreeTypes.h"
 #include "IDetailPropertyExtensionHandler.h"
 
 class IPropertyHandle;
@@ -12,7 +12,8 @@ struct FStateTreeEditorPropertyPath;
 class IPropertyAccessEditor;
 struct FStateTreeEditorPropertyPath;
 
-namespace UE { namespace StateTree { namespace PropertyBinding {
+namespace UE::StateTree::PropertyBinding
+{
 
 	/**
 	 * Get nearest Outer that implements IStateTreeEditorPropertyBindingsOwner.
@@ -21,17 +22,23 @@ namespace UE { namespace StateTree { namespace PropertyBinding {
 	UObject* FindEditorBindingsOwner(UObject* InObject);
 
 	/**
-	 * Returns property path for a specific property. It will look for nearest outer struct that contains property named "ID",
-	 * and build property path using that struct as root.
+	 * Returns property path for a specific property. The property's metadata "StateTreeItemID" is expected to specify the containing struct ID.
 	 * @param InPropertyHandle Handle to the property to find path for.
 	 * @param OutPath Resulting property path.
 	 */
-	void GetOuterStructPropertyPath(TSharedPtr<IPropertyHandle> InPropertyHandle, FStateTreeEditorPropertyPath& OutPath);
-			
+	void GetStructPropertyPath(TSharedPtr<IPropertyHandle> InPropertyHandle, FStateTreeEditorPropertyPath& OutPath);
+
+	/**
+	 * Parses property usage from a property handle.
+	 * @param InPropertyHandle Handle to the property to find path for.
+	 * @return Parsed property usage.
+	 */
+	EStateTreePropertyUsage ParsePropertyUsage(TSharedPtr<const IPropertyHandle> InPropertyHandle);
+	
 	DECLARE_MULTICAST_DELEGATE_TwoParams(FOnStateTreeBindingChanged, const FStateTreeEditorPropertyPath& /*SourcePath*/, const FStateTreeEditorPropertyPath& /*TargetPath*/);
 	extern STATETREEEDITORMODULE_API FOnStateTreeBindingChanged OnStateTreeBindingChanged;
 
-} } } // UE::StateTree::PropertyBinding
+} // UE::StateTree::PropertyBinding
 
 class FStateTreeBindingExtension : public IDetailPropertyExtensionHandler
 {
diff --git a/Engine/Plugins/Runtime/StateTree/Source/StateTreeEditorModule/Private/SStateTreeView.cpp b/Engine/Plugins/Runtime/StateTree/Source/StateTreeEditorModule/Private/SStateTreeView.cpp
--- a/Engine/Plugins/Runtime/StateTree/Source/StateTreeEditorModule/Private/SStateTreeView.cpp
+++ b/Engine/Plugins/Runtime/StateTree/Source/StateTreeEditorModule/Private/SStateTreeView.cpp
@@ -62,16 +62,23 @@ void SStateTreeView::Construct(const FArguments& InArgs, TSharedRef<FStateTreeVi
 
 	bUpdatingSelection = false;
 
+	TSharedRef<SScrollBar> HorizontalScrollBar = SNew(SScrollBar)
+		.Orientation(Orient_Horizontal)
+		.Thickness(FVector2D(12.0f, 12.0f));
+
+	TSharedRef<SScrollBar> VerticalScrollBar = SNew(SScrollBar)
+		.Orientation(Orient_Vertical)
+		.Thickness(FVector2D(12.0f, 12.0f));
+
 	TreeView = SNew(STreeView<UStateTreeState*>)
 		.OnGenerateRow(this, &SStateTreeView::HandleGenerateRow)
 		.OnGetChildren(this, &SStateTreeView::HandleGetChildren)
 		.TreeItemsSource(StateTreeViewModel->GetRoutines())
 		.ItemHeight(32)
 		.OnSelectionChanged(this, &SStateTreeView::HandleTreeSelectionChanged)
-		.OnContextMenuOpening(this, &SStateTreeView::HandleContextMenuOpening);
-
-	ExternalScrollbar = SNew(SScrollBar)
-		.AlwaysShowScrollbar(true);
+		.OnContextMenuOpening(this, &SStateTreeView::HandleContextMenuOpening)
+		.AllowOverscroll(EAllowOverscroll::No)
+		.ExternalScrollbar(VerticalScrollBar);
 
 	ChildSlot
 	[
@@ -123,7 +130,6 @@ void SStateTreeView::Construct(const FArguments& InArgs, TSharedRef<FStateTreeVi
 		]
 
 		+SVerticalBox::Slot()
-		.FillHeight(1.0f)
 		.Padding(0.0f, 6.0f, 0.0f, 0.0f)
 		[
 			SNew(SHorizontalBox)
@@ -133,30 +139,24 @@ void SStateTreeView::Construct(const FArguments& InArgs, TSharedRef<FStateTreeVi
 			[
 				SNew(SScrollBox)
 				.Orientation(Orient_Horizontal)
-				+ SScrollBox::Slot()
-				.HAlign(HAlign_Fill)
+				.ExternalScrollbar(HorizontalScrollBar)
+				+SScrollBox::Slot()
 				[
-					SNew(SBorder)
-					.HAlign(HAlign_Fill)
-					.BorderImage(FEditorStyle::GetBrush("ToolPanel.GroupBorder"))
-					.Padding(0.0f)
-					[
-						TreeView.ToSharedRef()
-					]
+					TreeView.ToSharedRef()
 				]
 			]
 
 			+ SHorizontalBox::Slot()
 			.AutoWidth()
-			.Padding(0.0f)
 			[
-				SNew(SBox)
-				.WidthOverride(FOptionalSize(13.0f))
-				[
-					ExternalScrollbar.ToSharedRef()
-				]
+				VerticalScrollBar
 			]
 		]
+		+SVerticalBox::Slot()
+		.AutoHeight()
+		[
+			HorizontalScrollBar
+		]
 	];
 
 	UpdateTree(true);
diff --git a/Engine/Plugins/Runtime/StateTree/Source/StateTreeEditorModule/Private/SStateTreeViewRow.cpp b/Engine/Plugins/Runtime/StateTree/Source/StateTreeEditorModule/Private/SStateTreeViewRow.cpp
--- a/Engine/Plugins/Runtime/StateTree/Source/StateTreeEditorModule/Private/SStateTreeViewRow.cpp
+++ b/Engine/Plugins/Runtime/StateTree/Source/StateTreeEditorModule/Private/SStateTreeViewRow.cpp
@@ -419,7 +419,7 @@ EVisibility SStateTreeViewRow::GetEvaluatorsVisibility() const
 		int32 ValidCount = 0;
 		for (int32 i = 0; i < State->Evaluators.Num(); i++)
 		{
-			if (const FStateTreeEvaluatorBase* Eval = State->Evaluators[i].Type.GetPtr<FStateTreeEvaluatorBase>())
+			if (const FStateTreeEvaluatorBase* Eval = State->Evaluators[i].Item.GetPtr<FStateTreeEvaluatorBase>())
 			{
 				ValidCount++;
 			}
@@ -441,7 +441,7 @@ FText SStateTreeViewRow::GetEvaluatorsDesc() const
 
 	for (int32 i = 0; i < State->Evaluators.Num(); i++)
 	{
-		if (const FStateTreeEvaluatorBase* Eval = State->Evaluators[i].Type.GetPtr<FStateTreeEvaluatorBase>())
+		if (const FStateTreeEvaluatorBase* Eval = State->Evaluators[i].Item.GetPtr<FStateTreeEvaluatorBase>())
 		{
 			Names.Add(FText::FromName(Eval->Name));
 		}
@@ -457,7 +457,7 @@ EVisibility SStateTreeViewRow::GetTasksVisibility() const
 		int32 ValidCount = 0;
 		for (int32 i = 0; i < State->Tasks.Num(); i++)
 		{
-			if (const FStateTreeTaskBase* Task = State->Tasks[i].Type.GetPtr<FStateTreeTaskBase>())
+			if (const FStateTreeTaskBase* Task = State->Tasks[i].Item.GetPtr<FStateTreeTaskBase>())
 			{
 				ValidCount++;
 			}
@@ -478,7 +478,7 @@ FText SStateTreeViewRow::GetTasksDesc() const
 	TArray<FText> Names;
 	for (int32 i = 0; i < State->Tasks.Num(); i++)
 	{
-		if (const FStateTreeTaskBase* Task = State->Tasks[i].Type.GetPtr<FStateTreeTaskBase>())
+		if (const FStateTreeTaskBase* Task = State->Tasks[i].Item.GetPtr<FStateTreeTaskBase>())
 		{
 			Names.Add(FText::FromName(Task->Name));
 		}
diff --git a/Engine/Plugins/Runtime/StateTree/Source/StateTreeEditorModule/Private/StateTreeBaker.cpp b/Engine/Plugins/Runtime/StateTree/Source/StateTreeEditorModule/Private/StateTreeBaker.cpp
--- a/Engine/Plugins/Runtime/StateTree/Source/StateTreeEditorModule/Private/StateTreeBaker.cpp
+++ b/Engine/Plugins/Runtime/StateTree/Source/StateTreeEditorModule/Private/StateTreeBaker.cpp
@@ -47,7 +47,7 @@ bool FStateTreeBaker::Bake(UStateTree& InStateTree)
 
 	StateTree->Link();
 	
-	StateTree->InitRuntimeStorage();
+	StateTree->InitInstanceStorage();
 
 	return true;
 }
@@ -66,8 +66,16 @@ FStateTreeHandle FStateTreeBaker::GetStateHandle(const FGuid& StateID) const
 bool FStateTreeBaker::CreateStates()
 {
 	// Create item for the runtime execution state
-	StateTree->RuntimeStorageItems.Add(FInstancedStruct::Make<FStateTreeExecutionState>());
+	StateTree->Instances.Add(FInstancedStruct::Make<FStateTreeExecutionState>());
 
+	// The struct index is currently shared between runtime item index, make sure the state has a dummy struct too. 
+	FStateTreeBindableStructDesc StructDesc;
+	StructDesc.Struct = FStateTreeExecutionState::StaticStruct();
+	StructDesc.Name = FName(TEXT("ExecutionState"));
+	StructDesc.ID = FGuid();
+	BindingsCompiler.AddSourceStruct(StructDesc);
+
+	
 	for (UStateTreeState* Routine : TreeData->Routines)
 	{
 		if (Routine)
@@ -78,6 +86,7 @@ bool FStateTreeBaker::CreateStates()
 			}
 		}
 	}
+	
 	return true;
 }
 
@@ -92,7 +101,7 @@ bool FStateTreeBaker::CreateStateTransitions()
 		FStateTreeCompilerLogStateScope LogStateScope(SourceState, Log);
 		
 		// Enter conditions.
-		BakedState.EnterConditionsBegin = uint16(StateTree->Conditions.Num());
+		BakedState.EnterConditionsBegin = uint16(StateTree->Items.Num());
 		for (FStateTreeConditionItem& ConditionItem : SourceState->EnterConditions)
 		{
 			if (!CreateCondition(ConditionItem))
@@ -102,7 +111,7 @@ bool FStateTreeBaker::CreateStateTransitions()
 				return false;
 			}
 		}
-		BakedState.EnterConditionsNum = uint8(uint16(StateTree->Conditions.Num()) - BakedState.EnterConditionsBegin);
+		BakedState.EnterConditionsNum = uint8(uint16(StateTree->Items.Num()) - BakedState.EnterConditionsBegin);
 
 		// Transitions
 		BakedState.TransitionsBegin = uint16(StateTree->Transitions.Num());
@@ -119,7 +128,7 @@ bool FStateTreeBaker::CreateStateTransitions()
 			}
 			// Note: Unset transition is allowed here. It can be used to mask a transition at parent.
 
-			BakedTransition.ConditionsBegin = uint16(StateTree->Conditions.Num());
+			BakedTransition.ConditionsBegin = uint16(StateTree->Items.Num());
 			for (FStateTreeConditionItem& ConditionItem : Transition.Conditions)
 			{
 				if (!CreateCondition(ConditionItem))
@@ -130,7 +139,7 @@ bool FStateTreeBaker::CreateStateTransitions()
 					return false;
 				}
 			}
-			BakedTransition.ConditionsNum = uint8(uint16(StateTree->Conditions.Num()) - BakedTransition.ConditionsBegin);
+			BakedTransition.ConditionsNum = uint8(uint16(StateTree->Items.Num()) - BakedTransition.ConditionsBegin);
 		}
 		BakedState.TransitionsNum = uint8(uint16(StateTree->Transitions.Num()) - BakedState.TransitionsBegin);
 	}
@@ -178,24 +187,35 @@ bool FStateTreeBaker::ResolveTransitionState(const UStateTreeState& SourceState,
 
 bool FStateTreeBaker::CreateCondition(const FStateTreeConditionItem& CondItem)
 {
-	if (!CondItem.Type.IsValid())
+	if (!CondItem.Item.IsValid())
 	{
 		// Empty line in conditions array, just silently ignore.
 		return true;
 	}
 
+	// Assume valid items have valid instance initialized.
+	check(CondItem.Instance.IsValid());
+
 	// Copy the condition
-	FInstancedStruct& CondPtr = StateTree->Conditions.AddDefaulted_GetRef();
-	CondPtr = CondItem.Type;
-	check(CondPtr.IsValid());
+	FInstancedStruct& Item = StateTree->Items.AddDefaulted_GetRef();
+	Item = CondItem.Item;
 
-	FStateTreeConditionBase& Cond = CondPtr.GetMutable<FStateTreeConditionBase>();
+	FInstancedStruct& Instance = StateTree->Instances.AddDefaulted_GetRef();
+	Instance = CondItem.Instance;
+			
+	check(Item.IsValid() && Instance.IsValid());
+	FStateTreeConditionBase& Cond = Item.GetMutable<FStateTreeConditionBase>();
 
 	// Create binding source struct descriptor. Note: not exposing the struct for reading.
 	FStateTreeBindableStructDesc StructDesc;
-	StructDesc.Struct = CondPtr.GetScriptStruct();
-	StructDesc.Name = CondPtr.GetScriptStruct()->GetFName();
-	StructDesc.ID = Cond.ID;
+	StructDesc.Struct = Instance.GetScriptStruct();
+	StructDesc.Name = Instance.GetScriptStruct()->GetFName();
+	StructDesc.ID = CondItem.ID;
+
+	// Mark the struct as binding source.
+	const int32 SourceStructIndex = BindingsCompiler.AddSourceStruct(StructDesc);
+
+	check((StateTree->Instances.Num() - 1) == SourceStructIndex);
 
 	// Check that the bindings for this struct are still all valid.
 	TArray<FStateTreeEditorPropertyBinding> Bindings;
@@ -213,6 +233,9 @@ bool FStateTreeBaker::CreateCondition(const FStateTreeConditionItem& CondItem)
 	check(BatchIndex < int32(MAX_uint16));
 	Cond.BindingsBatch = BatchIndex == INDEX_NONE ? FStateTreeHandle::Invalid : FStateTreeHandle(uint16(BatchIndex));
 
+	check(SourceStructIndex <= int32(MAX_uint16));
+	Cond.DataViewIndex = uint16(SourceStructIndex);
+	
 	return true;
 }
 
@@ -313,29 +336,38 @@ bool FStateTreeBaker::CreateStateRecursive(UStateTreeState& State, const FStateT
 	SourceStates.Add(&State);
 	IDToState.Add(State.ID, StateIdx);
 
+
+	check(StateTree->Instances.Num() <= int32(MAX_uint16));
+
 	// Collect evaluators
-	check(StateTree->RuntimeStorageItems.Num() <= int32(MAX_uint16));
-	BakedState.EvaluatorsBegin = uint16(StateTree->RuntimeStorageItems.Num());
+	check(StateTree->Items.Num() <= int32(MAX_uint16));
+	BakedState.EvaluatorsBegin = uint16(StateTree->Items.Num());
 
 	for (FStateTreeEvaluatorItem& EvaluatorItem : State.Evaluators)
 	{
-		if (EvaluatorItem.Type.IsValid())
+		if (EvaluatorItem.Item.IsValid() && EvaluatorItem.Instance.IsValid())
 		{
 			// Copy the evaluator
-			FInstancedStruct& EvalPtr = StateTree->RuntimeStorageItems.AddDefaulted_GetRef();
-			EvalPtr = EvaluatorItem.Type;
-			check(EvalPtr.IsValid());
-			FStateTreeEvaluatorBase& Eval = EvalPtr.GetMutable<FStateTreeEvaluatorBase>();
+			FInstancedStruct& Item = StateTree->Items.AddDefaulted_GetRef();
+			Item = EvaluatorItem.Item;
+
+			FInstancedStruct& Instance = StateTree->Instances.AddDefaulted_GetRef();
+			Instance = EvaluatorItem.Instance;
+			
+			check(Item.IsValid() && Instance.IsValid());
+			FStateTreeEvaluatorBase& Eval = Item.GetMutable<FStateTreeEvaluatorBase>();
 
 			// Create binding source struct descriptor.
 			FStateTreeBindableStructDesc StructDesc;
-			StructDesc.Struct = EvalPtr.GetScriptStruct();
+			StructDesc.Struct = Instance.GetScriptStruct();
 			StructDesc.Name = Eval.Name;
-			StructDesc.ID = Eval.ID;
+			StructDesc.ID = EvaluatorItem.ID;
 
 			// Mark the struct as binding source.
 			const int32 SourceStructIndex = BindingsCompiler.AddSourceStruct(StructDesc);
 
+			check((StateTree->Instances.Num() - 1) == SourceStructIndex);
+			
 			// Check that the bindings for this struct are still all valid.
 			TArray<FStateTreeEditorPropertyBinding> Bindings;
 			if (!GetAndValidateBindings(StructDesc, Bindings))
@@ -349,39 +381,48 @@ bool FStateTreeBaker::CreateStateRecursive(UStateTreeState& State, const FStateT
 			{
 				return false;
 			}
+			
 			check(BatchIndex < int32(MAX_uint16));
 			Eval.BindingsBatch = BatchIndex == INDEX_NONE ? FStateTreeHandle::Invalid : FStateTreeHandle(uint16(BatchIndex));
+			
 			check(SourceStructIndex <= int32(MAX_uint16));
-			Eval.SourceStructIndex = uint16(SourceStructIndex);
+			Eval.DataViewIndex = uint16(SourceStructIndex);
 		}
 	}
-	const int32 EvaluatorsNum = StateTree->RuntimeStorageItems.Num() - int32(BakedState.EvaluatorsBegin);
+	const int32 EvaluatorsNum = StateTree->Items.Num() - int32(BakedState.EvaluatorsBegin);
 	check(EvaluatorsNum <= int32(MAX_uint8));
 	BakedState.EvaluatorsNum = uint8(EvaluatorsNum);
 
 	// Collect tasks
-	check(StateTree->RuntimeStorageItems.Num() <= int32(MAX_uint16));
-	BakedState.TasksBegin = uint16(StateTree->RuntimeStorageItems.Num());
+	check(StateTree->Items.Num() <= int32(MAX_uint16));
+	BakedState.TasksBegin = uint16(StateTree->Items.Num());
 
 	for (FStateTreeTaskItem& TaskItem : State.Tasks)
 	{
-		if (TaskItem.Type.IsValid())
+		if (TaskItem.Item.IsValid() && TaskItem.Instance.IsValid())
 		{
 			// Copy the task
-			FInstancedStruct& TaskPtr = StateTree->RuntimeStorageItems.AddDefaulted_GetRef();
-			TaskPtr = TaskItem.Type;
-			check(TaskPtr.IsValid());
-			FStateTreeTaskBase& Task = TaskPtr.GetMutable<FStateTreeTaskBase>();
+			FInstancedStruct& Item = StateTree->Items.AddDefaulted_GetRef();
+			Item = TaskItem.Item;
+
+			FInstancedStruct& Instance = StateTree->Instances.AddDefaulted_GetRef();
+			Instance = TaskItem.Instance;
+
+			check(Item.IsValid() && Instance.IsValid());
+			
+			FStateTreeTaskBase& Task = Item.GetMutable<FStateTreeTaskBase>();
 
 			// Create binding source struct descriptor.
 			FStateTreeBindableStructDesc StructDesc;
-			StructDesc.Struct = TaskPtr.GetScriptStruct();
+			StructDesc.Struct = Instance.GetScriptStruct();
 			StructDesc.Name = Task.Name;
-			StructDesc.ID = Task.ID;
+			StructDesc.ID = TaskItem.ID;
 
 			// Mark the struct as binding source.
 			const int32 SourceStructIndex = BindingsCompiler.AddSourceStruct(StructDesc);
 
+			check((StateTree->Instances.Num() - 1) == SourceStructIndex);
+
 			// Check that the bindings for this struct are still all valid.
 			TArray<FStateTreeEditorPropertyBinding> Bindings;
 			if (!GetAndValidateBindings(StructDesc, Bindings))
@@ -395,13 +436,15 @@ bool FStateTreeBaker::CreateStateRecursive(UStateTreeState& State, const FStateT
 			{
 				return false;
 			}
+			
 			check(BatchIndex < int32(MAX_uint16));
 			Task.BindingsBatch = BatchIndex == INDEX_NONE ? FStateTreeHandle::Invalid : FStateTreeHandle(uint16(BatchIndex));
+			
 			check(SourceStructIndex <= int32(MAX_uint16));
-			Task.SourceStructIndex = uint16(SourceStructIndex);
+			Task.DataViewIndex = uint16(SourceStructIndex);
 		}
 	}
-	const int32 TasksNum = StateTree->RuntimeStorageItems.Num() - int32(BakedState.TasksBegin);
+	const int32 TasksNum = StateTree->Items.Num() - int32(BakedState.TasksBegin);
 	check(TasksNum <= int32(MAX_uint8));
 	BakedState.TasksNum = uint8(TasksNum);
 
diff --git a/Engine/Plugins/Runtime/StateTree/Source/StateTreeEditorModule/Private/StateTreeEditor.cpp b/Engine/Plugins/Runtime/StateTree/Source/StateTreeEditorModule/Private/StateTreeEditor.cpp
--- a/Engine/Plugins/Runtime/StateTree/Source/StateTreeEditorModule/Private/StateTreeEditor.cpp
+++ b/Engine/Plugins/Runtime/StateTree/Source/StateTreeEditorModule/Private/StateTreeEditor.cpp
@@ -289,7 +289,7 @@ TSharedRef<SDockTab> FStateTreeEditor::SpawnTab_StateTreeStatistics(const FSpawn
 		.Label(LOCTEXT("StatisticsTitle", "StateTree Statistics"))
 		[
 			SNew(SMultiLineEditableTextBox)
-			.Padding(10)
+			.Padding(10.0f)
 			.Style(FEditorStyle::Get(), "Log.TextBox")
 			.Font(FCoreStyle::GetDefaultFontStyle("Mono", 9))
 			.ForegroundColor(FLinearColor::Gray)
@@ -320,10 +320,10 @@ FText FStateTreeEditor::GetStatisticsText() const
 		return FText::GetEmpty();
 	}
 
-	if (const UScriptStruct* RuntimeStruct = StateTree->GetRuntimeStorageStruct())
+	if (const UScriptStruct* StorageStruct = StateTree->GetInstanceStorageStruct())
 	{
-		const FText SizeText = FText::AsMemory((uint64)RuntimeStruct->GetStructureSize());
-		const FText NumItemsText = FText::AsNumber(StateTree->GetRuntimeStorageItemCount());
+		const FText SizeText = FText::AsMemory((uint64)StorageStruct->GetStructureSize());
+		const FText NumItemsText = FText::AsNumber(StateTree->GetNumInstances());
 
 		return FText::Format(LOCTEXT("RuntimeSize", "Runtime size: {0}, {1} items"), SizeText, NumItemsText);
 	}
@@ -541,7 +541,7 @@ namespace UE::StateTree::Editor::Internal
 			return;
 		}
 
-		TMap<FGuid, const UScriptStruct*> AllStructIDs;
+		TMap<FGuid, const UStruct*> AllStructIDs;
 		TreeData->GetAllStructIDs(AllStructIDs);
 		TreeData->GetPropertyEditorBindings()->RemoveUnusedBindings(AllStructIDs);
 	}
diff --git a/Engine/Plugins/Runtime/StateTree/Source/StateTreeEditorModule/Private/StateTreeEditorData.cpp b/Engine/Plugins/Runtime/StateTree/Source/StateTreeEditorModule/Private/StateTreeEditorData.cpp
--- a/Engine/Plugins/Runtime/StateTree/Source/StateTreeEditorModule/Private/StateTreeEditorData.cpp
+++ b/Engine/Plugins/Runtime/StateTree/Source/StateTreeEditorModule/Private/StateTreeEditorData.cpp
@@ -19,7 +19,7 @@ void UStateTreeEditorData::GetAccessibleStructs(const FGuid TargetStructID, TArr
 
 	if (ValidStates.Num() > 0)
 	{
-		VisitHierarchy([&ValidStates, &OutStructDescs, &EvalDescs, &TaskDescs, TargetStructID, CurrentState](const UStateTreeState& State, const FGuid& ID, const FName& Name, const UScriptStruct* Struct) -> bool
+		VisitHierarchy([&ValidStates, &OutStructDescs, &EvalDescs, &TaskDescs, TargetStructID, CurrentState](const UStateTreeState& State, const FGuid& ID, const FName& Name, const UScriptStruct* ItemStruct, const UStruct* InstanceStruct) -> bool
 			{
 				if (ValidStates.Contains(&State))
 				{
@@ -28,7 +28,7 @@ void UStateTreeEditorData::GetAccessibleStructs(const FGuid TargetStructID, TArr
 						OutStructDescs.Append(EvalDescs);
 
 						// Only tasks can see other tasks too.
-						if (Struct->IsChildOf(FStateTreeTaskBase::StaticStruct()))
+						if (ItemStruct->IsChildOf(FStateTreeTaskBase::StaticStruct()))
 						{
 							OutStructDescs.Append(TaskDescs);
 						}
@@ -36,19 +36,19 @@ void UStateTreeEditorData::GetAccessibleStructs(const FGuid TargetStructID, TArr
 						return false; // Stop visit
 					}
 
-					if (Struct->IsChildOf(FStateTreeEvaluatorBase::StaticStruct()))
+					if (ItemStruct->IsChildOf(FStateTreeEvaluatorBase::StaticStruct()))
 					{
 						// All evaluators up to the target struct are accessible.
 						FStateTreeBindableStructDesc& Desc = EvalDescs.AddDefaulted_GetRef();
-						Desc.Struct = Struct;
+						Desc.Struct = InstanceStruct;
 						Desc.Name = Name;
 						Desc.ID = ID;
 					}
-					else if (Struct->IsChildOf(FStateTreeTaskBase::StaticStruct()))
+					else if (ItemStruct->IsChildOf(FStateTreeTaskBase::StaticStruct()))
 					{
 						// All tasks up to the target struct are accessible (for other tasks).
 						FStateTreeBindableStructDesc& Desc = TaskDescs.AddDefaulted_GetRef();
-						Desc.Struct = Struct;
+						Desc.Struct = InstanceStruct;
 						Desc.Name = Name;
 						Desc.ID = ID;
 					}
@@ -62,11 +62,11 @@ bool UStateTreeEditorData::GetStructByID(const FGuid StructID, FStateTreeBindabl
 {
 	bool bResult = false;
 
-	VisitHierarchy([&bResult, &OutStructDesc, StructID](const UStateTreeState& State, const FGuid& ID, const FName& Name, const UScriptStruct* Struct) -> bool
+	VisitHierarchy([&bResult, &OutStructDesc, StructID](const UStateTreeState& State, const FGuid& ID, const FName& Name, const UScriptStruct* ItemStruct, const UStruct* InstanceStruct) -> bool
 		{
 			if (ID == StructID)
 			{
-				OutStructDesc.Struct = Struct;
+				OutStructDesc.Struct = InstanceStruct;
 				OutStructDesc.Name = Name;
 				OutStructDesc.ID = ID;
 				bResult = true;
@@ -82,7 +82,7 @@ const UStateTreeState* UStateTreeEditorData::GetStateByStructID(const FGuid Targ
 {
 	const UStateTreeState* Result = nullptr;
 
-	VisitHierarchy([&Result, TargetStructID](const UStateTreeState& State, const FGuid& ID, const FName& Name, const UScriptStruct* Struct) -> bool
+	VisitHierarchy([&Result, TargetStructID](const UStateTreeState& State, const FGuid& ID, const FName& Name, const UScriptStruct* ItemStruct, const UStruct* InstanceStruct) -> bool
 		{
 			if (ID == TargetStructID)
 			{
@@ -95,18 +95,18 @@ const UStateTreeState* UStateTreeEditorData::GetStateByStructID(const FGuid Targ
 	return Result;
 }
 
-void UStateTreeEditorData::GetAllStructIDs(TMap<FGuid, const UScriptStruct*>& AllStructs) const
+void UStateTreeEditorData::GetAllStructIDs(TMap<FGuid, const UStruct*>& AllStructs) const
 {
 	AllStructs.Reset();
 
-	VisitHierarchy([&AllStructs](const UStateTreeState& State, const FGuid& ID, const FName& Name, const UScriptStruct* Struct) -> bool
+	VisitHierarchy([&AllStructs](const UStateTreeState& State, const FGuid& ID, const FName& Name, const UScriptStruct* ItemStruct, const UStruct* InstanceStruct) -> bool
 		{
-			AllStructs.Add(ID, Struct);
+			AllStructs.Add(ID, InstanceStruct);
 			return true; // Continue
 		});
 }
 
-void UStateTreeEditorData::VisitHierarchy(TFunctionRef<bool(const UStateTreeState& State, const FGuid& ID, const FName& Name, const UScriptStruct* Struct)> InFunc) const
+void UStateTreeEditorData::VisitHierarchy(TFunctionRef<bool(const UStateTreeState& State, const FGuid& ID, const FName& Name, const UScriptStruct* ItemStruct, const UStruct* InstanceStruct)> InFunc) const
 {
 	TArray<const UStateTreeState*> Stack;
 	bool bContinue = true;
@@ -130,9 +130,9 @@ void UStateTreeEditorData::VisitHierarchy(TFunctionRef<bool(const UStateTreeStat
 			// Evaluators
 			for (const FStateTreeEvaluatorItem& Item : State->Evaluators)
 			{
-				if (const FStateTreeEvaluatorBase* Evaluator = Item.Type.GetPtr<FStateTreeEvaluatorBase>())
+				if (const FStateTreeEvaluatorBase* Evaluator = Item.Item.GetPtr<FStateTreeEvaluatorBase>())
 				{
-					if (!InFunc(*State, Evaluator->ID, Evaluator->Name, Item.Type.GetScriptStruct()))
+					if (!InFunc(*State, Item.ID, Evaluator->Name, Item.Item.GetScriptStruct(), Evaluator->GetInstanceDataType()))
 					{
 						bContinue = false;
 						break;
@@ -144,9 +144,9 @@ void UStateTreeEditorData::VisitHierarchy(TFunctionRef<bool(const UStateTreeStat
 				// Enter conditions
 				for (const FStateTreeConditionItem& Item : State->EnterConditions)
 				{
-					if (const FStateTreeConditionBase* Cond = Item.Type.GetPtr<FStateTreeConditionBase>())
+					if (const FStateTreeConditionBase* Cond = Item.Item.GetPtr<FStateTreeConditionBase>())
 					{
-						if (!InFunc(*State, Cond->ID, Item.Type.GetScriptStruct()->GetFName(), Item.Type.GetScriptStruct()))
+						if (!InFunc(*State, Item.ID, Item.Item.GetScriptStruct()->GetFName(), Item.Item.GetScriptStruct(), Cond->GetInstanceDataType()))
 						{
 							bContinue = false;
 							break;
@@ -159,9 +159,9 @@ void UStateTreeEditorData::VisitHierarchy(TFunctionRef<bool(const UStateTreeStat
 				// Tasks
 				for (const FStateTreeTaskItem& Item : State->Tasks)
 				{
-					if (const FStateTreeTaskBase* Task = Item.Type.GetPtr<FStateTreeTaskBase>())
+					if (const FStateTreeTaskBase* Task = Item.Item.GetPtr<FStateTreeTaskBase>())
 					{
-						if (!InFunc(*State, Task->ID, Task->Name, Item.Type.GetScriptStruct()))
+						if (!InFunc(*State, Item.ID, Task->Name, Item.Item.GetScriptStruct(), Task->GetInstanceDataType()))
 						{
 							bContinue = false;
 							break;
@@ -176,9 +176,9 @@ void UStateTreeEditorData::VisitHierarchy(TFunctionRef<bool(const UStateTreeStat
 				{
 					for (const FStateTreeConditionItem& Item : Transition.Conditions)
 					{
-						if (const FStateTreeConditionBase* Cond = Item.Type.GetPtr<FStateTreeConditionBase>())
+						if (const FStateTreeConditionBase* Cond = Item.Item.GetPtr<FStateTreeConditionBase>())
 						{
-							if (!InFunc(*State, Cond->ID, Item.Type.GetScriptStruct()->GetFName(), Item.Type.GetScriptStruct()))
+							if (!InFunc(*State, Item.ID, Item.Item.GetScriptStruct()->GetFName(), Item.Item.GetScriptStruct(), Cond->GetInstanceDataType()))
 							{
 								bContinue = false;
 								break;
diff --git a/Engine/Plugins/Runtime/StateTree/Source/StateTreeEditorModule/Private/StateTreeEditorModule.cpp b/Engine/Plugins/Runtime/StateTree/Source/StateTreeEditorModule/Private/StateTreeEditorModule.cpp
--- a/Engine/Plugins/Runtime/StateTree/Source/StateTreeEditorModule/Private/StateTreeEditorModule.cpp
+++ b/Engine/Plugins/Runtime/StateTree/Source/StateTreeEditorModule/Private/StateTreeEditorModule.cpp
@@ -9,7 +9,6 @@
 #include "Customizations/StateTreeStateLinkDetails.h"
 #include "Customizations/StateTreeStateDetails.h"
 #include "Customizations/StateTreeBindableItemDetails.h"
-#include "Customizations/StateTreeConditionItemDetails.h"
 #include "Customizations/StateTreeAnyEnumDetails.h"
 #include "StateTreeEditor.h"
 #include "StateTree.h"
@@ -40,7 +39,7 @@ void FStateTreeEditorModule::StartupModule()
 	PropertyModule.RegisterCustomPropertyTypeLayout("StateTreeStateLink", FOnGetPropertyTypeCustomizationInstance::CreateStatic(&FStateTreeStateLinkDetails::MakeInstance));
 	PropertyModule.RegisterCustomPropertyTypeLayout("StateTreeEvaluatorItem", FOnGetPropertyTypeCustomizationInstance::CreateStatic(&FStateTreeBindableItemDetails::MakeInstance));
 	PropertyModule.RegisterCustomPropertyTypeLayout("StateTreeTaskItem", FOnGetPropertyTypeCustomizationInstance::CreateStatic(&FStateTreeBindableItemDetails::MakeInstance));
-	PropertyModule.RegisterCustomPropertyTypeLayout("StateTreeConditionItem", FOnGetPropertyTypeCustomizationInstance::CreateStatic(&FStateTreeConditionItemDetails::MakeInstance));
+	PropertyModule.RegisterCustomPropertyTypeLayout("StateTreeConditionItem", FOnGetPropertyTypeCustomizationInstance::CreateStatic(&FStateTreeBindableItemDetails::MakeInstance));
 	PropertyModule.RegisterCustomPropertyTypeLayout("StateTreeAnyEnum", FOnGetPropertyTypeCustomizationInstance::CreateStatic(&FStateTreeAnyEnumDetails::MakeInstance));
 	PropertyModule.RegisterCustomClassLayout("StateTreeState", FOnGetDetailCustomizationInstance::CreateStatic(&FStateTreeStateDetails::MakeInstance));
 
