StateTree: allow task/evaluator/condition to have separate data for isntance and item in the tree.

diff --git a/Engine/Plugins/Runtime/StateTree/Source/StateTreeEditorModule/Private/StateTreeEditorPropertyBindings.cpp b/Engine/Plugins/Runtime/StateTree/Source/StateTreeEditorModule/Private/StateTreeEditorPropertyBindings.cpp
--- a/Engine/Plugins/Runtime/StateTree/Source/StateTreeEditorModule/Private/StateTreeEditorPropertyBindings.cpp
+++ b/Engine/Plugins/Runtime/StateTree/Source/StateTreeEditorModule/Private/StateTreeEditorPropertyBindings.cpp
@@ -51,7 +51,7 @@ void FStateTreeEditorPropertyBindings::GetPropertyBindingsFor(const FGuid Struct
 		});
 }
 
-void FStateTreeEditorPropertyBindings::RemoveUnusedBindings(const TMap<FGuid, const UScriptStruct*>& ValidStructs)
+void FStateTreeEditorPropertyBindings::RemoveUnusedBindings(const TMap<FGuid, const UStruct*>& ValidStructs)
 {
 	PropertyBindings.RemoveAll([ValidStructs](const FStateTreeEditorPropertyBinding& Binding)
 		{
@@ -63,7 +63,7 @@ void FStateTreeEditorPropertyBindings::RemoveUnusedBindings(const TMap<FGuid, co
 
 			// Remove binding if the target property cannot be resolved, likely renamed property.
 			// TODO: Try to use core redirect to find new name.
-			const UScriptStruct* Struct = ValidStructs.FindChecked(Binding.TargetPath.StructID);
+			const UStruct* Struct = ValidStructs.FindChecked(Binding.TargetPath.StructID);
 
 			FStateTreeBindableStructDesc StructDesc;
 			StructDesc.Struct = Struct;
diff --git a/Engine/Plugins/Runtime/StateTree/Source/StateTreeEditorModule/Private/StateTreeState.cpp b/Engine/Plugins/Runtime/StateTree/Source/StateTreeEditorModule/Private/StateTreeState.cpp
--- a/Engine/Plugins/Runtime/StateTree/Source/StateTreeEditorModule/Private/StateTreeState.cpp
+++ b/Engine/Plugins/Runtime/StateTree/Source/StateTreeEditorModule/Private/StateTreeState.cpp
@@ -76,11 +76,11 @@ void UStateTreeState::PostEditChangeChainProperty(FPropertyChangedChainEvent& Pr
 					const int32 ArrayIndex = PropertyChangedEvent.GetArrayIndex(MemberProperty->GetFName().ToString());
 					if (Evaluators.IsValidIndex(ArrayIndex))
 					{
-						if (FStateTreeEvaluatorBase* Eval = Evaluators[ArrayIndex].Type.GetMutablePtr<FStateTreeEvaluatorBase>())
+						if (FStateTreeEvaluatorBase* Eval = Evaluators[ArrayIndex].Item.GetMutablePtr<FStateTreeEvaluatorBase>())
 						{
-							Eval->ID = FGuid::NewGuid();
 							Eval->Name = FName(Eval->Name.ToString() + TEXT(" Duplicate"));
 						}
+						Evaluators[ArrayIndex].ID = FGuid::NewGuid();
 					}
 				}
 				if (MemberProperty->GetFName() == GET_MEMBER_NAME_CHECKED(UStateTreeState, Tasks))
@@ -88,11 +88,11 @@ void UStateTreeState::PostEditChangeChainProperty(FPropertyChangedChainEvent& Pr
 					const int32 ArrayIndex = PropertyChangedEvent.GetArrayIndex(MemberProperty->GetFName().ToString());
 					if (Tasks.IsValidIndex(ArrayIndex))
 					{
-						if (FStateTreeTaskBase* Task = Tasks[ArrayIndex].Type.GetMutablePtr<FStateTreeTaskBase>())
+						if (FStateTreeTaskBase* Task = Tasks[ArrayIndex].Item.GetMutablePtr<FStateTreeTaskBase>())
 						{
-							Task->ID = FGuid::NewGuid();
 							Task->Name = FName(Task->Name.ToString() + TEXT(" Duplicate"));
 						}
+						Tasks[ArrayIndex].ID = FGuid::NewGuid();
 					}
 				}
 				if (MemberProperty->GetFName() == GET_MEMBER_NAME_CHECKED(UStateTreeState, EnterConditions))
@@ -100,10 +100,7 @@ void UStateTreeState::PostEditChangeChainProperty(FPropertyChangedChainEvent& Pr
 					const int32 ArrayIndex = PropertyChangedEvent.GetArrayIndex(MemberProperty->GetFName().ToString());
 					if (EnterConditions.IsValidIndex(ArrayIndex))
 					{
-						if (FStateTreeConditionBase* Cond = EnterConditions[ArrayIndex].Type.GetMutablePtr<FStateTreeConditionBase>())
-						{
-							Cond->ID = FGuid::NewGuid();
-						}
+						EnterConditions[ArrayIndex].ID = FGuid::NewGuid();
 					}
 				}
 				// TODO: Transition conditions.
@@ -114,36 +111,6 @@ void UStateTreeState::PostEditChangeChainProperty(FPropertyChangedChainEvent& Pr
 
 #endif
 
-const FStateTreeTaskBase* UStateTreeState::GetTaskByID(FGuid InID) const
-{
-	for (const FStateTreeTaskItem& TaskItem : Tasks)
-	{
-		if (const FStateTreeTaskBase* Task = TaskItem.Type.GetPtr<FStateTreeTaskBase>())
-		{
-			if (Task->ID == InID)
-			{
-				return Task;
-			}
-		}
-	}
-	return nullptr;
-}
-
-FStateTreeTaskBase* UStateTreeState::GetTaskByID(FGuid InID)
-{
-	for (FStateTreeTaskItem& TaskItem : Tasks)
-	{
-		if (FStateTreeTaskBase* Task = TaskItem.Type.GetMutablePtr<FStateTreeTaskBase>())
-		{
-			if (Task->ID == InID)
-			{
-				return Task;
-			}
-		}
-	}
-	return nullptr;
-}
-
 UStateTreeState* UStateTreeState::GetNextSiblingState() const
 {
 	if (!Parent)
diff --git a/Engine/Plugins/Runtime/StateTree/Source/StateTreeEditorModule/Public/StateTreeEditorData.h b/Engine/Plugins/Runtime/StateTree/Source/StateTreeEditorModule/Public/StateTreeEditorData.h
--- a/Engine/Plugins/Runtime/StateTree/Source/StateTreeEditorModule/Public/StateTreeEditorData.h
+++ b/Engine/Plugins/Runtime/StateTree/Source/StateTreeEditorModule/Public/StateTreeEditorData.h
@@ -24,13 +24,13 @@ public:
 	const UStateTreeState* GetStateByStructID(const FGuid TargetStructID) const;
 
 	/** Gets the IDs of all bindable structs in the StateTree. */
-	void GetAllStructIDs(TMap<FGuid, const UScriptStruct*>& AllStructs) const;
+	void GetAllStructIDs(TMap<FGuid, const UStruct*>& AllStructs) const;
 
 	/**
 	 * Iterates over all structs that are related to binding
 	 * @param InFunc function called at each item, should return true if visiting is continued or false to stop.
 	 */
-	void VisitHierarchy(TFunctionRef<bool(const UStateTreeState& State, const FGuid& ID, const FName& Name, const UScriptStruct* Struct)> InFunc) const;
+	void VisitHierarchy(TFunctionRef<bool(const UStateTreeState& State, const FGuid& ID, const FName& Name, const UScriptStruct* ItemStruct, const UStruct* InstanceStruct)> InFunc) const;
 
 	// StateTree Builder API
 
diff --git a/Engine/Plugins/Runtime/StateTree/Source/StateTreeEditorModule/Public/StateTreeEditorPropertyBindings.h b/Engine/Plugins/Runtime/StateTree/Source/StateTreeEditorModule/Public/StateTreeEditorPropertyBindings.h
--- a/Engine/Plugins/Runtime/StateTree/Source/StateTreeEditorModule/Public/StateTreeEditorPropertyBindings.h
+++ b/Engine/Plugins/Runtime/StateTree/Source/StateTreeEditorModule/Public/StateTreeEditorPropertyBindings.h
@@ -72,9 +72,9 @@ struct STATETREEEDITORMODULE_API FStateTreeEditorPropertyBindings
 	
 	/**
 	 * Removes bindings which do not point to valid structs IDs.
-	 * @param ValidStructsIDs Set of struct IDs that are currently valid.
+	 * @param ValidStructs Set of struct IDs that are currently valid.
 	 */
-	void RemoveUnusedBindings(const TMap<FGuid, const UScriptStruct*>& ValidStructs);
+	void RemoveUnusedBindings(const TMap<FGuid, const UStruct*>& ValidStructs);
 
 	/** @return array view to all bindings. */
 	TConstArrayView<FStateTreeEditorPropertyBinding> GetBindings() const { return PropertyBindings; }
diff --git a/Engine/Plugins/Runtime/StateTree/Source/StateTreeEditorModule/Public/StateTreeState.h b/Engine/Plugins/Runtime/StateTree/Source/StateTreeEditorModule/Public/StateTreeState.h
--- a/Engine/Plugins/Runtime/StateTree/Source/StateTreeEditorModule/Public/StateTreeState.h
+++ b/Engine/Plugins/Runtime/StateTree/Source/StateTreeEditorModule/Public/StateTreeState.h
@@ -31,7 +31,7 @@ public:
 	UPROPERTY(EditDefaultsOnly, Category = Link)
 	FGuid ID;
 
-	UPROPERTY(EditDefaultsOnly, Category = Transition)
+	UPROPERTY(EditDefaultsOnly, Category = Link)
 	EStateTreeTransitionType Type;
 };
 
@@ -43,10 +43,34 @@ struct STATETREEEDITORMODULE_API FStateTreeEvaluatorItem
 {
 	GENERATED_BODY()
 
-	UPROPERTY(EditDefaultsOnly, Category = Evaluator, meta = (BaseStruct = "StateTreeEvaluator2Base", ExcludeBaseStruct))
-	FInstancedStruct Type;
+	FName GetName() const
+	{
+		if (const FStateTreeEvaluatorBase* Eval = Item.GetPtr<FStateTreeEvaluatorBase>())
+		{
+			return Eval->Name;
+		}
+		return FName();
+	}
+
+	UPROPERTY(EditDefaultsOnly, Category = Evaluator, meta = (BaseStruct = "StateTreeEvaluatorBase", ExcludeBaseStruct))
+	FInstancedStruct Item;
+
+	UPROPERTY(EditDefaultsOnly, Category = Evaluator)
+	FInstancedStruct Instance;
+
+	UPROPERTY(EditDefaultsOnly, Category = Evaluator)
+	FGuid ID;
+};
+
+template <typename T>
+struct STATETREEEDITORMODULE_API TStateTreeEvaluatorItem : public FStateTreeEvaluatorItem
+{
+	typedef T ItemType;
+	FORCEINLINE T& GetItem() { return Item.template GetMutable<T>(); }
+	FORCEINLINE typename T::InstanceDataType& GetInstance() { return Instance.template GetMutable<typename T::InstanceDataType>(); }
 };
 
+
 /**
  * Helper struct for Task details customization (i.e. show summary when collapsed)
  */
@@ -55,10 +79,32 @@ struct STATETREEEDITORMODULE_API FStateTreeTaskItem
 {
 	GENERATED_BODY()
 
-	UPROPERTY(EditDefaultsOnly, Category = Task, meta = (BaseStruct = "StateTreeTask2Base", ExcludeBaseStruct))
-	FInstancedStruct Type;
+	FName GetName() const
+	{
+		if (const FStateTreeTaskBase* Task = Item.GetPtr<FStateTreeTaskBase>())
+		{
+			return Task->Name;
+		}
+		return FName();
+	}
+	
+	UPROPERTY(EditDefaultsOnly, Category = Task, meta = (BaseStruct = "StateTreeTaskBase", ExcludeBaseStruct))
+	FInstancedStruct Item;
+
+	UPROPERTY(EditDefaultsOnly, Category = Task)
+	FInstancedStruct Instance;
+
+	UPROPERTY(EditDefaultsOnly, Category = Task)
+	FGuid ID;
 };
 
+template <typename T>
+struct STATETREEEDITORMODULE_API TStateTreeTaskItem : public FStateTreeTaskItem
+{
+	typedef T ItemType;
+	FORCEINLINE T& GetItem() { return Item.template GetMutable<T>(); }
+	FORCEINLINE typename T::InstanceDataType& GetInstance() { return Instance.template GetMutable<typename T::InstanceDataType>(); }
+};
 
 /**
  * Helper struct for Condition details customization (i.e. show summary when collapsed)
@@ -68,8 +114,22 @@ struct STATETREEEDITORMODULE_API FStateTreeConditionItem
 {
 	GENERATED_BODY()
 
-	UPROPERTY(EditAnywhere, Category = Condition, meta = (BaseStruct = "StateTreeConditionBase", ExcludeBaseStruct))
-	FInstancedStruct Type;
+	UPROPERTY(EditDefaultsOnly, Category = Condition, meta = (BaseStruct = "StateTreeConditionBase", ExcludeBaseStruct))
+	FInstancedStruct Item;
+
+	UPROPERTY(EditDefaultsOnly, Category = Condition)
+	FInstancedStruct Instance;
+
+	UPROPERTY(EditDefaultsOnly, Category = Condition)
+	FGuid ID;
+};
+
+template <typename T>
+struct STATETREEEDITORMODULE_API TStateTreeConditionItem : public FStateTreeConditionItem
+{
+	typedef T ItemType;
+	FORCEINLINE T& GetItem() { return Item.template GetMutable<T>(); }
+	FORCEINLINE typename T::InstanceDataType& GetInstance() { return Instance.template GetMutable<typename T::InstanceDataType>(); }
 };
 
 /**
@@ -84,13 +144,17 @@ struct STATETREEEDITORMODULE_API FStateTreeTransition
 	FStateTreeTransition(const EStateTreeTransitionEvent InEvent, const EStateTreeTransitionType InType, const UStateTreeState* InState = nullptr);
 
 	template<typename T, typename... TArgs>
-	T& AddCondition(TArgs&&... InArgs)
+	TStateTreeConditionItem<T>& AddCondition(TArgs&&... InArgs)
 	{
 		FStateTreeConditionItem& CondItem = Conditions.AddDefaulted_GetRef();
-		CondItem.Type.InitializeAs<T>(Forward<TArgs>(InArgs)...);
-		T& Condition = CondItem.Type.GetMutable<T>();
-		Condition.ID = FGuid::NewGuid();
-		return Condition;
+		CondItem.ID = FGuid::NewGuid();
+		CondItem.Item.InitializeAs<T>(Forward<TArgs>(InArgs)...);
+		T& Cond = CondItem.Item.GetMutable<T>();
+		if (const UScriptStruct* InstanceType = Cast<const UScriptStruct>(Cond.GetInstanceDataType()))
+		{
+			CondItem.Instance.InitializeAs(InstanceType);
+		}
+		return static_cast<TStateTreeConditionItem<T>&>(CondItem);
 	}
 
 	UPROPERTY(EditDefaultsOnly, Category = Transition)
@@ -123,9 +187,6 @@ public:
 	virtual void PostEditChangeChainProperty(FPropertyChangedChainEvent& PropertyChangedEvent) override;
 #endif
 
-	const FStateTreeTaskBase* GetTaskByID(FGuid ID) const;
-	FStateTreeTaskBase* GetTaskByID(FGuid ID);
-
 	UStateTreeState* GetNextSiblingState() const;
 
 	// StateTree Builder API
@@ -146,11 +207,17 @@ public:
 	 * @return reference to the new condition. 
 	 */
 	template<typename T, typename... TArgs>
-	T& AddEnterCondition(TArgs&&... InArgs)
+	TStateTreeConditionItem<T>& AddEnterCondition(TArgs&&... InArgs)
 	{
 		FStateTreeConditionItem& CondItem = EnterConditions.AddDefaulted_GetRef();
-		CondItem.Type.InitializeAs<T>(Forward<TArgs>(InArgs)...);
-		return CondItem.Type.GetMutable<T>();
+		CondItem.ID = FGuid::NewGuid();
+		CondItem.Item.InitializeAs<T>(Forward<TArgs>(InArgs)...);
+		T& Cond = CondItem.Item.GetMutable<T>();
+		if (const UScriptStruct* InstanceType = Cast<const UScriptStruct>(Cond.GetInstanceDataType()))
+		{
+			CondItem.Instance.InitializeAs(InstanceType);
+		}
+		return static_cast<TStateTreeConditionItem<T>&>(CondItem);
 	}
 
 	/**
@@ -158,13 +225,17 @@ public:
 	 * @return reference to the new Task. 
 	 */
 	template<typename T, typename... TArgs>
-	T& AddTask(TArgs&&... InArgs)
+	TStateTreeTaskItem<T>& AddTask(TArgs&&... InArgs)
 	{
 		FStateTreeTaskItem& TaskItem = Tasks.AddDefaulted_GetRef();
-		TaskItem.Type.InitializeAs<T>(Forward<TArgs>(InArgs)...);
-		T& Task = TaskItem.Type.GetMutable<T>();
-		Task.ID = FGuid::NewGuid();
-		return Task;
+		TaskItem.ID = FGuid::NewGuid();
+		TaskItem.Item.InitializeAs<T>(Forward<TArgs>(InArgs)...);
+		T& Task = TaskItem.Item.GetMutable<T>();
+		if (const UScriptStruct* InstanceType = Cast<const UScriptStruct>(Task.GetInstanceDataType()))
+		{
+			TaskItem.Instance.InitializeAs(InstanceType);
+		}
+		return static_cast<TStateTreeTaskItem<T>&>(TaskItem);
 	}
 
 	/**
@@ -172,13 +243,17 @@ public:
 	 * @return reference to the new Evaluator. 
 	 */
 	template<typename T, typename... TArgs>
-    T& AddEvaluator(TArgs&&... InArgs)
+    TStateTreeEvaluatorItem<T>& AddEvaluator(TArgs&&... InArgs)
 	{
 		FStateTreeEvaluatorItem& EvalItem = Evaluators.AddDefaulted_GetRef();
-		EvalItem.Type.InitializeAs<T>(Forward<TArgs>(InArgs)...);
-		T& Eval = EvalItem.Type.GetMutable<T>();
-		Eval.ID = FGuid::NewGuid();
-		return Eval;
+		EvalItem.ID = FGuid::NewGuid();
+		EvalItem.Item.InitializeAs<T>(Forward<TArgs>(InArgs)...);
+		T& Eval = EvalItem.Item.GetMutable<T>();
+		if (const UScriptStruct* InstanceType = Cast<const UScriptStruct>(Eval.GetInstanceDataType()))
+		{
+			EvalItem.Instance.InitializeAs(InstanceType);
+		}
+		return static_cast<TStateTreeEvaluatorItem<T>&>(EvalItem);
 	}
 
 	/**
@@ -192,7 +267,7 @@ public:
 
 	// ~StateTree Builder API
 
-	UPROPERTY(EditDefaultsOnly, Category = State)
+	UPROPERTY(EditDefaultsOnly, Category = "State")
 	FName Name;
 
 	UPROPERTY()
diff --git a/Engine/Plugins/Runtime/StateTree/Source/StateTreeModule/Private/Components/StateTreeComponent.cpp b/Engine/Plugins/Runtime/StateTree/Source/StateTreeModule/Private/Components/StateTreeComponent.cpp
--- a/Engine/Plugins/Runtime/StateTree/Source/StateTreeModule/Private/Components/StateTreeComponent.cpp
+++ b/Engine/Plugins/Runtime/StateTree/Source/StateTreeModule/Private/Components/StateTreeComponent.cpp
@@ -8,6 +8,7 @@
 #include "StateTree.h"
 #include "StateTreeEvaluatorBase.h"
 #include "StateTreeTaskBase.h"
+#include "StateTreeConditionBase.h"
 #include "AIController.h"
 #include "GameFramework/CharacterMovementComponent.h"
 
@@ -21,7 +22,9 @@
 bool UBrainComponentStateTreeSchema::IsStructAllowed(const UScriptStruct* InScriptStruct) const
 {
 	// @todo: Proper base class for actor nodes.
-	return InScriptStruct->IsChildOf(FStateTreeEvaluatorBase::StaticStruct()) || InScriptStruct->IsChildOf(FStateTreeTaskBase::StaticStruct());
+	return InScriptStruct->IsChildOf(FStateTreeEvaluatorBase::StaticStruct())
+			|| InScriptStruct->IsChildOf(FStateTreeTaskBase::StaticStruct())
+			|| InScriptStruct->IsChildOf(FStateTreeConditionBase::StaticStruct());
 }
 
 
diff --git a/Engine/Plugins/Runtime/StateTree/Source/StateTreeModule/Private/Conditions/StateTreeCondition_Common.cpp b/Engine/Plugins/Runtime/StateTree/Source/StateTreeModule/Private/Conditions/StateTreeCondition_Common.cpp
--- a/Engine/Plugins/Runtime/StateTree/Source/StateTreeModule/Private/Conditions/StateTreeCondition_Common.cpp
+++ b/Engine/Plugins/Runtime/StateTree/Source/StateTreeModule/Private/Conditions/StateTreeCondition_Common.cpp
@@ -1,11 +1,14 @@
 // Copyright Epic Games, Inc. All Rights Reserved.
 
 #include "Conditions/StateTreeCondition_Common.h"
+
+#include "StateTreeExecutionContext.h"
 #include "StateTreePropertyBindings.h"
 
 #define LOCTEXT_NAMESPACE "StateTreeEditor"
 
-namespace UE { namespace StateTree { namespace Conditions {
+namespace UE::StateTree::Conditions
+{
 
 #if WITH_EDITOR
 FText GetOperatorText(const EGenericAICheck Operator)
@@ -69,13 +72,35 @@ bool CompareNumbers(const T Left, const T Right, const EGenericAICheck Operator)
 	return false;
 }
 
-}}} // UE::StateTree::Conditions
+} // UE::StateTree::Conditions
+
+
+//----------------------------------------------------------------------//
+//  FStateTreeCondition_CompareInt
+//----------------------------------------------------------------------//
+
+bool FStateTreeCondition_CompareInt::Link(FStateTreeLinker& Linker)
+{
+	Linker.LinkInstanceDataProperty(LeftHandle, STATETREE_INSTANCEDATA_PROPERTY(FStateTreeCondition_CompareIntInstanceData, Left));
+	Linker.LinkInstanceDataProperty(RightHandle, STATETREE_INSTANCEDATA_PROPERTY(FStateTreeCondition_CompareIntInstanceData, Right));
+
+	return true;
+}
+
+bool FStateTreeCondition_CompareInt::TestCondition(FStateTreeExecutionContext& Context) const
+{
+	const int32 Left = Context.GetInstanceData(LeftHandle);
+	const int32 Right = Context.GetInstanceData(RightHandle);
+	const bool bResult = UE::StateTree::Conditions::CompareNumbers<int32>(Left, Right, Operator);
+	return bResult ^ bInvert;
+}
 
 #if WITH_EDITOR
-FText FStateTreeCondition_CompareInt::GetDescription(const IStateTreeBindingLookup& BindingLookup) const
+FText FStateTreeCondition_CompareInt::GetDescription(const FGuid& ID, FConstStructView InstanceData, const IStateTreeBindingLookup& BindingLookup) const
 {
-	const FStateTreeEditorPropertyPath LeftPath(ID, GET_MEMBER_NAME_STRING_CHECKED(FStateTreeCondition_CompareInt, Left));
-	const FStateTreeEditorPropertyPath RightPath(ID, GET_MEMBER_NAME_STRING_CHECKED(FStateTreeCondition_CompareInt, Right));
+	const FStateTreeCondition_CompareIntInstanceData& Instance = InstanceData.Get<FStateTreeCondition_CompareIntInstanceData>();
+	const FStateTreeEditorPropertyPath LeftPath(ID, GET_MEMBER_NAME_STRING_CHECKED(FStateTreeCondition_CompareIntInstanceData, Left));
+	const FStateTreeEditorPropertyPath RightPath(ID, GET_MEMBER_NAME_STRING_CHECKED(FStateTreeCondition_CompareIntInstanceData, Right));
 
 	FText InvertText;
 	if (bInvert)
@@ -90,7 +115,7 @@ FText FStateTreeCondition_CompareInt::GetDescription(const IStateTreeBindingLook
 	}
 	else
 	{
-		LeftText = FText::AsNumber(Left);
+		LeftText = LOCTEXT("NotBound", "Not Bound");
 	}
 
 	FText OperatorText = UE::StateTree::Conditions::GetOperatorText(Operator);
@@ -102,25 +127,41 @@ FText FStateTreeCondition_CompareInt::GetDescription(const IStateTreeBindingLook
 	}
 	else
 	{
-		RightText = FText::AsNumber(Right);
+		RightText = FText::AsNumber(Instance.Right);
 	}
 
 	return FText::Format(LOCTEXT("CompareIntDesc", "{0} <Details.Bold>{1}</> {2} <Details.Bold>{3}</>"), InvertText, LeftText, OperatorText, RightText);
 }
 #endif
 
-bool FStateTreeCondition_CompareInt::TestCondition() const
+
+//----------------------------------------------------------------------//
+//  FStateTreeCondition_CompareFloat
+//----------------------------------------------------------------------//
+
+bool FStateTreeCondition_CompareFloat::Link(FStateTreeLinker& Linker)
 {
-	const bool bResult = UE::StateTree::Conditions::CompareNumbers<int32>(Left, Right, Operator);
+	Linker.LinkInstanceDataProperty(LeftHandle, STATETREE_INSTANCEDATA_PROPERTY(FStateTreeCondition_CompareFloatInstanceData, Left));
+	Linker.LinkInstanceDataProperty(RightHandle, STATETREE_INSTANCEDATA_PROPERTY(FStateTreeCondition_CompareFloatInstanceData, Right));
+
+	return true;
+}
+
+bool FStateTreeCondition_CompareFloat::TestCondition(FStateTreeExecutionContext& Context) const
+{
+	const float Left = Context.GetInstanceData(LeftHandle);
+	const float Right = Context.GetInstanceData(RightHandle);
+	const bool bResult = UE::StateTree::Conditions::CompareNumbers<float>(Left, Right, Operator);
 	return bResult ^ bInvert;
 }
 
 
 #if WITH_EDITOR
-FText FStateTreeCondition_CompareFloat::GetDescription(const IStateTreeBindingLookup& BindingLookup) const
+FText FStateTreeCondition_CompareFloat::GetDescription(const FGuid& ID, FConstStructView InstanceData, const IStateTreeBindingLookup& BindingLookup) const
 {
-	const FStateTreeEditorPropertyPath LeftPath(ID, GET_MEMBER_NAME_STRING_CHECKED(FStateTreeCondition_CompareFloat, Left));
-	const FStateTreeEditorPropertyPath RightPath(ID, GET_MEMBER_NAME_STRING_CHECKED(FStateTreeCondition_CompareFloat, Right));
+	const FStateTreeCondition_CompareFloatInstanceData& Instance = InstanceData.Get<FStateTreeCondition_CompareFloatInstanceData>();
+	const FStateTreeEditorPropertyPath LeftPath(ID, GET_MEMBER_NAME_STRING_CHECKED(FStateTreeCondition_CompareFloatInstanceData, Left));
+	const FStateTreeEditorPropertyPath RightPath(ID, GET_MEMBER_NAME_STRING_CHECKED(FStateTreeCondition_CompareFloatInstanceData, Right));
 
 	FText InvertText;
 	if (bInvert)
@@ -135,7 +176,7 @@ FText FStateTreeCondition_CompareFloat::GetDescription(const IStateTreeBindingLo
 	}
 	else
 	{
-		LeftText = FText::AsNumber(Left);
+		LeftText = LOCTEXT("NotBound", "Not Bound");
 	}
 
 	FText OperatorText = UE::StateTree::Conditions::GetOperatorText(Operator);
@@ -147,25 +188,39 @@ FText FStateTreeCondition_CompareFloat::GetDescription(const IStateTreeBindingLo
 	}
 	else
 	{
-		RightText = FText::AsNumber(Right);
+		RightText = FText::AsNumber(Instance.Right);
 	}
 
 	return FText::Format(LOCTEXT("CompareFloatDesc", "{0} <Details.Bold>{1}</> {2} <Details.Bold>{3}</>"), InvertText, LeftText, OperatorText, RightText);
 }
 #endif
 
-bool FStateTreeCondition_CompareFloat::TestCondition() const
+
+//----------------------------------------------------------------------//
+//  FStateTreeCondition_CompareBool
+//----------------------------------------------------------------------//
+
+bool FStateTreeCondition_CompareBool::Link(FStateTreeLinker& Linker)
 {
-	const bool bResult = UE::StateTree::Conditions::CompareNumbers<float>(Left, Right, Operator);
-	return bResult ^ bInvert;
+	Linker.LinkInstanceDataProperty(LeftHandle, STATETREE_INSTANCEDATA_PROPERTY(FStateTreeCondition_CompareBoolInstanceData, bLeft));
+	Linker.LinkInstanceDataProperty(RightHandle, STATETREE_INSTANCEDATA_PROPERTY(FStateTreeCondition_CompareBoolInstanceData, bRight));
+
+	return true;
 }
 
+bool FStateTreeCondition_CompareBool::TestCondition(FStateTreeExecutionContext& Context) const
+{
+	const bool bLeft = Context.GetInstanceData(LeftHandle);
+	const bool bRight = Context.GetInstanceData(RightHandle);
+	return (bLeft == bRight) ^ bInvert;
+}
 
 #if WITH_EDITOR
-FText FStateTreeCondition_CompareBool::GetDescription(const IStateTreeBindingLookup& BindingLookup) const
+FText FStateTreeCondition_CompareBool::GetDescription(const FGuid& ID, FConstStructView InstanceData, const IStateTreeBindingLookup& BindingLookup) const
 {
-	const FStateTreeEditorPropertyPath LeftPath(ID, GET_MEMBER_NAME_STRING_CHECKED(FStateTreeCondition_CompareBool, bLeft));
-	const FStateTreeEditorPropertyPath RightPath(ID, GET_MEMBER_NAME_STRING_CHECKED(FStateTreeCondition_CompareBool, bRight));
+	const FStateTreeCondition_CompareBoolInstanceData& Instance = InstanceData.Get<FStateTreeCondition_CompareBoolInstanceData>();
+	const FStateTreeEditorPropertyPath LeftPath(ID, GET_MEMBER_NAME_STRING_CHECKED(FStateTreeCondition_CompareBoolInstanceData, bLeft));
+	const FStateTreeEditorPropertyPath RightPath(ID, GET_MEMBER_NAME_STRING_CHECKED(FStateTreeCondition_CompareBoolInstanceData, bRight));
 
 	FText InvertText;
 	if (bInvert)
@@ -180,7 +235,7 @@ FText FStateTreeCondition_CompareBool::GetDescription(const IStateTreeBindingLoo
 	}
 	else
 	{
-		LeftText = FText::FromString(LexToSanitizedString(bLeft));
+		LeftText = LOCTEXT("NotBound", "Not Bound");
 	}
 
 	FText RightText;
@@ -190,24 +245,39 @@ FText FStateTreeCondition_CompareBool::GetDescription(const IStateTreeBindingLoo
 	}
 	else
 	{
-		RightText = FText::FromString(LexToSanitizedString(bRight));
+		RightText = FText::FromString(LexToSanitizedString(Instance.bRight));
 	}
 
 	return FText::Format(LOCTEXT("CompareBoolDesc", "{0} <Details.Bold>{1}</> is <Details.Bold>{2}</>"), InvertText, LeftText, RightText);
 }
 #endif
 
-bool FStateTreeCondition_CompareBool::TestCondition() const
+
+//----------------------------------------------------------------------//
+//  FStateTreeCondition_CompareEnum
+//----------------------------------------------------------------------//
+
+bool FStateTreeCondition_CompareEnum::Link(FStateTreeLinker& Linker)
 {
-	return (bLeft == bRight) ^ bInvert;
+	Linker.LinkInstanceDataProperty(LeftHandle, STATETREE_INSTANCEDATA_PROPERTY(FStateTreeCondition_CompareEnumInstanceData, Left));
+	Linker.LinkInstanceDataProperty(RightHandle, STATETREE_INSTANCEDATA_PROPERTY(FStateTreeCondition_CompareEnumInstanceData, Right));
+
+	return true;
 }
 
+bool FStateTreeCondition_CompareEnum::TestCondition(FStateTreeExecutionContext& Context) const
+{
+	const FStateTreeAnyEnum Left = Context.GetInstanceData(LeftHandle);
+	const FStateTreeAnyEnum Right = Context.GetInstanceData(RightHandle);
+	return (Left == Right) ^ bInvert;
+}
 
 #if WITH_EDITOR
-FText FStateTreeCondition_CompareEnum::GetDescription(const IStateTreeBindingLookup& BindingLookup) const
+FText FStateTreeCondition_CompareEnum::GetDescription(const FGuid& ID, FConstStructView InstanceData, const IStateTreeBindingLookup& BindingLookup) const
 {
-	const FStateTreeEditorPropertyPath LeftPath(ID, GET_MEMBER_NAME_STRING_CHECKED(FStateTreeCondition_CompareEnum, Left));
-	const FStateTreeEditorPropertyPath RightPath(ID, GET_MEMBER_NAME_STRING_CHECKED(FStateTreeCondition_CompareEnum, Right));
+	const FStateTreeCondition_CompareEnumInstanceData& Instance = InstanceData.Get<FStateTreeCondition_CompareEnumInstanceData>();
+	const FStateTreeEditorPropertyPath LeftPath(ID, GET_MEMBER_NAME_STRING_CHECKED(FStateTreeCondition_CompareEnumInstanceData, Left));
+	const FStateTreeEditorPropertyPath RightPath(ID, GET_MEMBER_NAME_STRING_CHECKED(FStateTreeCondition_CompareEnumInstanceData, Right));
 
 	FText InvertText;
 	if (bInvert)
@@ -222,14 +292,7 @@ FText FStateTreeCondition_CompareEnum::GetDescription(const IStateTreeBindingLoo
 	}
 	else
 	{
-		if (Left.Enum)
-		{
-			LeftText = Left.Enum->GetDisplayNameTextByValue(int64(Left.Value));
-		}
-		else
-		{
-			LeftText = LOCTEXT("Invalid", "Invalid");
-		}
+		LeftText = LOCTEXT("NotBound", "Not Bound");
 	}
 
 	FText RightText;
@@ -239,9 +302,9 @@ FText FStateTreeCondition_CompareEnum::GetDescription(const IStateTreeBindingLoo
 	}
 	else
 	{
-		if (Right.Enum)
+		if (Instance.Right.Enum)
 		{
-			RightText = Right.Enum->GetDisplayNameTextByValue(int64(Right.Value));
+			RightText = Instance.Right.Enum->GetDisplayNameTextByValue(int64(Instance.Right.Value));
 		}
 		else
 		{
@@ -252,15 +315,17 @@ FText FStateTreeCondition_CompareEnum::GetDescription(const IStateTreeBindingLoo
 	return FText::Format(LOCTEXT("EnumEqualsDesc", "{0} <Details.Bold>{1}</> is <Details.Bold>{2}</>"), InvertText, LeftText, RightText);
 }
 
-void FStateTreeCondition_CompareEnum::OnBindingChanged(const FStateTreeEditorPropertyPath& SourcePath, const FStateTreeEditorPropertyPath& TargetPath, const IStateTreeBindingLookup& BindingLookup)
+void FStateTreeCondition_CompareEnum::OnBindingChanged(const FGuid& ID, FStructView InstanceData, const FStateTreeEditorPropertyPath& SourcePath, const FStateTreeEditorPropertyPath& TargetPath, const IStateTreeBindingLookup& BindingLookup)
 {
 	if (!TargetPath.IsValid())
 	{
 		return;
 	}
 
+	FStateTreeCondition_CompareEnumInstanceData& Instance = InstanceData.GetMutable<FStateTreeCondition_CompareEnumInstanceData>();
+
 	// Left has changed, update enums from the leaf property.
-	if (TargetPath.Path.Last() == GET_MEMBER_NAME_STRING_CHECKED(FStateTreeCondition_CompareEnum, Left))
+	if (TargetPath.Path.Last() == GET_MEMBER_NAME_STRING_CHECKED(FStateTreeCondition_CompareEnumInstanceData, Left))
 	{
 		if (const FProperty* LeafProperty = BindingLookup.GetPropertyPathLeafProperty(SourcePath))
 		{
@@ -275,36 +340,58 @@ void FStateTreeCondition_CompareEnum::OnBindingChanged(const FStateTreeEditorPro
 				NewEnum = EnumProperty->GetEnum();
 			}
 
-			if (Left.Enum != NewEnum)
+			if (Instance.Left.Enum != NewEnum)
 			{
-				Left.Initialize(NewEnum);
+				Instance.Left.Initialize(NewEnum);
 			}
 		}
 		else
 		{
-			Left.Initialize(nullptr);
+			Instance.Left.Initialize(nullptr);
 		}
 
-		if (Right.Enum != Left.Enum)
+		if (Instance.Right.Enum != Instance.Left.Enum)
 		{
-			Right.Initialize(Left.Enum);
+			Instance.Right.Initialize(Instance.Left.Enum);
 		}
 	}
 }
 
 #endif
-bool FStateTreeCondition_CompareEnum::TestCondition() const
+
+
+//----------------------------------------------------------------------//
+//  FStateTreeCondition_CompareDistance
+//----------------------------------------------------------------------//
+
+bool FStateTreeCondition_CompareDistance::Link(FStateTreeLinker& Linker)
 {
-	return (Left == Right) ^ bInvert;
+	Linker.LinkInstanceDataProperty(SourceHandle, STATETREE_INSTANCEDATA_PROPERTY(FStateTreeCondition_CompareDistanceInstanceData, Source));
+	Linker.LinkInstanceDataProperty(TargetHandle, STATETREE_INSTANCEDATA_PROPERTY(FStateTreeCondition_CompareDistanceInstanceData, Target));
+	Linker.LinkInstanceDataProperty(DistanceHandle, STATETREE_INSTANCEDATA_PROPERTY(FStateTreeCondition_CompareDistanceInstanceData, Distance));
+
+	return true;
 }
 
+bool FStateTreeCondition_CompareDistance::TestCondition(FStateTreeExecutionContext& Context) const
+{
+	const FVector& Source = Context.GetInstanceData(SourceHandle);
+	const FVector& Target = Context.GetInstanceData(TargetHandle);
+	const float Distance = Context.GetInstanceData(DistanceHandle);
+
+	const float Left = FVector::DistSquared(Source, Target);
+	const float Right = FMath::Square(Distance);
+	const bool bResult = UE::StateTree::Conditions::CompareNumbers<float>(Left, Right, Operator);
+	return bResult ^ bInvert;
+}
 
 #if WITH_EDITOR
-FText FStateTreeCondition_CompareDistance::GetDescription(const IStateTreeBindingLookup& BindingLookup) const
+FText FStateTreeCondition_CompareDistance::GetDescription(const FGuid& ID, FConstStructView InstanceData, const IStateTreeBindingLookup& BindingLookup) const
 {
-	const FStateTreeEditorPropertyPath SourcePath(ID, GET_MEMBER_NAME_STRING_CHECKED(FStateTreeCondition_CompareDistance, Source));
-	const FStateTreeEditorPropertyPath TargetPath(ID, GET_MEMBER_NAME_STRING_CHECKED(FStateTreeCondition_CompareDistance, Target));
-	const FStateTreeEditorPropertyPath DistancePath(ID, GET_MEMBER_NAME_STRING_CHECKED(FStateTreeCondition_CompareDistance, Distance));
+	const FStateTreeCondition_CompareDistanceInstanceData& Instance = InstanceData.Get<FStateTreeCondition_CompareDistanceInstanceData>();
+	const FStateTreeEditorPropertyPath SourcePath(ID, GET_MEMBER_NAME_STRING_CHECKED(FStateTreeCondition_CompareDistanceInstanceData, Source));
+	const FStateTreeEditorPropertyPath TargetPath(ID, GET_MEMBER_NAME_STRING_CHECKED(FStateTreeCondition_CompareDistanceInstanceData, Target));
+	const FStateTreeEditorPropertyPath DistancePath(ID, GET_MEMBER_NAME_STRING_CHECKED(FStateTreeCondition_CompareDistanceInstanceData, Distance));
 
 	FText InvertText;
 	if (bInvert)
@@ -319,7 +406,7 @@ FText FStateTreeCondition_CompareDistance::GetDescription(const IStateTreeBindin
 	}
 	else
 	{
-		SourceText = Source.ToCompactText();
+		SourceText = LOCTEXT("NotBound", "Not Bound");
 	}
 
 	FText TargetText;
@@ -329,7 +416,7 @@ FText FStateTreeCondition_CompareDistance::GetDescription(const IStateTreeBindin
 	}
 	else
 	{
-		TargetText = Target.ToCompactText();
+		TargetText = Instance.Target.ToCompactText();
 	}
 
 	FText OperatorText = UE::StateTree::Conditions::GetOperatorText(Operator);
@@ -341,19 +428,48 @@ FText FStateTreeCondition_CompareDistance::GetDescription(const IStateTreeBindin
 	}
 	else
 	{
-		DistanceText = FText::AsNumber(Distance);
+		DistanceText = FText::AsNumber(Instance.Distance);
 	}
 
 	return FText::Format(LOCTEXT("CompareDistanceDesc", "{0} Distance <Details.Subdued>from</> <Details.Bold>{1}</> <Details.Subdued>to</> <Details.Bold>{2}</> {3} <Details.Bold>{4}</>"), InvertText, SourceText, TargetText, OperatorText, DistanceText);
 }
 #endif
 
-bool FStateTreeCondition_CompareDistance::TestCondition() const
+//----------------------------------------------------------------------//
+//  FStateTreeCondition_Random
+//----------------------------------------------------------------------//
+
+bool FStateTreeCondition_Random::Link(FStateTreeLinker& Linker)
 {
-	const float Left = FVector::DistSquared(Source, Target);
-	const float Right = FMath::Square(Distance);
-	const bool bResult = UE::StateTree::Conditions::CompareNumbers<float>(Left, Right, Operator);
-	return bResult ^ bInvert;
+	Linker.LinkInstanceDataProperty(ThresholdHandle, STATETREE_INSTANCEDATA_PROPERTY(FStateTreeCondition_RandomInstanceData, Threshold));
+
+	return true;
+}
+
+bool FStateTreeCondition_Random::TestCondition(FStateTreeExecutionContext& Context) const
+{
+	const float Threshold = Context.GetInstanceData(ThresholdHandle);
+	return FMath::FRandRange(0.0f, 1.0f) < Threshold;
+}
+
+#if WITH_EDITOR
+FText FStateTreeCondition_Random::GetDescription(const FGuid& ID, FConstStructView InstanceData, const IStateTreeBindingLookup& BindingLookup) const
+{
+	const FStateTreeCondition_RandomInstanceData& Instance = InstanceData.Get<FStateTreeCondition_RandomInstanceData>();
+	const FStateTreeEditorPropertyPath ThresholdPath(ID, GET_MEMBER_NAME_STRING_CHECKED(FStateTreeCondition_RandomInstanceData, Threshold));
+
+	FText ThresholdText;
+	if (const FStateTreeEditorPropertyPath* Binding = BindingLookup.GetPropertyBindingSource(ThresholdPath))
+	{
+		ThresholdText = BindingLookup.GetPropertyPathDisplayName(*Binding);
+	}
+	else
+	{
+		ThresholdText = FText::AsPercent(Instance.Threshold);
+	}
+
+	return FText::Format(LOCTEXT("RandomDesc", "<Details.Bold>Random</> {1}"), ThresholdText);
 }
+#endif
 
 #undef LOCTEXT_NAMESPACE
diff --git a/Engine/Plugins/Runtime/StateTree/Source/StateTreeModule/Private/StateTree.cpp b/Engine/Plugins/Runtime/StateTree/Source/StateTreeModule/Private/StateTree.cpp
--- a/Engine/Plugins/Runtime/StateTree/Source/StateTreeModule/Private/StateTree.cpp
+++ b/Engine/Plugins/Runtime/StateTree/Source/StateTreeModule/Private/StateTree.cpp
@@ -4,6 +4,7 @@
 #include "StateTreeEvaluatorBase.h"
 #include "StateTreeTaskBase.h"
 #include "CoreMinimal.h"
+#include "StateTreeConditionBase.h"
 #include "StateTreeDelegates.h"
 
 UStateTree::UStateTree(const FObjectInitializer& ObjectInitializer)
@@ -22,18 +23,18 @@ bool UStateTree::IsValidStateTree() const
 void UStateTree::ResetBaked()
 {
 	States.Reset();
-	Conditions.Reset();
 	Transitions.Reset();
 
-	RuntimeStorageItems.Reset();
-	RuntimeStorageStruct = nullptr;
-	RuntimeStorageOffsets.Reset();
-	RuntimeStorageDefaultValue.Reset();
-	ExternalItems.Reset();
+	Items.Reset();
+	Instances.Reset();
+	InstanceStorageStruct = nullptr;
+	InstanceStorageOffsets.Reset();
+	InstanceStorageDefaultValue.Reset();
+	ExternalDataDescs.Reset();
 	PropertyBindings.Reset();
 
-	NumLinkedItems = 0;
-	ExternalItemBaseIndex = 0;
+	NumDataViews = 0;
+	ExternalDataBaseIndex = 0;
 }
 
 void UStateTree::PostEditChangeProperty(struct FPropertyChangedEvent& PropertyChangedEvent)
@@ -67,16 +68,16 @@ void UStateTree::PostLoad()
 	Link();
 	
 #if WITH_EDITOR
-	InitRuntimeStorage();
+	InitInstanceStorage();
 #else
 	// Item offsets still need to be calculated in non editor target since the struct sizes might be different.
-	checkf(RuntimeStorageOffsets.Num() == 0, TEXT("RuntimeStorageOffsets is transient and should only be computed once."));
+	checkf(InstanceStorageOffsets.Num() == 0, TEXT("RuntimeStorageOffsets is transient and should only be computed once."));
 
-	for (TFieldIterator<FProperty> PropertyIt(RuntimeStorageStruct); PropertyIt; ++PropertyIt)
+	for (TFieldIterator<FProperty> PropertyIt(InstanceStorageStruct); PropertyIt; ++PropertyIt)
 	{
 		FStructProperty* StructProperty = CastField<FStructProperty>(*PropertyIt);
-		checkf(StructProperty, TEXT("RuntimeStorageStruct is expected to only contain Struct properties"));
-		RuntimeStorageOffsets.Emplace(StructProperty->Struct, StructProperty->GetOffset_ForInternal());
+		checkf(StructProperty, TEXT("InstanceStorageStruct is expected to only contain Struct properties"));
+		InstanceStorageOffsets.Emplace(StructProperty->Struct, StructProperty->GetOffset_ForInternal());
 	}
 #endif // WITH_EDITOR
 }
@@ -86,7 +87,7 @@ void UStateTree::BeginDestroy()
 	Super::BeginDestroy();
 	
 	// Destroy the runtime storage before the UScriptStruct it uses gets removed.
-	RuntimeStorageDefaultValue.Reset();
+	InstanceStorageDefaultValue.Reset();
 }
 
 void UStateTree::ResolvePropertyPaths()
@@ -99,36 +100,43 @@ void UStateTree::Link()
 {
 	FStateTreeLinker Linker;
 
-	ExternalItemBaseIndex = PropertyBindings.GetSourceStructNum();
-	Linker.SetItemBaseIndex(ExternalItemBaseIndex);
+	ExternalDataBaseIndex = PropertyBindings.GetSourceStructNum();
+	Linker.SetExternalDataBaseIndex(ExternalDataBaseIndex);
 	
-	for (FInstancedStruct& RuntimeItem : RuntimeStorageItems)
+	for (FInstancedStruct& Item : Items)
 	{
-		if (FStateTreeEvaluatorBase* Eval = RuntimeItem.GetMutablePtr<FStateTreeEvaluatorBase>())
+		if (FStateTreeEvaluatorBase* Eval = Item.GetMutablePtr<FStateTreeEvaluatorBase>())
 		{
+			Linker.SetCurrentInstanceDataType(Eval->GetInstanceDataType(), Eval->DataViewIndex);
 			Eval->Link(Linker);
 		}
-		else if (FStateTreeTaskBase* Task = RuntimeItem.GetMutablePtr<FStateTreeTaskBase>())
+		else if (FStateTreeTaskBase* Task = Item.GetMutablePtr<FStateTreeTaskBase>())
 		{
+			Linker.SetCurrentInstanceDataType(Task->GetInstanceDataType(), Task->DataViewIndex);
 			Task->Link(Linker);
 		}
+		else if (FStateTreeConditionBase* Cond = Item.GetMutablePtr<FStateTreeConditionBase>())
+		{
+			Linker.SetCurrentInstanceDataType(Cond->GetInstanceDataType(), Cond->DataViewIndex);
+			Cond->Link(Linker);
+		}
 	}
 
-	ExternalItems = Linker.GetItemDescs();
+	ExternalDataDescs = Linker.GetExternalDataDescs();
 	
-	NumLinkedItems = ExternalItemBaseIndex + ExternalItems.Num();
+	NumDataViews = ExternalDataBaseIndex + ExternalDataDescs.Num();
 }
 
-void UStateTree::InitRuntimeStorage()
+void UStateTree::InitInstanceStorage()
 {
-	RuntimeStorageStruct = nullptr;
-	RuntimeStorageOffsets.Reset();
-	RuntimeStorageDefaultValue.Reset();
+	InstanceStorageStruct = nullptr;
+	InstanceStorageOffsets.Reset();
+	InstanceStorageDefaultValue.Reset();
 
 	// Check that the items are valid before trying to create the type.
 	// The structs can become invalid i.e. because of missing type.
 	bool bValid = true;
-	for (const FInstancedStruct& Item : RuntimeStorageItems)
+	for (const FInstancedStruct& Item : Instances)
 	{
 		if (!Item.IsValid())
 		{
@@ -141,7 +149,7 @@ void UStateTree::InitRuntimeStorage()
 		return;
 	}
 
-	const FString StructName = GetName() + TEXT("_RuntimeStorage");
+	const FString StructName = GetName() + TEXT("_InstanceStorage");
 
 	// Remove existing struct of same name.
 	UScriptStruct* OldStruct = FindObject<UScriptStruct>(this, *StructName);
@@ -164,10 +172,10 @@ void UStateTree::InitRuntimeStorage()
 	// Append all evaluators and tasks.
 	// Since properties are stored in linked list, add in reverse order so that we retain the correct order in output.
 	TArray<FStructProperty*> NewProperties;
-	NewProperties.SetNumZeroed(RuntimeStorageItems.Num());
-	for (int32 ItemIndex = RuntimeStorageItems.Num() - 1; ItemIndex >= 0; ItemIndex--)
+	NewProperties.SetNumZeroed(Instances.Num());
+	for (int32 ItemIndex = Instances.Num() - 1; ItemIndex >= 0; ItemIndex--)
 	{
-		const FInstancedStruct& ItemPtr = RuntimeStorageItems[ItemIndex];
+		const FInstancedStruct& ItemPtr = Instances[ItemIndex];
 		UScriptStruct* ItemStruct = const_cast<UScriptStruct*>(ItemPtr.GetScriptStruct());
 		FName PropName(FString::Printf(TEXT("%s%d"), *ItemStruct->GetName(), NewProperties.Num()));
 		FStructProperty* NewStructProperty = new FStructProperty(NewStruct, PropName, RF_Public);
@@ -183,22 +191,22 @@ void UStateTree::InitRuntimeStorage()
 	NewStruct->StaticLink(/*bRelinkExistingProperties*/true);
 
 	// Store item offsets for fast struct view creation at runtime.
-	TArray<FStateTreeRuntimeStorageItemOffset> NewItemOffsets;
+	TArray<FStateTreeInstanceStorageOffset> NewOffsets;
 	for (FStructProperty* StructProperty : NewProperties)
 	{
-		NewItemOffsets.Emplace(StructProperty->Struct, StructProperty->GetOffset_ForInternal());
+		NewOffsets.Emplace(StructProperty->Struct, StructProperty->GetOffset_ForInternal());
 	}
-	check(RuntimeStorageItems.Num() == NewItemOffsets.Num());
+	check(Instances.Num() == NewOffsets.Num());
 
 	// Instantiate default value
-	RuntimeStorageDefaultValue.InitializeAs(NewStruct);
-	for (int i = 0; i < NewItemOffsets.Num(); i++)
+	InstanceStorageDefaultValue.InitializeAs(NewStruct);
+	for (int i = 0; i < NewOffsets.Num(); i++)
 	{
-		const FStateTreeRuntimeStorageItemOffset& Item = NewItemOffsets[i];
-		uint8* Dest = RuntimeStorageDefaultValue.GetMutableMemory() + Item.Offset;
-		Item.Struct->CopyScriptStruct(Dest, RuntimeStorageItems[i].GetMemory());
+		const FStateTreeInstanceStorageOffset& Item = NewOffsets[i];
+		uint8* Dest = InstanceStorageDefaultValue.GetMutableMemory() + Item.Offset;
+		Item.Struct->CopyScriptStruct(Dest, Instances[i].GetMemory());
 	}
 
-	RuntimeStorageOffsets = NewItemOffsets;
-	RuntimeStorageStruct = NewStruct;
+	InstanceStorageOffsets = NewOffsets;
+	InstanceStorageStruct = NewStruct;
 }
diff --git a/Engine/Plugins/Runtime/StateTree/Source/StateTreeModule/Private/StateTreeExecutionContext.cpp b/Engine/Plugins/Runtime/StateTree/Source/StateTreeModule/Private/StateTreeExecutionContext.cpp
--- a/Engine/Plugins/Runtime/StateTree/Source/StateTreeModule/Private/StateTreeExecutionContext.cpp
+++ b/Engine/Plugins/Runtime/StateTree/Source/StateTreeModule/Private/StateTreeExecutionContext.cpp
@@ -44,12 +44,12 @@ bool FStateTreeExecutionContext::Init(UObject& InOwner, const UStateTree& InStat
 	StorageType = InStorageType;
 	if (StorageType == EStateTreeStorage::Internal)
 	{
-		// Duplicate runtime state.
-		StorageInstance = StateTree->RuntimeStorageDefaultValue;
+		// Duplicate instance state.
+		StorageInstance = StateTree->InstanceStorageDefaultValue;
 	}
 
-	// Initialize struct view for all possible source items.
-	ItemViews.SetNum(StateTree->GetLinkedItemCount());
+	// Initialize data views for all possible items.
+	DataViews.SetNum(StateTree->GetNumDataViews());
 
 	// Initialize array to keep track of which states have been updated.
 	VisitedStates.Init(false, StateTree->States.Num());
@@ -57,13 +57,13 @@ bool FStateTreeExecutionContext::Init(UObject& InOwner, const UStateTree& InStat
 	return true;
 }
 
-void FStateTreeExecutionContext::Start(FStateTreeItemView ExternalStorage)
+void FStateTreeExecutionContext::Start(FStateTreeDataView ExternalStorage)
 {
 	if (!Owner || !StateTree)
 	{
 		return;
 	}
-	FStateTreeItemView Storage = SelectMutableStorage(ExternalStorage);
+	FStateTreeDataView Storage = SelectMutableStorage(ExternalStorage);
 	FStateTreeExecutionState& Exec = GetExecState(Storage);
 
 	// Stop if still running 
@@ -79,13 +79,13 @@ void FStateTreeExecutionContext::Start(FStateTreeItemView ExternalStorage)
 	Exec.LastTickStatus = EStateTreeRunStatus::Unset;
 }
 
-void FStateTreeExecutionContext::Stop(FStateTreeItemView ExternalStorage)
+void FStateTreeExecutionContext::Stop(FStateTreeDataView ExternalStorage)
 {
 	if (!Owner || !StateTree)
 	{
 		return;
 	}
-	FStateTreeItemView Storage = SelectMutableStorage(ExternalStorage);
+	FStateTreeDataView Storage = SelectMutableStorage(ExternalStorage);
 	FStateTreeExecutionState& Exec = GetExecState(Storage);
 
 	// Stop if still running.
@@ -100,7 +100,7 @@ void FStateTreeExecutionContext::Stop(FStateTreeItemView ExternalStorage)
 	Exec.LastTickStatus = EStateTreeRunStatus::Unset;
 }
 
-EStateTreeRunStatus FStateTreeExecutionContext::Tick(const float DeltaTime, FStateTreeItemView ExternalStorage)
+EStateTreeRunStatus FStateTreeExecutionContext::Tick(const float DeltaTime, FStateTreeDataView ExternalStorage)
 {
 	CSV_SCOPED_TIMING_STAT_EXCLUSIVE(StateTreeCtxTick);
 
@@ -108,7 +108,7 @@ EStateTreeRunStatus FStateTreeExecutionContext::Tick(const float DeltaTime, FSta
 	{
 		return EStateTreeRunStatus::Failed;
 	}
-	FStateTreeItemView Storage = SelectMutableStorage(ExternalStorage);
+	FStateTreeDataView Storage = SelectMutableStorage(ExternalStorage);
 	FStateTreeExecutionState& Exec = GetExecState(Storage);
 
 	// No ticking of the tree is done or stopped.
@@ -210,13 +210,13 @@ void FStateTreeExecutionContext::Reset()
 	StateTree = nullptr;
 	Owner = nullptr;
 	World = nullptr;
-	ItemViews.Reset();
+	DataViews.Reset();
 	StorageInstance.Reset();
 	StorageType = EStateTreeStorage::Internal;
 	VisitedStates.Reset();
 }
 
-EStateTreeRunStatus FStateTreeExecutionContext::EnterState(FStateTreeItemView Storage, const FStateTreeTransitionResult& Transition)
+EStateTreeRunStatus FStateTreeExecutionContext::EnterState(FStateTreeDataView Storage, const FStateTreeTransitionResult& Transition)
 {
 	if (!Transition.Next.IsValid())
 	{
@@ -253,18 +253,16 @@ EStateTreeRunStatus FStateTreeExecutionContext::EnterState(FStateTreeItemView St
 		if (bWasActive && !bOnTargetBranch)
 		{
 			// States which will keep on begin active and were not part of the transition will not get enter/exit state.
-			// Must update item views.
+			// Must update data views.
 			for (int32 EvalIndex = 0; EvalIndex < int32(State.EvaluatorsNum); EvalIndex++)
 			{
-				FStateTreeItemView EvalView = GetItem(Storage, int32(State.EvaluatorsBegin) + EvalIndex);
-				const FStateTreeEvaluatorBase& Eval = EvalView.Get<FStateTreeEvaluatorBase>();
-				ItemViews[Eval.SourceStructIndex] = EvalView;
+				const FStateTreeEvaluatorBase& Eval = GetItem<FStateTreeEvaluatorBase>(int32(State.EvaluatorsBegin) + EvalIndex);
+				DataViews[Eval.DataViewIndex] = GetInstanceData(Storage, Eval.DataViewIndex);
 			}
 			for (int32 TaskIndex = 0; TaskIndex < int32(State.TasksNum); TaskIndex++)
 			{
-				FStateTreeItemView TaskView = GetItem(Storage, int32(State.TasksBegin) + TaskIndex);
-				const FStateTreeTaskBase& Task = TaskView.Get<FStateTreeTaskBase>();
-				ItemViews[Task.SourceStructIndex] = TaskView;
+				const FStateTreeTaskBase& Task = GetItem<FStateTreeTaskBase>(int32(State.TasksBegin) + TaskIndex);
+				DataViews[Task.DataViewIndex] = GetInstanceData(Storage, Task.DataViewIndex);
 			}
 			continue;
 		}
@@ -277,30 +275,30 @@ EStateTreeRunStatus FStateTreeExecutionContext::EnterState(FStateTreeItemView St
 		
 		for (int32 EvalIndex = 0; EvalIndex < int32(State.EvaluatorsNum); EvalIndex++)
 		{
-			FStateTreeItemView EvalView = GetItem(Storage, int32(State.EvaluatorsBegin) + EvalIndex);
-			FStateTreeEvaluatorBase& Eval = EvalView.GetMutable<FStateTreeEvaluatorBase>();
+			const FStateTreeEvaluatorBase& Eval = GetItem<FStateTreeEvaluatorBase>(int32(State.EvaluatorsBegin) + EvalIndex);
+			const FStateTreeDataView EvalInstanceView = GetInstanceData(Storage, Eval.DataViewIndex);
+			DataViews[Eval.DataViewIndex] = EvalInstanceView;
 
 			// Copy bound properties.
 			if (Eval.BindingsBatch.IsValid())
 			{
-				StateTree->PropertyBindings.CopyTo(ItemViews, Eval.BindingsBatch.Index, EvalView);
+				StateTree->PropertyBindings.CopyTo(DataViews, Eval.BindingsBatch.Index, EvalInstanceView);
 			}
 			STATETREE_LOG(Verbose, TEXT("%*s  Notify Evaluator '%s'."), Index*UE::StateTree::DebugIndentSize, TEXT(""), *Eval.Name.ToString());
 			Eval.EnterState(*this, ChangeType, CurrentTransition);
-
-			ItemViews[Eval.SourceStructIndex] = EvalView;
 		}
 		
 		// Activate tasks on current state.
 		for (int32 TaskIndex = 0; TaskIndex < int32(State.TasksNum); TaskIndex++)
 		{
-			FStateTreeItemView TaskView = GetItem(Storage, int32(State.TasksBegin) + TaskIndex);
-			FStateTreeTaskBase& Task = TaskView.GetMutable<FStateTreeTaskBase>();
+			const FStateTreeTaskBase& Task = GetItem<FStateTreeTaskBase>(int32(State.TasksBegin) + TaskIndex);
+			const FStateTreeDataView TaskInstanceView = GetInstanceData(Storage, Task.DataViewIndex);
+			DataViews[Task.DataViewIndex] = TaskInstanceView;
 
 			// Copy bound properties.
 			if (Task.BindingsBatch.IsValid())
 			{
-				StateTree->PropertyBindings.CopyTo(ItemViews, Task.BindingsBatch.Index, TaskView);
+				StateTree->PropertyBindings.CopyTo(DataViews, Task.BindingsBatch.Index, TaskInstanceView);
 			}
 		
 			STATETREE_LOG(Verbose, TEXT("%*s  Notify Task '%s'"), Index*UE::StateTree::DebugIndentSize, TEXT(""), *Task.Name.ToString());
@@ -312,17 +310,14 @@ EStateTreeRunStatus FStateTreeExecutionContext::EnterState(FStateTreeItemView St
 				// break;
 				// As a workaround, we mark EnterStateStatus as 'failed' in context so remaining calls to EnterState could be aware of the failure and act accordingly.
 				EnterStateStatus = Status;
-
 			}
-
-			ItemViews[Task.SourceStructIndex] = TaskView;
 		}
 	}
 
 	return Result;
 }
 
-void FStateTreeExecutionContext::ExitState(FStateTreeItemView Storage, const FStateTreeTransitionResult& Transition)
+void FStateTreeExecutionContext::ExitState(FStateTreeDataView Storage, const FStateTreeTransitionResult& Transition)
 {
 	if (!Transition.Source.State.IsValid())
 	{
@@ -367,15 +362,13 @@ void FStateTreeExecutionContext::ExitState(FStateTreeItemView Storage, const FSt
 			// Must update item views.
 			for (int32 EvalIndex = 0; EvalIndex < int32(State.EvaluatorsNum); EvalIndex++)
 			{
-				FStateTreeItemView EvalView = GetItem(Storage, int32(State.EvaluatorsBegin) + EvalIndex);
-				const FStateTreeEvaluatorBase& Eval = EvalView.Get<FStateTreeEvaluatorBase>();
-				ItemViews[Eval.SourceStructIndex] = EvalView;
+				FStateTreeEvaluatorBase& Eval = GetItem<FStateTreeEvaluatorBase>(int32(State.EvaluatorsBegin) + EvalIndex);
+				DataViews[Eval.DataViewIndex] = GetInstanceData(Storage, Eval.DataViewIndex);
 			}
 			for (int32 TaskIndex = 0; TaskIndex < int32(State.TasksNum); TaskIndex++)
 			{
-				FStateTreeItemView TaskView = GetItem(Storage, int32(State.TasksBegin) + TaskIndex);
-				const FStateTreeTaskBase& Task = TaskView.Get<FStateTreeTaskBase>();
-				ItemViews[Task.SourceStructIndex] = TaskView;
+				FStateTreeTaskBase& Task = GetItem<FStateTreeTaskBase>(int32(State.TasksBegin) + TaskIndex);
+				DataViews[Task.DataViewIndex] = GetInstanceData(Storage, Task.DataViewIndex);
 			}
 			continue;
 		}
@@ -388,41 +381,39 @@ void FStateTreeExecutionContext::ExitState(FStateTreeItemView Storage, const FSt
 
 		for (int32 EvalIndex = 0; EvalIndex < int32(State.EvaluatorsNum); EvalIndex++)
 		{
-			FStateTreeItemView EvalView = GetItem(Storage, int32(State.EvaluatorsBegin) + EvalIndex);
-			FStateTreeEvaluatorBase& Eval = EvalView.GetMutable<FStateTreeEvaluatorBase>();
+			const FStateTreeEvaluatorBase& Eval = GetItem<FStateTreeEvaluatorBase>(int32(State.EvaluatorsBegin) + EvalIndex);
+			const FStateTreeDataView EvalInstanceView = GetInstanceData(Storage, Eval.DataViewIndex);
+			DataViews[Eval.DataViewIndex] = EvalInstanceView;
 
 			// Copy bound properties.
 			if (Eval.BindingsBatch.IsValid())
 			{
-				StateTree->PropertyBindings.CopyTo(ItemViews, Eval.BindingsBatch.Index, EvalView);
+				StateTree->PropertyBindings.CopyTo(DataViews, Eval.BindingsBatch.Index, EvalInstanceView);
 			}
 			STATETREE_LOG(Verbose, TEXT("%*s  Notify Evaluator '%s'."), Index*UE::StateTree::DebugIndentSize, TEXT(""), *Eval.Name.ToString());
 			Eval.ExitState(*this, ChangeType, CurrentTransition);
-
-			ItemViews[Eval.SourceStructIndex] = EvalView;
 		}
 
 		// Deactivate tasks on current State
 		for (int32 TaskIndex = 0; TaskIndex < int32(State.TasksNum); TaskIndex++)
 		{
-			FStateTreeItemView TaskView = GetItem(Storage, int32(State.TasksBegin) + TaskIndex);
-			FStateTreeTaskBase& Task = TaskView.GetMutable<FStateTreeTaskBase>();
+			const FStateTreeTaskBase& Task = GetItem<FStateTreeTaskBase>(int32(State.TasksBegin) + TaskIndex);
+			const FStateTreeDataView TaskInstanceView = GetInstanceData(Storage, Task.DataViewIndex);
+			DataViews[Task.DataViewIndex] = TaskInstanceView;
 
 			// Copy bound properties.
 			if (Task.BindingsBatch.IsValid())
 			{
-				StateTree->PropertyBindings.CopyTo(ItemViews, Task.BindingsBatch.Index, TaskView);
+				StateTree->PropertyBindings.CopyTo(DataViews, Task.BindingsBatch.Index, TaskInstanceView);
 			}
 
 			STATETREE_LOG(Verbose, TEXT("%*s  Notify Task '%s'"), Index*UE::StateTree::DebugIndentSize, TEXT(""), *Task.Name.ToString());
 			Task.ExitState(*this, ChangeType, CurrentTransition);
-
-			ItemViews[Task.SourceStructIndex] = TaskView;
 		}
 	}
 }
 
-void FStateTreeExecutionContext::StateCompleted(FStateTreeItemView Storage, const FStateTreeHandle CurrentState, const EStateTreeRunStatus CompletionStatus)
+void FStateTreeExecutionContext::StateCompleted(FStateTreeDataView Storage, const FStateTreeHandle CurrentState, const EStateTreeRunStatus CompletionStatus)
 {
 	if (!CurrentState.IsValid())
 	{
@@ -444,8 +435,7 @@ void FStateTreeExecutionContext::StateCompleted(FStateTreeItemView Storage, cons
 		// Notify Tasks
 		for (int32 TaskIndex = int32(State.TasksNum) - 1; TaskIndex >= 0; TaskIndex--)
 		{
-			FStateTreeItemView TaskView = GetItem(Storage, int32(State.TasksBegin) + TaskIndex);
-			FStateTreeTaskBase& Task = TaskView.GetMutable<FStateTreeTaskBase>();
+			const FStateTreeTaskBase& Task = GetItem<FStateTreeTaskBase>(int32(State.TasksBegin) + TaskIndex);
 
 			STATETREE_LOG(Verbose, TEXT("%*s  Notify Task '%s'"), Index*UE::StateTree::DebugIndentSize, TEXT(""), *Task.Name.ToString());
 			Task.StateCompleted(*this, CompletionStatus, CurrentState);
@@ -454,8 +444,7 @@ void FStateTreeExecutionContext::StateCompleted(FStateTreeItemView Storage, cons
 		// Notify evaluators
 		for (int32 EvalIndex = int32(State.EvaluatorsNum) - 1; EvalIndex >= 0; EvalIndex--)
 		{
-			FStateTreeItemView EvalView = GetItem(Storage, int32(State.EvaluatorsBegin) + EvalIndex);
-			FStateTreeEvaluatorBase& Eval = EvalView.GetMutable<FStateTreeEvaluatorBase>();
+			const FStateTreeEvaluatorBase& Eval = GetItem<FStateTreeEvaluatorBase>(int32(State.EvaluatorsBegin) + EvalIndex);
 
 			STATETREE_LOG(Verbose, TEXT("%*s  Notify Evaluator '%s'"), Index*UE::StateTree::DebugIndentSize, TEXT(""), *Eval.Name.ToString());
 			Eval.StateCompleted(*this, CompletionStatus, CurrentState);
@@ -463,7 +452,7 @@ void FStateTreeExecutionContext::StateCompleted(FStateTreeItemView Storage, cons
 	}
 }
 
-void FStateTreeExecutionContext::TickEvaluators(FStateTreeItemView Storage, const FStateTreeHandle CurrentState, const EStateTreeEvaluationType EvalType, const float DeltaTime)
+void FStateTreeExecutionContext::TickEvaluators(FStateTreeDataView Storage, const FStateTreeHandle CurrentState, const EStateTreeEvaluationType EvalType, const float DeltaTime)
 {
 	CSV_SCOPED_TIMING_STAT_EXCLUSIVE(StateTreeCtxTickEvaluators);
 
@@ -491,24 +480,24 @@ void FStateTreeExecutionContext::TickEvaluators(FStateTreeItemView Storage, cons
 		// Tick evaluators
 		for (int32 EvalIndex = 0; EvalIndex < int32(State.EvaluatorsNum); EvalIndex++)
 		{
-			FStateTreeItemView EvalView = GetItem(Storage, int32(State.EvaluatorsBegin) + EvalIndex);
-			FStateTreeEvaluatorBase& Eval = EvalView.GetMutable<FStateTreeEvaluatorBase>();
+			const FStateTreeEvaluatorBase& Eval = GetItem<FStateTreeEvaluatorBase>(int32(State.EvaluatorsBegin) + EvalIndex);
+			const FStateTreeDataView EvalInstanceView = GetInstanceData(Storage, Eval.DataViewIndex);
+			DataViews[Eval.DataViewIndex] = EvalInstanceView;
 
 			// Copy bound properties.
 			if (Eval.BindingsBatch.IsValid())
 			{
-				StateTree->PropertyBindings.CopyTo(ItemViews, Eval.BindingsBatch.Index, EvalView);
+				StateTree->PropertyBindings.CopyTo(DataViews, Eval.BindingsBatch.Index, EvalInstanceView);
 			}
 			STATETREE_LOG(Verbose, TEXT("%*s  Evaluate: '%s'"), Index*UE::StateTree::DebugIndentSize, TEXT(""), *Eval.Name.ToString());
 			Eval.Evaluate(*this, EvalType, DeltaTime);
-			ItemViews[Eval.SourceStructIndex] = EvalView;
 		}
 
 		VisitedStates[CurrentHandle.Index] = true;
 	}
 }
 
-EStateTreeRunStatus FStateTreeExecutionContext::TickTasks(FStateTreeItemView Storage, const FStateTreeHandle CurrentState, const float DeltaTime)
+EStateTreeRunStatus FStateTreeExecutionContext::TickTasks(FStateTreeDataView Storage, const FStateTreeHandle CurrentState, const float DeltaTime)
 {
 	CSV_SCOPED_TIMING_STAT_EXCLUSIVE(StateTreeCtxTickTasks);
 
@@ -533,13 +522,14 @@ EStateTreeRunStatus FStateTreeExecutionContext::TickTasks(FStateTreeItemView Sto
 		// Tick Tasks
 		for (int32 TaskIndex = 0; TaskIndex < int32(State.TasksNum); TaskIndex++)
 		{
-			FStateTreeItemView TaskView = GetItem(Storage, int32(State.TasksBegin) + TaskIndex);
-			FStateTreeTaskBase& Task = TaskView.GetMutable<FStateTreeTaskBase>();
+			const FStateTreeTaskBase& Task = GetItem<FStateTreeTaskBase>(int32(State.TasksBegin) + TaskIndex);
+			const FStateTreeDataView TaskInstanceView = GetInstanceData(Storage, Task.DataViewIndex);
+			DataViews[Task.DataViewIndex] = TaskInstanceView;
 
 			// Copy bound properties.
 			if (Task.BindingsBatch.IsValid())
 			{
-				StateTree->PropertyBindings.CopyTo(ItemViews, Task.BindingsBatch.Index, TaskView);
+				StateTree->PropertyBindings.CopyTo(DataViews, Task.BindingsBatch.Index, TaskInstanceView);
 			}
 			STATETREE_LOG(Verbose, TEXT("%*s  Tick: '%s'"), Index*UE::StateTree::DebugIndentSize, TEXT(""), *Task.Name.ToString());
 			const EStateTreeRunStatus TaskResult = Task.Tick(*this, DeltaTime);
@@ -553,8 +543,6 @@ EStateTreeRunStatus FStateTreeExecutionContext::TickTasks(FStateTreeItemView Sto
 			{
 				break;
 			}
-
-			ItemViews[Task.SourceStructIndex] = TaskView;
 		}
 		NumTotalTasks += State.TasksNum;
 	}
@@ -568,39 +556,23 @@ EStateTreeRunStatus FStateTreeExecutionContext::TickTasks(FStateTreeItemView Sto
 	return Result;
 }
 
-bool FStateTreeExecutionContext::TestAllConditions(const uint32 ConditionsOffset, const uint32 ConditionsNum)
+bool FStateTreeExecutionContext::TestAllConditions(FStateTreeDataView Storage, const uint32 ConditionsOffset, const uint32 ConditionsNum)
 {
 	CSV_SCOPED_TIMING_STAT_EXCLUSIVE(StateTreeCtxTestConditions);
 
-	uint8 CondCopyBuffer[128];	// This buffer is intentionally small, conditions be small.
-	
 	for (uint32 i = 0; i < ConditionsNum; i++)
 	{
-		// Copy struct values from the StateTree asset to a stack copy.
-		const FInstancedStruct& SourceCondPtr = StateTree->Conditions[ConditionsOffset + i];
-		const FStateTreeConditionBase& SourceCond = SourceCondPtr.Get<FStateTreeConditionBase>();
-		const UScriptStruct* CondStruct = SourceCondPtr.GetScriptStruct();
-		check(CondStruct);
-		check(CondStruct->GetStructureSize() <= sizeof(CondCopyBuffer));
-
-		CondStruct->InitializeStruct(CondCopyBuffer);
-		CondStruct->CopyScriptStruct(CondCopyBuffer, &SourceCond);
-
-		// Copy properties over
-		FStateTreeItemView CondCopyView(CondStruct, CondCopyBuffer);
-		if (SourceCond.BindingsBatch.IsValid())
+		const FStateTreeConditionBase& Cond = GetItem<FStateTreeConditionBase>(ConditionsOffset + i);
+		const FStateTreeDataView CondInstanceView = GetInstanceData(Storage, Cond.DataViewIndex);
+		DataViews[Cond.DataViewIndex] = CondInstanceView;
+
+		// Copy bound properties.
+		if (Cond.BindingsBatch.IsValid())
 		{
-			StateTree->PropertyBindings.CopyTo(ItemViews, SourceCond.BindingsBatch.Index, CondCopyView);
+			StateTree->PropertyBindings.CopyTo(DataViews, Cond.BindingsBatch.Index, CondInstanceView);
 		}
-
-		// Do test (virtual call)
-		const FStateTreeConditionBase& Cond = CondCopyView.Get<FStateTreeConditionBase>();
-		const bool bResult = Cond.TestCondition();
-
-		// Struct destructor
-		CondStruct->DestroyStruct(CondCopyBuffer);
-
-		if (!bResult)
+		
+		if (!Cond.TestCondition(*this))
 		{
 			return false;
 		}
@@ -608,7 +580,7 @@ bool FStateTreeExecutionContext::TestAllConditions(const uint32 ConditionsOffset
 	return true;
 }
 
-FStateTreeTransitionResult FStateTreeExecutionContext::TriggerTransitions(FStateTreeItemView Storage, const FStateTreeStateStatus CurrentStatus, const int32 Depth)
+FStateTreeTransitionResult FStateTreeExecutionContext::TriggerTransitions(FStateTreeDataView Storage, const FStateTreeStateStatus CurrentStatus, const int32 Depth)
 {
 	FStateTreeExecutionState& Exec = GetExecState(Storage);
 	EStateTreeTransitionEvent Event = EStateTreeTransitionEvent::OnCondition;
@@ -630,7 +602,7 @@ FStateTreeTransitionResult FStateTreeExecutionContext::TriggerTransitions(FState
 			// All transition conditions must pass
 			const int16 TransitionIndex = State.TransitionsBegin + i;
 			const FBakedStateTransition& Transition = StateTree->Transitions[TransitionIndex];
-			if (EnumHasAllFlags(Transition.Event, Event) && TestAllConditions(Transition.ConditionsBegin, Transition.ConditionsNum))
+			if (EnumHasAllFlags(Transition.Event, Event) && TestAllConditions(Storage, Transition.ConditionsBegin, Transition.ConditionsNum))
 			{
 				// If a transition has delay, we stop testing other transitions, but the transition will not pass the condition until the delay time passes.
 				if (Transition.GateDelay > 0)
@@ -697,7 +669,7 @@ FStateTreeTransitionResult FStateTreeExecutionContext::TriggerTransitions(FState
 	return FStateTreeTransitionResult(CurrentStatus, FStateTreeHandle::Invalid);
 }
 
-FStateTreeTransitionResult FStateTreeExecutionContext::SelectState(FStateTreeItemView Storage, const FStateTreeStateStatus InitialStateStatus, const FStateTreeHandle InitialTargetState, const FStateTreeHandle NextState, const int Depth)
+FStateTreeTransitionResult FStateTreeExecutionContext::SelectState(FStateTreeDataView Storage, const FStateTreeStateStatus InitialStateStatus, const FStateTreeHandle InitialTargetState, const FStateTreeHandle NextState, const int Depth)
 {
 	if (!NextState.IsValid())
 	{
@@ -713,7 +685,7 @@ FStateTreeTransitionResult FStateTreeExecutionContext::SelectState(FStateTreeIte
 	TickEvaluators(Storage, NextState, EStateTreeEvaluationType::PreSelect, 0.0f);
 	
 	// Check that the state can be entered
-	if (TestAllConditions(State.EnterConditionsBegin, State.EnterConditionsNum))
+	if (TestAllConditions(Storage, State.EnterConditionsBegin, State.EnterConditionsNum))
 	{
 		// If the state has children, proceed to select children.
 		if (State.HasChildren())
@@ -816,7 +788,7 @@ FString FStateTreeExecutionContext::GetStateStatusString(const FStateTreeStateSt
 	}
 }
 
-EStateTreeRunStatus FStateTreeExecutionContext::GetLastTickStatus(FStateTreeItemView ExternalStorage) const
+EStateTreeRunStatus FStateTreeExecutionContext::GetLastTickStatus(FStateTreeDataView ExternalStorage) const
 {
 	const FStateTreeExecutionState& Exec = GetExecState(SelectStorage(ExternalStorage));
 	return Exec.LastTickStatus;
@@ -824,13 +796,13 @@ EStateTreeRunStatus FStateTreeExecutionContext::GetLastTickStatus(FStateTreeItem
 
 #if WITH_GAMEPLAY_DEBUGGER
 
-FString FStateTreeExecutionContext::GetDebugInfoString(FStateTreeItemView ExternalStorage) const
+FString FStateTreeExecutionContext::GetDebugInfoString(FStateTreeDataView ExternalStorage) const
 {
 	if (!StateTree)
 	{
 		return FString(TEXT("No StateTree asset."));
 	}
-	FStateTreeItemView Storage = SelectStorage(ExternalStorage);
+	FStateTreeDataView Storage = SelectStorage(ExternalStorage);
 	FStateTreeExecutionState& Exec = GetExecState(Storage);
 
 	FString DebugString = FString::Printf(TEXT("StateTree (asset: '%s')\n"), *GetNameSafe(StateTree));
@@ -869,8 +841,7 @@ FString FStateTreeExecutionContext::GetDebugInfoString(FStateTreeItemView Extern
 				DebugString += TEXT("\nTasks:\n");
 				for (int32 j = 0; j < int32(State.TasksNum); j++)
 				{
-					FStateTreeItemView TaskView = GetItem(Storage, int32(State.TasksBegin) + j);
-					const FStateTreeTaskBase& Task = TaskView.Get<FStateTreeTaskBase>();
+					const FStateTreeTaskBase& Task = GetItem<FStateTreeTaskBase>(int32(State.TasksBegin) + j);
 					Task.AppendDebugInfoString(DebugString, *this);
 				}
 			}
@@ -879,8 +850,7 @@ FString FStateTreeExecutionContext::GetDebugInfoString(FStateTreeItemView Extern
 				DebugString += TEXT("\nEvaluators:\n");
 				for (int32 EvalIndex = 0; EvalIndex < int32(State.EvaluatorsNum); EvalIndex++)
 				{
-					FStateTreeItemView EvalView = GetItem(Storage, int32(State.EvaluatorsBegin) + EvalIndex);
-					const FStateTreeEvaluatorBase& Eval = EvalView.Get<FStateTreeEvaluatorBase>();
+					const FStateTreeEvaluatorBase& Eval = GetItem<FStateTreeEvaluatorBase>(int32(State.EvaluatorsBegin) + EvalIndex);
 					Eval.AppendDebugInfoString(DebugString, *this);
 				}
 			}
@@ -892,7 +862,7 @@ FString FStateTreeExecutionContext::GetDebugInfoString(FStateTreeItemView Extern
 #endif // WITH_GAMEPLAY_DEBUGGER
 
 #if WITH_STATETREE_DEBUG
-void FStateTreeExecutionContext::DebugPrintInternalLayout(FStateTreeItemView ExternalStorage)
+void FStateTreeExecutionContext::DebugPrintInternalLayout(FStateTreeDataView ExternalStorage)
 {
 	LOG_SCOPE_VERBOSITY_OVERRIDE(LogStateTree, ELogVerbosity::Log);
 
@@ -904,32 +874,32 @@ void FStateTreeExecutionContext::DebugPrintInternalLayout(FStateTreeItemView Ext
 
 	FString DebugString = FString::Printf(TEXT("StateTree (asset: '%s')\n"), *GetNameSafe(StateTree));
 
-	// Runtime storage items (e.g. tasks, evaluators)
-	DebugString += FString::Printf(TEXT("\nItems(%d)\n"), StateTree->RuntimeStorageItems.Num());
-	for (const FInstancedStruct& Item : StateTree->RuntimeStorageItems)
+	// Tree items (e.g. tasks, evaluators, conditions)
+	DebugString += FString::Printf(TEXT("\nItems(%d)\n"), StateTree->Items.Num());
+	for (const FInstancedStruct& Item : StateTree->Items)
 	{
 		DebugString += FString::Printf(TEXT("  %s\n"), Item.IsValid() ? *Item.GetScriptStruct()->GetName() : TEXT("null"));
 	}
 
-	// Runtime storage offsets (e.g. tasks, evaluators)
-	DebugString += FString::Printf(TEXT("\nOffsets(%d)\n  [ %-40s | %-6s ]\n"), StateTree->RuntimeStorageOffsets.Num(), TEXT("Name"), TEXT("Offset"));
-	for (const FStateTreeRuntimeStorageItemOffset& Offset : StateTree->RuntimeStorageOffsets)
+	// Instance storage data (e.g. tasks, evaluators, conditions)
+	DebugString += FString::Printf(TEXT("\nInstance Data(%d)\n"), StateTree->Instances.Num());
+	for (const FInstancedStruct& Data : StateTree->Instances)
 	{
-		DebugString += FString::Printf(TEXT("  | %-40s | %6d |\n"), Offset.Struct ? *Offset.Struct->GetName() : TEXT("null"),  Offset.Offset);
+		DebugString += FString::Printf(TEXT("  %s\n"), Data.IsValid() ? *Data.GetScriptStruct()->GetName() : TEXT("null"));
 	}
 
-	// Conditions
-	DebugString += FString::Printf(TEXT("\nConditions(%d)\n"), StateTree->Conditions.Num());
-	for (const FInstancedStruct& Condition : StateTree->Conditions)
+	// Instance storage offsets (e.g. tasks, evaluators, conditions)
+	DebugString += FString::Printf(TEXT("\nInstance Data Offsets(%d)\n  [ %-40s | %-6s ]\n"), StateTree->InstanceStorageOffsets.Num(), TEXT("Name"), TEXT("Offset"));
+	for (const FStateTreeInstanceStorageOffset& Offset : StateTree->InstanceStorageOffsets)
 	{
-		DebugString += FString::Printf(TEXT("  %s\n"), Condition.IsValid() ? *Condition.GetScriptStruct()->GetName() : TEXT("null"));
+		DebugString += FString::Printf(TEXT("  | %-40s | %6d |\n"), Offset.Struct ? *Offset.Struct->GetName() : TEXT("null"),  Offset.Offset);
 	}
 
-	// External items (e.g. fragments, subsystems)
-	DebugString += FString::Printf(TEXT("\nExternal items(%d)\n  [ %-40s | %-8s | %5s ]\n"), StateTree->ExternalItems.Num(), TEXT("Name"), TEXT("Optional"), TEXT("Index"));
-	for (const FStateTreeExternalItemDesc& Desc : StateTree->ExternalItems)
+	// External data (e.g. fragments, subsystems)
+	DebugString += FString::Printf(TEXT("\nExternal Data(%d)\n  [ %-40s | %-8s | %5s ]\n"), StateTree->ExternalDataDescs.Num(), TEXT("Name"), TEXT("Optional"), TEXT("Index"));
+	for (const FStateTreeExternalDataDesc& Desc : StateTree->ExternalDataDescs)
 	{
-		DebugString += FString::Printf(TEXT("  | %-40s | %8s | %5d |\n"), Desc.Struct ? *Desc.Struct->GetName() : TEXT("null"), *UEnum::GetValueAsString(Desc.Requirement), Desc.Handle.ItemIndex);
+		DebugString += FString::Printf(TEXT("  | %-40s | %8s | %5d |\n"), Desc.Struct ? *Desc.Struct->GetName() : TEXT("null"), *UEnum::GetValueAsString(Desc.Requirement), Desc.Handle.DataViewIndex);
 	}
 
 	// Bindings
@@ -947,15 +917,15 @@ void FStateTreeExecutionContext::DebugPrintInternalLayout(FStateTreeItemView Ext
 									Transition.ConditionsNum);
 	}
 
-	// ItemViews
-	DebugString += FString::Printf(TEXT("\nItemViews(%d)\n"), ItemViews.Num());
-	for (const FStateTreeItemView& ItemView : ItemViews)
+	// DataViews
+	DebugString += FString::Printf(TEXT("\nDataViews(%d)\n"), DataViews.Num());
+	for (const FStateTreeDataView& DataView : DataViews)
 	{
-		DebugString += FString::Printf(TEXT("  [%s]\n"), ItemView.IsValid() ? *ItemView.GetStruct()->GetName() : TEXT("null"));
+		DebugString += FString::Printf(TEXT("  [%s]\n"), DataView.IsValid() ? *DataView.GetStruct()->GetName() : TEXT("null"));
 	}
 
 	// States
-	FStateTreeItemView Storage = SelectMutableStorage(ExternalStorage);
+	FStateTreeDataView Storage = SelectMutableStorage(ExternalStorage);
 	DebugString += FString::Printf(TEXT("\nStates(%d)\n"
 		"  [ %-30s | %15s | %5s [%3s:%-3s[ | Begin Idx : %4s %4s %4s %4s | Num : %4s %4s %4s %4s | Transitions : %-16s %-40s %-16s %-40s ]\n"),
 		StateTree->States.Num(),
@@ -974,8 +944,8 @@ void FStateTreeExecutionContext::DebugPrintInternalLayout(FStateTreeItemView Ext
 									*State.StateFailedTransitionState.Describe(), *StaticEnum<EStateTreeTransitionType>()->GetValueAsString(State.StateFailedTransitionType));
 	}
 
-	DebugString += FString::Printf(TEXT("\nStates Items\n  [ %-30s | %-12s | %-30s | %15s | %10s ]\n"),
-	TEXT("State"), TEXT("Type"), TEXT("Name"), TEXT("Bindings"), TEXT("Struct Idx"));
+	DebugString += FString::Printf(TEXT("\nStates\n  [ %-30s | %-12s | %-30s | %15s | %10s ]\n"),
+		TEXT("State"), TEXT("Type"), TEXT("Name"), TEXT("Bindings"), TEXT("Struct Idx"));
 	for (const FBakedStateTreeState& State : StateTree->States)
 	{
 		// Tasks
@@ -983,9 +953,9 @@ void FStateTreeExecutionContext::DebugPrintInternalLayout(FStateTreeItemView Ext
 		{
 			for (int32 j = 0; j < State.TasksNum; j++)
 			{
-				const FStateTreeTaskBase& Task = GetItem(Storage, State.TasksBegin + j).Get<FStateTreeTaskBase>();
+				const FStateTreeTaskBase& Task = GetInstanceData(Storage, State.TasksBegin + j).Get<FStateTreeTaskBase>();
 				DebugString += FString::Printf(TEXT("  | %-30s | %-12s | %-30s | %15s | %10d |\n"), *State.Name.ToString(),
-					TEXT("  Task"), *Task.Name.ToString(), *Task.BindingsBatch.Describe(), Task.SourceStructIndex);
+					TEXT("  Task"), *Task.Name.ToString(), *Task.BindingsBatch.Describe(), Task.DataViewIndex);
 			}
 		}
 
@@ -994,9 +964,9 @@ void FStateTreeExecutionContext::DebugPrintInternalLayout(FStateTreeItemView Ext
 		{
 			for (int32 EvalIndex = 0; EvalIndex < State.EvaluatorsNum; EvalIndex++)
 			{
-				const FStateTreeEvaluatorBase& Eval = GetItem(Storage, State.EvaluatorsBegin + EvalIndex).Get<FStateTreeEvaluatorBase>();
+				const FStateTreeEvaluatorBase& Eval = GetInstanceData(Storage, State.EvaluatorsBegin + EvalIndex).Get<FStateTreeEvaluatorBase>();
 				DebugString += FString::Printf(TEXT("  | %-30s | %-12s | %-30s | %15s | %10d |\n"), *State.Name.ToString(),
-					TEXT("  Evaluator"), *Eval.Name.ToString(), *Eval.BindingsBatch.Describe(), Eval.SourceStructIndex);
+					TEXT("  Evaluator"), *Eval.Name.ToString(), *Eval.BindingsBatch.Describe(), Eval.DataViewIndex);
 			}
 		}
 	}
@@ -1004,13 +974,13 @@ void FStateTreeExecutionContext::DebugPrintInternalLayout(FStateTreeItemView Ext
 	UE_LOG(LogStateTree, Log, TEXT("%s"), *DebugString);
 }
 
-FString FStateTreeExecutionContext::GetActiveStateName(FStateTreeItemView ExternalStorage) const
+FString FStateTreeExecutionContext::GetActiveStateName(FStateTreeDataView ExternalStorage) const
 {
 	if (!StateTree)
 	{
 		return FString(TEXT("<None>"));
 	}
-	const FStateTreeItemView Storage = SelectStorage(ExternalStorage);
+	const FStateTreeDataView Storage = SelectStorage(ExternalStorage);
 	FStateTreeExecutionState& Exec = GetExecState(Storage);
 
 	FString FullStateName;
diff --git a/Engine/Plugins/Runtime/StateTree/Source/StateTreeModule/Private/StateTreePropertyBindings.cpp b/Engine/Plugins/Runtime/StateTree/Source/StateTreeModule/Private/StateTreePropertyBindings.cpp
--- a/Engine/Plugins/Runtime/StateTree/Source/StateTreeModule/Private/StateTreePropertyBindings.cpp
+++ b/Engine/Plugins/Runtime/StateTree/Source/StateTreeModule/Private/StateTreePropertyBindings.cpp
@@ -292,7 +292,7 @@ bool FStateTreePropertyBindings::ValidateCopy(FStateTreePropCopy& Copy) const
 	return bResult;
 }
 
-uint8* FStateTreePropertyBindings::GetAddress(FStateTreeItemView InStructView, const FStateTreePropertyAccess& InAccess) const
+uint8* FStateTreePropertyBindings::GetAddress(FStateTreeDataView InStructView, const FStateTreePropertyAccess& InAccess) const
 {
 	check(InStructView.IsValid());
 
@@ -453,7 +453,7 @@ void FStateTreePropertyBindings::PerformCopy(const FStateTreePropCopy& Copy, con
 	}
 }
 
-void FStateTreePropertyBindings::CopyTo(TConstArrayView<FStateTreeItemView> SourceStructViews, const int32 TargetBatchIndex, FStateTreeItemView TargetStructView) const
+void FStateTreePropertyBindings::CopyTo(TConstArrayView<FStateTreeDataView> SourceStructViews, const int32 TargetBatchIndex, FStateTreeDataView TargetStructView) const
 {
 	// This is made ensure so that the programmers have the change to catch it (it's usually programming error not to call ResolvePaths(), and it wont spam log for others.
 	if (!ensureMsgf(bBindingsResolved, TEXT("Bindings must be resolved successfully before copying. See ResolvePaths()")))
diff --git a/Engine/Plugins/Runtime/StateTree/Source/StateTreeModule/Private/StateTreeTypes.cpp b/Engine/Plugins/Runtime/StateTree/Source/StateTreeModule/Private/StateTreeTypes.cpp
--- a/Engine/Plugins/Runtime/StateTree/Source/StateTreeModule/Private/StateTreeTypes.cpp
+++ b/Engine/Plugins/Runtime/StateTree/Source/StateTreeModule/Private/StateTreeTypes.cpp
@@ -8,4 +8,4 @@ const FStateTreeHandle FStateTreeHandle::Invalid = FStateTreeHandle();
 const FStateTreeHandle FStateTreeHandle::Succeeded = FStateTreeHandle(FStateTreeHandle::SucceededIndex);
 const FStateTreeHandle FStateTreeHandle::Failed = FStateTreeHandle(FStateTreeHandle::FailedIndex);
 
-const FStateTreeItemHandle FStateTreeItemHandle::Invalid = FStateTreeItemHandle();
+const FStateTreeExternalDataHandle FStateTreeExternalDataHandle::Invalid = FStateTreeExternalDataHandle();
diff --git a/Engine/Plugins/Runtime/StateTree/Source/StateTreeModule/Public/Conditions/StateTreeCondition_Common.h b/Engine/Plugins/Runtime/StateTree/Source/StateTreeModule/Public/Conditions/StateTreeCondition_Common.h
--- a/Engine/Plugins/Runtime/StateTree/Source/StateTreeModule/Public/Conditions/StateTreeCondition_Common.h
+++ b/Engine/Plugins/Runtime/StateTree/Source/StateTreeModule/Public/Conditions/StateTreeCondition_Common.h
@@ -4,173 +4,280 @@
 
 #include "CoreMinimal.h"
 #include "AITypes.h"
+#include "InstancedStruct.h"
 #include "StateTreeConditionBase.h"
 #include "StateTreeAnyEnum.h"
 #include "StateTreeCondition_Common.generated.h"
 
+enum class EStateTreeCompare : uint8
+{
+	Default,
+	Invert,
+};
+
 /**
  * Condition comparing two integers.
  */
+USTRUCT()
+struct STATETREEMODULE_API FStateTreeCondition_CompareIntInstanceData
+{
+	GENERATED_BODY()
+	
+	UPROPERTY(EditAnywhere, Category = Input)
+	int32 Left = 0;
+
+	UPROPERTY(EditAnywhere, Category = Parameter)
+	int32 Right = 0;
+};
+
 USTRUCT(DisplayName="Integer Compare")
 struct STATETREEMODULE_API FStateTreeCondition_CompareInt : public FStateTreeConditionBase
 {
 	GENERATED_BODY()
 
-	FStateTreeCondition_CompareInt() : FStateTreeConditionBase() { }
-	FStateTreeCondition_CompareInt(const int32 InLeft, const EGenericAICheck InOperator, const int32 InRight)
-		: FStateTreeConditionBase()
-		, Left(InLeft)
+	typedef FStateTreeCondition_CompareIntInstanceData InstanceDataType;
+
+	FStateTreeCondition_CompareInt() = default;
+	explicit FStateTreeCondition_CompareInt(const EGenericAICheck InOperator, const EStateTreeCompare InInverts = EStateTreeCompare::Default)
+		: bInvert(InInverts == EStateTreeCompare::Invert)
 		, Operator(InOperator)
-		, Right(InRight)
 	{}
 
+	virtual bool Link(FStateTreeLinker& Linker) override;
+	virtual const UStruct* GetInstanceDataType() const override { return FStateTreeCondition_CompareIntInstanceData::StaticStruct(); }
+	virtual bool TestCondition(FStateTreeExecutionContext& Context) const override;
+
 #if WITH_EDITOR
-	virtual FText GetDescription(const IStateTreeBindingLookup& BindingLookup) const override;
+	virtual FText GetDescription(const FGuid& ID, FConstStructView InstanceData, const IStateTreeBindingLookup& BindingLookup) const override;
 #endif
-	virtual bool TestCondition() const override;
 
-	UPROPERTY(EditAnywhere, Category = Condition, meta = (Bindable))
-	int32 Left = 0;
+	UPROPERTY(EditAnywhere, Category = Parameter)
+	bool bInvert = false;
 
-	UPROPERTY(EditAnywhere, Category = Condition)
+	UPROPERTY(EditAnywhere, Category = Parameter)
 	EGenericAICheck Operator = EGenericAICheck::Equal;
-
-	UPROPERTY(EditAnywhere, Category = Condition, meta = (Bindable))
-	int32 Right = 0;
-
-	UPROPERTY(EditAnywhere, Category = Condition)
-	bool bInvert = false;
+	
+	TStateTreeInstanceDataPropertyHandle<int32> LeftHandle;
+	TStateTreeInstanceDataPropertyHandle<int32> RightHandle;
 };
 
 /**
  * Condition comparing two floats.
  */
+USTRUCT()
+struct STATETREEMODULE_API FStateTreeCondition_CompareFloatInstanceData
+{
+	GENERATED_BODY()
+
+	UPROPERTY(EditAnywhere, Category = Input)
+	float Left = 0.0f;
+
+	UPROPERTY(EditAnywhere, Category = Parameter)
+	float Right = 0.0f;
+};
+
 USTRUCT(DisplayName = "Float Compare")
 struct STATETREEMODULE_API FStateTreeCondition_CompareFloat : public FStateTreeConditionBase
 {
 	GENERATED_BODY()
 
-	FStateTreeCondition_CompareFloat() : FStateTreeConditionBase() {}
-	FStateTreeCondition_CompareFloat(const float InLeft, const EGenericAICheck InOperator, const float InRight)
-        : FStateTreeConditionBase()
-        , Left(InLeft)
-        , Operator(InOperator)
-        , Right(InRight)
+	typedef FStateTreeCondition_CompareFloatInstanceData InstanceDataType;
+
+	FStateTreeCondition_CompareFloat() = default;
+	explicit FStateTreeCondition_CompareFloat(const EGenericAICheck InOperator, const EStateTreeCompare InInverts = EStateTreeCompare::Default)
+		: bInvert(InInverts == EStateTreeCompare::Invert)
+		, Operator(InOperator)
 	{}
 
+	virtual bool Link(FStateTreeLinker& Linker) override;
+	virtual const UStruct* GetInstanceDataType() const override { return FStateTreeCondition_CompareFloatInstanceData::StaticStruct(); }
+	virtual bool TestCondition(FStateTreeExecutionContext& Context) const override;
+
 #if WITH_EDITOR
-	virtual FText GetDescription(const IStateTreeBindingLookup& BindingLookup) const override;
+	virtual FText GetDescription(const FGuid& ID, FConstStructView InstanceData, const IStateTreeBindingLookup& BindingLookup) const override;
 #endif
-	virtual bool TestCondition() const override;
 
-	UPROPERTY(EditAnywhere, Category = Condition, meta = (Bindable))
-	float Left = 0.0f;
+	UPROPERTY(EditAnywhere, Category = Parameter)
+	bool bInvert = false;
 
-	UPROPERTY(EditAnywhere, Category = Condition)
+	UPROPERTY(EditAnywhere, Category = Parameter)
 	EGenericAICheck Operator = EGenericAICheck::Equal;
-
-	UPROPERTY(EditAnywhere, Category = Condition, meta = (Bindable))
-	float Right = 0.0f;
-
-	UPROPERTY(EditAnywhere, Category = Condition)
-	bool bInvert = false;
+	
+	TStateTreeInstanceDataPropertyHandle<float> LeftHandle;
+	TStateTreeInstanceDataPropertyHandle<float> RightHandle;
 };
 
 /**
  * Condition comparing two booleans.
  */
+USTRUCT()
+struct STATETREEMODULE_API FStateTreeCondition_CompareBoolInstanceData
+{
+	GENERATED_BODY()
+
+	UPROPERTY(EditAnywhere, Category = Input)
+	bool bLeft = false;
+
+	UPROPERTY(EditAnywhere, Category = Parameter)
+	bool bRight = false;
+};
+
 USTRUCT(DisplayName = "Bool Compare")
 struct STATETREEMODULE_API FStateTreeCondition_CompareBool : public FStateTreeConditionBase
 {
 	GENERATED_BODY()
 
-	FStateTreeCondition_CompareBool() : FStateTreeConditionBase() {}
-	FStateTreeCondition_CompareBool(const bool bInLeft, const bool bInRight)
-        : FStateTreeConditionBase()
-        , bLeft(bInLeft)
-        , bRight(bInRight)
+	typedef FStateTreeCondition_CompareBoolInstanceData InstanceDataType;
+
+	FStateTreeCondition_CompareBool() = default;
+	explicit FStateTreeCondition_CompareBool(const EStateTreeCompare InInverts)
+		: bInvert(InInverts == EStateTreeCompare::Invert)
 	{}
 
-#if WITH_EDITOR
-	virtual FText GetDescription(const IStateTreeBindingLookup& BindingLookup) const override;
-#endif
-	virtual bool TestCondition() const override;
+	FStateTreeCondition_CompareBool(const bool bInInverts)
+		: bInvert(bInInverts)
+	{}
 
-	UPROPERTY(EditAnywhere, Category = Condition, meta = (Bindable))
-	bool bLeft = false;
+	virtual bool Link(FStateTreeLinker& Linker) override;
+	virtual const UStruct* GetInstanceDataType() const override { return FStateTreeCondition_CompareBoolInstanceData::StaticStruct(); }
+	virtual bool TestCondition(FStateTreeExecutionContext& Context) const override;
 
-	UPROPERTY(EditAnywhere, Category = Condition, meta = (Bindable))
-	bool bRight = false;
+#if WITH_EDITOR
+	virtual FText GetDescription(const FGuid& ID, FConstStructView InstanceData, const IStateTreeBindingLookup& BindingLookup) const override;
+#endif
 
-	UPROPERTY(EditAnywhere, Category = Condition)
+	UPROPERTY(EditAnywhere, Category = Parameter)
 	bool bInvert = false;
+
+	TStateTreeInstanceDataPropertyHandle<bool> LeftHandle;
+	TStateTreeInstanceDataPropertyHandle<bool> RightHandle;
 };
 
 /**
  * Condition comparing two enums.
  */
+USTRUCT()
+struct STATETREEMODULE_API FStateTreeCondition_CompareEnumInstanceData
+{
+	GENERATED_BODY()
+
+	UPROPERTY(EditAnywhere, Category = Input)
+	FStateTreeAnyEnum Left;
+
+	UPROPERTY(EditAnywhere, Category = Parameter)
+	FStateTreeAnyEnum Right;
+};
+
 USTRUCT(DisplayName = "Enum Compare")
 struct STATETREEMODULE_API FStateTreeCondition_CompareEnum : public FStateTreeConditionBase
 {
 	GENERATED_BODY()
 
-	FStateTreeCondition_CompareEnum() : FStateTreeConditionBase() {}
-	FStateTreeCondition_CompareEnum(const FStateTreeAnyEnum& InLeft, const FStateTreeAnyEnum& InRight)
-        : FStateTreeConditionBase()
-        , Left(InLeft)
-        , Right(InRight)
+	typedef FStateTreeCondition_CompareEnumInstanceData InstanceDataType;
+
+	FStateTreeCondition_CompareEnum() = default;
+	explicit FStateTreeCondition_CompareEnum(const EStateTreeCompare InInverts)
+		: bInvert(InInverts == EStateTreeCompare::Invert)
 	{}
 
-#if WITH_EDITOR
-	virtual FText GetDescription(const IStateTreeBindingLookup& BindingLookup) const override;
-	virtual void OnBindingChanged(const FStateTreeEditorPropertyPath& SourcePath, const FStateTreeEditorPropertyPath& TargetPath, const IStateTreeBindingLookup& BindingLookup) override;
-#endif
-	virtual bool TestCondition() const override;
+	FStateTreeCondition_CompareEnum(const bool bInInverts)
+		: bInvert(bInInverts)
+	{}
 
-	UPROPERTY(EditAnywhere, Category = Condition, meta = (Bindable))
-	FStateTreeAnyEnum Left;
+	virtual bool Link(FStateTreeLinker& Linker) override;
+	virtual const UStruct* GetInstanceDataType() const override { return FStateTreeCondition_CompareEnumInstanceData::StaticStruct(); }
+	virtual bool TestCondition(FStateTreeExecutionContext& Context) const override;
 
-	UPROPERTY(EditAnywhere, Category = Condition, meta = (Bindable))
-	FStateTreeAnyEnum Right;
+#if WITH_EDITOR
+	virtual FText GetDescription(const FGuid& ID, FConstStructView InstanceData, const IStateTreeBindingLookup& BindingLookup) const override;
+	virtual void OnBindingChanged(const FGuid& ID, FStructView InstanceData, const FStateTreeEditorPropertyPath& SourcePath, const FStateTreeEditorPropertyPath& TargetPath, const IStateTreeBindingLookup& BindingLookup) override;
+#endif
 
-	UPROPERTY(EditAnywhere, Category = Condition)
+	UPROPERTY(EditAnywhere, Category = Parameter)
 	bool bInvert = false;
+
+	TStateTreeInstanceDataPropertyHandle<FStateTreeAnyEnum> LeftHandle;
+	TStateTreeInstanceDataPropertyHandle<FStateTreeAnyEnum> RightHandle;
 };
 
 /**
  * Condition comparing distance between two vectors.
  */
+USTRUCT()
+struct STATETREEMODULE_API FStateTreeCondition_CompareDistanceInstanceData
+{
+	GENERATED_BODY()
+
+	UPROPERTY(EditAnywhere, Category = Input)
+	FVector Source = FVector(EForceInit::ForceInitToZero);
+
+	UPROPERTY(EditAnywhere, Category = Parameter)
+	FVector Target = FVector(EForceInit::ForceInitToZero);
+
+	UPROPERTY(EditAnywhere, Category = Parameter)
+	float Distance = 0.0f;
+};
+
 USTRUCT(DisplayName = "Distance Compare")
 struct STATETREEMODULE_API FStateTreeCondition_CompareDistance : public FStateTreeConditionBase
 {
 	GENERATED_BODY()
 
-	FStateTreeCondition_CompareDistance() : FStateTreeConditionBase() { }
-	FStateTreeCondition_CompareDistance(const FVector& InSource, const FVector& InTarget, const EGenericAICheck InOperator, const float InDistance)
-        : FStateTreeConditionBase()
-        , Source(InSource)
-        , Target(InTarget)
-        , Operator(InOperator)
-        , Distance(InDistance)
+	typedef FStateTreeCondition_CompareDistanceInstanceData InstanceDataType;
+
+	FStateTreeCondition_CompareDistance() = default;
+	explicit FStateTreeCondition_CompareDistance(const EGenericAICheck InOperator, const EStateTreeCompare InInverts = EStateTreeCompare::Default)
+		: bInvert(InInverts == EStateTreeCompare::Invert)
+		, Operator(InOperator)
 	{}
+	
+	virtual bool Link(FStateTreeLinker& Linker) override;
+	virtual const UStruct* GetInstanceDataType() const override { return FStateTreeCondition_CompareDistanceInstanceData::StaticStruct(); }
+	virtual bool TestCondition(FStateTreeExecutionContext& Context) const override;
 
 #if WITH_EDITOR
-	virtual FText GetDescription(const IStateTreeBindingLookup& BindingLookup) const override;
+	virtual FText GetDescription(const FGuid& ID, FConstStructView InstanceData, const IStateTreeBindingLookup& BindingLookup) const override;
 #endif
-	virtual bool TestCondition() const override;
 
-	UPROPERTY(EditAnywhere, Category = Condition, meta = (Bindable))
-	FVector Source = FVector(EForceInit::ForceInitToZero);
-
-	UPROPERTY(EditAnywhere, Category = Condition, meta = (Bindable))
-	FVector Target = FVector(EForceInit::ForceInitToZero);
+	UPROPERTY(EditAnywhere, Category = Condition)
+	bool bInvert = false;
 
 	UPROPERTY(EditAnywhere, Category = Condition)
 	EGenericAICheck Operator = EGenericAICheck::Equal;
 
-	UPROPERTY(EditAnywhere, Category = Condition, meta = (Bindable))
-	float Distance = 0.0f;
+	TStateTreeInstanceDataPropertyHandle<FVector> SourceHandle;
+	TStateTreeInstanceDataPropertyHandle<FVector> TargetHandle;
+	TStateTreeInstanceDataPropertyHandle<float> DistanceHandle;
+};
 
-	UPROPERTY(EditAnywhere, Category = Condition)
-	bool bInvert = false;
+/**
+* Random condition
+*/
+
+USTRUCT()
+struct STATETREEMODULE_API FStateTreeCondition_RandomInstanceData
+{
+	GENERATED_BODY()
+	UPROPERTY(EditAnywhere, Category = Parameter, meta = (ClampMin = "0.0", ClampMax = "1.0", UIMin = "0.0", UIMax = "1.0"))
+	float Threshold = 0.5f;
+};
+
+USTRUCT(DisplayName = "Random")
+struct STATETREEMODULE_API FStateTreeCondition_Random : public FStateTreeConditionBase
+{
+	GENERATED_BODY()
+
+	typedef FStateTreeCondition_RandomInstanceData InstanceDataType;
+
+	FStateTreeCondition_Random() = default;
+	
+	virtual bool Link(FStateTreeLinker& Linker) override;
+	virtual const UStruct* GetInstanceDataType() const override { return FStateTreeCondition_RandomInstanceData::StaticStruct(); }
+	virtual bool TestCondition(FStateTreeExecutionContext& Context) const override;
+
+#if WITH_EDITOR
+	virtual FText GetDescription(const FGuid& ID, FConstStructView InstanceData, const IStateTreeBindingLookup& BindingLookup) const override;
+#endif
+
+	TStateTreeInstanceDataPropertyHandle<float> ThresholdHandle;
 };
diff --git a/Engine/Plugins/Runtime/StateTree/Source/StateTreeModule/Public/StateTree.h b/Engine/Plugins/Runtime/StateTree/Source/StateTreeModule/Public/StateTree.h
--- a/Engine/Plugins/Runtime/StateTree/Source/StateTreeModule/Public/StateTree.h
+++ b/Engine/Plugins/Runtime/StateTree/Source/StateTreeModule/Public/StateTree.h
@@ -23,22 +23,22 @@ public:
 	UStateTree(const FObjectInitializer& ObjectInitializer = FObjectInitializer::Get());
 
 	/** @return Script Struct that can be used to instantiate the runtime storage */
-	const UScriptStruct* GetRuntimeStorageStruct() const { return RuntimeStorageStruct; }
+	const UScriptStruct* GetInstanceStorageStruct() const { return InstanceStorageStruct; }
 
 	/** @return Instance of the runtime storage that contains the default values */
-	const FInstancedStruct& GetRuntimeStorageDefaultValue() const { return RuntimeStorageDefaultValue; }
+	const FInstancedStruct& GetInstanceStorageDefaultValue() const { return InstanceStorageDefaultValue; }
 
-	/** @return Number of items (Evaluators & Tasks) in the runtime storage. */
-	int32 GetRuntimeStorageItemCount() const { return RuntimeStorageOffsets.Num(); }
+	/** @return Number of runtime data (Evaluators, Tasks, Conditions) in the runtime storage. */
+	int32 GetNumInstances() const { return Instances.Num(); }
 
-	/** @return Number of linked items (Evaluators, Tasks, external items). */
-	int32 GetLinkedItemCount() const { return NumLinkedItems; }
+	/** @return Number of data views required for StateTree execution (Evaluators, Tasks, Conditions, External data). */
+	int32 GetNumDataViews() const { return NumDataViews; }
 
-	/** @return Base index for external items. */
-	int32 GetExternalItemBaseIndex() const { return ExternalItemBaseIndex; }
+	/** @return Base index in data views for external data. */
+	int32 GetExternalDataBaseIndex() const { return ExternalDataBaseIndex; }
 
-	/** @return List of external items required by the state tree */
-	TConstArrayView<FStateTreeExternalItemDesc> GetExternalItems() const { return ExternalItems; }
+	/** @return List of external data required by the state tree */
+	TConstArrayView<FStateTreeExternalDataDesc> GetExternalDataDescs() const { return ExternalDataDescs; }
 
 	/** @return Schema describing which inputs, evaluators, and tasks a StateTree can contain */
 	const UStateTreeSchema* GetSchema() const { return Schema; }
@@ -70,7 +70,7 @@ protected:
 #endif
 
 	/** Initializes the types and default values related to StateTree runtime storage */
-	void InitRuntimeStorage();
+	void InitInstanceStorage();
 
 	/** Resolved references between data in the StateTree. */
 	void Link();
@@ -82,33 +82,34 @@ private:
 	UPROPERTY(EditDefaultsOnly, Category = Common, Instanced)
 	UStateTreeSchema* Schema = nullptr;
 
-	/** Evaluators and Tasks that require runtime state */
+	/** Evaluators, Tasks, and Condition items */
 	UPROPERTY()
-	TArray<FInstancedStruct> RuntimeStorageItems;
+	TArray<FInstancedStruct> Items;
+
+	/** Evaluators, Tasks, and Conditions runtime data. */
+	UPROPERTY()
+	TArray<FInstancedStruct> Instances;
 
 	/** Script Struct that can be used to instantiate the runtime storage */
 	UPROPERTY()
-	UScriptStruct* RuntimeStorageStruct;
+	UScriptStruct* InstanceStorageStruct;
 
 	/** Offsets into the runtime type to quickly get a struct view to a specific Task or Evaluator */
-	TArray<FStateTreeRuntimeStorageItemOffset> RuntimeStorageOffsets;
+	TArray<FStateTreeInstanceStorageOffset> InstanceStorageOffsets;
 
 	/** Instance of the runtime storage that contains the default values. */
 	UPROPERTY()
-	FInstancedStruct RuntimeStorageDefaultValue;
-
-	UPROPERTY(meta = (BaseStruct = "StateTreeConditionBase", ExcludeBaseStruct))
-	TArray<FInstancedStruct> Conditions;
+	FInstancedStruct InstanceStorageDefaultValue;
 
-	/** List of external items required by the state tree, creating during linking. */
+	/** List of external data required by the state tree, creating during linking. */
 	UPROPERTY(Transient)
-	TArray<FStateTreeExternalItemDesc> ExternalItems;
+	TArray<FStateTreeExternalDataDesc> ExternalDataDescs;
 
 	UPROPERTY(Transient)
-	int32 NumLinkedItems = 0;
+	int32 NumDataViews = 0;
 
 	UPROPERTY(Transient)
-	int32 ExternalItemBaseIndex = 0;
+	int32 ExternalDataBaseIndex = 0;
 
 	UPROPERTY()
 	FStateTreePropertyBindings PropertyBindings;
diff --git a/Engine/Plugins/Runtime/StateTree/Source/StateTreeModule/Public/StateTreeConditionBase.h b/Engine/Plugins/Runtime/StateTree/Source/StateTreeModule/Public/StateTreeConditionBase.h
--- a/Engine/Plugins/Runtime/StateTree/Source/StateTreeModule/Public/StateTreeConditionBase.h
+++ b/Engine/Plugins/Runtime/StateTree/Source/StateTreeModule/Public/StateTreeConditionBase.h
@@ -5,10 +5,12 @@
 #include "CoreMinimal.h"
 #include "EngineDefines.h"
 #include "StateTreeTypes.h"
+#include "InstancedStruct.h"
 #include "StateTreeConditionBase.generated.h"
 
 struct IStateTreeBindingLookup;
 struct FStateTreeEditorPropertyPath;
+struct FStateTreeExecutionContext;
 
 /**
  * Base struct for all conditions.
@@ -18,33 +20,41 @@ struct STATETREEMODULE_API FStateTreeConditionBase
 {
 	GENERATED_BODY()
 
-	FStateTreeConditionBase()
-	{
-	}
-
 	virtual ~FStateTreeConditionBase() {}
 
+	/**
+	* @return Struct that represents the runtime data of the condition.
+	*/
+	virtual const UStruct* GetInstanceDataType() const PURE_VIRTUAL(FStateTreeConditionBase::GetInstanceDataType(), return nullptr;);
+	
+	/**
+	 * Called when the StateTree asset is linked. Allows to resolve references to other StateTree data.
+	 * @see TStateTreeExternalDataHandle
+	 * @see TStateTreeInstanceDataPropertyHandle
+	 * @param Linker Reference to the linker
+	 * @return true if linking succeeded. 
+	 */
+	virtual bool Link(FStateTreeLinker& Linker) { return true; }
+
 #if WITH_EDITOR
 	/** @return Rich text description of the condition. */
-	virtual FText GetDescription(const IStateTreeBindingLookup& BindingLookup) const { return FText::GetEmpty(); }
+	virtual FText GetDescription(const FGuid& ID, FConstStructView LayoutData, const IStateTreeBindingLookup& BindingLookup) const { return FText::GetEmpty(); }
 	/**
 	 * Called when binding of any of the properties in the condition changes.
 	 * @param SourcePath Source path of the new binding.
 	 * @param TargetPath Target path of the new binding (the property in the condition).
 	 * @param BindingLookup Reference to binding lookup which can be used to reason about property paths.
 	 */
-	virtual void OnBindingChanged(const FStateTreeEditorPropertyPath& SourcePath, const FStateTreeEditorPropertyPath& TargetPath, const IStateTreeBindingLookup& BindingLookup) {}
+	virtual void OnBindingChanged(const FGuid& ID, FStructView LayoutData, const FStateTreeEditorPropertyPath& SourcePath, const FStateTreeEditorPropertyPath& TargetPath, const IStateTreeBindingLookup& BindingLookup) {}
 #endif
 	/** @return True if the condition passes. */
-	virtual bool TestCondition() const { return false; }
-
-#if WITH_EDITORONLY_DATA
-	/** Edit time struct ID for property binding. */
-	UPROPERTY(EditDefaultsOnly, Category = Evaluator, meta = (IgnoreForMemberInitializationTest, EditCondition = "false", EditConditionHides))	// Hack, we want the ID to be found as IPropertyHandle, but do not want to display it.
-	FGuid ID;
-#endif
+	virtual bool TestCondition(FStateTreeExecutionContext& Context) const { return false; }
 
 	/** Property binding copy batch handle. */
 	UPROPERTY()
 	FStateTreeHandle BindingsBatch = FStateTreeHandle::Invalid;
+
+	/** The runtime data's data view index in the StateTreeExecutionContext, and source struct index in property binding. */
+	UPROPERTY()
+	uint16 DataViewIndex = 0;
 };
