StateTree: allow task/evaluator/condition to have separate data for isntance and item in the tree.

diff --git a/Engine/Plugins/Runtime/StateTree/Source/StateTreeTestSuite/Private/StateTreeTestTypes.h b/Engine/Plugins/Runtime/StateTree/Source/StateTreeTestSuite/Private/StateTreeTestTypes.h
--- a/Engine/Plugins/Runtime/StateTree/Source/StateTreeTestSuite/Private/StateTreeTestTypes.h
+++ b/Engine/Plugins/Runtime/StateTree/Source/StateTreeTestSuite/Private/StateTreeTestTypes.h
@@ -8,15 +8,6 @@
 #include "StateTreeExecutionContext.h"
 #include "StateTreeTestTypes.generated.h"
 
-/**
- * This macro can be used to create property path for property binding. Checks that the member exists.
- * Example:
- *		EditorData.AddPropertyBinding(STATETREE_PROPPATH_CHECKED(EvalA, IntA), STATETREE_PROPPATH_CHECKED(TaskB1, IntB));
- **/
-#define STATETREE_PROPPATH_CHECKED(Struct, MemberName) \
-FStateTreeEditorPropertyPath(Struct.ID, ((void)sizeof(UEAsserts_Private::GetMemberNameCheckedJunk((&Struct)->MemberName)), TEXT(#MemberName)))
-
-
 
 USTRUCT()
 struct FTestStateTreeExecutionContext : public FStateTreeExecutionContext
@@ -74,59 +65,95 @@ struct FTestStateTreeExecutionContext : public FStateTreeExecutionContext
 	
 };
 
-
 USTRUCT()
-struct FTestEval_A : public FStateTreeEvaluatorBase
+struct FTestEval_AInstanceData
 {
 	GENERATED_BODY()
-
-	FTestEval_A() = default;
-	virtual ~FTestEval_A() {}
-
-	UPROPERTY(EditAnywhere, Category = Test, meta = (Bindable))
+	
+	UPROPERTY(EditAnywhere, Category = Parameter)
 	float FloatA = 0.0f;
 
-	UPROPERTY(EditAnywhere, Category = Test, meta = (Bindable))
+	UPROPERTY(EditAnywhere, Category = Parameter)
 	int32 IntA = 0;
 
-	UPROPERTY(EditAnywhere, Category = Test, meta = (Bindable))
+	UPROPERTY(EditAnywhere, Category = Parameter)
 	bool bBoolA = false;
 };
 
 USTRUCT()
-struct FTestTask_B : public FStateTreeTaskBase
+struct FTestEval_A : public FStateTreeEvaluatorBase
 {
 	GENERATED_BODY()
 
-	FTestTask_B() = default;
-	virtual ~FTestTask_B() {}
+	typedef FTestEval_AInstanceData InstanceDataType;
+
+	FTestEval_A() = default;
+	virtual ~FTestEval_A() override {}
+
+	virtual const UStruct* GetInstanceDataType() const override { return FTestEval_AInstanceData::StaticStruct(); }
 	
-	UPROPERTY(EditAnywhere, Category = Task, meta = (Bindable))
+	virtual bool Link(FStateTreeLinker& Linker) override
+	{
+		Linker.LinkInstanceDataProperty(FloatAHandle, STATETREE_INSTANCEDATA_PROPERTY(FTestEval_AInstanceData, FloatA));
+		Linker.LinkInstanceDataProperty(IntAHandle, STATETREE_INSTANCEDATA_PROPERTY(FTestEval_AInstanceData, IntA));
+		Linker.LinkInstanceDataProperty(BoolAHandle, STATETREE_INSTANCEDATA_PROPERTY(FTestEval_AInstanceData, bBoolA));
+		return true;
+	}
+
+	TStateTreeInstanceDataPropertyHandle<float> FloatAHandle;
+	TStateTreeInstanceDataPropertyHandle<int32> IntAHandle;
+	TStateTreeInstanceDataPropertyHandle<bool> BoolAHandle;
+};
+
+USTRUCT()
+struct FTestTask_BInstanceData
+{
+	GENERATED_BODY()
+
+	UPROPERTY(EditAnywhere, Category = Parameter)
 	float FloatB = 0.0f;
 
-	UPROPERTY(EditAnywhere, Category = Task, meta = (Bindable))
+	UPROPERTY(EditAnywhere, Category = Parameter)
 	int32 IntB = 0;
 
-	UPROPERTY(EditAnywhere, Category = Task, meta = (Bindable))
+	UPROPERTY(EditAnywhere, Category = Parameter)
 	bool bBoolB = false;
 };
 
+USTRUCT()
+struct FTestTask_B : public FStateTreeTaskBase
+{
+	GENERATED_BODY()
 
+	typedef FTestTask_BInstanceData InstanceDataType;
 
+	FTestTask_B() = default;
+	virtual ~FTestTask_B() override {}
+	
+	virtual const UStruct* GetInstanceDataType() const override { return FTestTask_BInstanceData::StaticStruct(); }
+	
+	virtual bool Link(FStateTreeLinker& Linker) override
+	{
+		Linker.LinkInstanceDataProperty(FloatBHandle, STATETREE_INSTANCEDATA_PROPERTY(FTestTask_BInstanceData, FloatB));
+		Linker.LinkInstanceDataProperty(IntBHandle, STATETREE_INSTANCEDATA_PROPERTY(FTestTask_BInstanceData, IntB));
+		Linker.LinkInstanceDataProperty(BoolBHandle, STATETREE_INSTANCEDATA_PROPERTY(FTestTask_BInstanceData, bBoolB));
+		return true;
+	}
+
+	TStateTreeInstanceDataPropertyHandle<float> FloatBHandle;
+	TStateTreeInstanceDataPropertyHandle<int32> IntBHandle;
+	TStateTreeInstanceDataPropertyHandle<bool> BoolBHandle;
+};
 
 USTRUCT()
-struct FTestMoveLocationResult : public FStateTreeResult
+struct FTestMoveLocationResult
 {
 	GENERATED_BODY()
 
-	virtual ~FTestMoveLocationResult() {}
-	virtual const UScriptStruct& GetStruct() const override { return *StaticStruct(); }
-
 	void Reset()
 	{
 		Location = FVector::ZeroVector;
 		Direction = FVector::ZeroVector;
-		Status = EStateTreeResultStatus::Unset;
 	}
 	
 	UPROPERTY(EditAnywhere, Category = Value)
@@ -134,18 +161,13 @@ struct FTestMoveLocationResult : public FStateTreeResult
 
 	UPROPERTY(EditAnywhere, Category = Value)
 	FVector Direction = FVector::ZeroVector;
-
-	EStateTreeResultStatus Status = EStateTreeResultStatus::Unset;
 };
 
 USTRUCT()
-struct FTestPotentialSmartObjectsResult : public FStateTreeResult
+struct FTestPotentialSmartObjectsResult
 {
 	GENERATED_BODY()
 
-	virtual ~FTestPotentialSmartObjectsResult() {}
-	virtual const UScriptStruct& GetStruct() const override { return *StaticStruct(); }
-
 	bool HasSmartObjects() const { return NumSmartObjects > 0 && FailedCoolDown == 0; }
 	
 	FName PopSmartObjectName()
@@ -179,18 +201,16 @@ struct FTestPotentialSmartObjectsResult : public FStateTreeResult
 };
 
 USTRUCT()
-struct FTestSmartObjectResult : public FStateTreeResult
+struct FTestSmartObjectResult
 {
 	GENERATED_BODY()
 
-	virtual ~FTestSmartObjectResult() {}
-	virtual const UScriptStruct& GetStruct() const override { return *StaticStruct(); }
-
+	bool IsValid() const { return SmartObjectHandle != INDEX_NONE; }
+	
 	void Reset()
 	{
 		AnimationName = FName();
 		SmartObjectHandle = INDEX_NONE;
-		Status = EStateTreeResultStatus::Unset;
 	}
 	
 	UPROPERTY(EditAnywhere, Category = Value)
@@ -198,52 +218,85 @@ struct FTestSmartObjectResult : public FStateTreeResult
 
 	UPROPERTY(EditAnywhere, Category = Value)
 	int32 SmartObjectHandle = INDEX_NONE;
-
-	EStateTreeResultStatus Status = EStateTreeResultStatus::Unset;
 };
 
+USTRUCT()
+struct FTestEval_WanderInstanceData
+{
+	GENERATED_BODY()
+
+	// True if wander location available
+	UPROPERTY(EditAnywhere, Category = Output)
+	bool bHasWanderLocation = false;
+
+	UPROPERTY(EditAnywhere, Category = Output)
+	FTestMoveLocationResult NextMoveLocation;
+};
 
 USTRUCT()
 struct FTestEval_Wander : public FStateTreeEvaluatorBase
 {
 	GENERATED_BODY()
 
+	typedef FTestEval_WanderInstanceData InstanceDataType;
+	
 	FTestEval_Wander() = default;
 	FTestEval_Wander(const FName InName) { Name = InName; }
-	virtual ~FTestEval_Wander() {}
+	virtual ~FTestEval_Wander() override {}
 
-	virtual void EnterState(FStateTreeExecutionContext& Context, const EStateTreeStateChangeType ChangeType, const FStateTreeTransitionResult& Transition) override
+	virtual const UStruct* GetInstanceDataType() const override { return FTestEval_WanderInstanceData::StaticStruct(); }
+	
+	virtual bool Link(FStateTreeLinker& Linker) override
+	{
+		Linker.LinkInstanceDataProperty(HasWanderLocationHandle, STATETREE_INSTANCEDATA_PROPERTY(FTestEval_WanderInstanceData, bHasWanderLocation));
+		Linker.LinkInstanceDataProperty(NextMoveLocationHandle, STATETREE_INSTANCEDATA_PROPERTY(FTestEval_WanderInstanceData, NextMoveLocation));
+		return true;
+	}
+
+	virtual void EnterState(FStateTreeExecutionContext& Context, const EStateTreeStateChangeType ChangeType, const FStateTreeTransitionResult& Transition) const override
 	{
 		FTestStateTreeExecutionContext& TestContext = static_cast<FTestStateTreeExecutionContext&>(Context);
 		TestContext.Log(Name, TEXT("EnterState"));
 
-		WanderLocation = &NextMoveLocation;
+		bool& bHasWanderLocation = Context.GetInstanceData(HasWanderLocationHandle);
+		FTestMoveLocationResult& NextMoveLocation = Context.GetInstanceData(NextMoveLocationHandle);
+
+		bHasWanderLocation = true;
+		NextMoveLocation.Location = FVector::ZeroVector;
+		NextMoveLocation.Direction = FVector::ForwardVector;
 	}
 
-	virtual void ExitState(FStateTreeExecutionContext& Context, const EStateTreeStateChangeType ChangeType, const FStateTreeTransitionResult& Transition) override
+	virtual void ExitState(FStateTreeExecutionContext& Context, const EStateTreeStateChangeType ChangeType, const FStateTreeTransitionResult& Transition) const override
 	{
 		FTestStateTreeExecutionContext& TestContext = static_cast<FTestStateTreeExecutionContext&>(Context);
 		TestContext.Log(Name, TEXT("ExitState"));
-		
-		WanderLocation = &NextMoveLocation;
 	}
 
-	virtual void Evaluate(FStateTreeExecutionContext& Context, const EStateTreeEvaluationType EvalType, const float DeltaTime) override
+	virtual void Evaluate(FStateTreeExecutionContext& Context, const EStateTreeEvaluationType EvalType, const float DeltaTime) const override
 	{
 		FTestStateTreeExecutionContext& TestContext = static_cast<FTestStateTreeExecutionContext&>(Context);
 		TestContext.Log(Name, TEXT("Evaluate"));
-
-		WanderLocation = &NextMoveLocation;
 	}
+	
+	TStateTreeInstanceDataPropertyHandle<bool> HasWanderLocationHandle;
+	TStateTreeInstanceDataPropertyHandle<FTestMoveLocationResult> NextMoveLocationHandle;
+};
 
-	// True if wander location available
-	UPROPERTY(EditAnywhere, Category = Eval)
-	bool bHasWanderLocation = false;
+USTRUCT()
+struct FTestEval_SmartObjectSensorInstanceData
+{
+	GENERATED_BODY()
 
-	UPROPERTY(EditAnywhere, Category = Eval)
-	FStateTreeResultRef WanderLocation;
-	
-	FTestMoveLocationResult NextMoveLocation;
+	// Result output
+	UPROPERTY(EditAnywhere, Category = Output)
+	FTestPotentialSmartObjectsResult PotentialSmartObjects;
+
+	// True if any SmartObjects available
+	UPROPERTY(EditAnywhere, Category = Output)
+	bool bHasSmartObjects = false;
+
+	UPROPERTY()
+	bool bIsQueryingSmartObjects = false;
 };
 
 USTRUCT()
@@ -251,22 +304,32 @@ struct FTestEval_SmartObjectSensor : public FStateTreeEvaluatorBase
 {
 	GENERATED_BODY()
 
+	typedef FTestEval_SmartObjectSensorInstanceData InstanceDataType;
+
 	FTestEval_SmartObjectSensor() = default;
 	FTestEval_SmartObjectSensor(const FName InName) { Name = InName; }
 	virtual ~FTestEval_SmartObjectSensor() {}
 
-	virtual void EnterState(FStateTreeExecutionContext& Context, const EStateTreeStateChangeType ChangeType, const FStateTreeTransitionResult& Transition) override
+	virtual const UStruct* GetInstanceDataType() const override { return FTestEval_SmartObjectSensorInstanceData::StaticStruct(); }
+	
+	virtual bool Link(FStateTreeLinker& Linker) override
+	{
+		Linker.LinkInstanceDataProperty(PotentialSmartObjectsHandle, STATETREE_INSTANCEDATA_PROPERTY(FTestEval_SmartObjectSensorInstanceData, PotentialSmartObjects));
+		Linker.LinkInstanceDataProperty(HasSmartObjectsHandle, STATETREE_INSTANCEDATA_PROPERTY(FTestEval_SmartObjectSensorInstanceData, bHasSmartObjects));
+		Linker.LinkInstanceDataProperty(IsQueryingSmartObjectsHandle, STATETREE_INSTANCEDATA_PROPERTY(FTestEval_SmartObjectSensorInstanceData, bIsQueryingSmartObjects));
+		return true;
+	}
+
+	virtual void EnterState(FStateTreeExecutionContext& Context, const EStateTreeStateChangeType ChangeType, const FStateTreeTransitionResult& Transition) const override
 	{
 		FTestStateTreeExecutionContext& TestContext = static_cast<FTestStateTreeExecutionContext&>(Context);
 		TestContext.Log(Name, TEXT("EnterState"));
 
 		// We expect the state of the evaluator to be valid on EnterState.
 		// Evaluate() have been called during the state selection process, and some properties may have been used as conditions during select.
-		
-		PotentialSmartObjects = &SmartObjects;
 	}
 	
-	virtual void ExitState(FStateTreeExecutionContext& Context, const EStateTreeStateChangeType ChangeType, const FStateTreeTransitionResult& Transition) override
+	virtual void ExitState(FStateTreeExecutionContext& Context, const EStateTreeStateChangeType ChangeType, const FStateTreeTransitionResult& Transition) const override
 	{
 		FTestStateTreeExecutionContext& TestContext = static_cast<FTestStateTreeExecutionContext&>(Context);
 		TestContext.Log(Name, TEXT("ExitState"));
@@ -275,66 +338,91 @@ struct FTestEval_SmartObjectSensor : public FStateTreeEvaluatorBase
 		// during state selection.
 		if (ChangeType == EStateTreeStateChangeType::Changed)
 		{
+			FTestPotentialSmartObjectsResult& PotentialSmartObjects = Context.GetInstanceData(PotentialSmartObjectsHandle);
+			bool& bHasSmartObjects = Context.GetInstanceData(HasSmartObjectsHandle);
+			bool& bIsQueryingSmartObjects = Context.GetInstanceData(IsQueryingSmartObjectsHandle);
+
 			bHasSmartObjects = false;
 			bIsQueryingSmartObjects = false; // Cancel a query
-			SmartObjects.Reset(); // Note: this does not reset the cool down.
+			PotentialSmartObjects.Reset(); // Note: this does not reset the cool down.
 		}
-
-		PotentialSmartObjects = &SmartObjects;
 	}
 
-	virtual void Evaluate(FStateTreeExecutionContext& Context, const EStateTreeEvaluationType EvalType, const float DeltaTime) override
+	virtual void Evaluate(FStateTreeExecutionContext& Context, const EStateTreeEvaluationType EvalType, const float DeltaTime) const override
 	{
 		FTestStateTreeExecutionContext& TestContext = static_cast<FTestStateTreeExecutionContext&>(Context);
 		TestContext.Log(Name, TEXT("Evaluate"));
 
+		FTestPotentialSmartObjectsResult& PotentialSmartObjects = Context.GetInstanceData(PotentialSmartObjectsHandle);
+		bool& bHasSmartObjects = Context.GetInstanceData(HasSmartObjectsHandle);
+		bool& bIsQueryingSmartObjects = Context.GetInstanceData(IsQueryingSmartObjectsHandle);
+
 		// The cool down should be a timestamp instead of a counter, to catch a case where we re-enter the evaluator after being in different state.
-		if (SmartObjects.FailedCoolDown > 0)
+		if (PotentialSmartObjects.FailedCoolDown > 0)
 		{
-			SmartObjects.FailedCoolDown--;
+			PotentialSmartObjects.FailedCoolDown--;
 		}
-		bHasSmartObjects = SmartObjects.HasSmartObjects();
+		bHasSmartObjects = PotentialSmartObjects.HasSmartObjects();
 		
 		if (bIsQueryingSmartObjects)
 		{
 			// Simulate Async query result
 			bIsQueryingSmartObjects = false;
 			bHasSmartObjects = true;
-			SmartObjects.NumSmartObjects = 2;
-			SmartObjects.SmartObjectNames[0] = FName();
-			SmartObjects.SmartObjectNames[1] = FName(TEXT("Foo"));
+			PotentialSmartObjects.NumSmartObjects = 2;
+			PotentialSmartObjects.SmartObjectNames[0] = FName();
+			PotentialSmartObjects.SmartObjectNames[1] = FName(TEXT("Foo"));
 		}
-		if (!bHasSmartObjects && !bIsQueryingSmartObjects && SmartObjects.FailedCoolDown == 0)
+		if (!bHasSmartObjects && !bIsQueryingSmartObjects && PotentialSmartObjects.FailedCoolDown == 0)
 		{
 			// Simulate Async query start
 			bIsQueryingSmartObjects = true;
 		}
-
-		// Set ref
-		PotentialSmartObjects = &SmartObjects;
 	}
 
-	// Result output
-	UPROPERTY(EditAnywhere, Category = Eval, meta=(Output, Struct="TestPotentialSmartObjectsResult"))
-	FStateTreeResultRef PotentialSmartObjects;
+	TStateTreeInstanceDataPropertyHandle<FTestPotentialSmartObjectsResult> PotentialSmartObjectsHandle;
+	TStateTreeInstanceDataPropertyHandle<bool> HasSmartObjectsHandle;
+	TStateTreeInstanceDataPropertyHandle<bool> IsQueryingSmartObjectsHandle;
 
-	// True if any SmartObjects available
-	UPROPERTY(EditAnywhere, Category = Eval)
-	bool bHasSmartObjects = false;
-
-	bool bIsQueryingSmartObjects = false;
-	FTestPotentialSmartObjectsResult SmartObjects;
 };
 
 
+USTRUCT()
+struct FTestTask_ReserveSmartObjectInstanceData : public FStateTreeTaskBase
+{
+	GENERATED_BODY()
+
+	UPROPERTY(EditAnywhere, Category = Input)
+	FTestPotentialSmartObjectsResult PotentialSmartObjects;
+
+	// Because reserve is a task, these can be seen only by other Task (not evals, not conditions)
+	UPROPERTY(EditAnywhere, Category = Output)
+	FTestMoveLocationResult ReservedSmartObjectLocation;
+
+	UPROPERTY(EditAnywhere, Category = Output)
+	FTestSmartObjectResult ReservedSmartObject;
+};
+
 USTRUCT()
 struct FTestTask_ReserveSmartObject : public FStateTreeTaskBase
 {
 	GENERATED_BODY()
 
+	typedef FTestTask_ReserveSmartObjectInstanceData InstanceDataType;
+
 	FTestTask_ReserveSmartObject() = default;
 	FTestTask_ReserveSmartObject(const FName InName) { Name = InName; }
-	virtual ~FTestTask_ReserveSmartObject() {}
+	virtual ~FTestTask_ReserveSmartObject() override {}
+
+	virtual const UStruct* GetInstanceDataType() const override { return FTestTask_ReserveSmartObjectInstanceData::StaticStruct(); }
+	
+	virtual bool Link(FStateTreeLinker& Linker) override
+	{
+		Linker.LinkInstanceDataProperty(PotentialSmartObjectsHandle, STATETREE_INSTANCEDATA_PROPERTY(FTestTask_ReserveSmartObjectInstanceData, PotentialSmartObjects));
+		Linker.LinkInstanceDataProperty(ReservedSmartObjectLocationHandle, STATETREE_INSTANCEDATA_PROPERTY(FTestTask_ReserveSmartObjectInstanceData, ReservedSmartObjectLocation));
+		Linker.LinkInstanceDataProperty(ReservedSmartObjectHandle, STATETREE_INSTANCEDATA_PROPERTY(FTestTask_ReserveSmartObjectInstanceData, ReservedSmartObject));
+		return true;
+	}
 
 	static int32 AcquireSmartObject(const FName Name)
 	{
@@ -346,37 +434,33 @@ struct FTestTask_ReserveSmartObject : public FStateTreeTaskBase
 		// empty
 	}
 	
-	virtual EStateTreeRunStatus EnterState(FStateTreeExecutionContext& Context, const EStateTreeStateChangeType ChangeType, const FStateTreeTransitionResult& Transition) override
+	virtual EStateTreeRunStatus EnterState(FStateTreeExecutionContext& Context, const EStateTreeStateChangeType ChangeType, const FStateTreeTransitionResult& Transition) const override
 	{
 		FTestStateTreeExecutionContext& TestContext = static_cast<FTestStateTreeExecutionContext&>(Context);
 		TestContext.Log(Name, TEXT("EnterState"));
 
-		FTestPotentialSmartObjectsResult* PotentialSmartObjectsPtr = PotentialSmartObjects.GetMutablePtr<FTestPotentialSmartObjectsResult>();
-		if (PotentialSmartObjectsPtr == nullptr)
-		{
-			return EStateTreeRunStatus::Failed;
-		}
+		const FTestPotentialSmartObjectsResult& PotentialSmartObjects = Context.GetInstanceData(PotentialSmartObjectsHandle);
+		FTestMoveLocationResult& ReservedSmartObjectLocation = Context.GetInstanceData(ReservedSmartObjectLocationHandle);
+		FTestSmartObjectResult& ReservedSmartObject = Context.GetInstanceData(ReservedSmartObjectHandle);
 
 		// Reset outputs
-		SmartObject.Reset();
-		SmartObjectLocation.Reset();
+		ReservedSmartObject.Reset();
+		ReservedSmartObjectLocation.Reset();
 
 		// This consumes the sensor inputs
 		bool bFoundSO = false;
-		while (PotentialSmartObjectsPtr->HasSmartObjects())
+		while (PotentialSmartObjects.HasSmartObjects())
 		{
-			const FName SOName = PotentialSmartObjectsPtr->PopSmartObjectName();
+			const FName SOName = FName(); // PotentialSmartObjects.PopSmartObjectName(); // @todo: This is currently not allowed, all inputs are const
 			const int32 NewHandle = AcquireSmartObject(SOName);
 			if (NewHandle != INDEX_NONE)
 			{
 				// Acquire and expose one SmartObject and the location
-				SmartObject.AnimationName = SOName;
-				SmartObject.SmartObjectHandle = NewHandle;
-				SmartObject.Status = EStateTreeResultStatus::Unset; // Not available just yet
+				ReservedSmartObject.AnimationName = SOName;
+				ReservedSmartObject.SmartObjectHandle = NewHandle;
 
-				SmartObjectLocation.Location = FVector(10,10,0);
-				SmartObjectLocation.Direction = FVector::ForwardVector;
-				SmartObjectLocation.Status = EStateTreeResultStatus::Available;
+				ReservedSmartObjectLocation.Location = FVector(10,10,0);
+				ReservedSmartObjectLocation.Direction = FVector::ForwardVector;
 
 				bFoundSO = true;
 				break;
@@ -386,114 +470,114 @@ struct FTestTask_ReserveSmartObject : public FStateTreeTaskBase
 		if (!bFoundSO)
 		{
 			TestContext.Log(Name, TEXT("ReserveFailed"));
-			PotentialSmartObjectsPtr->SetFailedCoolDown();
+			// PotentialSmartObjects.SetFailedCoolDown(); // @todo: This is currently not allowed, all inputs are const
 			return EStateTreeRunStatus::Failed;
 		}
 
-		ReservedSmartObject = &SmartObject;
-		ReservedSmartObjectLocation = &SmartObjectLocation;
-
 		return EStateTreeRunStatus::Running;
 	}
 
-	virtual void ExitState(FStateTreeExecutionContext& Context, const EStateTreeStateChangeType ChangeType, const FStateTreeTransitionResult& Transition) override
+	virtual void ExitState(FStateTreeExecutionContext& Context, const EStateTreeStateChangeType ChangeType, const FStateTreeTransitionResult& Transition) const override
 	{
 		FTestStateTreeExecutionContext& TestContext = static_cast<FTestStateTreeExecutionContext&>(Context);
 		TestContext.Log(Name, TEXT("ExitState"));
-		
-		if (SmartObjectLocation.Status != EStateTreeResultStatus::Succeeded
+
+		FTestSmartObjectResult& ReservedSmartObject = Context.GetInstanceData(ReservedSmartObjectHandle);
+
+		// @todo: redo this logic
+/*		if (SmartObjectLocation.Status != EStateTreeResultStatus::Succeeded
 			|| SmartObject.Status != EStateTreeResultStatus::Succeeded)
 		{
 			// Did not fully complete the SO
-		}
+		}*/
 
 		// Clean up
-		if (SmartObject.SmartObjectHandle != INDEX_NONE)
+		if (ReservedSmartObject.SmartObjectHandle != INDEX_NONE)
 		{
-			ReleaseSmartObject(SmartObject.SmartObjectHandle);
-			SmartObject.SmartObjectHandle = INDEX_NONE;
+			ReleaseSmartObject(ReservedSmartObject.SmartObjectHandle);
+			ReservedSmartObject.SmartObjectHandle = INDEX_NONE;
 		}
-		
-		ReservedSmartObject = &SmartObject;
-		ReservedSmartObjectLocation = &SmartObjectLocation;
 	}
 
-	virtual void StateCompleted(FStateTreeExecutionContext& Context, const EStateTreeRunStatus CompletionStatus, const FStateTreeHandle StateCompleted) override
+	virtual void StateCompleted(FStateTreeExecutionContext& Context, const EStateTreeRunStatus CompletionStatus, const FStateTreeHandle StateCompleted) const override
 	{
-		// Make SmartObject available only if successfully moved to it.
-		if (SmartObjectLocation.Status == EStateTreeResultStatus::Succeeded)
+		const FTestPotentialSmartObjectsResult& PotentialSmartObjects = Context.GetInstanceData(PotentialSmartObjectsHandle);
+		FTestMoveLocationResult& ReservedSmartObjectLocation = Context.GetInstanceData(ReservedSmartObjectLocationHandle);
+		FTestSmartObjectResult& ReservedSmartObject = Context.GetInstanceData(ReservedSmartObjectHandle);
+
+		// @todo: redo this logic
+/*		// Make SmartObject available only if successfully moved to it.
+		if (ReservedSmartObjectLocation.Status == EStateTreeResultStatus::Succeeded)
 		{
-			SmartObject.Status = EStateTreeResultStatus::Available;
-			SmartObjectLocation.Status = EStateTreeResultStatus::Unset;
+			ReservedSmartObject.Status = EStateTreeResultStatus::Available;
+			ReservedSmartObjectLocation.Status = EStateTreeResultStatus::Unset;
 		}
 
-		if (SmartObjectLocation.Status == EStateTreeResultStatus::Failed
-			|| SmartObject.Status == EStateTreeResultStatus::Failed)
+		if (ReservedSmartObjectLocation.Status == EStateTreeResultStatus::Failed
+			|| ReservedSmartObject.Status == EStateTreeResultStatus::Failed)
 		{
 			// Failed to use the SmartObject, set cool down here to prevent SO being used in next State Select.
 			// The Cooldown probably needs to be 
-			FTestPotentialSmartObjectsResult* PotentialSmartObjectsPtr = PotentialSmartObjects.GetMutablePtr<FTestPotentialSmartObjectsResult>();
-			if (PotentialSmartObjectsPtr)
-			{
-				PotentialSmartObjectsPtr->SetFailedCoolDown();
-			}
-		}
+			PotentialSmartObjects.SetFailedCoolDown();
+		}*/
 	}
 
-	virtual EStateTreeRunStatus Tick(FStateTreeExecutionContext& Context, const float DeltaTime) override
+	virtual EStateTreeRunStatus Tick(FStateTreeExecutionContext& Context, const float DeltaTime) const override
 	{
 		FTestStateTreeExecutionContext& TestContext = static_cast<FTestStateTreeExecutionContext&>(Context);
 		TestContext.Log(Name, TEXT("Tick"));
 
-		// Need to update on each update, i.e. fragment may have moved.
-		ReservedSmartObject = &SmartObject;
-		ReservedSmartObjectLocation = &SmartObjectLocation;
-
 		return EStateTreeRunStatus::Running;
 	}
 
-	
-	UPROPERTY(EditAnywhere, Category = Eval, meta = (Bindable, Struct="TestPotentialSmartObjectsResult"))  // In
-	FStateTreeResultRef PotentialSmartObjects;
+	TStateTreeInstanceDataPropertyHandle<FTestPotentialSmartObjectsResult> PotentialSmartObjectsHandle;
+	TStateTreeInstanceDataPropertyHandle<FTestMoveLocationResult> ReservedSmartObjectLocationHandle;
+	TStateTreeInstanceDataPropertyHandle<FTestSmartObjectResult> ReservedSmartObjectHandle;
+};
 
-	// Because reserve is a task, these can be seen only by other Task (not evals, not conditions)
-	UPROPERTY(EditAnywhere, Category = Eval, meta = (Output, Struct="TestMoveLocationResult")) // Out
-	FStateTreeResultRef ReservedSmartObjectLocation;
 
-	UPROPERTY(EditAnywhere, Category = Eval, meta = (Output, Struct="TestSmartObjectResult")) // Out
-	FStateTreeResultRef ReservedSmartObject;
+USTRUCT()
+struct FTestTask_MoveToInstanceData
+{
+	GENERATED_BODY()
 
-	FTestMoveLocationResult SmartObjectLocation;
-	FTestSmartObjectResult SmartObject;
-};
+	UPROPERTY(EditAnywhere, Category = Input)
+	FTestMoveLocationResult MoveLocation;
 
+	UPROPERTY()
+	int32 CurrentTick = 0;
+};
 
 USTRUCT()
 struct FTestTask_MoveTo : public FStateTreeTaskBase
 {
 	GENERATED_BODY()
 
+	typedef FTestTask_MoveToInstanceData InstanceDataType;
+
 	FTestTask_MoveTo() = default;
 	FTestTask_MoveTo(const FName InName) { Name = InName; }
 	virtual ~FTestTask_MoveTo() {}
 
-	virtual EStateTreeRunStatus EnterState(FStateTreeExecutionContext& Context, const EStateTreeStateChangeType ChangeType, const FStateTreeTransitionResult& Transition) override
+	virtual const UStruct* GetInstanceDataType() const override { return FTestTask_MoveToInstanceData::StaticStruct(); }
+	
+	virtual bool Link(FStateTreeLinker& Linker) override
+	{
+		Linker.LinkInstanceDataProperty(MoveLocationHandle, STATETREE_INSTANCEDATA_PROPERTY(FTestTask_MoveToInstanceData, MoveLocation));
+		Linker.LinkInstanceDataProperty(CurrentTickHandle, STATETREE_INSTANCEDATA_PROPERTY(FTestTask_MoveToInstanceData, CurrentTick));
+		return true;
+	}
+
+	virtual EStateTreeRunStatus EnterState(FStateTreeExecutionContext& Context, const EStateTreeStateChangeType ChangeType, const FStateTreeTransitionResult& Transition) const override
 	{
 		FTestStateTreeExecutionContext& TestContext = static_cast<FTestStateTreeExecutionContext&>(Context);
 		TestContext.Log(Name, TEXT("EnterState"));
 
-		// Move restarts only when state changes for good, this allows to have child states running too.
-		FTestMoveLocationResult* MoveLocationPtr = MoveLocation.GetMutablePtr<FTestMoveLocationResult>();
-		if (MoveLocationPtr == nullptr)
-		{
-			return EStateTreeRunStatus::Failed;
-		}
-		if (MoveLocationPtr->Status != EStateTreeResultStatus::Available)
-		{
-			MoveLocationPtr->Status = EStateTreeResultStatus::Failed; // This probably should be a function we call, to allow prioritize the return values.
-			return EStateTreeRunStatus::Failed;
-		}
-		MoveLocationPtr->Status = EStateTreeResultStatus::InUse;
+		FTestMoveLocationResult& MoveLocation = Context.GetInstanceData(MoveLocationHandle);
+		int32& CurrentTick = Context.GetInstanceData(CurrentTickHandle);
+
+		// @todo: This is currently not allowed, all inputs are const
+//		MoveLocation.Status = EStateTreeResultStatus::InUse;
 
 		CurrentTick = 0;
 
@@ -501,71 +585,87 @@ struct FTestTask_MoveTo : public FStateTreeTaskBase
 	}
 
 
-	virtual void ExitState(FStateTreeExecutionContext& Context, const EStateTreeStateChangeType ChangeType, const FStateTreeTransitionResult& Transition) override
+	virtual void ExitState(FStateTreeExecutionContext& Context, const EStateTreeStateChangeType ChangeType, const FStateTreeTransitionResult& Transition) const override
 	{
 		FTestStateTreeExecutionContext& TestContext = static_cast<FTestStateTreeExecutionContext&>(Context);
 		TestContext.Log(Name, TEXT("ExitState"));
 
-		FTestMoveLocationResult* MoveLocationPtr = MoveLocation.GetMutablePtr<FTestMoveLocationResult>();
-		if (MoveLocationPtr != nullptr)
+		FTestMoveLocationResult& MoveLocation = Context.GetInstanceData(MoveLocationHandle);
+
+		// @todo: This is currently not allowed, all inputs are const
+/*		if (MoveLocation.Status == EStateTreeResultStatus::InUse)
 		{
-			if (MoveLocationPtr->Status == EStateTreeResultStatus::InUse)
-			{
-				// We got interrupted
-				MoveLocationPtr->Status = EStateTreeResultStatus::Failed;
-			}
-		}
+			// We got interrupted
+			MoveLocation.Status = EStateTreeResultStatus::Failed;
+		}*/
 	}
 
-	virtual EStateTreeRunStatus Tick(FStateTreeExecutionContext& Context, const float DeltaTime) override
+	virtual EStateTreeRunStatus Tick(FStateTreeExecutionContext& Context, const float DeltaTime) const override
 	{
 		FTestStateTreeExecutionContext& TestContext = static_cast<FTestStateTreeExecutionContext&>(Context);
 		TestContext.Log(Name, TEXT("Tick"));
 
-		FTestMoveLocationResult* MoveLocationPtr = MoveLocation.GetMutablePtr<FTestMoveLocationResult>();
-		if (MoveLocationPtr == nullptr)
-		{
-			return EStateTreeRunStatus::Failed;
-		}
+		FTestMoveLocationResult& MoveLocation = Context.GetInstanceData(MoveLocationHandle);
+		int32& CurrentTick = Context.GetInstanceData(CurrentTickHandle);
 
 		CurrentTick++;
 		const bool bDone = CurrentTick >= TicksToCompletion;
 		if (bDone)
 		{
-			MoveLocationPtr->Status = TickResult == EStateTreeRunStatus::Succeeded ? EStateTreeResultStatus::Succeeded : EStateTreeResultStatus::Failed;
+			// @todo: This is currently not allowed, all inputs are const
+//			MoveLocation.Status = TickResult == EStateTreeRunStatus::Succeeded ? EStateTreeResultStatus::Succeeded : EStateTreeResultStatus::Failed;
 		}
 		
 		return bDone ? TickResult : EStateTreeRunStatus::Running;
 	}
 
-	int32 CurrentTick = 0;
-
-	UPROPERTY(EditAnywhere, Category = Eval, meta = (Bindable))
-	FStateTreeResultRef MoveLocation;
-
-	UPROPERTY(EditAnywhere, Category = Task, meta = (Bindable))
+	TStateTreeInstanceDataPropertyHandle<FTestMoveLocationResult> MoveLocationHandle;
+	TStateTreeInstanceDataPropertyHandle<int32> CurrentTickHandle;
+	
+	UPROPERTY(EditAnywhere, Category = Parameter)
 	int32 TicksToCompletion = 2;
 
-	UPROPERTY(EditAnywhere, Category = Task, meta = (Bindable))
+	UPROPERTY(EditAnywhere, Category = Parameter)
 	EStateTreeRunStatus TickResult = EStateTreeRunStatus::Succeeded;
 
-	UPROPERTY(EditAnywhere, Category = Task, meta = (Bindable))
+	UPROPERTY(EditAnywhere, Category = Parameter)
 	EStateTreeRunStatus EnterStateResult = EStateTreeRunStatus::Running;
 };
 
+USTRUCT()
+struct FTestTask_StandInstanceData
+{
+	GENERATED_BODY()
+
+	UPROPERTY()
+	int32 CurrentTick = 0;
+};
+
 USTRUCT()
 struct FTestTask_Stand : public FStateTreeTaskBase
 {
 	GENERATED_BODY()
 
+	typedef FTestTask_StandInstanceData InstanceDataType;
+	
 	FTestTask_Stand() = default;
 	FTestTask_Stand(const FName InName) { Name = InName; }
 	virtual ~FTestTask_Stand() {}
 
-	virtual EStateTreeRunStatus EnterState(FStateTreeExecutionContext& Context, const EStateTreeStateChangeType ChangeType, const FStateTreeTransitionResult& Transition) override
+	virtual const UStruct* GetInstanceDataType() const { return FTestTask_StandInstanceData::StaticStruct(); }
+
+	virtual bool Link(FStateTreeLinker& Linker) override
+	{
+		Linker.LinkInstanceDataProperty(CurrentTickHandle, STATETREE_INSTANCEDATA_PROPERTY(FTestTask_StandInstanceData, CurrentTick));
+		return true;
+	}
+	
+	virtual EStateTreeRunStatus EnterState(FStateTreeExecutionContext& Context, const EStateTreeStateChangeType ChangeType, const FStateTreeTransitionResult& Transition) const override
 	{
 		FTestStateTreeExecutionContext& TestContext = static_cast<FTestStateTreeExecutionContext&>(Context);
 		TestContext.Log(Name, TEXT("EnterState"));
+
+		int32& CurrentTick = Context.GetInstanceData(CurrentTickHandle);
 		
 		if (ChangeType == EStateTreeStateChangeType::Changed)
 		{
@@ -574,112 +674,137 @@ struct FTestTask_Stand : public FStateTreeTaskBase
 		return EnterStateResult;
 	}
 
-	virtual void ExitState(FStateTreeExecutionContext& Context, const EStateTreeStateChangeType ChangeType, const FStateTreeTransitionResult& Transition) override
+	virtual void ExitState(FStateTreeExecutionContext& Context, const EStateTreeStateChangeType ChangeType, const FStateTreeTransitionResult& Transition) const override
 	{
 		FTestStateTreeExecutionContext& TestContext = static_cast<FTestStateTreeExecutionContext&>(Context);
 		TestContext.Log(Name, TEXT("ExitState"));
 	}
 
-	virtual EStateTreeRunStatus Tick(FStateTreeExecutionContext& Context, const float DeltaTime) override
+	virtual EStateTreeRunStatus Tick(FStateTreeExecutionContext& Context, const float DeltaTime) const override
 	{
 		FTestStateTreeExecutionContext& TestContext = static_cast<FTestStateTreeExecutionContext&>(Context);
 		TestContext.Log(Name, TEXT("Tick"));
-		
+
+		int32& CurrentTick = Context.GetInstanceData(CurrentTickHandle);
 		CurrentTick++;
+		
 		return (CurrentTick >= TicksToCompletion) ? TickResult : EStateTreeRunStatus::Running;
 	};
 
-	int32 CurrentTick = 0;
-	
-	UPROPERTY(EditAnywhere, Category = Task, meta = (Bindable))
+	TStateTreeInstanceDataPropertyHandle<int32> CurrentTickHandle;
+
+	UPROPERTY(EditAnywhere, Category = Parameter)
 	int32 TicksToCompletion = 1;
 
-	UPROPERTY(EditAnywhere, Category = Task, meta = (Bindable))
+	UPROPERTY(EditAnywhere, Category = Parameter)
 	EStateTreeRunStatus TickResult = EStateTreeRunStatus::Succeeded;
 
-	UPROPERTY(EditAnywhere, Category = Task, meta = (Bindable))
+	UPROPERTY(EditAnywhere, Category = Parameter)
 	EStateTreeRunStatus EnterStateResult = EStateTreeRunStatus::Running;
 };
 
+
+USTRUCT()
+struct FTestTask_UseSmartObjectInstanceData
+{
+	GENERATED_BODY()
+
+	UPROPERTY()
+	int32 CurrentTick = 0;
+
+	UPROPERTY(EditAnywhere, Category = Input) // In
+	FTestSmartObjectResult SmartObject;
+};
+
+
 USTRUCT()
 struct FTestTask_UseSmartObject : public FStateTreeTaskBase
 {
 	GENERATED_BODY()
 
+	typedef FTestEval_WanderInstanceData InstanceDataType;
+	
 	FTestTask_UseSmartObject() = default;
 	FTestTask_UseSmartObject(const FName InName) { Name = InName; }
 	virtual ~FTestTask_UseSmartObject() {}
 
-	virtual EStateTreeRunStatus EnterState(FStateTreeExecutionContext& Context, const EStateTreeStateChangeType ChangeType, const FStateTreeTransitionResult& Transition) override
+	virtual const UStruct* GetRuntimeDataType() const { return FTestTask_UseSmartObjectInstanceData::StaticStruct(); }
+
+	virtual bool Link(FStateTreeLinker& Linker) override
+	{
+		Linker.LinkInstanceDataProperty(CurrentTickHandle, STATETREE_INSTANCEDATA_PROPERTY(FTestTask_UseSmartObjectInstanceData, CurrentTick));
+		Linker.LinkInstanceDataProperty(SmartObjectHandle, STATETREE_INSTANCEDATA_PROPERTY(FTestTask_UseSmartObjectInstanceData, SmartObject));
+
+		return true;
+	}
+
+	virtual EStateTreeRunStatus EnterState(FStateTreeExecutionContext& Context, const EStateTreeStateChangeType ChangeType, const FStateTreeTransitionResult& Transition) const override
 	{
 		FTestStateTreeExecutionContext& TestContext = static_cast<FTestStateTreeExecutionContext&>(Context);
 		TestContext.Log(Name, TEXT("EnterState"));
+
+		FTestSmartObjectResult& SmartObject = Context.GetInstanceData(SmartObjectHandle);
+		int32& CurrentTick = Context.GetInstanceData(CurrentTickHandle);
 		
-		FTestSmartObjectResult* SmartObjectPtr = SmartObject.GetMutablePtr<FTestSmartObjectResult>();
-		if (SmartObjectPtr == nullptr)
-		{
-			return EStateTreeRunStatus::Failed;
-		}
-		if (SmartObjectPtr->Status != EStateTreeResultStatus::Available)
+		// @todo: This is currently not allowed, all inputs are const
+/*		if (SmartObjectPtr.Status != EStateTreeResultStatus::Available)
 		{
-			SmartObjectPtr->Status = EStateTreeResultStatus::Failed;
+			SmartObjectPtr.Status = EStateTreeResultStatus::Failed;
 			return EStateTreeRunStatus::Failed;
-		}
+		}*/
+		
 		CurrentTick = 0;
-		SmartObjectPtr->Status = EStateTreeResultStatus::InUse;
+//		SmartObjectPtr.Status = EStateTreeResultStatus::InUse;
 
 		return EnterStateResult;
 	}
 
-	virtual void ExitState(FStateTreeExecutionContext& Context, const EStateTreeStateChangeType ChangeType, const FStateTreeTransitionResult& Transition) override
+	virtual void ExitState(FStateTreeExecutionContext& Context, const EStateTreeStateChangeType ChangeType, const FStateTreeTransitionResult& Transition) const override
 	{
 		FTestStateTreeExecutionContext& TestContext = static_cast<FTestStateTreeExecutionContext&>(Context);
 		TestContext.Log(Name, TEXT("ExitState"));
 
-		FTestSmartObjectResult* SmartObjectPtr = SmartObject.GetMutablePtr<FTestSmartObjectResult>();
-		if (SmartObjectPtr != nullptr)
+		FTestSmartObjectResult& SmartObject = Context.GetInstanceData(SmartObjectHandle);
+
+		// @todo: This is currently not allowed, all inputs are const
+/*		if (SmartObject.Status != EStateTreeResultStatus::Succeeded)
 		{
-			if (SmartObjectPtr->Status != EStateTreeResultStatus::Succeeded)
-			{
-				// We got interrupted
-				SmartObjectPtr->Status = EStateTreeResultStatus::Failed;
-			}
-		}
+			// We got interrupted
+			SmartObject.Status = EStateTreeResultStatus::Failed;
+		}*/
+
 	}
 
-	virtual EStateTreeRunStatus Tick(FStateTreeExecutionContext& Context, const float DeltaTime) override
+	virtual EStateTreeRunStatus Tick(FStateTreeExecutionContext& Context, const float DeltaTime) const override
 	{
 		FTestStateTreeExecutionContext& TestContext = static_cast<FTestStateTreeExecutionContext&>(Context);
 		TestContext.Log(Name, TEXT("Tick"));
 
-		FTestSmartObjectResult* SmartObjectPtr = SmartObject.GetMutablePtr<FTestSmartObjectResult>();
-		if (SmartObjectPtr == nullptr)
-		{
-			return EStateTreeRunStatus::Failed;
-		}
+		FTestSmartObjectResult& SmartObject = Context.GetInstanceData(SmartObjectHandle);
+		int32& CurrentTick = Context.GetInstanceData(CurrentTickHandle);
 		
 		CurrentTick++;
 		
 		const bool bDone = CurrentTick >= TicksToCompletion;
 		if (bDone)
 		{
-			SmartObjectPtr->Status = TickResult == EStateTreeRunStatus::Succeeded ? EStateTreeResultStatus::Succeeded : EStateTreeResultStatus::Failed;
+			// @todo: This is currently not allowed, all inputs are const
+//			SmartObject.Status = TickResult == EStateTreeRunStatus::Succeeded ? EStateTreeResultStatus::Succeeded : EStateTreeResultStatus::Failed;
 		}
 		
 		return bDone ? TickResult : EStateTreeRunStatus::Running;
 	}
 
-	int32 CurrentTick = 0;
-
-	UPROPERTY(EditAnywhere, Category = Eval, meta = (Bindable, Struct="TestSmartObjectResult")) // In
-	FStateTreeResultRef SmartObject;
+	TStateTreeInstanceDataPropertyHandle<int32> CurrentTickHandle;
+	TStateTreeInstanceDataPropertyHandle<FTestSmartObjectResult> SmartObjectHandle;
 
-	UPROPERTY(EditAnywhere, Category = Task)
+	UPROPERTY(EditAnywhere, Category = Parameter)
 	int32 TicksToCompletion = 2;
 
-	UPROPERTY(EditAnywhere, Category = Task)
+	UPROPERTY(EditAnywhere, Category = Parameter)
 	EStateTreeRunStatus TickResult = EStateTreeRunStatus::Succeeded;
 
-	UPROPERTY(EditAnywhere, Category = Task)
+	UPROPERTY(EditAnywhere, Category = Parameter)
 	EStateTreeRunStatus EnterStateResult = EStateTreeRunStatus::Running;
-};
\ No newline at end of file
+};
+
