StateTree: handle transitions after ticking tasks and evaluators
- changed at which stage transitions are handled
- updated the binding visibility of tasks and evaluators (transition conditions can see task data now)

diff --git a/Engine/Plugins/Runtime/StateTree/Source/StateTreeEditorModule/Private/StateTreeEditorData.cpp b/Engine/Plugins/Runtime/StateTree/Source/StateTreeEditorModule/Private/StateTreeEditorData.cpp
--- a/Engine/Plugins/Runtime/StateTree/Source/StateTreeEditorModule/Private/StateTreeEditorData.cpp
+++ b/Engine/Plugins/Runtime/StateTree/Source/StateTreeEditorModule/Private/StateTreeEditorData.cpp
@@ -19,26 +19,42 @@ void UStateTreeEditorData::GetAccessibleStructs(const FGuid TargetStructID, TArr
 
 	if (ValidStates.Num() > 0)
 	{
-		VisitHierarchy([&ValidStates, &OutStructDescs, &EvalDescs, &TaskDescs, TargetStructID, CurrentState](const UStateTreeState& State, const FGuid& ID, const FName& Name, const UScriptStruct* ItemStruct, const UStruct* InstanceStruct) -> bool
+		VisitHierarchy([&ValidStates, &OutStructDescs, &EvalDescs, &TaskDescs, TargetStructID, CurrentState]
+			(const UStateTreeState& State, const FGuid& ID, const FName& Name, const EStateTreeItemType ItemType, const UScriptStruct* ItemStruct, const UStruct* InstanceStruct) -> bool
 			{
 				if (ValidStates.Contains(&State))
 				{
 					if (ID == TargetStructID)
 					{
-						OutStructDescs.Append(EvalDescs);
-
-						// Only tasks can see other tasks too.
-						if (ItemStruct->IsChildOf(FStateTreeTaskBase::StaticStruct()))
+						if (ItemType == EStateTreeItemType::EnterCondition)
+						{
+							// Enter conditions can only see evaluators
+							OutStructDescs.Append(EvalDescs);
+						}
+						else if (ItemType == EStateTreeItemType::Evaluator)
+						{
+							// Evaluators can only see other evaluators
+							OutStructDescs.Append(EvalDescs);
+						}
+						else if (ItemType == EStateTreeItemType::Task)
+						{
+							// Tasks can see evals and tasks.
+							OutStructDescs.Append(EvalDescs);
+							OutStructDescs.Append(TaskDescs);
+						}
+						else if (ItemType == EStateTreeItemType::TransitionCondition)
 						{
+							// Transitions can see evals and tasks.
+							OutStructDescs.Append(EvalDescs);
 							OutStructDescs.Append(TaskDescs);
 						}
 						
 						return false; // Stop visit
 					}
 
+					// Not at target yet, collect all evaluators and tasks visible so far.
 					if (ItemStruct->IsChildOf(FStateTreeEvaluatorBase::StaticStruct()))
 					{
-						// All evaluators up to the target struct are accessible.
 						FStateTreeBindableStructDesc& Desc = EvalDescs.AddDefaulted_GetRef();
 						Desc.Struct = InstanceStruct;
 						Desc.Name = Name;
@@ -46,7 +62,6 @@ void UStateTreeEditorData::GetAccessibleStructs(const FGuid TargetStructID, TArr
 					}
 					else if (ItemStruct->IsChildOf(FStateTreeTaskBase::StaticStruct()))
 					{
-						// All tasks up to the target struct are accessible (for other tasks).
 						FStateTreeBindableStructDesc& Desc = TaskDescs.AddDefaulted_GetRef();
 						Desc.Struct = InstanceStruct;
 						Desc.Name = Name;
@@ -62,7 +77,7 @@ bool UStateTreeEditorData::GetStructByID(const FGuid StructID, FStateTreeBindabl
 {
 	bool bResult = false;
 
-	VisitHierarchy([&bResult, &OutStructDesc, StructID](const UStateTreeState& State, const FGuid& ID, const FName& Name, const UScriptStruct* ItemStruct, const UStruct* InstanceStruct) -> bool
+	VisitHierarchy([&bResult, &OutStructDesc, StructID](const UStateTreeState& State, const FGuid& ID, const FName& Name, const EStateTreeItemType ItemType, const UScriptStruct* ItemStruct, const UStruct* InstanceStruct) -> bool
 		{
 			if (ID == StructID)
 			{
@@ -82,7 +97,7 @@ const UStateTreeState* UStateTreeEditorData::GetStateByStructID(const FGuid Targ
 {
 	const UStateTreeState* Result = nullptr;
 
-	VisitHierarchy([&Result, TargetStructID](const UStateTreeState& State, const FGuid& ID, const FName& Name, const UScriptStruct* ItemStruct, const UStruct* InstanceStruct) -> bool
+	VisitHierarchy([&Result, TargetStructID](const UStateTreeState& State, const FGuid& ID, const FName& Name, const EStateTreeItemType ItemType, const UScriptStruct* ItemStruct, const UStruct* InstanceStruct) -> bool
 		{
 			if (ID == TargetStructID)
 			{
@@ -99,14 +114,14 @@ void UStateTreeEditorData::GetAllStructIDs(TMap<FGuid, const UStruct*>& AllStruc
 {
 	AllStructs.Reset();
 
-	VisitHierarchy([&AllStructs](const UStateTreeState& State, const FGuid& ID, const FName& Name, const UScriptStruct* ItemStruct, const UStruct* InstanceStruct) -> bool
+	VisitHierarchy([&AllStructs](const UStateTreeState& State, const FGuid& ID, const FName& Name, const EStateTreeItemType ItemType, const UScriptStruct* ItemStruct, const UStruct* InstanceStruct) -> bool
 		{
 			AllStructs.Add(ID, InstanceStruct);
 			return true; // Continue
 		});
 }
 
-void UStateTreeEditorData::VisitHierarchy(TFunctionRef<bool(const UStateTreeState& State, const FGuid& ID, const FName& Name, const UScriptStruct* ItemStruct, const UStruct* InstanceStruct)> InFunc) const
+void UStateTreeEditorData::VisitHierarchy(TFunctionRef<bool(const UStateTreeState& State, const FGuid& ID, const FName& Name, const EStateTreeItemType ItemType, const UScriptStruct* ItemStruct, const UStruct* InstanceStruct)> InFunc) const
 {
 	TArray<const UStateTreeState*> Stack;
 	bool bContinue = true;
@@ -132,7 +147,7 @@ void UStateTreeEditorData::VisitHierarchy(TFunctionRef<bool(const UStateTreeStat
 			{
 				if (const FStateTreeEvaluatorBase* Evaluator = Item.Item.GetPtr<FStateTreeEvaluatorBase>())
 				{
-					if (!InFunc(*State, Item.ID, Evaluator->Name, Item.Item.GetScriptStruct(), Evaluator->GetInstanceDataType()))
+					if (!InFunc(*State, Item.ID, Evaluator->Name, EStateTreeItemType::Evaluator, Item.Item.GetScriptStruct(), Evaluator->GetInstanceDataType()))
 					{
 						bContinue = false;
 						break;
@@ -146,7 +161,7 @@ void UStateTreeEditorData::VisitHierarchy(TFunctionRef<bool(const UStateTreeStat
 				{
 					if (const FStateTreeConditionBase* Cond = Item.Item.GetPtr<FStateTreeConditionBase>())
 					{
-						if (!InFunc(*State, Item.ID, Item.Item.GetScriptStruct()->GetFName(), Item.Item.GetScriptStruct(), Cond->GetInstanceDataType()))
+						if (!InFunc(*State, Item.ID, Item.Item.GetScriptStruct()->GetFName(), EStateTreeItemType::EnterCondition, Item.Item.GetScriptStruct(), Cond->GetInstanceDataType()))
 						{
 							bContinue = false;
 							break;
@@ -161,7 +176,7 @@ void UStateTreeEditorData::VisitHierarchy(TFunctionRef<bool(const UStateTreeStat
 				{
 					if (const FStateTreeTaskBase* Task = Item.Item.GetPtr<FStateTreeTaskBase>())
 					{
-						if (!InFunc(*State, Item.ID, Task->Name, Item.Item.GetScriptStruct(), Task->GetInstanceDataType()))
+						if (!InFunc(*State, Item.ID, Task->Name, EStateTreeItemType::Task, Item.Item.GetScriptStruct(), Task->GetInstanceDataType()))
 						{
 							bContinue = false;
 							break;
@@ -178,7 +193,7 @@ void UStateTreeEditorData::VisitHierarchy(TFunctionRef<bool(const UStateTreeStat
 					{
 						if (const FStateTreeConditionBase* Cond = Item.Item.GetPtr<FStateTreeConditionBase>())
 						{
-							if (!InFunc(*State, Item.ID, Item.Item.GetScriptStruct()->GetFName(), Item.Item.GetScriptStruct(), Cond->GetInstanceDataType()))
+							if (!InFunc(*State, Item.ID, Item.Item.GetScriptStruct()->GetFName(), EStateTreeItemType::TransitionCondition, Item.Item.GetScriptStruct(), Cond->GetInstanceDataType()))
 							{
 								bContinue = false;
 								break;
diff --git a/Engine/Plugins/Runtime/StateTree/Source/StateTreeEditorModule/Public/StateTreeEditorData.h b/Engine/Plugins/Runtime/StateTree/Source/StateTreeEditorModule/Public/StateTreeEditorData.h
--- a/Engine/Plugins/Runtime/StateTree/Source/StateTreeEditorModule/Public/StateTreeEditorData.h
+++ b/Engine/Plugins/Runtime/StateTree/Source/StateTreeEditorModule/Public/StateTreeEditorData.h
@@ -6,6 +6,15 @@
 #include "StateTreeEditorPropertyBindings.h"
 #include "StateTreeEditorData.generated.h"
 
+UENUM()
+enum class EStateTreeItemType : uint8
+{
+	EnterCondition,
+	Evaluator,
+	Task,
+	TransitionCondition,
+};
+
 /**
  * Edit time data for StateTree asset. This data gets baked into runtime format before being used by the StateTreeInstance.
  */
@@ -30,7 +39,7 @@ public:
 	 * Iterates over all structs that are related to binding
 	 * @param InFunc function called at each item, should return true if visiting is continued or false to stop.
 	 */
-	void VisitHierarchy(TFunctionRef<bool(const UStateTreeState& State, const FGuid& ID, const FName& Name, const UScriptStruct* ItemStruct, const UStruct* InstanceStruct)> InFunc) const;
+	void VisitHierarchy(TFunctionRef<bool(const UStateTreeState& State, const FGuid& ID, const FName& Name, const EStateTreeItemType ItemType, const UScriptStruct* ItemStruct, const UStruct* InstanceStruct)> InFunc) const;
 
 	// StateTree Builder API
 
diff --git a/Engine/Plugins/Runtime/StateTree/Source/StateTreeModule/Private/StateTreeExecutionContext.cpp b/Engine/Plugins/Runtime/StateTree/Source/StateTreeModule/Private/StateTreeExecutionContext.cpp
--- a/Engine/Plugins/Runtime/StateTree/Source/StateTreeModule/Private/StateTreeExecutionContext.cpp
+++ b/Engine/Plugins/Runtime/StateTree/Source/StateTreeModule/Private/StateTreeExecutionContext.cpp
@@ -57,34 +57,74 @@ bool FStateTreeExecutionContext::Init(UObject& InOwner, const UStateTree& InStat
 	return true;
 }
 
-void FStateTreeExecutionContext::Start(FStateTreeDataView ExternalStorage)
+EStateTreeRunStatus FStateTreeExecutionContext::Start(FStateTreeDataView ExternalStorage)
 {
 	if (!Owner || !StateTree)
 	{
-		return;
+		return EStateTreeRunStatus::Failed;
 	}
+	
 	FStateTreeDataView Storage = SelectMutableStorage(ExternalStorage);
 	FStateTreeExecutionState& Exec = GetExecState(Storage);
 
-	// Stop if still running 
+	// Stop if still running previous state.
 	if (Exec.TreeRunStatus == EStateTreeRunStatus::Running)
 	{
 		const FStateTreeTransitionResult Transition(FStateTreeStateStatus(Exec.CurrentState, Exec.LastTickStatus), FStateTreeHandle::Succeeded);
 		ExitState(Storage, Transition);
 	}
-
-	// Initialize to unset running state, tick will choose the first state.
+	
+	// Initialize to unset running state.
 	Exec.TreeRunStatus = EStateTreeRunStatus::Running;
 	Exec.CurrentState = FStateTreeHandle::Invalid;
 	Exec.LastTickStatus = EStateTreeRunStatus::Unset;
+
+	// Select new Starting from root.
+	VisitedStates.Init(false, StateTree->States.Num());
+	static const FStateTreeHandle RootState = FStateTreeHandle(0);
+	const FStateTreeTransitionResult Transition = SelectState(Storage, FStateTreeStateStatus(FStateTreeHandle::Invalid), RootState, RootState, 0);
+
+	// Handle potential transition.
+	if (Transition.Next.IsValid())
+	{
+		if (Transition.Next == FStateTreeHandle::Succeeded || Transition.Next == FStateTreeHandle::Failed)
+		{
+			// Transition to a terminal state (succeeded/failed), or default transition failed.
+			STATETREE_LOG(Warning, TEXT("%s: Tree %s at StateTree start on '%s' using StateTree '%s'."),
+				ANSI_TO_TCHAR(__FUNCTION__), Transition.Next == FStateTreeHandle::Succeeded ? TEXT("succeeded") : TEXT("failed"), *GetNameSafe(Owner), *GetNameSafe(StateTree));
+			Exec.TreeRunStatus = Transition.Next == FStateTreeHandle::Succeeded ? EStateTreeRunStatus::Succeeded : EStateTreeRunStatus::Failed;
+		}
+		else
+		{
+			// Enter state tasks can fail/succeed, treat it same as tick.
+			Exec.LastTickStatus = EnterState(Storage, Transition);
+			Exec.CurrentState = Transition.Next;
+			// Report state completed immediately.
+			if (Exec.LastTickStatus != EStateTreeRunStatus::Running)
+			{
+				StateCompleted(Storage, Exec.CurrentState, Exec.LastTickStatus);
+			}
+		}
+	}
+
+	if (!Exec.CurrentState.IsValid())
+	{
+		// Should not happen. This may happen if initial state could not be selected.
+		STATETREE_LOG(Error, TEXT("%s: Failed to select initial state on '%s' using StateTree '%s'. This should not happen, check that the StateTree logic can always select a state at start."),
+			ANSI_TO_TCHAR(__FUNCTION__), *GetNameSafe(Owner), *GetNameSafe(StateTree));
+		Exec.TreeRunStatus = EStateTreeRunStatus::Failed;
+	}
+
+	return Exec.TreeRunStatus;
 }
 
-void FStateTreeExecutionContext::Stop(FStateTreeDataView ExternalStorage)
+EStateTreeRunStatus FStateTreeExecutionContext::Stop(FStateTreeDataView ExternalStorage)
 {
 	if (!Owner || !StateTree)
 	{
-		return;
+		return EStateTreeRunStatus::Failed;
 	}
+	
 	FStateTreeDataView Storage = SelectMutableStorage(ExternalStorage);
 	FStateTreeExecutionState& Exec = GetExecState(Storage);
 
@@ -98,6 +138,8 @@ void FStateTreeExecutionContext::Stop(FStateTreeDataView ExternalStorage)
 	Exec.TreeRunStatus = EStateTreeRunStatus::Succeeded;
 	Exec.CurrentState = FStateTreeHandle::Succeeded;
 	Exec.LastTickStatus = EStateTreeRunStatus::Unset;
+
+	return Exec.TreeRunStatus;
 }
 
 EStateTreeRunStatus FStateTreeExecutionContext::Tick(const float DeltaTime, FStateTreeDataView ExternalStorage)
@@ -123,35 +165,32 @@ EStateTreeRunStatus FStateTreeExecutionContext::Tick(const float DeltaTime, FSta
 		Exec.GatedTransitionTime -= DeltaTime;
 	}
 	
-	// The state selection is repeated up to MaxIteration time. This allows failed EnterState() to potentially find a new state immediately.
-	// This helps event driven StateTrees to not require another event/tick to find a suitable state.
-	static const int32 MaxIterations = 5;
-	for (int32 Iter = 0; Iter < MaxIterations; Iter++)
-	{
-		// Reset visited state, used to keep track which state's evaluators have been ticked.
-		VisitedStates.Init(false, StateTree->States.Num());
+	// Reset visited state, used to keep track which state's evaluators have been ticked.
+	VisitedStates.Init(false, StateTree->States.Num());
 
+	if (Exec.LastTickStatus == EStateTreeRunStatus::Running)
+	{
 		// Tick evaluators on active states. Further evaluator may be ticked during selection as non-active states are visited.
 		// Ticked states are kept track of and evaluators are ticked just once per frame.
-		if (Exec.LastTickStatus == EStateTreeRunStatus::Running)
-		{
-			TickEvaluators(Storage, Exec.CurrentState, EStateTreeEvaluationType::Tick, DeltaTime);
-		}
+		TickEvaluators(Storage, Exec.CurrentState, EStateTreeEvaluationType::Tick, DeltaTime);
 
-		// Select initial state or trigger transitions.
-		FStateTreeTransitionResult Transition;
-		if (Exec.LastTickStatus == EStateTreeRunStatus::Unset && Exec.CurrentState == FStateTreeHandle::Invalid)
-		{
-			// No state has been selected yet, select starting from root.
-			static const FStateTreeHandle RootState = FStateTreeHandle(0);
-			Transition = SelectState(Storage, FStateTreeStateStatus(FStateTreeHandle::Invalid), RootState, RootState, 0);
-		}
-		else
+		// Tick tasks on active states.
+		Exec.LastTickStatus = TickTasks(Storage, Exec.CurrentState, DeltaTime);
+		
+		// Report state completed immediately.
+		if (Exec.LastTickStatus != EStateTreeRunStatus::Running)
 		{
-			// Trigger running transitions or state succeed/failed transitions.
-			// Transitions are triggered early in the frame, to prevent one frame lag in games which require quick state changes (i.e. combat).
-			Transition = TriggerTransitions(Storage, FStateTreeStateStatus(Exec.CurrentState, Exec.LastTickStatus), 0);
+			StateCompleted(Storage, Exec.CurrentState, Exec.LastTickStatus);
 		}
+	}
+
+	// The state selection is repeated up to MaxIteration time. This allows failed EnterState() to potentially find a new state immediately.
+	// This helps event driven StateTrees to not require another event/tick to find a suitable state.
+	static const int32 MaxIterations = 5;
+	for (int32 Iter = 0; Iter < MaxIterations; Iter++)
+	{
+		// Trigger conditional transitions or state succeed/failed transitions. First tick transition is handled here too.
+		FStateTreeTransitionResult Transition = TriggerTransitions(Storage, FStateTreeStateStatus(Exec.CurrentState, Exec.LastTickStatus), 0);
 
 		// Handle potential transition.
 		if (Transition.Next.IsValid())
@@ -184,24 +223,13 @@ EStateTreeRunStatus FStateTreeExecutionContext::Tick(const float DeltaTime, FSta
 	
 	if (!Exec.CurrentState.IsValid())
 	{
-		// Should not happen. This may happen if initial state or default transition could not be selected. 
-		STATETREE_LOG(Error, TEXT("%s: Failed to select state on '%s' using StateTree '%s'."), ANSI_TO_TCHAR(__FUNCTION__), *GetNameSafe(Owner), *GetNameSafe(StateTree));
+		// Should not happen. This may happen if a state completion transition could not be selected. 
+		STATETREE_LOG(Error, TEXT("%s: Failed to select state on '%s' using StateTree '%s'. This should not happen, state completion transition is likely missing."),
+			ANSI_TO_TCHAR(__FUNCTION__), *GetNameSafe(Owner), *GetNameSafe(StateTree));
 		Exec.TreeRunStatus = EStateTreeRunStatus::Failed;
 		return Exec.TreeRunStatus;
 	}
 
-	// Tick tasks on active states.
-	if (Exec.LastTickStatus == EStateTreeRunStatus::Running)
-	{
-		Exec.LastTickStatus = TickTasks(Storage, Exec.CurrentState, DeltaTime);
-		// Report state completed immediately.
-		if (Exec.LastTickStatus != EStateTreeRunStatus::Running)
-		{
-			StateCompleted(Storage, Exec.CurrentState, Exec.LastTickStatus);
-		}
-	}
-	
-	
 	return Exec.TreeRunStatus;
 }
 
diff --git a/Engine/Plugins/Runtime/StateTree/Source/StateTreeModule/Public/StateTreeExecutionContext.h b/Engine/Plugins/Runtime/StateTree/Source/StateTreeModule/Public/StateTreeExecutionContext.h
--- a/Engine/Plugins/Runtime/StateTree/Source/StateTreeModule/Public/StateTreeExecutionContext.h
+++ b/Engine/Plugins/Runtime/StateTree/Source/StateTreeModule/Public/StateTreeExecutionContext.h
@@ -66,9 +66,9 @@ public:
 	void SetWorld(UWorld* InWorld) { World = InWorld; };
 
 	/** Start executing. */
-	void Start(FStateTreeDataView ExternalStorage = FStateTreeDataView());
+	EStateTreeRunStatus Start(FStateTreeDataView ExternalStorage = FStateTreeDataView());
 	/** Stop executing. */
-	void Stop(FStateTreeDataView ExternalStorage = FStateTreeDataView());
+	EStateTreeRunStatus Stop(FStateTreeDataView ExternalStorage = FStateTreeDataView());
 
 	/** Tick the state tree logic. */
 	EStateTreeRunStatus Tick(const float DeltaTime, FStateTreeDataView ExternalStorage = FStateTreeDataView());
