[StateTree] fixed clang compiler CIS issues caused by exported templates

diff --git a/Engine/Plugins/Runtime/StateTree/Source/StateTreeEditorModule/Public/StateTreeState.h b/Engine/Plugins/Runtime/StateTree/Source/StateTreeEditorModule/Public/StateTreeState.h
--- a/Engine/Plugins/Runtime/StateTree/Source/StateTreeEditorModule/Public/StateTreeState.h
+++ b/Engine/Plugins/Runtime/StateTree/Source/StateTreeEditorModule/Public/StateTreeState.h
@@ -63,7 +63,7 @@ struct STATETREEEDITORMODULE_API FStateTreeEvaluatorItem
 };
 
 template <typename T>
-struct STATETREEEDITORMODULE_API TStateTreeEvaluatorItem : public FStateTreeEvaluatorItem
+struct TStateTreeEvaluatorItem : public FStateTreeEvaluatorItem
 {
 	typedef T ItemType;
 	FORCEINLINE T& GetItem() { return Item.template GetMutable<T>(); }
@@ -99,7 +99,7 @@ struct STATETREEEDITORMODULE_API FStateTreeTaskItem
 };
 
 template <typename T>
-struct STATETREEEDITORMODULE_API TStateTreeTaskItem : public FStateTreeTaskItem
+struct TStateTreeTaskItem : public FStateTreeTaskItem
 {
 	typedef T ItemType;
 	FORCEINLINE T& GetItem() { return Item.template GetMutable<T>(); }
@@ -125,7 +125,7 @@ struct STATETREEEDITORMODULE_API FStateTreeConditionItem
 };
 
 template <typename T>
-struct STATETREEEDITORMODULE_API TStateTreeConditionItem : public FStateTreeConditionItem
+struct TStateTreeConditionItem : public FStateTreeConditionItem
 {
 	typedef T ItemType;
 	FORCEINLINE T& GetItem() { return Item.template GetMutable<T>(); }
