[StateTreeTest] Provide world as outer when creating test StateTrees to allow VLOG without ensure

diff --git a/Engine/Plugins/Runtime/StateTree/Source/StateTreeTestSuite/Private/StateTreeTest.cpp b/Engine/Plugins/Runtime/StateTree/Source/StateTreeTestSuite/Private/StateTreeTest.cpp
--- a/Engine/Plugins/Runtime/StateTree/Source/StateTreeTestSuite/Private/StateTreeTest.cpp
+++ b/Engine/Plugins/Runtime/StateTree/Source/StateTreeTestSuite/Private/StateTreeTest.cpp
@@ -12,6 +12,7 @@
 #include "StateTree.h"
 #include "StateTreeTestTypes.h"
 #include "StateTreeExecutionContext.h"
+#include "Engine/World.h"
 
 #define LOCTEXT_NAMESPACE "AITestSuite_StateTreeTest"
 
@@ -19,9 +20,9 @@ PRAGMA_DISABLE_OPTIMIZATION
 
 namespace UE::StateTree::Tests
 {
-	UStateTree& NewStateTree()
+	UStateTree& NewStateTree(UObject* Outer = GetTransientPackage())
 	{
-		UStateTree* StateTree = NewObject<UStateTree>();
+		UStateTree* StateTree = NewObject<UStateTree>(Outer);
 		check(StateTree);
 		UStateTreeEditorData* EditorData = NewObject<UStateTreeEditorData>(StateTree);
 		check(EditorData);
@@ -36,7 +37,7 @@ struct FStateTreeTest_MakeAndBakeStateTree : FAITestBase
 {
 	virtual bool InstantTest() override
 	{
-		UStateTree& StateTree = UE::StateTree::Tests::NewStateTree();
+		UStateTree& StateTree = UE::StateTree::Tests::NewStateTree(&GetWorld());
 		UStateTreeEditorData& EditorData = *Cast<UStateTreeEditorData>(StateTree.EditorData);
 		
 		UStateTreeState& Root = EditorData.AddSubTree(FName(TEXT("Root")));
@@ -85,7 +86,7 @@ struct FStateTreeTest_WanderLoop : FAITestBase
 {
 	virtual bool InstantTest() override
 	{
-		UStateTree& StateTree = UE::StateTree::Tests::NewStateTree();
+		UStateTree& StateTree = UE::StateTree::Tests::NewStateTree(&GetWorld());
 		UStateTreeEditorData& EditorData = *Cast<UStateTreeEditorData>(StateTree.EditorData);
 		
 		UStateTreeState& Root = EditorData.AddSubTree(FName(TEXT("Root")));
@@ -218,7 +219,7 @@ struct FStateTreeTest_Sequence : FAITestBase
 {
 	virtual bool InstantTest() override
 	{
-		UStateTree& StateTree = UE::StateTree::Tests::NewStateTree();
+		UStateTree& StateTree = UE::StateTree::Tests::NewStateTree(&GetWorld());
 		UStateTreeEditorData& EditorData = *Cast<UStateTreeEditorData>(StateTree.EditorData);
 		
 		UStateTreeState& Root = EditorData.AddSubTree(FName(TEXT("Root")));
@@ -275,7 +276,7 @@ struct FStateTreeTest_Select : FAITestBase
 {
 	virtual bool InstantTest() override
 	{
-		UStateTree& StateTree = UE::StateTree::Tests::NewStateTree();
+		UStateTree& StateTree = UE::StateTree::Tests::NewStateTree(&GetWorld());
 		UStateTreeEditorData& EditorData = *Cast<UStateTreeEditorData>(StateTree.EditorData);
 		
 		UStateTreeState& Root = EditorData.AddSubTree(FName(TEXT("Root")));
