[StateTree] fixed bad inheritance of a struct used for tests and added trait `WithPureVirtual` for FStateTreeConditionBase struct

diff --git a/Engine/Plugins/Runtime/StateTree/Source/StateTreeModule/Public/StateTreeConditionBase.h b/Engine/Plugins/Runtime/StateTree/Source/StateTreeModule/Public/StateTreeConditionBase.h
--- a/Engine/Plugins/Runtime/StateTree/Source/StateTreeModule/Public/StateTreeConditionBase.h
+++ b/Engine/Plugins/Runtime/StateTree/Source/StateTreeModule/Public/StateTreeConditionBase.h
@@ -58,3 +58,12 @@ struct STATETREEMODULE_API FStateTreeConditionBase
 	UPROPERTY()
 	uint16 DataViewIndex = 0;
 };
+
+template<>
+struct TStructOpsTypeTraits<FStateTreeConditionBase> : public TStructOpsTypeTraitsBase2<FStateTreeConditionBase>
+{
+	enum
+	{
+		WithPureVirtual = true,
+	};
+};
diff --git a/Engine/Plugins/Runtime/StateTree/Source/StateTreeTestSuite/Private/StateTreeTestTypes.h b/Engine/Plugins/Runtime/StateTree/Source/StateTreeTestSuite/Private/StateTreeTestTypes.h
--- a/Engine/Plugins/Runtime/StateTree/Source/StateTreeTestSuite/Private/StateTreeTestTypes.h
+++ b/Engine/Plugins/Runtime/StateTree/Source/StateTreeTestSuite/Private/StateTreeTestTypes.h
@@ -388,7 +388,7 @@ struct FTestEval_SmartObjectSensor : public FStateTreeEvaluatorBase
 
 
 USTRUCT()
-struct FTestTask_ReserveSmartObjectInstanceData : public FStateTreeTaskBase
+struct FTestTask_ReserveSmartObjectInstanceData
 {
 	GENERATED_BODY()
 
@@ -728,7 +728,7 @@ struct FTestTask_UseSmartObject : public FStateTreeTaskBase
 	FTestTask_UseSmartObject(const FName InName) { Name = InName; }
 	virtual ~FTestTask_UseSmartObject() {}
 
-	virtual const UStruct* GetRuntimeDataType() const { return FTestTask_UseSmartObjectInstanceData::StaticStruct(); }
+	virtual const UStruct* GetInstanceDataType()const override { return FTestTask_UseSmartObjectInstanceData::StaticStruct(); }
 
 	virtual bool Link(FStateTreeLinker& Linker) override
 	{
