StateTree: Fixed tests to reflect recent changes in state selection

diff --git a/Engine/Plugins/Runtime/StateTree/Source/StateTreeTestSuite/Private/StateTreeTest.cpp b/Engine/Plugins/Runtime/StateTree/Source/StateTreeTestSuite/Private/StateTreeTest.cpp
--- a/Engine/Plugins/Runtime/StateTree/Source/StateTreeTestSuite/Private/StateTreeTest.cpp
+++ b/Engine/Plugins/Runtime/StateTree/Source/StateTreeTestSuite/Private/StateTreeTest.cpp
@@ -246,26 +246,28 @@ struct FStateTreeTest_Sequence : FAITestBase
 		const FString EnterStateStr(TEXT("EnterState"));
 		const FString ExitStateStr(TEXT("ExitState"));
 		
-		Exec.Start();
+		Status = Exec.Start();
+		AITEST_TRUE("StateTree Task1 should enter state", Exec.Expect(Task1.GetName(), EnterStateStr));
+		AITEST_FALSE("StateTree Task1 should not tick", Exec.Expect(Task1.GetName(), TickStr));
+		Exec.LogClear();
 
 		Status = Exec.Tick(0.1f);
-		AITEST_TRUE("StateTree Task1 should enter state and tick", Exec.Expect(Task1.GetName(), EnterStateStr).Then(Task1.GetName(), TickStr));
+		AITEST_TRUE("StateTree Task1 should tick, and exit state", Exec.Expect(Task1.GetName(), TickStr).Then(Task1.GetName(), ExitStateStr));
+		AITEST_TRUE("StateTree Task2 should enter state", Exec.Expect(Task2.GetName(), EnterStateStr));
 		AITEST_FALSE("StateTree Task2 should not tick", Exec.Expect(Task2.GetName(), TickStr));
 		AITEST_TRUE("StateTree should be running", Status == EStateTreeRunStatus::Running);
 		Exec.LogClear();
 		
 		Status = Exec.Tick(0.1f);
-		AITEST_TRUE("StateTree Task1 should exit state", Exec.Expect(Task1.GetName(), ExitStateStr));
+        AITEST_TRUE("StateTree Task2 should tick, and exit state", Exec.Expect(Task2.GetName(), TickStr).Then(Task2.GetName(), ExitStateStr));
 		AITEST_FALSE("StateTree Task1 should not tick", Exec.Expect(Task1.GetName(), TickStr));
-        AITEST_TRUE("StateTree Task2 should enter state and tick", Exec.Expect(Task2.GetName(), EnterStateStr).Then(Task2.GetName(), TickStr));
-        AITEST_TRUE("StateTree should be running", Status == EStateTreeRunStatus::Running);
-        Exec.LogClear();
-		
+        AITEST_TRUE("StateTree should be completed", Status == EStateTreeRunStatus::Succeeded);
+		Exec.LogClear();
+
 		Status = Exec.Tick(0.1f);
-        AITEST_TRUE("StateTree Task2 should exit state", Exec.Expect(Task2.GetName(), ExitStateStr));
 		AITEST_FALSE("StateTree Task1 should not tick", Exec.Expect(Task1.GetName(), TickStr));
-        AITEST_FALSE("StateTree Task2 should not tick", Exec.Expect(Task2.GetName(), TickStr));
-        AITEST_TRUE("StateTree should be completed", Status == EStateTreeRunStatus::Succeeded);
+		AITEST_FALSE("StateTree Task2 should not tick", Exec.Expect(Task2.GetName(), TickStr));
+		Exec.LogClear();
 
 		return true;
 	}
@@ -307,12 +309,14 @@ struct FStateTreeTest_Select : FAITestBase
 		const FString EnterStateStr(TEXT("EnterState"));
 		const FString ExitStateStr(TEXT("ExitState"));
 
-		Exec.Start();
-		
-		Status = Exec.Tick(0.1f);
-		AITEST_TRUE("StateTree TaskRoot should enter state, and tick", Exec.Expect(TaskRoot.GetName(), EnterStateStr).Then(TaskRoot.GetName(), TickStr));
-		AITEST_TRUE("StateTree Task1 should enter state, and tick", Exec.Expect(Task1.GetName(), EnterStateStr).Then(Task1.GetName(), TickStr));
-		AITEST_TRUE("StateTree Task1A should enter state, and tick", Exec.Expect(Task1A.GetName(), EnterStateStr).Then(Task1A.GetName(), TickStr));
+		// Start and enter state
+		Status = Exec.Start();
+		AITEST_TRUE("StateTree TaskRoot should enter state", Exec.Expect(TaskRoot.GetName(), EnterStateStr));
+		AITEST_TRUE("StateTree Task1 should enter state", Exec.Expect(Task1.GetName(), EnterStateStr));
+		AITEST_TRUE("StateTree Task1A should enter state", Exec.Expect(Task1A.GetName(), EnterStateStr));
+		AITEST_FALSE("StateTree TaskRoot should not tick", Exec.Expect(TaskRoot.GetName(), TickStr));
+		AITEST_FALSE("StateTree Task1 should not tick", Exec.Expect(Task1.GetName(), TickStr));
+		AITEST_FALSE("StateTree Task1A should not tick", Exec.Expect(Task1A.GetName(), TickStr));
 		AITEST_TRUE("StateTree should be running", Status == EStateTreeRunStatus::Running);
 		Exec.LogClear();
 
@@ -331,8 +335,8 @@ struct FStateTreeTest_Select : FAITestBase
 		// Partial reselect, Root should not get EnterState
 		Status = Exec.Tick(0.1f);
 		AITEST_FALSE("StateTree TaskRoot should not enter state", Exec.Expect(TaskRoot.GetName(), EnterStateStr));
-		AITEST_TRUE("StateTree Task1 should exit state, and enter state", Exec.Expect(Task1A.GetName(), ExitStateStr).Then(Task1.GetName(), EnterStateStr));
-		AITEST_TRUE("StateTree Task1A should exit state, enter state, and tick", Exec.Expect(Task1A.GetName(), ExitStateStr).Then(Task1A.GetName(), EnterStateStr).Then(Task1A.GetName(), TickStr));
+		AITEST_TRUE("StateTree Task1 should tick, exit state, and enter state", Exec.Expect(Task1.GetName(), TickStr).Then(Task1.GetName(), ExitStateStr).Then(Task1.GetName(), EnterStateStr));
+		AITEST_TRUE("StateTree Task1A should tick, exit state, and enter state", Exec.Expect(Task1A.GetName(), TickStr).Then(Task1A.GetName(), ExitStateStr).Then(Task1A.GetName(), EnterStateStr));
 		AITEST_TRUE("StateTree should be running", Status == EStateTreeRunStatus::Running);
         Exec.LogClear();
 
