StateTree: UObject and BP support
- connect GetWorld() of FStateTreeExecutionContext explicitly to the owner
- take MassSignalSubsystem as parameter to FMassStateTreeExecutionContext directly
- separated DataViewIndex (what others see) and InstanceIndex/bInstanceIsObject (where the pointert comes from) on StateTree items
- allow item instance to be a struct or object
- added editor support for both struct or object based instance
- added Blueprint base classes for Eval, Task and Condition
- Update UStateTreeBrainComponent

diff --git a/Engine/Plugins/AI/MassAI/Source/MassAIBehavior/Private/MassStateTreeExecutionContext.cpp b/Engine/Plugins/AI/MassAI/Source/MassAIBehavior/Private/MassStateTreeExecutionContext.cpp
--- a/Engine/Plugins/AI/MassAI/Source/MassAIBehavior/Private/MassStateTreeExecutionContext.cpp
+++ b/Engine/Plugins/AI/MassAI/Source/MassAIBehavior/Private/MassStateTreeExecutionContext.cpp
@@ -7,11 +7,10 @@
 #include "Engine/World.h"
 
 FMassStateTreeExecutionContext::FMassStateTreeExecutionContext(UMassEntitySubsystem& InEntitySubsystem,
+                                                               UMassSignalSubsystem& InSignalSubsystem,
                                                                FMassExecutionContext& InContext):
-	EntitySubsystem(&InEntitySubsystem), EntitySubsystemExecutionContext(&InContext)
+	EntitySubsystem(&InEntitySubsystem), SignalSubsystem(&InSignalSubsystem), EntitySubsystemExecutionContext(&InContext)
 {
-	World = InEntitySubsystem.GetWorld();
-	SignalSubsystem = UWorld::GetSubsystem<UMassSignalSubsystem>(World);
 }
 
 void FMassStateTreeExecutionContext::BeginGatedTransition(const FStateTreeExecutionState& Exec)
diff --git a/Engine/Plugins/AI/MassAI/Source/MassAIBehavior/Private/MassStateTreeProcessors.cpp b/Engine/Plugins/AI/MassAI/Source/MassAIBehavior/Private/MassStateTreeProcessors.cpp
--- a/Engine/Plugins/AI/MassAI/Source/MassAIBehavior/Private/MassStateTreeProcessors.cpp
+++ b/Engine/Plugins/AI/MassAI/Source/MassAIBehavior/Private/MassStateTreeProcessors.cpp
@@ -140,47 +140,17 @@ void ForEachEntityInChunk(
 
 
 //----------------------------------------------------------------------//
-// UMassStateTreeFragmentInitializer
+// UMassStateTreeFragmentDestructor
 //----------------------------------------------------------------------//
-UMassStateTreeFragmentInitializer::UMassStateTreeFragmentInitializer()
+UMassStateTreeFragmentDestructor::UMassStateTreeFragmentDestructor()
 {
 	ExecutionFlags = (int32)(EProcessorExecutionFlags::Standalone | EProcessorExecutionFlags::Server);
 	FragmentType = FMassStateTreeFragment::StaticStruct();
 }
 
-void UMassStateTreeFragmentInitializer::ConfigureQueries()
-{
-	EntityQuery.AddRequirement<FMassStateTreeFragment>(EMassFragmentAccess::ReadOnly);
-}
-
-void UMassStateTreeFragmentInitializer::Execute(UMassEntitySubsystem& EntitySubsystem, FMassExecutionContext& Context)
-{
-	FMassStateTreeExecutionContext StateTreeContext(EntitySubsystem, Context);
-	UMassStateTreeSubsystem* MassStateTreeSubsystem = UWorld::GetSubsystem<UMassStateTreeSubsystem>(EntitySubsystem.GetWorld());
-
-	EntityQuery.ForEachEntityChunk(
-		EntitySubsystem,
-		Context,
-		[this, &StateTreeContext, &MassStateTreeSubsystem](const FMassExecutionContext& Context)
-		{
-			UE::MassBehavior::ForEachEntityInChunk(
-				StateTreeContext,
-				*MassStateTreeSubsystem,
-				[](FMassStateTreeExecutionContext& StateTreeExecutionContext, FStateTreeDataView Storage)
-				{
-					// Start the tree instance
-					StateTreeExecutionContext.Start(Storage);
-				});
-		});
-}
-
-//----------------------------------------------------------------------//
-// UMassStateTreeFragmentDestructor
-//----------------------------------------------------------------------//
-UMassStateTreeFragmentDestructor::UMassStateTreeFragmentDestructor()
+void UMassStateTreeFragmentDestructor::Initialize(UObject& Owner)
 {
-	ExecutionFlags = (int32)(EProcessorExecutionFlags::Standalone | EProcessorExecutionFlags::Server);
-	FragmentType = FMassStateTreeFragment::StaticStruct();
+	SignalSubsystem = UWorld::GetSubsystem<UMassSignalSubsystem>(Owner.GetWorld());
 }
 
 void UMassStateTreeFragmentDestructor::ConfigureQueries()
@@ -190,7 +160,12 @@ void UMassStateTreeFragmentDestructor::ConfigureQueries()
 
 void UMassStateTreeFragmentDestructor::Execute(UMassEntitySubsystem& EntitySubsystem, FMassExecutionContext& Context)
 {
-	FMassStateTreeExecutionContext StateTreeContext(EntitySubsystem, Context);
+	if (SignalSubsystem == nullptr)
+	{
+		return;
+	}
+		
+	FMassStateTreeExecutionContext StateTreeContext(EntitySubsystem, *SignalSubsystem, Context);
 	UMassStateTreeSubsystem* MassStateTreeSubsystem = UWorld::GetSubsystem<UMassStateTreeSubsystem>(EntitySubsystem.GetWorld());
 
 	EntityQuery.ForEachEntityChunk(
@@ -221,23 +196,35 @@ UMassStateTreeActivationProcessor::UMassStateTreeActivationProcessor()
 	MaxActivationsPerLOD[EMassLOD::Low] = 100;
 	MaxActivationsPerLOD[EMassLOD::Off] = 100;
 }
+
 void UMassStateTreeActivationProcessor::Initialize(UObject& Owner)
 {
 	SignalSubsystem = UWorld::GetSubsystem<UMassSignalSubsystem>(Owner.GetWorld());
 }
+
 void UMassStateTreeActivationProcessor::ConfigureQueries()
 {
 	EntityQuery.AddRequirement<FMassStateTreeFragment>(EMassFragmentAccess::ReadOnly);
 	EntityQuery.AddTagRequirement<FMassStateTreeActivated>(EMassFragmentPresence::None);
 	EntityQuery.AddChunkRequirement<FMassSimulationVariableTickChunkFragment>(EMassFragmentAccess::ReadOnly, EMassFragmentPresence::Optional);
 }
+
 void UMassStateTreeActivationProcessor::Execute(UMassEntitySubsystem& EntitySubsystem, FMassExecutionContext& Context)
 {
+	if (SignalSubsystem == nullptr)
+	{
+		return;
+	}
+
 	// StateTree processor relies on signals to be ticked but we need an 'initial tick' to set the tree in the proper state.
 	// The initializer provides that by sending a signal to all new entities that use StateTree.
+
+	FMassStateTreeExecutionContext StateTreeContext(EntitySubsystem, *SignalSubsystem, Context);
+	UMassStateTreeSubsystem* MassStateTreeSubsystem = UWorld::GetSubsystem<UMassStateTreeSubsystem>(EntitySubsystem.GetWorld());
+
 	TArray<FMassEntityHandle> EntitiesToSignal;
 	int32 ActivationCounts[EMassLOD::Max] {0,0,0,0};
-	EntityQuery.ForEachEntityChunk(EntitySubsystem, Context, [&EntitiesToSignal, &ActivationCounts, MaxActivationsPerLOD = MaxActivationsPerLOD](FMassExecutionContext& Context)
+	EntityQuery.ForEachEntityChunk(EntitySubsystem, Context, [&EntitiesToSignal, &ActivationCounts, MaxActivationsPerLOD = MaxActivationsPerLOD, &StateTreeContext, MassStateTreeSubsystem](FMassExecutionContext& Context)
 	{
 		const int32 NumEntities = Context.GetNumEntities();
 		// Check if we already reached the maximum for this frame
@@ -247,6 +234,19 @@ void UMassStateTreeActivationProcessor::Execute(UMassEntitySubsystem& EntitySubs
 			return;
 		}
 		ActivationCounts[ChunkLOD] += NumEntities;
+
+		// Start StateTree. This may do substantial amount of work, as we select and enter the first state.
+		UE::MassBehavior::ForEachEntityInChunk(
+			StateTreeContext,
+			*MassStateTreeSubsystem,
+			[MassStateTreeSubsystem](FMassStateTreeExecutionContext& StateTreeExecutionContext, FStateTreeDataView Storage)
+			{
+				// Init object instances
+				StateTreeExecutionContext.InitInstanceData(*MassStateTreeSubsystem, Storage);
+				// Start the tree instance
+				StateTreeExecutionContext.Start(Storage);
+			});
+
 		// Append all entities of the current chunk to the consolidated list to send signal once
 		EntitiesToSignal.Append(Context.GetEntities().GetData(), Context.GetEntities().Num());
 		// Adding a tag on each entities to remember we have sent the state tree initialization signal
@@ -283,7 +283,9 @@ UMassStateTreeProcessor::UMassStateTreeProcessor(const FObjectInitializer& Objec
 void UMassStateTreeProcessor::Initialize(UObject& Owner)
 {
 	Super::Initialize(Owner);
+	
 	MassStateTreeSubsystem = UWorld::GetSubsystem<UMassStateTreeSubsystem>(Owner.GetWorld());
+	SignalSubsystem = UWorld::GetSubsystem<UMassSignalSubsystem>(Owner.GetWorld());
 
 	SubscribeToSignal(UE::Mass::Signals::StateTreeActivate);
 	SubscribeToSignal(UE::Mass::Signals::LookAtFinished);
@@ -315,7 +317,7 @@ void UMassStateTreeProcessor::ConfigureQueries()
 
 void UMassStateTreeProcessor::SignalEntities(UMassEntitySubsystem& EntitySubsystem, FMassExecutionContext& Context, FMassSignalNameLookup& EntitySignals)
 {
-	if (!MassStateTreeSubsystem)
+	if (MassStateTreeSubsystem == nullptr || SignalSubsystem == nullptr)
 	{
 		return;
 	}
@@ -323,7 +325,7 @@ void UMassStateTreeProcessor::SignalEntities(UMassEntitySubsystem& EntitySubsyst
 	CSV_SCOPED_TIMING_STAT_EXCLUSIVE(StateTreeProcessorExecute);
 
 	const float TimeDelta = Context.GetDeltaTimeSeconds();
-	FMassStateTreeExecutionContext StateTreeContext(EntitySubsystem, Context);
+	FMassStateTreeExecutionContext StateTreeContext(EntitySubsystem, *SignalSubsystem, Context);
 	const float TimeInSeconds = EntitySubsystem.GetWorld()->GetTimeSeconds();
 
 	TArray<FMassEntityHandle> EntitiesToSignal;
diff --git a/Engine/Plugins/AI/MassAI/Source/MassAIBehavior/Public/MassStateTreeExecutionContext.h b/Engine/Plugins/AI/MassAI/Source/MassAIBehavior/Public/MassStateTreeExecutionContext.h
--- a/Engine/Plugins/AI/MassAI/Source/MassAIBehavior/Public/MassStateTreeExecutionContext.h
+++ b/Engine/Plugins/AI/MassAI/Source/MassAIBehavior/Public/MassStateTreeExecutionContext.h
@@ -20,7 +20,7 @@ struct MASSAIBEHAVIOR_API FMassStateTreeExecutionContext : public FStateTreeExec
 public:
 	/** Should never be used but has to be public for 'void UScriptStruct::TCppStructOps<FMassStateTreeExecutionContext>::ConstructForTests(void *)' */
 	FMassStateTreeExecutionContext() = default;
-	FMassStateTreeExecutionContext(UMassEntitySubsystem& InEntitySubsystem, FMassExecutionContext& InContext);
+	FMassStateTreeExecutionContext(UMassEntitySubsystem& InEntitySubsystem, UMassSignalSubsystem& InSignalSubsystem, FMassExecutionContext& InContext);
 
 	UMassEntitySubsystem& GetEntitySubsystem() const { return *EntitySubsystem; }
 	FMassExecutionContext& GetEntitySubsystemExecutionContext() const { return *EntitySubsystemExecutionContext; }
diff --git a/Engine/Plugins/AI/MassAI/Source/MassAIBehavior/Public/MassStateTreeProcessors.h b/Engine/Plugins/AI/MassAI/Source/MassAIBehavior/Public/MassStateTreeProcessors.h
--- a/Engine/Plugins/AI/MassAI/Source/MassAIBehavior/Public/MassStateTreeProcessors.h
+++ b/Engine/Plugins/AI/MassAI/Source/MassAIBehavior/Public/MassStateTreeProcessors.h
@@ -11,23 +11,8 @@
 struct FMassStateTreeExecutionContext;
 
 /** 
- * Processor to prepare StateTree fragments for their execution
+ * Processor to stop and uninitialize StateTrees on entities.
  */
-UCLASS()
-class MASSAIBEHAVIOR_API UMassStateTreeFragmentInitializer : public UMassFragmentInitializer
-{
-	GENERATED_BODY()
-
-public:
-	UMassStateTreeFragmentInitializer();
-
-protected:
-	virtual void ConfigureQueries() override;
-	virtual void Execute(UMassEntitySubsystem& EntitySubsystem, FMassExecutionContext& Context) override;
-
-	FMassEntityQuery EntityQuery;
-};
-
 UCLASS()
 class MASSAIBEHAVIOR_API UMassStateTreeFragmentDestructor : public UMassFragmentDestructor
 {
@@ -37,10 +22,14 @@ public:
 	UMassStateTreeFragmentDestructor();
 
 protected:
+	virtual void Initialize(UObject& Owner) override;
 	virtual void ConfigureQueries() override;
 	virtual void Execute(UMassEntitySubsystem& EntitySubsystem, FMassExecutionContext& Context) override;
 
 	FMassEntityQuery EntityQuery;
+
+	UPROPERTY(Transient)
+	UMassSignalSubsystem* SignalSubsystem = nullptr;
 };
 
 /**
diff --git a/Engine/Plugins/AI/MassAI/Source/MassAIDebug/MassAIDebug.Build.cs b/Engine/Plugins/AI/MassAI/Source/MassAIDebug/MassAIDebug.Build.cs
--- a/Engine/Plugins/AI/MassAI/Source/MassAIDebug/MassAIDebug.Build.cs
+++ b/Engine/Plugins/AI/MassAI/Source/MassAIDebug/MassAIDebug.Build.cs
@@ -41,6 +41,7 @@ namespace UnrealBuildTool.Rules
 					"MassSpawner",
 					"MassSimulation",
 					"MassRepresentation",
+					"MassSignals",
 					"MassLOD",
 					"StructUtils",
 				}
diff --git a/Engine/Plugins/AI/MassAI/Source/MassAIDebug/Private/GameplayDebuggerCategory_Mass.cpp b/Engine/Plugins/AI/MassAI/Source/MassAIDebug/Private/GameplayDebuggerCategory_Mass.cpp
--- a/Engine/Plugins/AI/MassAI/Source/MassAIDebug/Private/GameplayDebuggerCategory_Mass.cpp
+++ b/Engine/Plugins/AI/MassAI/Source/MassAIDebug/Private/GameplayDebuggerCategory_Mass.cpp
@@ -8,6 +8,7 @@
 #include "GameplayDebuggerCategoryReplicator.h"
 #include "GameplayDebuggerPlayerManager.h"
 #include "MassDebuggerSubsystem.h"
+#include "MassSignalSubsystem.h"
 #include "MassActorSubsystem.h"
 #include "GameFramework/PlayerController.h"
 #include "MassAgentComponent.h"
@@ -354,11 +355,13 @@ void FGameplayDebuggerCategory_Mass::CollectData(APlayerController* OwnerPC, AAc
 		const float CurrentTime = World->GetTimeSeconds();
 		
 		UMassStateTreeSubsystem* MassStateTreeSubsystem = World->GetSubsystem<UMassStateTreeSubsystem>();
-		if (MassStateTreeSubsystem && EntitySystem)
+		UMassSignalSubsystem* SignalSubsystem = World->GetSubsystem<UMassSignalSubsystem>();
+		
+		if (MassStateTreeSubsystem && EntitySystem && SignalSubsystem)
 		{
 			FMassExecutionContext Context(0.0f);
 		
-			EntityQuery.ForEachEntityChunk(*EntitySystem, Context, [this, Debugger, MassStateTreeSubsystem, EntitySystem, OwnerPC, ViewLocation, ViewDirection, CurrentTime](FMassExecutionContext& Context)
+			EntityQuery.ForEachEntityChunk(*EntitySystem, Context, [this, Debugger, MassStateTreeSubsystem, SignalSubsystem, EntitySystem, OwnerPC, ViewLocation, ViewDirection, CurrentTime](FMassExecutionContext& Context)
 			{
 				const int32 NumEntities = Context.GetNumEntities();
 				const TConstArrayView<FMassStateTreeFragment> StateTreeList = Context.GetFragmentView<FMassStateTreeFragment>();
@@ -529,7 +532,7 @@ void FGameplayDebuggerCategory_Mass::CollectData(APlayerController* OwnerPC, AAc
 						// Current StateTree task
 						if (StateTree != nullptr)
 						{
-							FMassStateTreeExecutionContext StateTreeContext(*EntitySystem, Context);
+							FMassStateTreeExecutionContext StateTreeContext(*EntitySystem, *SignalSubsystem, Context);
 							StateTreeContext.Init(*OwnerPC, *StateTree, EStateTreeStorage::External);
 							StateTreeContext.SetEntity(Entity);
 							
diff --git a/Engine/Plugins/AI/MassAI/Source/MassAIDebug/Private/MassDebugStateTreeProcessor.cpp b/Engine/Plugins/AI/MassAI/Source/MassAIDebug/Private/MassDebugStateTreeProcessor.cpp
--- a/Engine/Plugins/AI/MassAI/Source/MassAIDebug/Private/MassDebugStateTreeProcessor.cpp
+++ b/Engine/Plugins/AI/MassAI/Source/MassAIDebug/Private/MassDebugStateTreeProcessor.cpp
@@ -2,6 +2,7 @@
 
 #include "MassDebugStateTreeProcessor.h"
 #include "MassDebuggerSubsystem.h"
+#include "MassSignalSubsystem.h"
 #include "MassStateTreeExecutionContext.h"
 #include "MassStateTreeFragments.h"
 #include "MassCommonFragments.h"
@@ -44,13 +45,19 @@ void UMassDebugStateTreeProcessor::Execute(UMassEntitySubsystem& EntitySubsystem
 		return;
 	}
 
+	UMassSignalSubsystem* MassSignalSubsystem = World->GetSubsystem<UMassSignalSubsystem>();
+	if (MassSignalSubsystem == nullptr)
+	{
+		return;
+	}
+
 	if (!Debugger->GetSelectedEntity().IsSet() && !UE::Mass::Debug::HasDebugEntities())
 	{
 		return;
 	}
 	
 	QUICK_SCOPE_CYCLE_COUNTER(UMassDebugStateTreeProcessor_Run);	
-	EntityQuery.ForEachEntityChunk(EntitySubsystem, Context, [this, Debugger, MassStateTreeSubsystem, &EntitySubsystem](FMassExecutionContext& Context)
+	EntityQuery.ForEachEntityChunk(EntitySubsystem, Context, [this, Debugger, MassStateTreeSubsystem, &EntitySubsystem, MassSignalSubsystem](FMassExecutionContext& Context)
 		{
 			const FMassEntityHandle SelectedEntity = Debugger->GetSelectedEntity();
 			const int32 NumEntities = Context.GetNumEntities();
@@ -69,7 +76,7 @@ void UMassDebugStateTreeProcessor::Execute(UMassEntitySubsystem& EntitySubsystem
 				FMassEntityHandle Entity = Context.GetEntity(i);
 				if (Entity == SelectedEntity)
 				{
-					FMassStateTreeExecutionContext StateTreeContext(EntitySubsystem, Context);
+					FMassStateTreeExecutionContext StateTreeContext(EntitySubsystem, *MassSignalSubsystem, Context);
 					StateTreeContext.Init(*this, *StateTree, EStateTreeStorage::External);
 					StateTreeContext.SetEntity(Entity);
 					
@@ -88,7 +95,7 @@ void UMassDebugStateTreeProcessor::Execute(UMassEntitySubsystem& EntitySubsystem
 					const FVector ZOffset(0,0,50);
 					const FVector Position = Transform.GetTransform().GetLocation() + ZOffset;
 
-					FMassStateTreeExecutionContext StateTreeContext(EntitySubsystem, Context);
+					FMassStateTreeExecutionContext StateTreeContext(EntitySubsystem, *MassSignalSubsystem, Context);
 					StateTreeContext.Init(*this, *StateTree, EStateTreeStorage::External);
 					StateTreeContext.SetEntity(Entity);
 
diff --git a/Engine/Plugins/Runtime/StateTree/Source/StateTreeEditorModule/Private/Customizations/StateTreeBindableItemDetails.cpp b/Engine/Plugins/Runtime/StateTree/Source/StateTreeEditorModule/Private/Customizations/StateTreeBindableItemDetails.cpp
--- a/Engine/Plugins/Runtime/StateTree/Source/StateTreeEditorModule/Private/Customizations/StateTreeBindableItemDetails.cpp
+++ b/Engine/Plugins/Runtime/StateTree/Source/StateTreeEditorModule/Private/Customizations/StateTreeBindableItemDetails.cpp
@@ -20,6 +20,8 @@
 #include "Modules/ModuleManager.h"
 #include "StructViewerModule.h"
 #include "StructViewerFilter.h"
+#include "ClassViewerModule.h"
+#include "ClassViewerFilter.h"
 #include "InstancedStructDetails.h"
 #include "Engine/UserDefinedStruct.h"
 #include "StateTreeBindingExtension.h"
@@ -29,26 +31,17 @@
 #include "StateTreeConditionBase.h"
 #include "StateTreePropertyHelpers.h"
 #include "StateTreeEditorStyle.h"
-
 #include "EdGraphSchema_K2.h"
 #include "Kismet2/BlueprintEditorUtils.h"
+#include "Blueprint/StateTreeEvaluatorBlueprintBase.h"
+#include "Blueprint/StateTreeTaskBlueprintBase.h"
+#include "Blueprint/StateTreeConditionBlueprintBase.h"
 
 #define LOCTEXT_NAMESPACE "StateTreeEditor"
 
-// Customized version of FInstancedStructDataDetails used to hide bindable properties.
-class FBindableItemDataDetails : public FInstancedStructDataDetails
+namespace UE::StateTreeEditor::Internal
 {
-public:
-
-	FBindableItemDataDetails(TSharedPtr<IPropertyHandle> InStructProperty, FGuid InID, UStateTreeEditorData* InEditorData)
-		: FInstancedStructDataDetails(InStructProperty)
-		, EditorData(InEditorData)
-	{
-		EditorPropBindings = EditorData ? EditorData->GetPropertyEditorBindings() : nullptr;
-		ID = InID;
-	}
-
-	virtual void OnChildRowAdded(IDetailPropertyRow& ChildRow)
+	void ModifyRow(IDetailPropertyRow& ChildRow, const FGuid& ID, FStateTreeEditorPropertyBindings* EditorPropBindings)
 	{
 		if (!EditorPropBindings)
 		{
@@ -71,7 +64,7 @@ public:
 			FDetailWidgetRow Row;
 			ChildRow.GetDefaultWidgets(NameWidget, ValueWidget, Row);
 
-			auto IsValueVisible = TAttribute<EVisibility>::Create([this, Path]() -> EVisibility
+			auto IsValueVisible = TAttribute<EVisibility>::Create([Path, EditorPropBindings]() -> EVisibility
 				{
 					return EditorPropBindings->HasPropertyBinding(Path) ? EVisibility::Collapsed : EVisibility::Visible;
 				});
@@ -134,13 +127,31 @@ public:
 			}
 		}
 	}
+} // UE::StateTreeEditor::Internal
+
+// Customized version of FInstancedStructDataDetails used to hide bindable properties.
+class FBindableItemInstanceDetails : public FInstancedStructDataDetails
+{
+public:
+
+	FBindableItemInstanceDetails(TSharedPtr<IPropertyHandle> InStructProperty, FGuid InID, UStateTreeEditorData* InEditorData)
+		: FInstancedStructDataDetails(InStructProperty)
+		, EditorData(InEditorData)
+	{
+		EditorPropBindings = EditorData ? EditorData->GetPropertyEditorBindings() : nullptr;
+		ID = InID;
+	}
+
+	virtual void OnChildRowAdded(IDetailPropertyRow& ChildRow)
+	{
+		UE::StateTreeEditor::Internal::ModifyRow(ChildRow, ID, EditorPropBindings);
+	}
 
 	UStateTreeEditorData* EditorData;
 	FStateTreeEditorPropertyBindings* EditorPropBindings;
 	FGuid ID;
 };
 
-
 ////////////////////////////////////
 
 class FBindableItemStructFilter : public IStructViewerFilter
@@ -181,7 +192,7 @@ public:
 		}
 
 		// Query the native struct to see if it has the correct parent type (if any)
-		return !BaseStruct || InStruct->IsChildOf(BaseStruct);
+		return BaseStruct != nullptr && InStruct->IsChildOf(BaseStruct);
 	}
 
 	virtual bool IsUnloadedStructAllowed(const FStructViewerInitializationOptions& InInitOptions, const FName InStructPath, TSharedRef<FStructViewerFilterFuncs> InFilterFuncs) override
@@ -191,6 +202,63 @@ public:
 	}
 };
 
+////////////////////////////////////
+
+class FBindableItemClassFilter : public IClassViewerFilter
+{
+public:
+	/** The base struct for the property that classes must be a child-of. */
+	const UClass* BaseClass = nullptr;
+
+	// A flag controlling whether we allow to select the BaseStruct
+	bool bAllowBaseClass = true;
+
+	// Schema to filter which structs are allowed.
+	const UStateTreeSchema* Schema = nullptr;
+
+	virtual bool IsClassAllowed(const FClassViewerInitializationOptions& InInitOptions, const UClass* InClass, TSharedRef<FClassViewerFilterFuncs> InFilterFuncs ) override
+	{
+		if (!Schema || !InClass)
+		{
+			return false;
+		}
+
+		if (!Schema->IsClassAllowed(InClass))
+		{
+			return false;
+		}
+		
+		if (InClass == BaseClass)
+		{
+			return bAllowBaseClass;
+		}
+
+		// Query the native struct to see if it has the correct parent type (if any)
+		return BaseClass != nullptr && InClass->IsChildOf(BaseClass);
+	}
+
+	virtual bool IsUnloadedClassAllowed(const FClassViewerInitializationOptions& InInitOptions, const TSharedRef<const IUnloadedBlueprintData> InUnloadedClassData, TSharedRef<FClassViewerFilterFuncs> InFilterFuncs) override
+	{
+		if (!Schema)
+		{
+			return false;
+		}
+
+		const UClass* Class = InUnloadedClassData->GetClassWithin();
+		
+		if (!Schema->IsClassAllowed(Class))
+		{
+			return false;
+		}
+		
+		if (Class == BaseClass)
+		{
+			return bAllowBaseClass;
+		}
+
+		return !BaseClass || Class->IsChildOf(BaseClass);
+	}
+};
 
 ////////////////////////////////////
 
@@ -211,17 +279,23 @@ void FStateTreeBindableItemDetails::CustomizeHeader(TSharedRef<class IPropertyHa
 
 	ItemProperty = StructProperty->GetChildHandle(TEXT("Item"));
 	InstanceProperty = StructProperty->GetChildHandle(TEXT("Instance"));
+	InstanceObjectProperty = StructProperty->GetChildHandle(TEXT("InstanceObject"));
 	IDProperty = StructProperty->GetChildHandle(TEXT("ID"));
 
 	check(ItemProperty.IsValid());
 	check(InstanceProperty.IsValid());
 	check(IDProperty.IsValid());
 	
-	// Find base class from meta data.
-	static const FName BaseClassMetaName(TEXT("BaseStruct")); // TODO: move these names into one central place.
-	const FString BaseClassName = ItemProperty->GetMetaData(BaseClassMetaName);
-	BaseScriptStruct = FindObject<UScriptStruct>(ANY_PACKAGE, *BaseClassName);
+	// Find base class and struct from meta data.
+	static const FName BaseStructMetaName(TEXT("BaseStruct")); // TODO: move these names into one central place.
+	static const FName BaseClassMetaName(TEXT("BaseClass")); // TODO: move these names into one central place.
+	
+	const FString BaseStructName = StructProperty->GetMetaData(BaseStructMetaName);
+	BaseScriptStruct = FindObject<UScriptStruct>(ANY_PACKAGE, *BaseStructName);
 
+	const FString BaseClassName = StructProperty->GetMetaData(BaseClassMetaName);
+	BaseClass = FindObject<UClass>(ANY_PACKAGE, *BaseClassName);
+	
 	const FIsResetToDefaultVisible IsResetVisible = FIsResetToDefaultVisible::CreateSP(this, &FStateTreeBindableItemDetails::ShouldResetToDefault);
 	const FResetToDefaultHandler ResetHandler = FResetToDefaultHandler::CreateSP(this, &FStateTreeBindableItemDetails::ResetToDefault);
 	const FResetToDefaultOverride ResetOverride = FResetToDefaultOverride::Create(IsResetVisible, ResetHandler);
@@ -268,7 +342,7 @@ void FStateTreeBindableItemDetails::CustomizeHeader(TSharedRef<class IPropertyHa
 			.VAlign(VAlign_Center)
 			[
 				SAssignNew(ComboButton, SComboButton)
-				.OnGetMenuContent(this, &FStateTreeBindableItemDetails::GenerateStructPicker)
+				.OnGetMenuContent(this, &FStateTreeBindableItemDetails::GeneratePicker)
 				.ContentPadding(0.f)
 				.ButtonContent()
 				[
@@ -301,17 +375,17 @@ void FStateTreeBindableItemDetails::CustomizeHeader(TSharedRef<class IPropertyHa
 
 bool FStateTreeBindableItemDetails::ShouldResetToDefault(TSharedPtr<IPropertyHandle> PropertyHandle) const
 {
-	check(ItemProperty);
+	check(StructProperty);
 	
 	bool bAnyValid = false;
 	
-	TArray<void*> RawData;
-	ItemProperty->AccessRawData(RawData);
-	for (void* Data : RawData)
+	TArray<const void*> RawItemData;
+	StructProperty->AccessRawData(RawItemData);
+	for (const void* Data : RawItemData)
 	{
-		if (const FInstancedStruct* Struct = static_cast<FInstancedStruct*>(Data))
+		if (const FStateTreeItem* Item = static_cast<const FStateTreeItem*>(Data))
 		{
-			if (Struct->IsValid())
+			if (Item->Item.IsValid())
 			{
 				bAnyValid = true;
 				break;
@@ -323,42 +397,27 @@ bool FStateTreeBindableItemDetails::ShouldResetToDefault(TSharedPtr<IPropertyHan
 	return bAnyValid;
 }
 
+
 void FStateTreeBindableItemDetails::ResetToDefault(TSharedPtr<IPropertyHandle> PropertyHandle)
 {
-	check(ItemProperty);
-	check(InstanceProperty);
+	check(StructProperty);
 	
 	GEditor->BeginTransaction(LOCTEXT("OnResetToDefault", "Reset to default"));
 
-	ItemProperty->NotifyPreChange();
+	StructProperty->NotifyPreChange();
+	
 	TArray<void*> RawItemData;
-	ItemProperty->AccessRawData(RawItemData);
+	StructProperty->AccessRawData(RawItemData);
 	for (void* Data : RawItemData)
 	{
-		if (FInstancedStruct* Struct = static_cast<FInstancedStruct*>(Data))
-		{
-			// Assume that the default value is empty.
-			Struct->Reset();
-		}
-	}
-
-	InstanceProperty->NotifyPostChange(EPropertyChangeType::ValueSet);
-	InstanceProperty->NotifyFinishedChangingProperties();
-
-	InstanceProperty->NotifyPreChange();
-	TArray<void*> RawInstanceData;
-	InstanceProperty->AccessRawData(RawInstanceData);
-	for (void* Data : RawInstanceData)
-	{
-		if (FInstancedStruct* Struct = static_cast<FInstancedStruct*>(Data))
+		if (FStateTreeItem* Item = static_cast<FStateTreeItem*>(Data))
 		{
-			// Assume that the default value is empty.
-			Struct->Reset();
+			Item->Reset();
 		}
 	}
 
-	InstanceProperty->NotifyPostChange(EPropertyChangeType::ValueSet);
-	InstanceProperty->NotifyFinishedChangingProperties();
+	StructProperty->NotifyPostChange(EPropertyChangeType::ValueSet);
+	StructProperty->NotifyFinishedChangingProperties();
 
 	GEditor->EndTransaction();
 
@@ -370,68 +429,100 @@ void FStateTreeBindableItemDetails::ResetToDefault(TSharedPtr<IPropertyHandle> P
 
 void FStateTreeBindableItemDetails::CustomizeChildren(TSharedRef<class IPropertyHandle> StructPropertyHandle, class IDetailChildrenBuilder& StructBuilder, IPropertyTypeCustomizationUtils& StructCustomizationUtils)
 {
+	FStateTreeEditorPropertyBindings* EditorPropBindings = EditorData ? EditorData->GetPropertyEditorBindings() : nullptr;
+
 	FGuid ID;
 	UE::StateTree::PropertyHelpers::GetStructValue<FGuid>(IDProperty, ID);
 
-	TSharedRef<FBindableItemDataDetails> ItemDetails = MakeShareable(new FBindableItemDataDetails(ItemProperty, FGuid(), EditorData));
+	TSharedRef<FBindableItemInstanceDetails> ItemDetails = MakeShareable(new FBindableItemInstanceDetails(ItemProperty, FGuid(), EditorData));
 	StructBuilder.AddCustomBuilder(ItemDetails);
 
-	TSharedRef<FBindableItemDataDetails> InstanceDetails = MakeShareable(new FBindableItemDataDetails(InstanceProperty, ID, EditorData));
+	// Instance
+	TSharedRef<FBindableItemInstanceDetails> InstanceDetails = MakeShareable(new FBindableItemInstanceDetails(InstanceProperty, ID, EditorData));
 	StructBuilder.AddCustomBuilder(InstanceDetails);
-}
 
-void FStateTreeBindableItemDetails::OnIdentifierChanged(const UStateTree& InStateTree)
-{
-	if (PropUtils && StateTree == &InStateTree)
+	// InstanceObject
+	// Get the actual UObject from the pointer.
+	TSharedPtr<IPropertyHandle> InstanceObjectValueProperty = GetInstancedObjectValueHandle(InstanceObjectProperty);
+	if (InstanceObjectValueProperty.IsValid())
 	{
-		PropUtils->ForceRefresh();
+		static const FName CategoryName(TEXT("Category"));
+
+		uint32 NumChildren = 0;
+		InstanceObjectValueProperty->GetNumChildren(NumChildren);
+		for (uint32 Index = 0; Index < NumChildren; Index++)
+		{
+			TSharedPtr<IPropertyHandle> Child = InstanceObjectValueProperty->GetChildHandle(Index);
+			
+			if (!Child.IsValid()
+				|| Child->GetProperty()->HasAllPropertyFlags(CPF_DisableEditOnInstance)
+				|| Child->GetMetaData(CategoryName) == TEXT("ExternalData"))
+			{
+				continue;
+			}
+			
+			IDetailPropertyRow& ChildRow = StructBuilder.AddProperty(Child.ToSharedRef());
+			UE::StateTreeEditor::Internal::ModifyRow(ChildRow, ID, EditorPropBindings);
+		}
 	}
 }
 
-void FStateTreeBindableItemDetails::OnBindingChanged(const FStateTreeEditorPropertyPath& SourcePath, const FStateTreeEditorPropertyPath& TargetPath)
+TSharedPtr<IPropertyHandle> FStateTreeBindableItemDetails::GetInstancedObjectValueHandle(TSharedPtr<IPropertyHandle> PropertyHandle)
 {
-	check(ItemProperty);
+	TSharedPtr<IPropertyHandle> ChildHandle;
 
-	FGuid ID;
-	if (UE::StateTree::PropertyHelpers::GetStructValue(IDProperty, ID) != FPropertyAccess::Success)
+	uint32 NumChildren;
+	PropertyHandle->GetNumChildren(NumChildren);
+
+	if (NumChildren > 0)
 	{
-		return;
+		// when the property is a (inlined) object property, the first child will be
+		// the object instance, and its properties are the children underneath that
+		ensure(NumChildren == 1);
+		ChildHandle = PropertyHandle->GetChildHandle(0);
 	}
 
-	if (ID != TargetPath.StructID)
+	return ChildHandle;
+}
+
+void FStateTreeBindableItemDetails::OnIdentifierChanged(const UStateTree& InStateTree)
+{
+	if (PropUtils && StateTree == &InStateTree)
 	{
-		return;
+		PropUtils->ForceRefresh();
 	}
+}
+
+void FStateTreeBindableItemDetails::OnBindingChanged(const FStateTreeEditorPropertyPath& SourcePath, const FStateTreeEditorPropertyPath& TargetPath)
+{
+	check(StructProperty);
 
 	TArray<UObject*> OuterObjects;
-	ItemProperty->GetOuterObjects(OuterObjects);
+	StructProperty->GetOuterObjects(OuterObjects);
 
 	TArray<void*> RawItemData;
-	ItemProperty->AccessRawData(RawItemData);
-
-	TArray<void*> RawInstanceData;
-	InstanceProperty->AccessRawData(RawInstanceData);
+	StructProperty->AccessRawData(RawItemData);
 
-	if (OuterObjects.Num() != RawItemData.Num() || OuterObjects.Num() != RawInstanceData.Num())
+	if (OuterObjects.Num() != RawItemData.Num())
 	{
 		return;
 	}
 
 	for (int32 i = 0; i < OuterObjects.Num(); i++)
 	{
-		// Dig out name from the struct without knowing the type.
-		const FInstancedStruct* ItemPtr = static_cast<FInstancedStruct*>(RawItemData[i]);
-		const FInstancedStruct* InstancePtr = static_cast<FInstancedStruct*>(RawInstanceData[i]);
+		const FStateTreeItem* Item = static_cast<FStateTreeItem*>(RawItemData[i]);
 		UObject* OuterObject = OuterObjects[i]; // Immediate outer, i.e StateTreeState
-
-		if (InstancePtr->IsValid() && EditorData != nullptr)
+		if (Item != nullptr && EditorData != nullptr && Item->Item.IsValid() && Item->Instance.IsValid())
 		{
-			if (FStateTreeConditionBase* Condition = ItemPtr->GetMutablePtr<FStateTreeConditionBase>())
+			if (Item->ID == TargetPath.StructID)
 			{
-				const FStateTreeBindingLookup BindingLookup(EditorData);
+				if (FStateTreeConditionBase* Condition = Item->Item.GetMutablePtr<FStateTreeConditionBase>())
+				{
+					const FStateTreeBindingLookup BindingLookup(EditorData);
 
-				OuterObject->Modify();
-				Condition->OnBindingChanged(ID, *InstancePtr, SourcePath, TargetPath, BindingLookup);
+					OuterObject->Modify();
+					Condition->OnBindingChanged(Item->ID, Item->Instance, SourcePath, TargetPath, BindingLookup);
+				}
 			}
 		}
 	}
@@ -439,13 +530,13 @@ void FStateTreeBindableItemDetails::OnBindingChanged(const FStateTreeEditorPrope
 
 void FStateTreeBindableItemDetails::FindOuterObjects()
 {
-	check(ItemProperty);
+	check(StructProperty);
 	
 	EditorData = nullptr;
 	StateTree = nullptr;
 
 	TArray<UObject*> OuterObjects;
-	ItemProperty->GetOuterObjects(OuterObjects);
+	StructProperty->GetOuterObjects(OuterObjects);
 	for (UObject* Outer : OuterObjects)
 	{
 		UStateTreeEditorData* OuterEditorData = Outer->GetTypedOuter<UStateTreeEditorData>();
@@ -461,39 +552,24 @@ void FStateTreeBindableItemDetails::FindOuterObjects()
 
 FText FStateTreeBindableItemDetails::GetDescription() const
 {
-	check(ItemProperty);
-	check(InstanceProperty);
+	check(StructProperty);
 
-	// Multiple descriptions do not make sense, just if only one item is selected.
-	if (ItemProperty->GetNumOuterObjects() != 1)
-	{
-		return LOCTEXT("MultipleSelected", "<Details.Subdued>Multiple selected</>");
-	}
+	TArray<void*> RawItemData;
+	StructProperty->AccessRawData(RawItemData);
 
-	FGuid ID;
-	if (UE::StateTree::PropertyHelpers::GetStructValue(IDProperty, ID) != FPropertyAccess::Success)
+	// Multiple descriptions do not make sense, just if only one item is selected.
+	if (RawItemData.Num() != 1)
 	{
 		return LOCTEXT("MultipleSelected", "<Details.Subdued>Multiple selected</>");
 	}
 
-	TArray<void*> RawItemData;
-	ItemProperty->AccessRawData(RawItemData);
-	check(RawItemData.Num() == 1);
-
-	TArray<void*> RawInstanceData;
-	InstanceProperty->AccessRawData(RawInstanceData);
-	check(RawInstanceData.Num() == 1);
-
-	// Dig out name from the struct without knowing the type.
-	const FInstancedStruct* ItemPtr = static_cast<FInstancedStruct*>(RawItemData[0]);
-	const FInstancedStruct* InstancePtr = static_cast<FInstancedStruct*>(RawInstanceData[0]);
-
-	if (InstancePtr->IsValid() && EditorData != nullptr)
+	FStateTreeItem* Item = static_cast<FStateTreeItem*>(RawItemData[0]);
+	if (Item != nullptr && EditorData != nullptr && Item->Item.IsValid() && Item->Instance.IsValid())
 	{
-		if (const FStateTreeConditionBase* Condition = ItemPtr->GetPtr<const FStateTreeConditionBase>())
+		if (const FStateTreeConditionBase* Condition = Item->Item.GetPtr<const FStateTreeConditionBase>())
 		{
 			const FStateTreeBindingLookup BindingLookup(EditorData);
-			return Condition->GetDescription(ID, *InstancePtr, BindingLookup);
+			return Condition->GetDescription(Item->ID, Item->Instance, BindingLookup);
 		}
 	}
 	
@@ -508,26 +584,26 @@ EVisibility FStateTreeBindableItemDetails::IsDescriptionVisible() const
 
 FText FStateTreeBindableItemDetails::GetName() const
 {
-	check(ItemProperty);
+	check(StructProperty);
 
 	// Multiple names do not make sense, just if only one item is selected.
 	TArray<void*> RawItemData;
-	ItemProperty->AccessRawData(RawItemData);
+	StructProperty->AccessRawData(RawItemData);
 	if (RawItemData.Num() == 1)
 	{
 		// Dig out name from the struct without knowing the type.
-		FInstancedStruct* StructPtr = static_cast<FInstancedStruct*>(RawItemData[0]);
+		FStateTreeItem* Item = static_cast<FStateTreeItem*>(RawItemData[0]);
 
 		// Make sure the associated property has valid data (e.g. while moving an item in array)
-		const UScriptStruct* ScriptStruct = StructPtr != nullptr ? StructPtr->GetScriptStruct() : nullptr;
+		const UScriptStruct* ScriptStruct = Item != nullptr ? Item->Item.GetScriptStruct() : nullptr;
 		if (ScriptStruct != nullptr)
 		{
 			if (FProperty* NameProperty = ScriptStruct->FindPropertyByName(TEXT("Name")))
 			{
 				if (NameProperty->IsA(FNameProperty::StaticClass()))
 				{
-					void* Ptr = const_cast<void*>(static_cast<const void*>(StructPtr->GetMemory()));
-					FName NameValue = *NameProperty->ContainerPtrToValuePtr<FName>(Ptr);
+					void* Ptr = const_cast<void*>(static_cast<const void*>(Item->Item.GetMemory()));
+					const FName NameValue = *NameProperty->ContainerPtrToValuePtr<FName>(Ptr);
 					return FText::FromName(NameValue);
 				}
 			}
@@ -546,7 +622,7 @@ EVisibility FStateTreeBindableItemDetails::IsNameVisible() const
 
 void FStateTreeBindableItemDetails::OnNameCommitted(const FText& NewText, ETextCommit::Type InTextCommit)
 {
-	check(ItemProperty);
+	check(StructProperty);
 
 	if (InTextCommit == ETextCommit::OnEnter || InTextCommit == ETextCommit::OnUserMovedFocus)
 	{
@@ -557,29 +633,32 @@ void FStateTreeBindableItemDetails::OnNameCommitted(const FText& NewText, ETextC
 		{
 			GEditor->BeginTransaction(LOCTEXT("SetName", "Set Name"));
 		}
-		ItemProperty->NotifyPreChange();
+		StructProperty->NotifyPreChange();
 
 		TArray<void*> RawItemData;
-		ItemProperty->AccessRawData(RawItemData);
+		StructProperty->AccessRawData(RawItemData);
+		
 		for (void* Data : RawItemData)
 		{
-			// Set Name
-			FInstancedStruct* ItemPtr = static_cast<FInstancedStruct*>(Data);
-			if (const UScriptStruct* ScriptStruct = ItemPtr->GetScriptStruct())
+			if (FStateTreeItem* Item = static_cast<FStateTreeItem*>(Data))
 			{
-				if (FProperty* NameProperty = ScriptStruct->FindPropertyByName(TEXT("Name")))
+				// Set Name
+				if (const UScriptStruct* ScriptStruct = Item->Item.GetScriptStruct())
 				{
-					if (NameProperty->IsA(FNameProperty::StaticClass()))
+					if (FProperty* NameProperty = ScriptStruct->FindPropertyByName(TEXT("Name")))
 					{
-						void* Ptr = const_cast<void*>(static_cast<const void*>(ItemPtr->GetMemory()));
-						FName& NameValue = *NameProperty->ContainerPtrToValuePtr<FName>(Ptr);
-						NameValue = FName(NewName);
+						if (NameProperty->IsA(FNameProperty::StaticClass()))
+						{
+							void* Ptr = const_cast<void*>(static_cast<const void*>(Item->Item.GetMemory()));
+							FName& NameValue = *NameProperty->ContainerPtrToValuePtr<FName>(Ptr);
+							NameValue = FName(NewName);
+						}
 					}
 				}
 			}
 		}
 
-		ItemProperty->NotifyPostChange(EPropertyChangeType::ValueSet);
+		StructProperty->NotifyPostChange(EPropertyChangeType::ValueSet);
 
 		if (StateTree)
 		{
@@ -591,7 +670,7 @@ void FStateTreeBindableItemDetails::OnNameCommitted(const FText& NewText, ETextC
 			GEditor->EndTransaction();
 		}
 
-		ItemProperty->NotifyFinishedChangingProperties();
+		StructProperty->NotifyFinishedChangingProperties();
 	}
 }
 
@@ -603,17 +682,18 @@ bool FStateTreeBindableItemDetails::IsNameEnabled() const
 
 const UScriptStruct* FStateTreeBindableItemDetails::GetCommonItemScriptStruct() const
 {
-	check(ItemProperty);
+	check(StructProperty);
 
 	const UScriptStruct* CommonScriptStruct = nullptr;
 	bool bMultipleValues = false;
 	TArray<void*> RawItemData;
-	ItemProperty->AccessRawData(RawItemData);
+	StructProperty->AccessRawData(RawItemData);
+
 	for (void* Data : RawItemData)
 	{
-		if (FInstancedStruct* Struct = static_cast<FInstancedStruct*>(Data))
+		if (FStateTreeItem* Item = static_cast<FStateTreeItem*>(Data))
 		{
-			const UScriptStruct* ScriptStruct = Struct->GetScriptStruct();
+			const UScriptStruct* ScriptStruct = Item->Item.GetScriptStruct();
 			if (!bMultipleValues && !CommonScriptStruct)
 			{
 				CommonScriptStruct = ScriptStruct;
@@ -644,15 +724,42 @@ const FSlateBrush* FStateTreeBindableItemDetails::GetDisplayValueIcon() const
 	return FSlateIconFinder::FindIconBrushForClass(UScriptStruct::StaticClass());
 }
 
-TSharedRef<SWidget> FStateTreeBindableItemDetails::GenerateStructPicker()
+TSharedRef<SWidget> FStateTreeBindableItemDetails::GeneratePicker()
 {
-	static const FName ExcludeBaseStruct(TEXT("ExcludeBaseStruct")); // TODO: move these names into one central place.
-	const bool bExcludeBaseStruct = ItemProperty->HasMetaData(ExcludeBaseStruct);
+	FMenuBuilder MenuBuilder(true, nullptr);
 
+	// @todo: Find a way to combine the menus. Now both are displayed even if one may be empty.
+	
+	MenuBuilder.AddWrapperSubMenu(
+		LOCTEXT("Items", "Items"),
+		LOCTEXT("Items_ToolTip", "Items"),
+		FOnGetContent::CreateSP(this, &FStateTreeBindableItemDetails::GenerateStructPicker),
+		FSlateIcon()
+		);
+
+	MenuBuilder.AddWrapperSubMenu(
+		LOCTEXT("BlueprintItems", "Blueprint Items"),
+		LOCTEXT("BlueprintItems_ToolTip", "Blueprint Items"),
+		FOnGetContent::CreateSP(this, &FStateTreeBindableItemDetails::GenerateClassPicker),
+		FSlateIcon()
+	);
+
+	MenuBuilder.AddMenuSeparator();
+
+	FUIAction ClearAction(FExecuteAction::CreateSP(this, &FStateTreeBindableItemDetails::OnStructPicked, (const UScriptStruct*)nullptr));
+	MenuBuilder.AddMenuEntry(LOCTEXT("ClearItem", "Clear Item"), TAttribute<FText>(), FSlateIcon(FEditorStyle::GetStyleSetName(), "Cross"), ClearAction);
+	
+	return MenuBuilder.MakeWidget();
+}
+
+
+
+TSharedRef<SWidget> FStateTreeBindableItemDetails::GenerateStructPicker()
+{
 	TSharedRef<FBindableItemStructFilter> StructFilter = MakeShared<FBindableItemStructFilter>();
 	StructFilter->BaseStruct = BaseScriptStruct;
 	StructFilter->bAllowUserDefinedStructs = false;
-	StructFilter->bAllowBaseStruct = !bExcludeBaseStruct;
+	StructFilter->bAllowBaseStruct = false;
 	StructFilter->Schema = StateTree ? StateTree->GetSchema() : nullptr;
 
 	FStructViewerInitializationOptions Options;
@@ -664,7 +771,7 @@ TSharedRef<SWidget> FStateTreeBindableItemDetails::GenerateStructPicker()
 	Options.bAllowViewOptions = true;
 
 	FOnStructPicked OnPicked(FOnStructPicked::CreateRaw(this, &FStateTreeBindableItemDetails::OnStructPicked));
-
+	
 	return SNew(SBox)
 		.WidthOverride(280.f)
 		[
@@ -680,109 +787,178 @@ TSharedRef<SWidget> FStateTreeBindableItemDetails::GenerateStructPicker()
 
 void FStateTreeBindableItemDetails::OnStructPicked(const UScriptStruct* InStruct)
 {
-	check(ItemProperty);
-	check(InstanceProperty);
-	check(IDProperty);
+	check(StructProperty);
+	check(StateTree);
 
 	TArray<void*> RawItemData;
-	ItemProperty->AccessRawData(RawItemData);
-
-	TArray<void*> RawInstanceData;
-	InstanceProperty->AccessRawData(RawInstanceData);
-
-	TArray<void*> RawIDData;
-	IDProperty->AccessRawData(RawIDData);
+	StructProperty->AccessRawData(RawItemData);
 
-	if (RawItemData.Num() != RawInstanceData.Num() || RawItemData.Num() != RawIDData.Num())
-	{
-		return;
-	}
+	GEditor->BeginTransaction(LOCTEXT("OnStructPicked", "Set item"));
 
-	GEditor->BeginTransaction(LOCTEXT("OnStructPicked", "Set Struct"));
+	StructProperty->NotifyPreChange();
 
-	ItemProperty->NotifyPreChange();
-	InstanceProperty->NotifyPreChange();
-	IDProperty->NotifyPreChange();
-
-	for (int32 Index = 0; Index < RawItemData.Num(); Index++)
+	for (void* Data : RawItemData)
 	{
-		FInstancedStruct* ItemPtr = static_cast<FInstancedStruct*>(RawItemData[Index]);
-		FInstancedStruct* InstancePtr = static_cast<FInstancedStruct*>(RawInstanceData[Index]);
-		FGuid* ID = static_cast<FGuid*>(RawIDData[Index]);
-		
-		if (ItemPtr != nullptr && InstancePtr != nullptr && ID != nullptr)
+		if (FStateTreeItem* Item = static_cast<FStateTreeItem*>(Data))
 		{
+		
+			Item->Reset();
+			
 			if (InStruct)
 			{
-				// Generate new ID and name.
-				*ID = FGuid::NewGuid();
+				// Generate new ID.
+				Item->ID = FGuid::NewGuid();
 
-				// Intialize item
-				ItemPtr->InitializeAs(InStruct);
+				// Initialize item
+				Item->Item.InitializeAs(InStruct);
 				
 				// Generate new name and instantiate instance data.
 				if (InStruct->IsChildOf(FStateTreeTaskBase::StaticStruct()))
 				{
-					FStateTreeTaskBase* Task = ItemPtr->GetMutablePtr<FStateTreeTaskBase>();
-					check(Task);
-					Task->Name = FName(InStruct->GetDisplayNameText().ToString());
+					FStateTreeTaskBase& Task = Item->Item.GetMutable<FStateTreeTaskBase>();
+					Task.Name = FName(InStruct->GetDisplayNameText().ToString());
 
-					if (const UScriptStruct* InstanceType = Cast<const UScriptStruct>(Task->GetInstanceDataType()))
+					if (const UScriptStruct* InstanceType = Cast<const UScriptStruct>(Task.GetInstanceDataType()))
 					{
-						InstancePtr->InitializeAs(InstanceType);
+						Item->Instance.InitializeAs(InstanceType);
 					}
-					else
+					else if (const UClass* InstanceClass = Cast<const UClass>(Task.GetInstanceDataType()))
 					{
-						InstancePtr->Reset();
+						Item->InstanceObject = NewObject<UObject>(EditorData, InstanceClass);
 					}
 				}
 				else if (InStruct->IsChildOf(FStateTreeEvaluatorBase::StaticStruct()))
 				{
-					FStateTreeEvaluatorBase* Eval = ItemPtr->GetMutablePtr<FStateTreeEvaluatorBase>();
-					check(Eval);
-					Eval->Name = FName(InStruct->GetDisplayNameText().ToString());
+					FStateTreeEvaluatorBase& Eval = Item->Item.GetMutable<FStateTreeEvaluatorBase>();
+					Eval.Name = FName(InStruct->GetDisplayNameText().ToString());
 
-					if (const UScriptStruct* InstanceType = Cast<const UScriptStruct>(Eval->GetInstanceDataType()))
+					if (const UScriptStruct* InstanceType = Cast<const UScriptStruct>(Eval.GetInstanceDataType()))
 					{
-						InstancePtr->InitializeAs(InstanceType);
+						Item->Instance.InitializeAs(InstanceType);
 					}
-					else
+					else if (const UClass* InstanceClass = Cast<const UClass>(Eval.GetInstanceDataType()))
 					{
-						InstancePtr->Reset();
+						Item->InstanceObject = NewObject<UObject>(EditorData, InstanceClass);
 					}
 				}
 				else if (InStruct->IsChildOf(FStateTreeConditionBase::StaticStruct()))
 				{
-					FStateTreeConditionBase* Cond = ItemPtr->GetMutablePtr<FStateTreeConditionBase>();
-					check(Cond);
+					FStateTreeConditionBase& Cond = Item->Item.GetMutable<FStateTreeConditionBase>();
 
-					if (const UScriptStruct* InstanceType = Cast<const UScriptStruct>(Cond->GetInstanceDataType()))
+					if (const UScriptStruct* InstanceType = Cast<const UScriptStruct>(Cond.GetInstanceDataType()))
 					{
-						InstancePtr->InitializeAs(InstanceType);
+						Item->Instance.InitializeAs(InstanceType);
 					}
-					else
+					else if (const UClass* InstanceClass = Cast<const UClass>(Cond.GetInstanceDataType()))
 					{
-						InstancePtr->Reset();
+						Item->InstanceObject = NewObject<UObject>(EditorData, InstanceClass);
 					}
 				}
 			}
-			else
-			{
-				ItemPtr->Reset();
-				InstancePtr->Reset();
-				ID->Invalidate();
-			}
 		}
 	}
 
-	ItemProperty->NotifyPostChange(EPropertyChangeType::ValueSet);
-	ItemProperty->NotifyFinishedChangingProperties();
+	StructProperty->NotifyPostChange(EPropertyChangeType::ValueSet);
+	StructProperty->NotifyFinishedChangingProperties();
 
-	InstanceProperty->NotifyPostChange(EPropertyChangeType::ValueSet);
-	InstanceProperty->NotifyFinishedChangingProperties();
+	GEditor->EndTransaction();
 
-	IDProperty->NotifyPostChange(EPropertyChangeType::ValueSet);
-	IDProperty->NotifyFinishedChangingProperties();
+	ComboButton->SetIsOpen(false);
+
+	if (PropUtils)
+	{
+		PropUtils->ForceRefresh();
+	}
+}
+
+
+TSharedRef<SWidget> FStateTreeBindableItemDetails::GenerateClassPicker()
+{
+	TSharedRef<FBindableItemClassFilter> ClassFilter = MakeShared<FBindableItemClassFilter>();
+	ClassFilter->BaseClass = BaseClass;
+	ClassFilter->bAllowBaseClass = false;
+	ClassFilter->Schema = StateTree ? StateTree->GetSchema() : nullptr;
+
+	FClassViewerInitializationOptions Options;
+	Options.bShowUnloadedBlueprints = true;
+	Options.bShowNoneOption = true;
+	Options.ClassFilters.Add(ClassFilter);
+	Options.NameTypeToDisplay = EClassViewerNameTypeToDisplay::DisplayName;
+	Options.DisplayMode = EClassViewerDisplayMode::ListView;
+	Options.bAllowViewOptions = true;
+
+	FOnClassPicked OnPicked(FOnClassPicked::CreateRaw(this, &FStateTreeBindableItemDetails::OnClassPicked));
+
+	return SNew(SBox)
+		.WidthOverride(280.f)
+		[
+			SNew(SVerticalBox)
+			+ SVerticalBox::Slot()
+			.AutoHeight()
+			.MaxHeight(500.f)
+			[
+				FModuleManager::LoadModuleChecked<FClassViewerModule>("ClassViewer").CreateClassViewer(Options, OnPicked)
+			]
+		];
+}
+
+void FStateTreeBindableItemDetails::OnClassPicked(UClass* InClass)
+{
+	check(StructProperty);
+	check(StateTree);
+
+	TArray<void*> RawItemData;
+	StructProperty->AccessRawData(RawItemData);
+
+	GEditor->BeginTransaction(LOCTEXT("OnClassPicked", "Set Blueprint item"));
+
+	StructProperty->NotifyPreChange();
+
+	for (void* Data : RawItemData)
+	{
+		if (FStateTreeItem* Item = static_cast<FStateTreeItem*>(Data))
+		{
+			Item->Reset();
+
+			if (InClass && InClass->IsChildOf(UStateTreeTaskBlueprintBase::StaticClass()))
+			{
+				Item->Item.InitializeAs(FStateTreeBlueprintTaskWrapper::StaticStruct());
+				FStateTreeBlueprintTaskWrapper& Task = Item->Item.GetMutable<FStateTreeBlueprintTaskWrapper>();
+				Task.TaskClass = InClass;
+				Task.Name = FName(InClass->GetDisplayNameText().ToString());
+				
+				Item->InstanceObject = NewObject<UObject>(EditorData, InClass);
+
+				Item->ID = FGuid::NewGuid();
+			}
+			else if (InClass && InClass->IsChildOf(UStateTreeEvaluatorBlueprintBase::StaticClass()))
+			{
+				Item->Item.InitializeAs(FStateTreeBlueprintEvaluatorWrapper::StaticStruct());
+				FStateTreeBlueprintEvaluatorWrapper& Eval = Item->Item.GetMutable<FStateTreeBlueprintEvaluatorWrapper>();
+				Eval.EvaluatorClass = InClass;
+				Eval.Name = FName(InClass->GetDisplayNameText().ToString());
+				
+				Item->InstanceObject = NewObject<UObject>(EditorData, InClass);
+
+				Item->ID = FGuid::NewGuid();
+			}
+			else if (InClass && InClass->IsChildOf(UStateTreeConditionBlueprintBase::StaticClass()))
+			{
+				Item->Item.InitializeAs(FStateTreeBlueprintConditionWrapper::StaticStruct());
+				FStateTreeBlueprintConditionWrapper& Cond = Item->Item.GetMutable<FStateTreeBlueprintConditionWrapper>();
+				Cond.ConditionClass = InClass;
+				
+				Item->InstanceObject = NewObject<UObject>(EditorData, InClass);
+
+				Item->ID = FGuid::NewGuid();
+			}
+
+			
+		}
+	}
+
+	StructProperty->NotifyPostChange(EPropertyChangeType::ValueSet);
+	StructProperty->NotifyFinishedChangingProperties();
 
 	GEditor->EndTransaction();
 
@@ -794,4 +970,5 @@ void FStateTreeBindableItemDetails::OnStructPicked(const UScriptStruct* InStruct
 	}
 }
 
+
 #undef LOCTEXT_NAMESPACE
diff --git a/Engine/Plugins/Runtime/StateTree/Source/StateTreeEditorModule/Private/Customizations/StateTreeBindableItemDetails.h b/Engine/Plugins/Runtime/StateTree/Source/StateTreeEditorModule/Private/Customizations/StateTreeBindableItemDetails.h
--- a/Engine/Plugins/Runtime/StateTree/Source/StateTreeEditorModule/Private/Customizations/StateTreeBindableItemDetails.h
+++ b/Engine/Plugins/Runtime/StateTree/Source/StateTreeEditorModule/Private/Customizations/StateTreeBindableItemDetails.h
@@ -31,6 +31,7 @@ private:
 
 	bool ShouldResetToDefault(TSharedPtr<IPropertyHandle> PropertyHandle) const;
 	void ResetToDefault(TSharedPtr<IPropertyHandle> PropertyHandle);
+	TSharedPtr<IPropertyHandle> GetInstancedObjectValueHandle(TSharedPtr<IPropertyHandle> PropertyHandle);
 
 	FText GetDescription() const;
 	EVisibility IsDescriptionVisible() const;
@@ -43,14 +44,21 @@ private:
 	const UScriptStruct* GetCommonItemScriptStruct() const;
 	FText GetDisplayValueString() const;
 	const FSlateBrush* GetDisplayValueIcon() const;
+
+	TSharedRef<SWidget> GeneratePicker();
+
 	TSharedRef<SWidget> GenerateStructPicker();
 	void OnStructPicked(const UScriptStruct* InStruct);
 
+	TSharedRef<SWidget> GenerateClassPicker();
+	void OnClassPicked(UClass* InClass);
+
 	void OnIdentifierChanged(const UStateTree& StateTree);
 	void OnBindingChanged(const FStateTreeEditorPropertyPath& SourcePath, const FStateTreeEditorPropertyPath& TargetPath);
 	void FindOuterObjects();
 
 	UScriptStruct* BaseScriptStruct = nullptr;
+	UClass* BaseClass = nullptr;
 	TSharedPtr<class SComboButton> ComboButton;
 
 	UStateTreeEditorData* EditorData = nullptr;
@@ -60,6 +68,7 @@ private:
 	TSharedPtr<IPropertyHandle> StructProperty;
 	TSharedPtr<IPropertyHandle> ItemProperty;
 	TSharedPtr<IPropertyHandle> InstanceProperty;
+	TSharedPtr<IPropertyHandle> InstanceObjectProperty;
 	TSharedPtr<IPropertyHandle> IDProperty;
 
 	FDelegateHandle OnBindingChangedHandle;
diff --git a/Engine/Plugins/Runtime/StateTree/Source/StateTreeEditorModule/Private/StateTreeBaker.cpp b/Engine/Plugins/Runtime/StateTree/Source/StateTreeEditorModule/Private/StateTreeBaker.cpp
--- a/Engine/Plugins/Runtime/StateTree/Source/StateTreeEditorModule/Private/StateTreeBaker.cpp
+++ b/Engine/Plugins/Runtime/StateTree/Source/StateTreeEditorModule/Private/StateTreeBaker.cpp
@@ -47,7 +47,7 @@ bool FStateTreeBaker::Bake(UStateTree& InStateTree)
 
 	StateTree->Link();
 	
-	StateTree->InitInstanceStorage();
+	StateTree->InitInstanceStorageType();
 
 	return true;
 }
@@ -68,14 +68,6 @@ bool FStateTreeBaker::CreateStates()
 	// Create item for the runtime execution state
 	StateTree->Instances.Add(FInstancedStruct::Make<FStateTreeExecutionState>());
 
-	// The struct index is currently shared between runtime item index, make sure the state has a dummy struct too. 
-	FStateTreeBindableStructDesc StructDesc;
-	StructDesc.Struct = FStateTreeExecutionState::StaticStruct();
-	StructDesc.Name = FName(TEXT("ExecutionState"));
-	StructDesc.ID = FGuid();
-	BindingsCompiler.AddSourceStruct(StructDesc);
-
-	
 	for (UStateTreeState* Routine : TreeData->Routines)
 	{
 		if (Routine)
@@ -96,7 +88,7 @@ bool FStateTreeBaker::CreateStateTransitions()
 	{
 		FBakedStateTreeState& BakedState = StateTree->States[i];
 		UStateTreeState* SourceState = SourceStates[i];
-		check(SourceState);
+		check(SourceState != nullptr);
 
 		FStateTreeCompilerLogStateScope LogStateScope(SourceState, Log);
 		
@@ -193,30 +185,62 @@ bool FStateTreeBaker::CreateCondition(const FStateTreeConditionItem& CondItem)
 		return true;
 	}
 
-	// Assume valid items have valid instance initialized.
-	check(CondItem.Instance.IsValid());
+	FStateTreeBindableStructDesc StructDesc;
+	StructDesc.ID = CondItem.ID;
+	StructDesc.Name = CondItem.Item.GetScriptStruct()->GetFName();
+
+	// Check that item has valid instance initialized.
+	if (!CondItem.Instance.IsValid() && CondItem.InstanceObject == nullptr)
+	{
+		Log.Reportf(EMessageSeverity::Error, StructDesc,
+			TEXT("Malformed condition, missing instance value."));
+		return false;
+	}
 
 	// Copy the condition
 	FInstancedStruct& Item = StateTree->Items.AddDefaulted_GetRef();
 	Item = CondItem.Item;
 
-	FInstancedStruct& Instance = StateTree->Instances.AddDefaulted_GetRef();
-	Instance = CondItem.Instance;
-			
-	check(Item.IsValid() && Instance.IsValid());
 	FStateTreeConditionBase& Cond = Item.GetMutable<FStateTreeConditionBase>();
 
-	// Create binding source struct descriptor. Note: not exposing the struct for reading.
-	FStateTreeBindableStructDesc StructDesc;
-	StructDesc.Struct = Instance.GetScriptStruct();
-	StructDesc.Name = Instance.GetScriptStruct()->GetFName();
-	StructDesc.ID = CondItem.ID;
+	if (CondItem.Instance.IsValid())
+	{
+		// Struct instance
+		FInstancedStruct& Instance = StateTree->Instances.AddDefaulted_GetRef();
+		const int32 InstanceIndex = StateTree->Instances.Num() - 1;
+		
+		Instance = CondItem.Instance;
+	
+		// Create binding source struct descriptor.
+		StructDesc.Struct = Instance.GetScriptStruct();
+		StructDesc.Name = StructDesc.Struct->GetFName();
+
+		check(InstanceIndex <= int32(MAX_uint16));
+		Cond.InstanceIndex = uint16(InstanceIndex);
+		Cond.bInstanceIsObject = false;
+	}
+	else
+	{
+		// Object Instance
+		check(CondItem.InstanceObject != nullptr);
+		
+		UObject* Instance = DuplicateObject(CondItem.InstanceObject, StateTree);
+		
+		StateTree->InstanceObjects.Add(Instance);
+		const int32 InstanceIndex = StateTree->InstanceObjects.Num() - 1;
+		
+		// Create binding source struct descriptor.
+		StructDesc.Struct = Instance->GetClass();
+		StructDesc.Name = StructDesc.Struct->GetFName();
+
+		check(InstanceIndex <= int32(MAX_uint16));
+		Cond.InstanceIndex = uint16(InstanceIndex);
+		Cond.bInstanceIsObject = true;
+	}
 
 	// Mark the struct as binding source.
 	const int32 SourceStructIndex = BindingsCompiler.AddSourceStruct(StructDesc);
 
-	check((StateTree->Instances.Num() - 1) == SourceStructIndex);
-
 	// Check that the bindings for this struct are still all valid.
 	TArray<FStateTreeEditorPropertyBinding> Bindings;
 	if (!GetAndValidateBindings(StructDesc, Bindings))
@@ -345,50 +369,90 @@ bool FStateTreeBaker::CreateStateRecursive(UStateTreeState& State, const FStateT
 
 	for (FStateTreeEvaluatorItem& EvaluatorItem : State.Evaluators)
 	{
-		if (EvaluatorItem.Item.IsValid() && EvaluatorItem.Instance.IsValid())
+		// Silently ignore empty items.
+		if (!EvaluatorItem.Item.IsValid())
 		{
-			// Copy the evaluator
-			FInstancedStruct& Item = StateTree->Items.AddDefaulted_GetRef();
-			Item = EvaluatorItem.Item;
-
+			continue;
+		}
+		
+		// Create binding source struct descriptor.
+		FStateTreeBindableStructDesc StructDesc;
+        StructDesc.ID = EvaluatorItem.ID;
+        StructDesc.Name = EvaluatorItem.Item.GetScriptStruct()->GetFName();
+    
+        // Check that item has valid instance initialized.
+        if (!EvaluatorItem.Instance.IsValid() && EvaluatorItem.InstanceObject == nullptr)
+        {
+        	Log.Reportf(EMessageSeverity::Error, StructDesc,
+        		TEXT("Malformed evaluator, missing instance value."));
+        	return false;
+        }
+
+		// Copy the evaluator
+		FInstancedStruct& Item = StateTree->Items.AddDefaulted_GetRef();
+		Item = EvaluatorItem.Item;
+
+		FStateTreeEvaluatorBase& Eval = Item.GetMutable<FStateTreeEvaluatorBase>();
+
+		if (EvaluatorItem.Instance.IsValid())
+		{
+			// Struct Instance
 			FInstancedStruct& Instance = StateTree->Instances.AddDefaulted_GetRef();
+			const int32 InstanceIndex = StateTree->Instances.Num() - 1;
+	
 			Instance = EvaluatorItem.Instance;
-			
-			check(Item.IsValid() && Instance.IsValid());
-			FStateTreeEvaluatorBase& Eval = Item.GetMutable<FStateTreeEvaluatorBase>();
 
 			// Create binding source struct descriptor.
-			FStateTreeBindableStructDesc StructDesc;
 			StructDesc.Struct = Instance.GetScriptStruct();
 			StructDesc.Name = Eval.Name;
-			StructDesc.ID = EvaluatorItem.ID;
 
-			// Mark the struct as binding source.
-			const int32 SourceStructIndex = BindingsCompiler.AddSourceStruct(StructDesc);
+			check(InstanceIndex <= int32(MAX_uint16));
+			Eval.InstanceIndex = uint16(InstanceIndex);
+			Eval.bInstanceIsObject = false;
+		}
+		else
+		{
+			// Object Instance
+			check(EvaluatorItem.InstanceObject != nullptr);
+	
+			UObject* Instance = DuplicateObject(EvaluatorItem.InstanceObject, StateTree);
+	
+			StateTree->InstanceObjects.Add(Instance);
+			const int32 InstanceIndex = StateTree->InstanceObjects.Num() - 1;
+	
+			// Create binding source struct descriptor.
+			StructDesc.Struct = Instance->GetClass();
+			StructDesc.Name = Eval.Name;
 
-			check((StateTree->Instances.Num() - 1) == SourceStructIndex);
+			check(InstanceIndex <= int32(MAX_uint16));
+			Eval.InstanceIndex = uint16(InstanceIndex);
+			Eval.bInstanceIsObject = true;
+		}
 			
-			// Check that the bindings for this struct are still all valid.
-			TArray<FStateTreeEditorPropertyBinding> Bindings;
-			if (!GetAndValidateBindings(StructDesc, Bindings))
-			{
-				return false;
-			}
+		// Mark the instance as binding source.
+		const int32 SourceStructIndex = BindingsCompiler.AddSourceStruct(StructDesc);
 
-			// Compile batch copy for this struct, we pass in all the bindings, the compiler will pick up the ones for the target structs.
-			int32 BatchIndex = INDEX_NONE;
-			if (!BindingsCompiler.CompileBatch(StructDesc, Bindings, BatchIndex))
-			{
-				return false;
-			}
-			
-			check(BatchIndex < int32(MAX_uint16));
-			Eval.BindingsBatch = BatchIndex == INDEX_NONE ? FStateTreeHandle::Invalid : FStateTreeHandle(uint16(BatchIndex));
-			
-			check(SourceStructIndex <= int32(MAX_uint16));
-			Eval.DataViewIndex = uint16(SourceStructIndex);
+		// Check that the bindings for this struct are still all valid.
+		TArray<FStateTreeEditorPropertyBinding> Bindings;
+		if (!GetAndValidateBindings(StructDesc, Bindings))
+		{
+			return false;
+		}
+
+		// Compile batch copy for this struct, we pass in all the bindings, the compiler will pick up the ones for the target structs.
+		int32 BatchIndex = INDEX_NONE;
+		if (!BindingsCompiler.CompileBatch(StructDesc, Bindings, BatchIndex))
+		{
+			return false;
 		}
+		
+		check(BatchIndex < int32(MAX_uint16));
+		Eval.BindingsBatch = BatchIndex == INDEX_NONE ? FStateTreeHandle::Invalid : FStateTreeHandle(uint16(BatchIndex));
+		
+		check(SourceStructIndex <= int32(MAX_uint16));
+		Eval.DataViewIndex = uint16(SourceStructIndex);
 	}
+	
 	const int32 EvaluatorsNum = StateTree->Items.Num() - int32(BakedState.EvaluatorsBegin);
 	check(EvaluatorsNum <= int32(MAX_uint8));
 	BakedState.EvaluatorsNum = uint8(EvaluatorsNum);
@@ -399,51 +463,90 @@ bool FStateTreeBaker::CreateStateRecursive(UStateTreeState& State, const FStateT
 
 	for (FStateTreeTaskItem& TaskItem : State.Tasks)
 	{
-		if (TaskItem.Item.IsValid() && TaskItem.Instance.IsValid())
+		// Silently ignore empty items.
+		if (!TaskItem.Item.IsValid())
+		{
+			continue;
+		}
+		
+		// Create binding source struct descriptor.
+		FStateTreeBindableStructDesc StructDesc;
+		StructDesc.ID = TaskItem.ID;
+		StructDesc.Name = TaskItem.Item.GetScriptStruct()->GetFName();
+    
+		// Check that item has valid instance initialized.
+		if (!TaskItem.Instance.IsValid() && TaskItem.InstanceObject == nullptr)
 		{
-			// Copy the task
-			FInstancedStruct& Item = StateTree->Items.AddDefaulted_GetRef();
-			Item = TaskItem.Item;
+			Log.Reportf(EMessageSeverity::Error, StructDesc,
+				TEXT("Malformed evaluator, missing instance value."));
+			return false;
+		}
 
+		// Copy the task
+		FInstancedStruct& Item = StateTree->Items.AddDefaulted_GetRef();
+		Item = TaskItem.Item;
+		
+		FStateTreeTaskBase& Task = Item.GetMutable<FStateTreeTaskBase>();
+
+		if (TaskItem.Instance.IsValid())
+		{
+			// Struct Instance
 			FInstancedStruct& Instance = StateTree->Instances.AddDefaulted_GetRef();
+			const int32 InstanceIndex = StateTree->Instances.Num() - 1;
+	
 			Instance = TaskItem.Instance;
 
-			check(Item.IsValid() && Instance.IsValid());
-			
-			FStateTreeTaskBase& Task = Item.GetMutable<FStateTreeTaskBase>();
-
 			// Create binding source struct descriptor.
-			FStateTreeBindableStructDesc StructDesc;
 			StructDesc.Struct = Instance.GetScriptStruct();
 			StructDesc.Name = Task.Name;
-			StructDesc.ID = TaskItem.ID;
 
-			// Mark the struct as binding source.
-			const int32 SourceStructIndex = BindingsCompiler.AddSourceStruct(StructDesc);
+			check(InstanceIndex <= int32(MAX_uint16));
+			Task.InstanceIndex = uint16(InstanceIndex);
+			Task.bInstanceIsObject = false;
+		}
+		else
+		{
+			// Object Instance
+			check(TaskItem.InstanceObject != nullptr);
+	
+			UObject* Instance = DuplicateObject(TaskItem.InstanceObject, StateTree);
+	
+			StateTree->InstanceObjects.Add(Instance);
+			const int32 InstanceIndex = StateTree->InstanceObjects.Num() - 1;
+	
+			// Create binding source struct descriptor.
+			StructDesc.Struct = Instance->GetClass();
+			StructDesc.Name = Task.Name;
 
-			check((StateTree->Instances.Num() - 1) == SourceStructIndex);
+			check(InstanceIndex <= int32(MAX_uint16));
+			Task.InstanceIndex = uint16(InstanceIndex);
+			Task.bInstanceIsObject = true;
+		}
 
-			// Check that the bindings for this struct are still all valid.
-			TArray<FStateTreeEditorPropertyBinding> Bindings;
-			if (!GetAndValidateBindings(StructDesc, Bindings))
-			{
-				return false;
-			}
+		// Mark the instance as binding source.
+		const int32 SourceStructIndex = BindingsCompiler.AddSourceStruct(StructDesc);
 
-			// Compile batch copy for this struct, we pass in all the bindings, the compiler will pick up the ones for the target structs.
-			int32 BatchIndex = INDEX_NONE;
-			if (!BindingsCompiler.CompileBatch(StructDesc, Bindings, BatchIndex))
-			{
-				return false;
-			}
-			
-			check(BatchIndex < int32(MAX_uint16));
-			Task.BindingsBatch = BatchIndex == INDEX_NONE ? FStateTreeHandle::Invalid : FStateTreeHandle(uint16(BatchIndex));
-			
-			check(SourceStructIndex <= int32(MAX_uint16));
-			Task.DataViewIndex = uint16(SourceStructIndex);
+		// Check that the bindings for this struct are still all valid.
+		TArray<FStateTreeEditorPropertyBinding> Bindings;
+		if (!GetAndValidateBindings(StructDesc, Bindings))
+		{
+			return false;
 		}
+
+		// Compile batch copy for this struct, we pass in all the bindings, the compiler will pick up the ones for the target structs.
+		int32 BatchIndex = INDEX_NONE;
+		if (!BindingsCompiler.CompileBatch(StructDesc, Bindings, BatchIndex))
+		{
+			return false;
+		}
+		
+		check(BatchIndex < int32(MAX_uint16));
+		Task.BindingsBatch = BatchIndex == INDEX_NONE ? FStateTreeHandle::Invalid : FStateTreeHandle(uint16(BatchIndex));
+		
+		check(SourceStructIndex <= int32(MAX_uint16));
+		Task.DataViewIndex = uint16(SourceStructIndex);
 	}
+	
 	const int32 TasksNum = StateTree->Items.Num() - int32(BakedState.TasksBegin);
 	check(TasksNum <= int32(MAX_uint8));
 	BakedState.TasksNum = uint8(TasksNum);
diff --git a/Engine/Plugins/Runtime/StateTree/Source/StateTreeEditorModule/Public/StateTreeState.h b/Engine/Plugins/Runtime/StateTree/Source/StateTreeEditorModule/Public/StateTreeState.h
--- a/Engine/Plugins/Runtime/StateTree/Source/StateTreeEditorModule/Public/StateTreeState.h
+++ b/Engine/Plugins/Runtime/StateTree/Source/StateTreeEditorModule/Public/StateTreeState.h
@@ -32,14 +32,44 @@ public:
 	FGuid ID;
 
 	UPROPERTY(EditDefaultsOnly, Category = Link)
-	EStateTreeTransitionType Type;
+	EStateTreeTransitionType Type = EStateTreeTransitionType::GotoState;
+};
+
+
+/**
+ * Base for Evaluator, Task and Condition items.
+ */
+USTRUCT()
+struct STATETREEEDITORMODULE_API FStateTreeItem
+{
+	GENERATED_BODY()
+
+	void Reset()
+	{
+		Item.Reset();
+		Instance.Reset();
+		InstanceObject = nullptr;
+		ID = FGuid();
+	}
+	
+	UPROPERTY(EditDefaultsOnly, Category = Item)
+	FInstancedStruct Item;
+
+	UPROPERTY(EditDefaultsOnly, Category = Item)
+	FInstancedStruct Instance;
+
+	UPROPERTY(EditDefaultsOnly, Instanced, Category = Item)
+	UObject* InstanceObject = nullptr;
+	
+	UPROPERTY(EditDefaultsOnly, Category = Item)
+	FGuid ID;
 };
 
 /**
  * Helper struct for Evaluator details customization (i.e. show summary when collapsed)
  */
 USTRUCT()
-struct STATETREEEDITORMODULE_API FStateTreeEvaluatorItem
+struct STATETREEEDITORMODULE_API FStateTreeEvaluatorItem : public FStateTreeItem
 {
 	GENERATED_BODY()
 
@@ -51,15 +81,6 @@ struct STATETREEEDITORMODULE_API FStateTreeEvaluatorItem
 		}
 		return FName();
 	}
-
-	UPROPERTY(EditDefaultsOnly, Category = Evaluator, meta = (BaseStruct = "StateTreeEvaluatorBase", ExcludeBaseStruct))
-	FInstancedStruct Item;
-
-	UPROPERTY(EditDefaultsOnly, Category = Evaluator)
-	FInstancedStruct Instance;
-
-	UPROPERTY(EditDefaultsOnly, Category = Evaluator)
-	FGuid ID;
 };
 
 template <typename T>
@@ -75,7 +96,7 @@ struct TStateTreeEvaluatorItem : public FStateTreeEvaluatorItem
  * Helper struct for Task details customization (i.e. show summary when collapsed)
  */
 USTRUCT()
-struct STATETREEEDITORMODULE_API FStateTreeTaskItem
+struct STATETREEEDITORMODULE_API FStateTreeTaskItem : public FStateTreeItem
 {
 	GENERATED_BODY()
 
@@ -87,15 +108,6 @@ struct STATETREEEDITORMODULE_API FStateTreeTaskItem
 		}
 		return FName();
 	}
-	
-	UPROPERTY(EditDefaultsOnly, Category = Task, meta = (BaseStruct = "StateTreeTaskBase", ExcludeBaseStruct))
-	FInstancedStruct Item;
-
-	UPROPERTY(EditDefaultsOnly, Category = Task)
-	FInstancedStruct Instance;
-
-	UPROPERTY(EditDefaultsOnly, Category = Task)
-	FGuid ID;
 };
 
 template <typename T>
@@ -110,18 +122,9 @@ struct TStateTreeTaskItem : public FStateTreeTaskItem
  * Helper struct for Condition details customization (i.e. show summary when collapsed)
  */
 USTRUCT()
-struct STATETREEEDITORMODULE_API FStateTreeConditionItem
+struct STATETREEEDITORMODULE_API FStateTreeConditionItem : public FStateTreeItem
 {
 	GENERATED_BODY()
-
-	UPROPERTY(EditDefaultsOnly, Category = Condition, meta = (BaseStruct = "StateTreeConditionBase", ExcludeBaseStruct))
-	FInstancedStruct Item;
-
-	UPROPERTY(EditDefaultsOnly, Category = Condition)
-	FInstancedStruct Instance;
-
-	UPROPERTY(EditDefaultsOnly, Category = Condition)
-	FGuid ID;
 };
 
 template <typename T>
@@ -158,7 +161,7 @@ struct STATETREEEDITORMODULE_API FStateTreeTransition
 	}
 
 	UPROPERTY(EditDefaultsOnly, Category = Transition)
-	EStateTreeTransitionEvent Event = EStateTreeTransitionEvent::OnCondition;
+	EStateTreeTransitionEvent Event = EStateTreeTransitionEvent::OnCompleted;
 
 	UPROPERTY(EditDefaultsOnly, Category = Transition)
 	FStateTreeStateLink State;
@@ -167,7 +170,7 @@ struct STATETREEEDITORMODULE_API FStateTreeTransition
 	UPROPERTY(EditDefaultsOnly, Category = Transition, meta = (UIMin = "0", ClampMin = "0", UIMax = "25", ClampMax = "25"))
 	float GateDelay = 0.0f;
 	
-	UPROPERTY(EditDefaultsOnly, Category = Transition)
+	UPROPERTY(EditDefaultsOnly, Category = Transition, meta = (BaseStruct = "StateTreeConditionBase", BaseClass = "StateTreeConditionBlueprintBase"))
 	TArray<FStateTreeConditionItem> Conditions;
 };
 
@@ -273,13 +276,13 @@ public:
 	UPROPERTY()
 	FGuid ID;
 
-	UPROPERTY(EditDefaultsOnly, Category = "Enter Conditions")
+	UPROPERTY(EditDefaultsOnly, Category = "Enter Conditions", meta = (BaseStruct = "StateTreeConditionBase", BaseClass = "StateTreeConditionBlueprintBase"))
 	TArray<FStateTreeConditionItem> EnterConditions;
 
-	UPROPERTY(EditDefaultsOnly, Category = "Evaluators")
+	UPROPERTY(EditDefaultsOnly, Category = "Evaluators", meta = (BaseStruct = "StateTreeEvaluatorBase", BaseClass = "StateTreeEvaluatorBlueprintBase"))
 	TArray<FStateTreeEvaluatorItem> Evaluators;
 
-	UPROPERTY(EditDefaultsOnly, Category = "Tasks")
+	UPROPERTY(EditDefaultsOnly, Category = "Tasks", meta = (BaseStruct = "StateTreeTaskBase", BaseClass = "StateTreeTaskBlueprintBase"))
 	TArray<FStateTreeTaskItem> Tasks;
 
 	UPROPERTY(EditDefaultsOnly, Category = "Transitions")
diff --git a/Engine/Plugins/Runtime/StateTree/Source/StateTreeModule/Private/Blueprint/StateTreeConditionBlueprintBase.cpp b/Engine/Plugins/Runtime/StateTree/Source/StateTreeModule/Private/Blueprint/StateTreeConditionBlueprintBase.cpp
new file mode 100644
--- /dev/null
+++ b/Engine/Plugins/Runtime/StateTree/Source/StateTreeModule/Private/Blueprint/StateTreeConditionBlueprintBase.cpp
@@ -0,0 +1,40 @@
+﻿// Copyright Epic Games, Inc. All Rights Reserved.
+
+#include "Blueprint/StateTreeConditionBlueprintBase.h"
+#include "CoreMinimal.h"
+#include "StateTreeExecutionContext.h"
+
+//----------------------------------------------------------------------//
+//  UStateTreeConditionBlueprintBase
+//----------------------------------------------------------------------//
+
+bool UStateTreeConditionBlueprintBase::TestCondition(FStateTreeExecutionContext& Context)
+{
+	AActor* OwnerActor = GetOwnerActor(Context);
+	return ReceiveTestCondition(OwnerActor);
+}
+
+//----------------------------------------------------------------------//
+//  FStateTreeBlueprintConditionWrapper
+//----------------------------------------------------------------------//
+
+bool FStateTreeBlueprintConditionWrapper::Link(FStateTreeLinker& Linker)
+{
+	const UStateTreeConditionBlueprintBase* CondCDO =  ConditionClass ? ConditionClass->GetDefaultObject<UStateTreeConditionBlueprintBase>() : nullptr;
+	if (CondCDO != nullptr)
+	{
+		CondCDO->LinkExternalData(Linker, ExternalDataHandles);
+	}
+
+	return true;
+}
+
+bool FStateTreeBlueprintConditionWrapper::TestCondition(FStateTreeExecutionContext& Context) const
+{
+	if (UStateTreeConditionBlueprintBase* Instance = Context.GetInstanceObjectInternal<UStateTreeConditionBlueprintBase>(DataViewIndex))
+	{
+		Instance->CopyExternalData(Context, ExternalDataHandles);
+		return Instance->TestCondition(Context);
+	}
+	return false;
+}
diff --git a/Engine/Plugins/Runtime/StateTree/Source/StateTreeModule/Private/Blueprint/StateTreeEvaluatorBlueprintBase.cpp b/Engine/Plugins/Runtime/StateTree/Source/StateTreeModule/Private/Blueprint/StateTreeEvaluatorBlueprintBase.cpp
new file mode 100644
--- /dev/null
+++ b/Engine/Plugins/Runtime/StateTree/Source/StateTreeModule/Private/Blueprint/StateTreeEvaluatorBlueprintBase.cpp
@@ -0,0 +1,84 @@
+﻿// Copyright Epic Games, Inc. All Rights Reserved.
+
+#include "Blueprint/StateTreeEvaluatorBlueprintBase.h"
+#include "CoreMinimal.h"
+#include "StateTreeExecutionContext.h"
+
+//----------------------------------------------------------------------//
+//  UStateTreeEvaluatorBlueprintBase
+//----------------------------------------------------------------------//
+
+void UStateTreeEvaluatorBlueprintBase::EnterState(FStateTreeExecutionContext& Context, const EStateTreeStateChangeType ChangeType, const FStateTreeTransitionResult& Transition)
+{
+	AActor* OwnerActor = GetOwnerActor(Context);
+	ReceiveEnterState(OwnerActor, ChangeType, Transition);
+}
+
+void UStateTreeEvaluatorBlueprintBase::ExitState(FStateTreeExecutionContext& Context, const EStateTreeStateChangeType ChangeType, const FStateTreeTransitionResult& Transition)
+{
+	AActor* OwnerActor = GetOwnerActor(Context);
+	ReceiveExitState(OwnerActor, ChangeType, Transition);
+}
+
+void UStateTreeEvaluatorBlueprintBase::StateCompleted(FStateTreeExecutionContext& Context, const EStateTreeRunStatus CompletionStatus, const FStateTreeHandle CompletedState)
+{
+	AActor* OwnerActor = GetOwnerActor(Context);
+	ReceiveStateCompleted(OwnerActor, CompletionStatus, CompletedState);
+}
+
+void UStateTreeEvaluatorBlueprintBase::Evaluate(FStateTreeExecutionContext& Context, const EStateTreeEvaluationType EvalType, const float DeltaTime)
+{
+	AActor* OwnerActor = GetOwnerActor(Context);
+	ReceiveEvaluate(OwnerActor, EvalType, DeltaTime);
+}
+
+//----------------------------------------------------------------------//
+//  FStateTreeBlueprintEvaluatorWrapper
+//----------------------------------------------------------------------//
+
+bool FStateTreeBlueprintEvaluatorWrapper::Link(FStateTreeLinker& Linker)
+{
+	const UStateTreeEvaluatorBlueprintBase* EvalCDO = EvaluatorClass ? EvaluatorClass->GetDefaultObject<UStateTreeEvaluatorBlueprintBase>() : nullptr;
+	if (EvalCDO != nullptr)
+	{
+		EvalCDO->LinkExternalData(Linker, ExternalDataHandles);
+	}
+
+	return true;
+}
+
+void FStateTreeBlueprintEvaluatorWrapper::EnterState(FStateTreeExecutionContext& Context, const EStateTreeStateChangeType ChangeType, const FStateTreeTransitionResult& Transition) const
+{
+	UStateTreeEvaluatorBlueprintBase* Instance = Context.GetInstanceObjectInternal<UStateTreeEvaluatorBlueprintBase>(DataViewIndex);
+	check(Instance);
+	
+	Instance->CopyExternalData(Context, ExternalDataHandles);
+	Instance->EnterState(Context, ChangeType, Transition);
+}
+
+void FStateTreeBlueprintEvaluatorWrapper::ExitState(FStateTreeExecutionContext& Context, const EStateTreeStateChangeType ChangeType, const FStateTreeTransitionResult& Transition) const
+{
+	UStateTreeEvaluatorBlueprintBase* Instance = Context.GetInstanceObjectInternal<UStateTreeEvaluatorBlueprintBase>(DataViewIndex);
+	check(Instance);
+	
+	Instance->CopyExternalData(Context, ExternalDataHandles);
+	Instance->ExitState(Context, ChangeType, Transition);
+}
+
+void FStateTreeBlueprintEvaluatorWrapper::StateCompleted(FStateTreeExecutionContext& Context, const EStateTreeRunStatus CompletionStatus, const FStateTreeHandle CompletedState) const
+{
+	UStateTreeEvaluatorBlueprintBase* Instance = Context.GetInstanceObjectInternal<UStateTreeEvaluatorBlueprintBase>(DataViewIndex);
+	check(Instance);
+	
+	Instance->CopyExternalData(Context, ExternalDataHandles);
+	Instance->StateCompleted(Context, CompletionStatus, CompletedState);
+}
+
+void FStateTreeBlueprintEvaluatorWrapper::Evaluate(FStateTreeExecutionContext& Context, const EStateTreeEvaluationType EvalType, const float DeltaTime) const
+{
+	UStateTreeEvaluatorBlueprintBase* Instance = Context.GetInstanceObjectInternal<UStateTreeEvaluatorBlueprintBase>(DataViewIndex);
+	check(Instance);
+	
+	Instance->CopyExternalData(Context, ExternalDataHandles);
+	Instance->Evaluate(Context, EvalType, DeltaTime);
+}
diff --git a/Engine/Plugins/Runtime/StateTree/Source/StateTreeModule/Private/Blueprint/StateTreeItemBlueprintBase.cpp b/Engine/Plugins/Runtime/StateTree/Source/StateTreeModule/Private/Blueprint/StateTreeItemBlueprintBase.cpp
new file mode 100644
--- /dev/null
+++ b/Engine/Plugins/Runtime/StateTree/Source/StateTreeModule/Private/Blueprint/StateTreeItemBlueprintBase.cpp
@@ -0,0 +1,113 @@
+﻿// Copyright Epic Games, Inc. All Rights Reserved.
+
+#include "Blueprint/StateTreeItemBlueprintBase.h"
+#include "CoreMinimal.h"
+#include "AIController.h"
+#include "StateTreeExecutionContext.h"
+
+UWorld* UStateTreeItemBlueprintBase::GetWorld() const
+{
+	// The items are duplicated as the StateTreeExecution context as outer, so this should be essentially the same as GetWorld() on StateTree context.
+	// The CDO is used by the BP editor to check for certain functionality, make it return nullptr so that the GetWorld() passes as overridden. 
+	if (!HasAnyFlags(RF_ClassDefaultObject))
+	{
+		if (UObject* Outer = GetOuter())
+		{
+			return Outer->GetWorld();
+		}
+	}
+	
+	return nullptr;
+}
+
+AActor* UStateTreeItemBlueprintBase::GetOwnerActor(const FStateTreeExecutionContext& Context) const
+{
+	if (const AAIController* Controller = Cast<AAIController>(Context.GetOwner()))
+	{
+		return Controller->GetPawn();
+	}
+	
+	return Cast<AActor>(Context.GetOwner());
+}
+
+#if WITH_EDITOR
+void UStateTreeItemBlueprintBase::PostCDOCompiled()
+{
+	// Save some property data for runtime use.
+	PropertyInfos.Reset();
+
+	static const FName CategoryName(TEXT("Category"));
+
+	for (TFieldIterator<FProperty> It(GetClass()); It; ++It)
+	{
+		const FString Category = It->GetMetaData(CategoryName);
+		if (Category == TEXT("Input"))
+		{
+			PropertyInfos.Emplace(It->GetFName(), EStateTreeBlueprintPropertyCategory::Input);
+		}
+		else if (Category == TEXT("Output"))
+		{
+			PropertyInfos.Emplace(It->GetFName(), EStateTreeBlueprintPropertyCategory::Output);
+		}
+		else if (Category == TEXT("Parameter"))
+		{
+			PropertyInfos.Emplace(It->GetFName(), EStateTreeBlueprintPropertyCategory::Parameter);
+		}
+		else if (Category == TEXT("ExternalData"))
+		{
+			PropertyInfos.Emplace(It->GetFName(), EStateTreeBlueprintPropertyCategory::ExternalData);
+		}
+	}
+}
+#endif
+
+void UStateTreeItemBlueprintBase::LinkExternalData(FStateTreeLinker& Linker, TArray<FStateTreeBlueprintExternalDataHandle>& OutExternalDataHandles) const
+{
+	const UClass* Class = GetClass();
+
+	// Find properties which are external data requirements.
+	for (const FStateTreeBlueprintPropertyInfo& PropertyInfo : PropertyInfos)
+	{
+		if (PropertyInfo.Category == EStateTreeBlueprintPropertyCategory::ExternalData)
+		{
+			const FProperty* Property = Class->FindPropertyByName(PropertyInfo.PropertyName);
+			if (const FObjectProperty* ObjectProperty = CastField<FObjectProperty>(Property))
+			{
+				FStateTreeBlueprintExternalDataHandle& ExtData = OutExternalDataHandles.AddDefaulted_GetRef();
+				ExtData.Property = Property;
+				Linker.LinkExternalData(ExtData.Handle, ObjectProperty->PropertyClass, EStateTreeExternalDataRequirement::Required);
+			}
+			else if (const FStructProperty* StructProperty = CastField<FStructProperty>(Property))
+			{
+				FStateTreeBlueprintExternalDataHandle& ExtData = OutExternalDataHandles.AddDefaulted_GetRef();
+				ExtData.Property = Property;
+				Linker.LinkExternalData(ExtData.Handle, StructProperty->Struct, EStateTreeExternalDataRequirement::Required);
+			}
+		}
+	}
+}
+
+void UStateTreeItemBlueprintBase::CopyExternalData(FStateTreeExecutionContext& Context, TConstArrayView<FStateTreeBlueprintExternalDataHandle> ExternalDataHandles)
+{
+	for (const FStateTreeBlueprintExternalDataHandle& ExtData : ExternalDataHandles)
+	{
+		FStateTreeDataView View = Context.GetExternalDataView(ExtData.Handle);
+		if (View.IsValid())
+		{
+			if (const FObjectProperty* ObjectProperty = CastField<FObjectProperty>(ExtData.Property))
+			{
+				check(View.GetStruct()->IsChildOf(ObjectProperty->PropertyClass));
+				// Copy pointer
+				uint8* Address = ObjectProperty->ContainerPtrToValuePtr<uint8>(this);
+				ObjectProperty->SetObjectPropertyValue(Address, View.GetMutablePtr<UObject>());
+			}
+			else if (const FStructProperty* StructProperty = CastField<FStructProperty>(ExtData.Property))
+			{
+				check(View.GetStruct()->IsChildOf(StructProperty->Struct));
+				// Copy value
+				uint8* Address = StructProperty->ContainerPtrToValuePtr<uint8>(this);
+				StructProperty->CopyCompleteValue(Address, View.GetMemory());
+			}
+		}
+	}
+}
\ No newline at end of file
diff --git a/Engine/Plugins/Runtime/StateTree/Source/StateTreeModule/Private/Blueprint/StateTreeTaskBlueprintBase.cpp b/Engine/Plugins/Runtime/StateTree/Source/StateTreeModule/Private/Blueprint/StateTreeTaskBlueprintBase.cpp
new file mode 100644
--- /dev/null
+++ b/Engine/Plugins/Runtime/StateTree/Source/StateTreeModule/Private/Blueprint/StateTreeTaskBlueprintBase.cpp
@@ -0,0 +1,89 @@
+﻿// Copyright Epic Games, Inc. All Rights Reserved.
+
+#include "Blueprint/StateTreeTaskBlueprintBase.h"
+#include "CoreMinimal.h"
+#include "StateTreeExecutionContext.h"
+
+//----------------------------------------------------------------------//
+//  UStateTreeTaskBlueprintBase
+//----------------------------------------------------------------------//
+
+EStateTreeRunStatus UStateTreeTaskBlueprintBase::ReceiveEnterState_Implementation(AActor* OwnerActor, const EStateTreeStateChangeType ChangeType, const FStateTreeTransitionResult& Transition)
+{
+	return EStateTreeRunStatus::Running;
+}
+
+EStateTreeRunStatus UStateTreeTaskBlueprintBase::EnterState(FStateTreeExecutionContext& Context, const EStateTreeStateChangeType ChangeType, const FStateTreeTransitionResult& Transition)
+{
+	AActor* OwnerActor = GetOwnerActor(Context);
+	return ReceiveEnterState(OwnerActor, ChangeType, Transition);
+}
+
+void UStateTreeTaskBlueprintBase::ExitState(FStateTreeExecutionContext& Context, const EStateTreeStateChangeType ChangeType, const FStateTreeTransitionResult& Transition)
+{
+	AActor* OwnerActor = GetOwnerActor(Context);
+	ReceiveExitState(OwnerActor, ChangeType, Transition);
+}
+
+void UStateTreeTaskBlueprintBase::StateCompleted(FStateTreeExecutionContext& Context, const EStateTreeRunStatus CompletionStatus, const FStateTreeHandle CompletedState)
+{
+	AActor* OwnerActor = GetOwnerActor(Context);
+	ReceiveStateCompleted(OwnerActor, CompletionStatus, CompletedState);
+}
+
+EStateTreeRunStatus UStateTreeTaskBlueprintBase::Tick(FStateTreeExecutionContext& Context, const float DeltaTime)
+{
+	AActor* OwnerActor = GetOwnerActor(Context);
+	return ReceiveTick(OwnerActor, DeltaTime);
+}
+
+//----------------------------------------------------------------------//
+//  FStateTreeBlueprintTaskWrapper
+//----------------------------------------------------------------------//
+
+bool FStateTreeBlueprintTaskWrapper::Link(FStateTreeLinker& Linker)
+{
+	const UStateTreeTaskBlueprintBase* TaskCDO = TaskClass ? TaskClass->GetDefaultObject<UStateTreeTaskBlueprintBase>() : nullptr;
+	if (TaskCDO != nullptr)
+	{
+		TaskCDO->LinkExternalData(Linker, ExternalDataHandles);
+	}
+	
+	return true;
+}
+
+EStateTreeRunStatus FStateTreeBlueprintTaskWrapper::EnterState(FStateTreeExecutionContext& Context, const EStateTreeStateChangeType ChangeType, const FStateTreeTransitionResult& Transition) const
+{
+	UStateTreeTaskBlueprintBase* Instance = Context.GetInstanceObjectInternal<UStateTreeTaskBlueprintBase>(DataViewIndex);
+	check(Instance);
+	
+	Instance->CopyExternalData(Context, ExternalDataHandles);
+	return Instance->EnterState(Context, ChangeType, Transition);
+}
+
+void FStateTreeBlueprintTaskWrapper::ExitState(FStateTreeExecutionContext& Context, const EStateTreeStateChangeType ChangeType, const FStateTreeTransitionResult& Transition) const
+{
+	UStateTreeTaskBlueprintBase* Instance = Context.GetInstanceObjectInternal<UStateTreeTaskBlueprintBase>(DataViewIndex);
+	check(Instance);
+	
+	Instance->CopyExternalData(Context, ExternalDataHandles);
+	Instance->ExitState(Context, ChangeType, Transition);
+}
+
+void FStateTreeBlueprintTaskWrapper::StateCompleted(FStateTreeExecutionContext& Context, const EStateTreeRunStatus CompletionStatus, const FStateTreeHandle CompletedState) const
+{
+	UStateTreeTaskBlueprintBase* Instance = Context.GetInstanceObjectInternal<UStateTreeTaskBlueprintBase>(DataViewIndex);
+	check(Instance);
+	
+	Instance->CopyExternalData(Context, ExternalDataHandles);
+	Instance->StateCompleted(Context, CompletionStatus, CompletedState);
+}
+
+EStateTreeRunStatus FStateTreeBlueprintTaskWrapper::Tick(FStateTreeExecutionContext& Context, const float DeltaTime) const
+{
+	UStateTreeTaskBlueprintBase* Instance = Context.GetInstanceObjectInternal<UStateTreeTaskBlueprintBase>(DataViewIndex);
+	check(Instance);
+	
+	Instance->CopyExternalData(Context, ExternalDataHandles);
+	return Instance->Tick(Context, DeltaTime);
+}
diff --git a/Engine/Plugins/Runtime/StateTree/Source/StateTreeModule/Private/Components/StateTreeBrainComponent.cpp b/Engine/Plugins/Runtime/StateTree/Source/StateTreeModule/Private/Components/StateTreeBrainComponent.cpp
new file mode 100644
--- /dev/null
+++ b/Engine/Plugins/Runtime/StateTree/Source/StateTreeModule/Private/Components/StateTreeBrainComponent.cpp
@@ -0,0 +1,295 @@
+// Copyright Epic Games, Inc. All Rights Reserved.
+
+#include "Components/StateTreeBrainComponent.h"
+#include "GameFramework/Actor.h"
+#include "GameplayTasksComponent.h"
+#include "VisualLogger/VisualLoggerTypes.h"
+#include "VisualLogger/VisualLogger.h"
+#include "StateTree.h"
+#include "StateTreeEvaluatorBase.h"
+#include "Conditions/StateTreeCondition_Common.h"
+#include "AIController.h"
+#include "GameFramework/CharacterMovementComponent.h"
+
+#define STATETREE_LOG(Verbosity, Format, ...) UE_VLOG(GetOwner(), LogStateTree, Verbosity, Format, ##__VA_ARGS__)
+#define STATETREE_CLOG(Condition, Verbosity, Format, ...) UE_CVLOG((Condition), GetOwner(), LogStateTree, Verbosity, Format, ##__VA_ARGS__)
+
+
+//////////////////////////////////////////////////////////////////////////
+// UBrainComponentStateTreeSchema
+
+bool UBrainComponentStateTreeSchema::IsStructAllowed(const UScriptStruct* InScriptStruct) const
+{
+	return InScriptStruct->IsChildOf(FStateTreeConditionCommonBase::StaticStruct());
+}
+
+bool UBrainComponentStateTreeSchema::IsClassAllowed(const UClass* InClass) const
+{
+	return IsChildOfBlueprintBase(InClass);
+}
+
+bool UBrainComponentStateTreeSchema::IsExternalItemAllowed(const UStruct& InStruct) const
+{
+	// Actors and components.
+	return InStruct.IsChildOf(AActor::StaticClass())
+			|| InStruct.IsChildOf(APawn::StaticClass())
+			|| InStruct.IsChildOf(AAIController::StaticClass())
+			|| InStruct.IsChildOf(UActorComponent::StaticClass())
+			|| InStruct.IsChildOf(UWorldSubsystem::StaticClass());
+}
+
+//////////////////////////////////////////////////////////////////////////
+// UStateTreeBrainComponent
+
+UStateTreeBrainComponent::UStateTreeBrainComponent(const FObjectInitializer& ObjectInitializer)
+	: Super(ObjectInitializer)
+{
+	bWantsInitializeComponent = true;
+	bIsRunning = false;
+	bIsPaused = false;
+}
+
+void UStateTreeBrainComponent::InitializeComponent()
+{
+	if (StateTree == nullptr)
+	{
+		STATETREE_LOG(Error, TEXT("%s: StateTree asset is not set, cannot initialize."), ANSI_TO_TCHAR(__FUNCTION__));
+		return;
+	}
+
+	if (!StateTreeContext.Init(*GetOwner(), *StateTree, EStateTreeStorage::Internal))
+	{
+		STATETREE_LOG(Error, TEXT("%s: Failed to init StateTreeContext."), ANSI_TO_TCHAR(__FUNCTION__));
+	}
+}
+
+void UStateTreeBrainComponent::UninitializeComponent()
+{
+}
+
+bool UStateTreeBrainComponent::SetContextRequirements()
+{
+	if (!StateTreeContext.IsValid())
+	{
+		return false;
+	}
+	
+	UWorld* World = GetWorld();
+	if (World == nullptr)
+	{
+		return false;
+	}
+
+	for (const FStateTreeExternalDataDesc& ItemDesc : StateTreeContext.GetExternalDataDescs())
+	{
+		if (ItemDesc.Struct != nullptr)
+		{
+			if (ItemDesc.Struct->IsChildOf(UWorldSubsystem::StaticClass()))
+			{
+				UWorldSubsystem* Subsystem = World->GetSubsystemBase(Cast<UClass>(const_cast<UStruct*>(ItemDesc.Struct)));
+				StateTreeContext.SetExternalData(ItemDesc.Handle, FStateTreeDataView(Subsystem));
+			}
+			else if (ItemDesc.Struct->IsChildOf(UActorComponent::StaticClass()))
+			{
+				UActorComponent* Component = GetOwner()->FindComponentByClass(Cast<UClass>(const_cast<UStruct*>(ItemDesc.Struct)));
+				StateTreeContext.SetExternalData(ItemDesc.Handle, FStateTreeDataView(Component));
+			}
+			else if (ItemDesc.Struct->IsChildOf(AActor::StaticClass()))
+			{
+				if (AIOwner != nullptr)
+				{
+					AActor* OwnerActor = Cast<AActor>(AIOwner->GetPawn());
+					StateTreeContext.SetExternalData(ItemDesc.Handle, FStateTreeDataView(OwnerActor));
+				}
+				else
+				{
+					AActor* OwnerActor = Cast<AActor>(GetOwner());
+					StateTreeContext.SetExternalData(ItemDesc.Handle, FStateTreeDataView(OwnerActor));
+				}
+			}
+			else if (ItemDesc.Struct->IsChildOf(APawn::StaticClass()))
+			{
+				if (AIOwner != nullptr)
+				{
+					APawn* OwnerActor = Cast<APawn>(AIOwner->GetPawn());
+					StateTreeContext.SetExternalData(ItemDesc.Handle, FStateTreeDataView(OwnerActor));
+				}
+				else
+				{
+					APawn* OwnerActor = Cast<APawn>(GetOwner());
+					StateTreeContext.SetExternalData(ItemDesc.Handle, FStateTreeDataView(OwnerActor));
+				}
+			}
+			else if (ItemDesc.Struct->IsChildOf(AAIController::StaticClass()))
+			{
+				if (AIOwner != nullptr)
+				{
+					StateTreeContext.SetExternalData(ItemDesc.Handle, FStateTreeDataView(AIOwner));
+				}
+			}
+		}
+	}
+
+	return StateTreeContext.AreExternalDataViewsValid();
+}
+
+void UStateTreeBrainComponent::TickComponent(float DeltaTime, enum ELevelTick TickType, FActorComponentTickFunction *ThisTickFunction)
+{
+	Super::TickComponent(DeltaTime, TickType, ThisTickFunction);
+
+	if (!bIsRunning || bIsPaused)
+	{
+		return;
+	}
+	
+	if (SetContextRequirements())
+	{
+		StateTreeContext.Tick(DeltaTime);
+	}
+}
+
+void UStateTreeBrainComponent::StartLogic()
+{
+	STATETREE_LOG(Log, TEXT("%s: Start Logic"), ANSI_TO_TCHAR(__FUNCTION__));
+
+	// TODO: Move this to UStateTree
+	StateTree->ResolvePropertyPaths();
+
+	if (SetContextRequirements())
+	{
+		StateTreeContext.Start();
+		bIsRunning = true;
+	}
+}
+
+void UStateTreeBrainComponent::RestartLogic()
+{
+	STATETREE_LOG(Log, TEXT("%s: Restart Logic"), ANSI_TO_TCHAR(__FUNCTION__));
+
+	if (SetContextRequirements())
+	{
+		StateTreeContext.Start();
+		bIsRunning = true;
+	}
+}
+
+void UStateTreeBrainComponent::StopLogic(const FString& Reason)
+{
+	STATETREE_LOG(Log, TEXT("%s: Stopping, reason: \'%s\'"), ANSI_TO_TCHAR(__FUNCTION__), *Reason);
+
+	if (SetContextRequirements())
+	{
+		StateTreeContext.Stop();
+		bIsRunning = false;
+	}
+}
+
+void UStateTreeBrainComponent::Cleanup()
+{
+}
+
+void UStateTreeBrainComponent::PauseLogic(const FString& Reason)
+{
+	STATETREE_LOG(Log, TEXT("%s: Execution updates: PAUSED (%s)"), ANSI_TO_TCHAR(__FUNCTION__), *Reason);
+	bIsPaused = true;
+}
+
+EAILogicResuming::Type UStateTreeBrainComponent::ResumeLogic(const FString& Reason)
+{
+	STATETREE_LOG(Log, TEXT("%s: Execution updates: RESUMED (%s)"), ANSI_TO_TCHAR(__FUNCTION__), *Reason);
+
+	const EAILogicResuming::Type SuperResumeResult = Super::ResumeLogic(Reason);
+
+	if (!!bIsPaused)
+	{
+		bIsPaused = false;
+
+		if (SuperResumeResult == EAILogicResuming::Continue)
+		{
+			// Nop
+		}
+		else if (SuperResumeResult == EAILogicResuming::RestartedInstead)
+		{
+			RestartLogic();
+		}
+	}
+
+	return SuperResumeResult;
+}
+
+bool UStateTreeBrainComponent::IsRunning() const
+{
+	return bIsRunning;
+}
+
+bool UStateTreeBrainComponent::IsPaused() const
+{
+	return bIsPaused;
+}
+
+UGameplayTasksComponent* UStateTreeBrainComponent::GetGameplayTasksComponent(const UGameplayTask& Task) const
+{
+	const UAITask* AITask = Cast<const UAITask>(&Task);
+	return (AITask && AITask->GetAIController()) ? AITask->GetAIController()->GetGameplayTasksComponent(Task) : Task.GetGameplayTasksComponent();
+}
+
+AActor* UStateTreeBrainComponent::GetGameplayTaskOwner(const UGameplayTask* Task) const
+{
+	if (Task == nullptr)
+	{
+		return GetAIOwner();
+	}
+
+	const UAITask* AITask = Cast<const UAITask>(Task);
+	if (AITask)
+	{
+		return AITask->GetAIController();
+	}
+
+	const UGameplayTasksComponent* TasksComponent = Task->GetGameplayTasksComponent();
+	return TasksComponent ? TasksComponent->GetGameplayTaskOwner(Task) : nullptr;
+}
+
+AActor* UStateTreeBrainComponent::GetGameplayTaskAvatar(const UGameplayTask* Task) const
+{
+	if (Task == nullptr)
+	{
+		return GetAIOwner() ? GetAIOwner()->GetPawn() : nullptr;
+	}
+
+	const UAITask* AITask = Cast<const UAITask>(Task);
+	if (AITask)
+	{
+		return AITask->GetAIController() ? AITask->GetAIController()->GetPawn() : nullptr;
+	}
+
+	const UGameplayTasksComponent* TasksComponent = Task->GetGameplayTasksComponent();
+	return TasksComponent ? TasksComponent->GetGameplayTaskAvatar(Task) : nullptr;
+}
+
+uint8 UStateTreeBrainComponent::GetGameplayTaskDefaultPriority() const
+{
+	return static_cast<uint8>(EAITaskPriority::AutonomousAI);
+}
+
+void UStateTreeBrainComponent::OnGameplayTaskInitialized(UGameplayTask& Task)
+{
+	const UAITask* AITask = Cast<const UAITask>(&Task);
+	if (AITask && (AITask->GetAIController() == nullptr))
+	{
+		// this means that the task has either been created without specifying 
+		// UAITAsk::OwnerController's value (like via BP's Construct Object node)
+		// or it has been created in C++ with inappropriate function
+		UE_LOG(LogStateTree, Error, TEXT("Missing AIController in AITask %s"), *AITask->GetName());
+	}
+}
+
+#if WITH_GAMEPLAY_DEBUGGER
+FString UStateTreeBrainComponent::GetDebugInfoString() const
+{
+	return StateTreeContext.GetDebugInfoString();
+}
+#endif // WITH_GAMEPLAY_DEBUGGER
+
+#undef STATETREE_LOG
+#undef STATETREE_CLOG
diff --git a/Engine/Plugins/Runtime/StateTree/Source/StateTreeModule/Private/Components/StateTreeComponent.cpp b/Engine/Plugins/Runtime/StateTree/Source/StateTreeModule/Private/Components/StateTreeComponent.cpp
deleted file mode 100644
diff --git a/Engine/Plugins/Runtime/StateTree/Source/StateTreeModule/Private/StateTree.cpp b/Engine/Plugins/Runtime/StateTree/Source/StateTreeModule/Private/StateTree.cpp
--- a/Engine/Plugins/Runtime/StateTree/Source/StateTreeModule/Private/StateTree.cpp
+++ b/Engine/Plugins/Runtime/StateTree/Source/StateTreeModule/Private/StateTree.cpp
@@ -6,10 +6,24 @@
 #include "CoreMinimal.h"
 #include "StateTreeConditionBase.h"
 #include "StateTreeDelegates.h"
+#if WITH_EDITOR
+#include "Editor.h"
+#endif
 
 UStateTree::UStateTree(const FObjectInitializer& ObjectInitializer)
 	: Super(ObjectInitializer)
 {
+
+#if WITH_EDITOR
+	FEditorDelegates::PreBeginPIE.AddUObject(this, &UStateTree::OnPIEStarted);
+#endif
+}
+
+UStateTree::~UStateTree()
+{
+#if WITH_EDITOR
+	FEditorDelegates::PreBeginPIE.RemoveAll(this);
+#endif
 }
 
 bool UStateTree::IsValidStateTree() const
@@ -18,8 +32,15 @@ bool UStateTree::IsValidStateTree() const
 	return States.Num() > 0;
 }
 
-
 #if WITH_EDITOR
+
+void UStateTree::OnPIEStarted(const bool bIsSimulating)
+{
+	PropertyBindings.ResolvePaths();
+	Link();
+	InitInstanceStorageType();
+}
+
 void UStateTree::ResetBaked()
 {
 	States.Reset();
@@ -27,6 +48,7 @@ void UStateTree::ResetBaked()
 
 	Items.Reset();
 	Instances.Reset();
+	InstanceObjects.Reset();
 	InstanceStorageStruct = nullptr;
 	InstanceStorageOffsets.Reset();
 	InstanceStorageDefaultValue.Reset();
@@ -68,7 +90,7 @@ void UStateTree::PostLoad()
 	Link();
 	
 #if WITH_EDITOR
-	InitInstanceStorage();
+	InitInstanceStorageType();
 #else
 	// Item offsets still need to be calculated in non editor target since the struct sizes might be different.
 	checkf(InstanceStorageOffsets.Num() == 0, TEXT("RuntimeStorageOffsets is transient and should only be computed once."));
@@ -127,7 +149,7 @@ void UStateTree::Link()
 	NumDataViews = ExternalDataBaseIndex + ExternalDataDescs.Num();
 }
 
-void UStateTree::InitInstanceStorage()
+void UStateTree::InitInstanceStorageType()
 {
 	InstanceStorageStruct = nullptr;
 	InstanceStorageOffsets.Reset();
diff --git a/Engine/Plugins/Runtime/StateTree/Source/StateTreeModule/Private/StateTreeExecutionContext.cpp b/Engine/Plugins/Runtime/StateTree/Source/StateTreeModule/Private/StateTreeExecutionContext.cpp
--- a/Engine/Plugins/Runtime/StateTree/Source/StateTreeModule/Private/StateTreeExecutionContext.cpp
+++ b/Engine/Plugins/Runtime/StateTree/Source/StateTreeModule/Private/StateTreeExecutionContext.cpp
@@ -32,7 +32,7 @@ bool FStateTreeExecutionContext::Init(UObject& InOwner, const UStateTree& InStat
 {
 	// Set owner first for proper logging (it will be reset in case of failure) 
 	Owner = &InOwner;
-
+	
 	if (!InStateTree.IsValidStateTree())
 	{
 		STATETREE_LOG(Error, TEXT("%s: StateTree asset '%s' is not valid."), ANSI_TO_TCHAR(__FUNCTION__), *InStateTree.GetName());
@@ -46,6 +46,7 @@ bool FStateTreeExecutionContext::Init(UObject& InOwner, const UStateTree& InStat
 	{
 		// Duplicate instance state.
 		StorageInstance = StateTree->InstanceStorageDefaultValue;
+		InitInstanceData(InOwner, FStateTreeDataView(StorageInstance));
 	}
 
 	// Initialize data views for all possible items.
@@ -57,6 +58,27 @@ bool FStateTreeExecutionContext::Init(UObject& InOwner, const UStateTree& InStat
 	return true;
 }
 
+bool FStateTreeExecutionContext::InitInstanceData(UObject& InOwner, FStateTreeDataView ExternalStorage) const
+{
+	if (!Owner || !StateTree)
+	{
+		STATETREE_LOG(Error, TEXT("%s: Failed to init instance data on '%s' using StateTree '%s'."),
+			ANSI_TO_TCHAR(__FUNCTION__), *GetNameSafe(Owner), *GetNameSafe(StateTree));
+		return false;
+	}
+
+	check (ExternalStorage.IsValid());
+	FStateTreeExecutionState& Exec = GetExecState(ExternalStorage);
+
+	Exec.InstanceObjects.Reset();
+	for (const UObject* Instance : StateTree->InstanceObjects)
+	{
+		Exec.InstanceObjects.Add(DuplicateObject(Instance, &InOwner));
+	}
+
+	return true;
+}
+
 EStateTreeRunStatus FStateTreeExecutionContext::Start(FStateTreeDataView ExternalStorage)
 {
 	if (!Owner || !StateTree)
@@ -237,7 +259,6 @@ void FStateTreeExecutionContext::Reset()
 {
 	StateTree = nullptr;
 	Owner = nullptr;
-	World = nullptr;
 	DataViews.Reset();
 	StorageInstance.Reset();
 	StorageType = EStateTreeStorage::Internal;
@@ -285,12 +306,12 @@ EStateTreeRunStatus FStateTreeExecutionContext::EnterState(FStateTreeDataView St
 			for (int32 EvalIndex = 0; EvalIndex < int32(State.EvaluatorsNum); EvalIndex++)
 			{
 				const FStateTreeEvaluatorBase& Eval = GetItem<FStateTreeEvaluatorBase>(int32(State.EvaluatorsBegin) + EvalIndex);
-				DataViews[Eval.DataViewIndex] = GetInstanceData(Storage, Eval.DataViewIndex);
+				DataViews[Eval.DataViewIndex] = GetInstanceData(Storage, Eval.bInstanceIsObject, Eval.InstanceIndex);
 			}
 			for (int32 TaskIndex = 0; TaskIndex < int32(State.TasksNum); TaskIndex++)
 			{
 				const FStateTreeTaskBase& Task = GetItem<FStateTreeTaskBase>(int32(State.TasksBegin) + TaskIndex);
-				DataViews[Task.DataViewIndex] = GetInstanceData(Storage, Task.DataViewIndex);
+				DataViews[Task.DataViewIndex] = GetInstanceData(Storage, Task.bInstanceIsObject, Task.InstanceIndex);
 			}
 			continue;
 		}
@@ -304,7 +325,7 @@ EStateTreeRunStatus FStateTreeExecutionContext::EnterState(FStateTreeDataView St
 		for (int32 EvalIndex = 0; EvalIndex < int32(State.EvaluatorsNum); EvalIndex++)
 		{
 			const FStateTreeEvaluatorBase& Eval = GetItem<FStateTreeEvaluatorBase>(int32(State.EvaluatorsBegin) + EvalIndex);
-			const FStateTreeDataView EvalInstanceView = GetInstanceData(Storage, Eval.DataViewIndex);
+			const FStateTreeDataView EvalInstanceView = GetInstanceData(Storage, Eval.bInstanceIsObject, Eval.InstanceIndex);
 			DataViews[Eval.DataViewIndex] = EvalInstanceView;
 
 			// Copy bound properties.
@@ -320,7 +341,7 @@ EStateTreeRunStatus FStateTreeExecutionContext::EnterState(FStateTreeDataView St
 		for (int32 TaskIndex = 0; TaskIndex < int32(State.TasksNum); TaskIndex++)
 		{
 			const FStateTreeTaskBase& Task = GetItem<FStateTreeTaskBase>(int32(State.TasksBegin) + TaskIndex);
-			const FStateTreeDataView TaskInstanceView = GetInstanceData(Storage, Task.DataViewIndex);
+			const FStateTreeDataView TaskInstanceView = GetInstanceData(Storage, Task.bInstanceIsObject, Task.InstanceIndex);
 			DataViews[Task.DataViewIndex] = TaskInstanceView;
 
 			// Copy bound properties.
@@ -391,12 +412,12 @@ void FStateTreeExecutionContext::ExitState(FStateTreeDataView Storage, const FSt
 			for (int32 EvalIndex = 0; EvalIndex < int32(State.EvaluatorsNum); EvalIndex++)
 			{
 				FStateTreeEvaluatorBase& Eval = GetItem<FStateTreeEvaluatorBase>(int32(State.EvaluatorsBegin) + EvalIndex);
-				DataViews[Eval.DataViewIndex] = GetInstanceData(Storage, Eval.DataViewIndex);
+				DataViews[Eval.DataViewIndex] = GetInstanceData(Storage, Eval.bInstanceIsObject, Eval.InstanceIndex);
 			}
 			for (int32 TaskIndex = 0; TaskIndex < int32(State.TasksNum); TaskIndex++)
 			{
 				FStateTreeTaskBase& Task = GetItem<FStateTreeTaskBase>(int32(State.TasksBegin) + TaskIndex);
-				DataViews[Task.DataViewIndex] = GetInstanceData(Storage, Task.DataViewIndex);
+				DataViews[Task.DataViewIndex] = GetInstanceData(Storage, Task.bInstanceIsObject, Task.InstanceIndex);
 			}
 			continue;
 		}
@@ -410,7 +431,7 @@ void FStateTreeExecutionContext::ExitState(FStateTreeDataView Storage, const FSt
 		for (int32 EvalIndex = 0; EvalIndex < int32(State.EvaluatorsNum); EvalIndex++)
 		{
 			const FStateTreeEvaluatorBase& Eval = GetItem<FStateTreeEvaluatorBase>(int32(State.EvaluatorsBegin) + EvalIndex);
-			const FStateTreeDataView EvalInstanceView = GetInstanceData(Storage, Eval.DataViewIndex);
+			const FStateTreeDataView EvalInstanceView = GetInstanceData(Storage, Eval.bInstanceIsObject, Eval.InstanceIndex);
 			DataViews[Eval.DataViewIndex] = EvalInstanceView;
 
 			// Copy bound properties.
@@ -426,7 +447,7 @@ void FStateTreeExecutionContext::ExitState(FStateTreeDataView Storage, const FSt
 		for (int32 TaskIndex = 0; TaskIndex < int32(State.TasksNum); TaskIndex++)
 		{
 			const FStateTreeTaskBase& Task = GetItem<FStateTreeTaskBase>(int32(State.TasksBegin) + TaskIndex);
-			const FStateTreeDataView TaskInstanceView = GetInstanceData(Storage, Task.DataViewIndex);
+			const FStateTreeDataView TaskInstanceView = GetInstanceData(Storage, Task.bInstanceIsObject, Task.InstanceIndex);
 			DataViews[Task.DataViewIndex] = TaskInstanceView;
 
 			// Copy bound properties.
@@ -509,7 +530,7 @@ void FStateTreeExecutionContext::TickEvaluators(FStateTreeDataView Storage, cons
 		for (int32 EvalIndex = 0; EvalIndex < int32(State.EvaluatorsNum); EvalIndex++)
 		{
 			const FStateTreeEvaluatorBase& Eval = GetItem<FStateTreeEvaluatorBase>(int32(State.EvaluatorsBegin) + EvalIndex);
-			const FStateTreeDataView EvalInstanceView = GetInstanceData(Storage, Eval.DataViewIndex);
+			const FStateTreeDataView EvalInstanceView = GetInstanceData(Storage, Eval.bInstanceIsObject, Eval.InstanceIndex);
 			DataViews[Eval.DataViewIndex] = EvalInstanceView;
 
 			// Copy bound properties.
@@ -551,7 +572,7 @@ EStateTreeRunStatus FStateTreeExecutionContext::TickTasks(FStateTreeDataView Sto
 		for (int32 TaskIndex = 0; TaskIndex < int32(State.TasksNum); TaskIndex++)
 		{
 			const FStateTreeTaskBase& Task = GetItem<FStateTreeTaskBase>(int32(State.TasksBegin) + TaskIndex);
-			const FStateTreeDataView TaskInstanceView = GetInstanceData(Storage, Task.DataViewIndex);
+			const FStateTreeDataView TaskInstanceView = GetInstanceData(Storage, Task.bInstanceIsObject, Task.InstanceIndex);
 			DataViews[Task.DataViewIndex] = TaskInstanceView;
 
 			// Copy bound properties.
@@ -591,7 +612,7 @@ bool FStateTreeExecutionContext::TestAllConditions(FStateTreeDataView Storage, c
 	for (uint32 i = 0; i < ConditionsNum; i++)
 	{
 		const FStateTreeConditionBase& Cond = GetItem<FStateTreeConditionBase>(ConditionsOffset + i);
-		const FStateTreeDataView CondInstanceView = GetInstanceData(Storage, Cond.DataViewIndex);
+		const FStateTreeDataView CondInstanceView = GetInstanceData(Storage, Cond.bInstanceIsObject, Cond.InstanceIndex);
 		DataViews[Cond.DataViewIndex] = CondInstanceView;
 
 		// Copy bound properties.
@@ -822,6 +843,23 @@ EStateTreeRunStatus FStateTreeExecutionContext::GetLastTickStatus(FStateTreeData
 	return Exec.LastTickStatus;
 }
 
+void FStateTreeExecutionContext::AddStructReferencedObjects(FReferenceCollector& Collector)
+{
+	if (StorageType == EStateTreeStorage::Internal)
+	{
+		AddStructReferencedObjects(FStateTreeDataView(StorageInstance), Collector);
+	}
+}
+
+void FStateTreeExecutionContext::AddStructReferencedObjects(const FStateTreeDataView Storage, FReferenceCollector& Collector)
+{
+	if (Storage.IsValid())
+	{
+		FStateTreeExecutionState& Exec = GetExecState(Storage);
+		Collector.AddReferencedObjects(Exec.InstanceObjects);
+	}
+}
+
 #if WITH_GAMEPLAY_DEBUGGER
 
 FString FStateTreeExecutionContext::GetDebugInfoString(FStateTreeDataView ExternalStorage) const
@@ -976,25 +1014,25 @@ void FStateTreeExecutionContext::DebugPrintInternalLayout(FStateTreeDataView Ext
 		TEXT("State"), TEXT("Type"), TEXT("Name"), TEXT("Bindings"), TEXT("Struct Idx"));
 	for (const FBakedStateTreeState& State : StateTree->States)
 	{
-		// Tasks
-		if (State.TasksNum)
+		// Evaluators
+		if (State.EvaluatorsNum)
 		{
-			for (int32 j = 0; j < State.TasksNum; j++)
+			for (int32 EvalIndex = 0; EvalIndex < State.EvaluatorsNum; EvalIndex++)
 			{
-				const FStateTreeTaskBase& Task = GetInstanceData(Storage, State.TasksBegin + j).Get<FStateTreeTaskBase>();
+				const FStateTreeEvaluatorBase& Eval = GetItem<FStateTreeEvaluatorBase>(State.EvaluatorsBegin + EvalIndex);
 				DebugString += FString::Printf(TEXT("  | %-30s | %-12s | %-30s | %15s | %10d |\n"), *State.Name.ToString(),
-					TEXT("  Task"), *Task.Name.ToString(), *Task.BindingsBatch.Describe(), Task.DataViewIndex);
+					TEXT("  Evaluator"), *Eval.Name.ToString(), *Eval.BindingsBatch.Describe(), Eval.DataViewIndex);
 			}
 		}
 
-		// Evaluators
-		if (State.EvaluatorsNum)
+		// Tasks
+		if (State.TasksNum)
 		{
-			for (int32 EvalIndex = 0; EvalIndex < State.EvaluatorsNum; EvalIndex++)
+			for (int32 TaskIndex = 0; TaskIndex < State.TasksNum; TaskIndex++)
 			{
-				const FStateTreeEvaluatorBase& Eval = GetInstanceData(Storage, State.EvaluatorsBegin + EvalIndex).Get<FStateTreeEvaluatorBase>();
+				const FStateTreeTaskBase& Task = GetItem<FStateTreeTaskBase>(State.TasksBegin + TaskIndex);
 				DebugString += FString::Printf(TEXT("  | %-30s | %-12s | %-30s | %15s | %10d |\n"), *State.Name.ToString(),
-					TEXT("  Evaluator"), *Eval.Name.ToString(), *Eval.BindingsBatch.Describe(), Eval.DataViewIndex);
+					TEXT("  Task"), *Task.Name.ToString(), *Task.BindingsBatch.Describe(), Task.DataViewIndex);
 			}
 		}
 	}
diff --git a/Engine/Plugins/Runtime/StateTree/Source/StateTreeModule/Public/Blueprint/StateTreeConditionBlueprintBase.h b/Engine/Plugins/Runtime/StateTree/Source/StateTreeModule/Public/Blueprint/StateTreeConditionBlueprintBase.h
new file mode 100644
--- /dev/null
+++ b/Engine/Plugins/Runtime/StateTree/Source/StateTreeModule/Public/Blueprint/StateTreeConditionBlueprintBase.h
@@ -0,0 +1,45 @@
+﻿// Copyright Epic Games, Inc. All Rights Reserved.
+
+#pragma once
+
+#include "CoreMinimal.h"
+#include "StateTreeTypes.h"
+#include "StateTreeConditionBase.h"
+#include "StateTreePropertyBindings.h"
+#include "StateTreeItemBlueprintBase.h"
+#include "StateTreeConditionBlueprintBase.generated.h"
+
+struct FStateTreeExecutionContext;
+
+/*
+ * Base class for Blueprint based Conditions. 
+ */
+UCLASS(Abstract, Blueprintable)
+class STATETREEMODULE_API UStateTreeConditionBlueprintBase : public UStateTreeItemBlueprintBase
+{
+	GENERATED_BODY()
+public:
+
+	UFUNCTION(BlueprintImplementableEvent)
+	bool ReceiveTestCondition(AActor* OwnerActor);
+
+	virtual bool TestCondition(FStateTreeExecutionContext& Context);
+};
+
+/**
+ * Wrapper for Blueprint based Conditions.
+ */
+USTRUCT()
+struct STATETREEMODULE_API FStateTreeBlueprintConditionWrapper : public FStateTreeConditionBase
+{
+	GENERATED_BODY()
+
+	virtual const UStruct* GetInstanceDataType() const override { return ConditionClass; };
+	virtual bool Link(FStateTreeLinker& Linker) override;
+	virtual bool TestCondition(FStateTreeExecutionContext& Context) const override;
+
+	UPROPERTY()
+	TSubclassOf<UStateTreeConditionBlueprintBase> ConditionClass = nullptr;
+
+	TArray<FStateTreeBlueprintExternalDataHandle> ExternalDataHandles;
+};
diff --git a/Engine/Plugins/Runtime/StateTree/Source/StateTreeModule/Public/Blueprint/StateTreeEvaluatorBlueprintBase.h b/Engine/Plugins/Runtime/StateTree/Source/StateTreeModule/Public/Blueprint/StateTreeEvaluatorBlueprintBase.h
new file mode 100644
--- /dev/null
+++ b/Engine/Plugins/Runtime/StateTree/Source/StateTreeModule/Public/Blueprint/StateTreeEvaluatorBlueprintBase.h
@@ -0,0 +1,64 @@
+﻿// Copyright Epic Games, Inc. All Rights Reserved.
+
+#pragma once
+
+#include "CoreMinimal.h"
+#include "StateTreeTypes.h"
+#include "StateTreeEvaluatorBase.h"
+#include "StateTreePropertyBindings.h"
+#include "StateTreeItemBlueprintBase.h"
+#include "StateTreeEvaluatorBlueprintBase.generated.h"
+
+struct FStateTreeExecutionContext;
+
+/*
+ * Base class for Blueprint based evaluators. 
+ */
+UCLASS(Abstract, Blueprintable)
+class STATETREEMODULE_API UStateTreeEvaluatorBlueprintBase : public UStateTreeItemBlueprintBase
+{
+	GENERATED_BODY()
+public:
+
+	UFUNCTION(BlueprintImplementableEvent, meta = (DisplayName = "EnterState"))
+	void ReceiveEnterState(AActor* OwnerActor, const EStateTreeStateChangeType ChangeType, const FStateTreeTransitionResult& Transition);
+
+	UFUNCTION(BlueprintImplementableEvent, meta = (DisplayName = "ExitState"))
+	void ReceiveExitState(AActor* OwnerActor, const EStateTreeStateChangeType ChangeType, const FStateTreeTransitionResult& Transition);
+
+	UFUNCTION(BlueprintImplementableEvent, meta = (DisplayName = "StateCompleted"))
+	void ReceiveStateCompleted(AActor* OwnerActor, const EStateTreeRunStatus CompletionStatus, const FStateTreeHandle CompletedState);
+
+	UFUNCTION(BlueprintImplementableEvent, meta = (DisplayName = "Evaluate"))
+	void ReceiveEvaluate(AActor* OwnerActor, const EStateTreeEvaluationType EvalType, const float DeltaTime);
+
+	
+	virtual void EnterState(FStateTreeExecutionContext& Context, const EStateTreeStateChangeType ChangeType, const FStateTreeTransitionResult& Transition);
+
+	virtual void ExitState(FStateTreeExecutionContext& Context, const EStateTreeStateChangeType ChangeType, const FStateTreeTransitionResult& Transition);
+
+	virtual void StateCompleted(FStateTreeExecutionContext& Context, const EStateTreeRunStatus CompletionStatus, const FStateTreeHandle CompletedState);
+
+	virtual void Evaluate(FStateTreeExecutionContext& Context, const EStateTreeEvaluationType EvalType, const float DeltaTime);
+};
+
+/**
+ * Wrapper for Blueprint based Evaluators.
+ */
+USTRUCT()
+struct STATETREEMODULE_API FStateTreeBlueprintEvaluatorWrapper : public FStateTreeEvaluatorBase
+{
+	GENERATED_BODY()
+
+	virtual const UStruct* GetInstanceDataType() const override { return EvaluatorClass; };
+	virtual bool Link(FStateTreeLinker& Linker) override;
+	virtual void EnterState(FStateTreeExecutionContext& Context, const EStateTreeStateChangeType ChangeType, const FStateTreeTransitionResult& Transition) const override;
+	virtual void ExitState(FStateTreeExecutionContext& Context, const EStateTreeStateChangeType ChangeType, const FStateTreeTransitionResult& Transition) const override;
+	virtual void StateCompleted(FStateTreeExecutionContext& Context, const EStateTreeRunStatus CompletionStatus, const FStateTreeHandle CompletedState) const override;
+	virtual void Evaluate(FStateTreeExecutionContext& Context, const EStateTreeEvaluationType EvalType, const float DeltaTime) const override;
+
+	UPROPERTY()
+	TSubclassOf<UStateTreeEvaluatorBlueprintBase> EvaluatorClass = nullptr;
+
+	TArray<FStateTreeBlueprintExternalDataHandle> ExternalDataHandles;
+};
diff --git a/Engine/Plugins/Runtime/StateTree/Source/StateTreeModule/Public/Blueprint/StateTreeItemBlueprintBase.h b/Engine/Plugins/Runtime/StateTree/Source/StateTreeModule/Public/Blueprint/StateTreeItemBlueprintBase.h
new file mode 100644
--- /dev/null
+++ b/Engine/Plugins/Runtime/StateTree/Source/StateTreeModule/Public/Blueprint/StateTreeItemBlueprintBase.h
@@ -0,0 +1,70 @@
+﻿// Copyright Epic Games, Inc. All Rights Reserved.
+
+#pragma once
+
+#include "CoreMinimal.h"
+#include "StateTreeTypes.h"
+#include "StateTreeItemBlueprintBase.generated.h"
+
+struct FStateTreeExecutionContext;
+
+UENUM()
+enum class EStateTreeBlueprintPropertyCategory : uint8
+{
+	NotSet,
+	Input,	
+	Parameter,
+	Output,
+	ExternalData,
+};
+
+/** Additional information for properties, captured from property metadata. */
+USTRUCT()
+struct STATETREEMODULE_API FStateTreeBlueprintPropertyInfo
+{
+	GENERATED_BODY()
+
+	FStateTreeBlueprintPropertyInfo() = default;
+
+	FStateTreeBlueprintPropertyInfo(const FName InPropertyName, const EStateTreeBlueprintPropertyCategory InCategory)
+		: PropertyName(InPropertyName), Category(InCategory)
+	{}
+	
+	UPROPERTY()
+	FName PropertyName;
+
+	UPROPERTY()
+	EStateTreeBlueprintPropertyCategory Category;
+};
+
+/** Struct use to copy external data to the Blueprint item instance, resolved during StateTree linking. */
+struct STATETREEMODULE_API FStateTreeBlueprintExternalDataHandle
+{
+	const FProperty* Property = nullptr;
+	FStateTreeExternalDataHandle Handle;
+};
+
+
+UCLASS(Abstract)
+class STATETREEMODULE_API UStateTreeItemBlueprintBase : public UObject
+{
+	GENERATED_BODY()
+
+public:
+	/** Called during link to collect external data handles. */
+	void LinkExternalData(FStateTreeLinker& Linker, TArray<FStateTreeBlueprintExternalDataHandle>& OutExternalDataHandles) const;
+	/** Called before call to logic functions to copy data from */
+	void CopyExternalData(FStateTreeExecutionContext& Context, TConstArrayView<FStateTreeBlueprintExternalDataHandle> ExternalDataHandles);
+
+protected:
+	virtual UWorld* GetWorld() const override;
+	AActor* GetOwnerActor(const FStateTreeExecutionContext& Context) const;
+
+#if WITH_EDITOR
+	virtual void PostCDOCompiled() override;
+#endif
+
+	/** Metadata for properties */
+	UPROPERTY()
+	TArray<FStateTreeBlueprintPropertyInfo> PropertyInfos;
+};
diff --git a/Engine/Plugins/Runtime/StateTree/Source/StateTreeModule/Public/Blueprint/StateTreeTaskBlueprintBase.h b/Engine/Plugins/Runtime/StateTree/Source/StateTreeModule/Public/Blueprint/StateTreeTaskBlueprintBase.h
new file mode 100644
--- /dev/null
+++ b/Engine/Plugins/Runtime/StateTree/Source/StateTreeModule/Public/Blueprint/StateTreeTaskBlueprintBase.h
@@ -0,0 +1,64 @@
+﻿// Copyright Epic Games, Inc. All Rights Reserved.
+
+#pragma once
+
+#include "CoreMinimal.h"
+#include "StateTreeTypes.h"
+#include "StateTreeTaskBase.h"
+#include "StateTreePropertyBindings.h"
+#include "StateTreeItemBlueprintBase.h"
+#include "StateTreeTaskBlueprintBase.generated.h"
+
+struct FStateTreeExecutionContext;
+
+/*
+ * Base class for Blueprint based Tasks. 
+ */
+UCLASS(Abstract, Blueprintable)
+class STATETREEMODULE_API UStateTreeTaskBlueprintBase : public UStateTreeItemBlueprintBase
+{
+	GENERATED_BODY()
+public:
+
+	UFUNCTION(BlueprintNativeEvent, meta = (DisplayName = "EnterState"))
+	EStateTreeRunStatus ReceiveEnterState(AActor* OwnerActor, const EStateTreeStateChangeType ChangeType, const FStateTreeTransitionResult& Transition);
+
+	UFUNCTION(BlueprintImplementableEvent, meta = (DisplayName = "ExitState"))
+	void ReceiveExitState(AActor* OwnerActor, const EStateTreeStateChangeType ChangeType, const FStateTreeTransitionResult& Transition);
+
+	UFUNCTION(BlueprintImplementableEvent, meta = (DisplayName = "StateCompleted"))
+	void ReceiveStateCompleted(AActor* OwnerActor, const EStateTreeRunStatus CompletionStatus, const FStateTreeHandle CompletedState);
+
+	UFUNCTION(BlueprintImplementableEvent, meta = (DisplayName = "Tick"))
+	EStateTreeRunStatus ReceiveTick(AActor* OwnerActor, const float DeltaTime);
+
+	
+	virtual EStateTreeRunStatus EnterState(FStateTreeExecutionContext& Context, const EStateTreeStateChangeType ChangeType, const FStateTreeTransitionResult& Transition);
+
+	virtual void ExitState(FStateTreeExecutionContext& Context, const EStateTreeStateChangeType ChangeType, const FStateTreeTransitionResult& Transition);
+
+	virtual void StateCompleted(FStateTreeExecutionContext& Context, const EStateTreeRunStatus CompletionStatus, const FStateTreeHandle CompletedState);
+
+	virtual EStateTreeRunStatus Tick(FStateTreeExecutionContext& Context, const float DeltaTime);
+};
+
+/**
+ * Wrapper for Blueprint based Tasks.
+ */
+USTRUCT()
+struct STATETREEMODULE_API FStateTreeBlueprintTaskWrapper : public FStateTreeTaskBase
+{
+	GENERATED_BODY()
+
+	virtual const UStruct* GetInstanceDataType() const override { return TaskClass; };
+	virtual bool Link(FStateTreeLinker& Linker) override;
+	virtual EStateTreeRunStatus EnterState(FStateTreeExecutionContext& Context, const EStateTreeStateChangeType ChangeType, const FStateTreeTransitionResult& Transition) const override;
+	virtual void ExitState(FStateTreeExecutionContext& Context, const EStateTreeStateChangeType ChangeType, const FStateTreeTransitionResult& Transition) const override;
+	virtual void StateCompleted(FStateTreeExecutionContext& Context, const EStateTreeRunStatus CompletionStatus, const FStateTreeHandle CompletedState) const override;
+	virtual EStateTreeRunStatus Tick(FStateTreeExecutionContext& Context, const float DeltaTime) const override;
+
+	UPROPERTY()
+	TSubclassOf<UStateTreeTaskBlueprintBase> TaskClass = nullptr;
+
+	TArray<FStateTreeBlueprintExternalDataHandle> ExternalDataHandles;
+};
diff --git a/Engine/Plugins/Runtime/StateTree/Source/StateTreeModule/Public/Components/StateTreeBrainComponent.h b/Engine/Plugins/Runtime/StateTree/Source/StateTreeModule/Public/Components/StateTreeBrainComponent.h
new file mode 100644
--- /dev/null
+++ b/Engine/Plugins/Runtime/StateTree/Source/StateTreeModule/Public/Components/StateTreeBrainComponent.h
@@ -0,0 +1,81 @@
+// Copyright Epic Games, Inc. All Rights Reserved.
+
+#pragma once
+
+#include "Components/ActorComponent.h"
+#include "GameplayTagAssetInterface.h"
+#include "GameplayTagContainer.h"
+#include "StateTreeExecutionContext.h"
+#include "BrainComponent.h"
+#include "Tasks/AITask.h"
+#include "StateTreeBrainComponent.generated.h"
+
+class UStateTree;
+
+
+UCLASS(BlueprintType, EditInlineNew, CollapseCategories, meta = (DisplayName = "Brain Component"))
+class STATETREEMODULE_API UBrainComponentStateTreeSchema : public UStateTreeSchema
+{
+	GENERATED_BODY()
+
+public:
+	virtual bool IsStructAllowed(const UScriptStruct* InScriptStruct) const override;
+	virtual bool IsClassAllowed(const UClass* InScriptStruct) const override;
+	virtual bool IsExternalItemAllowed(const UStruct& InStruct) const override;
+};
+
+
+UCLASS(ClassGroup = AI, HideCategories = (Activation, Collision), meta = (BlueprintSpawnableComponent))
+class STATETREEMODULE_API UStateTreeBrainComponent : public UBrainComponent, public IGameplayTaskOwnerInterface
+{
+	GENERATED_BODY()
+public:
+	UStateTreeBrainComponent(const FObjectInitializer& ObjectInitializer = FObjectInitializer::Get());
+
+	// BEGIN UActorComponent overrides
+	virtual void InitializeComponent() override;
+	virtual void UninitializeComponent() override;
+	virtual void TickComponent(float DeltaTime, enum ELevelTick TickType, FActorComponentTickFunction *ThisTickFunction) override;
+	// END UActorComponent overrides
+
+	// BEGIN UBrainComponent overrides
+	virtual void StartLogic() override;
+	virtual void RestartLogic() override;
+	virtual void StopLogic(const FString& Reason)  override;
+	virtual void Cleanup() override;
+	virtual void PauseLogic(const FString& Reason) override;
+	virtual EAILogicResuming::Type ResumeLogic(const FString& Reason)  override;
+	virtual bool IsRunning() const override;
+	virtual bool IsPaused() const override;
+	// END UBrainComponent overrides
+
+	// BEGIN IGameplayTaskOwnerInterface
+	virtual UGameplayTasksComponent* GetGameplayTasksComponent(const UGameplayTask& Task) const override;
+	virtual AActor* GetGameplayTaskOwner(const UGameplayTask* Task) const override;
+	virtual AActor* GetGameplayTaskAvatar(const UGameplayTask* Task) const override;
+	virtual uint8 GetGameplayTaskDefaultPriority() const override;
+	virtual void OnGameplayTaskInitialized(UGameplayTask& Task) override;
+	// END IGameplayTaskOwnerInterface
+
+#if WITH_GAMEPLAY_DEBUGGER
+	virtual FString GetDebugInfoString() const override;
+#endif // WITH_GAMEPLAY_DEBUGGER
+
+private:
+
+protected:
+
+	bool SetContextRequirements();
+	
+	UPROPERTY(EditDefaultsOnly, Category = AI, meta=(RequiredAssetDataTags="Schema=BrainComponentStateTreeSchema"))
+	UStateTree* StateTree;
+
+	UPROPERTY()
+	FStateTreeExecutionContext StateTreeContext;
+
+	/** if set, state tree execution is allowed */
+	uint8 bIsRunning : 1;
+
+	/** if set, execution requests will be postponed */
+	uint8 bIsPaused : 1;
+};
diff --git a/Engine/Plugins/Runtime/StateTree/Source/StateTreeModule/Public/Components/StateTreeComponent.h b/Engine/Plugins/Runtime/StateTree/Source/StateTreeModule/Public/Components/StateTreeComponent.h
deleted file mode 100644
diff --git a/Engine/Plugins/Runtime/StateTree/Source/StateTreeModule/Public/Conditions/StateTreeCondition_Common.h b/Engine/Plugins/Runtime/StateTree/Source/StateTreeModule/Public/Conditions/StateTreeCondition_Common.h
--- a/Engine/Plugins/Runtime/StateTree/Source/StateTreeModule/Public/Conditions/StateTreeCondition_Common.h
+++ b/Engine/Plugins/Runtime/StateTree/Source/StateTreeModule/Public/Conditions/StateTreeCondition_Common.h
@@ -15,6 +15,19 @@ enum class EStateTreeCompare : uint8
 	Invert,
 };
 
+/**
+ * Base class for all common conditions so that they are easy to single out in the schema.
+ */
+USTRUCT()
+struct STATETREEMODULE_API FStateTreeConditionCommonBase : public FStateTreeConditionBase
+{
+	GENERATED_BODY()
+
+	FStateTreeConditionCommonBase() {}
+	virtual ~FStateTreeConditionCommonBase() override {}
+};
+template<> struct TStructOpsTypeTraits<FStateTreeConditionCommonBase> : public TStructOpsTypeTraitsBase2<FStateTreeConditionCommonBase> { enum { WithPureVirtual = true, }; };
+
 /**
  * Condition comparing two integers.
  */
@@ -31,7 +44,7 @@ struct STATETREEMODULE_API FStateTreeCondition_CompareIntInstanceData
 };
 
 USTRUCT(DisplayName="Integer Compare")
-struct STATETREEMODULE_API FStateTreeCondition_CompareInt : public FStateTreeConditionBase
+struct STATETREEMODULE_API FStateTreeCondition_CompareInt : public FStateTreeConditionCommonBase
 {
 	GENERATED_BODY()
 
@@ -77,7 +90,7 @@ struct STATETREEMODULE_API FStateTreeCondition_CompareFloatInstanceData
 };
 
 USTRUCT(DisplayName = "Float Compare")
-struct STATETREEMODULE_API FStateTreeCondition_CompareFloat : public FStateTreeConditionBase
+struct STATETREEMODULE_API FStateTreeCondition_CompareFloat : public FStateTreeConditionCommonBase
 {
 	GENERATED_BODY()
 
@@ -123,7 +136,7 @@ struct STATETREEMODULE_API FStateTreeCondition_CompareBoolInstanceData
 };
 
 USTRUCT(DisplayName = "Bool Compare")
-struct STATETREEMODULE_API FStateTreeCondition_CompareBool : public FStateTreeConditionBase
+struct STATETREEMODULE_API FStateTreeCondition_CompareBool : public FStateTreeConditionCommonBase
 {
 	GENERATED_BODY()
 
@@ -169,7 +182,7 @@ struct STATETREEMODULE_API FStateTreeCondition_CompareEnumInstanceData
 };
 
 USTRUCT(DisplayName = "Enum Compare")
-struct STATETREEMODULE_API FStateTreeCondition_CompareEnum : public FStateTreeConditionBase
+struct STATETREEMODULE_API FStateTreeCondition_CompareEnum : public FStateTreeConditionCommonBase
 {
 	GENERATED_BODY()
 
@@ -219,7 +232,7 @@ struct STATETREEMODULE_API FStateTreeCondition_CompareDistanceInstanceData
 };
 
 USTRUCT(DisplayName = "Distance Compare")
-struct STATETREEMODULE_API FStateTreeCondition_CompareDistance : public FStateTreeConditionBase
+struct STATETREEMODULE_API FStateTreeCondition_CompareDistance : public FStateTreeConditionCommonBase
 {
 	GENERATED_BODY()
 
diff --git a/Engine/Plugins/Runtime/StateTree/Source/StateTreeModule/Public/StateTree.h b/Engine/Plugins/Runtime/StateTree/Source/StateTreeModule/Public/StateTree.h
--- a/Engine/Plugins/Runtime/StateTree/Source/StateTreeModule/Public/StateTree.h
+++ b/Engine/Plugins/Runtime/StateTree/Source/StateTreeModule/Public/StateTree.h
@@ -21,6 +21,7 @@ class STATETREEMODULE_API UStateTree : public UDataAsset
 public:
 
 	UStateTree(const FObjectInitializer& ObjectInitializer = FObjectInitializer::Get());
+	~UStateTree();
 
 	/** @return Script Struct that can be used to instantiate the runtime storage */
 	const UScriptStruct* GetInstanceStorageStruct() const { return InstanceStorageStruct; }
@@ -50,6 +51,8 @@ public:
 	void ResolvePropertyPaths();
 
 #if WITH_EDITOR
+	void OnPIEStarted(const bool bIsSimulating);
+	
 	/** Resets the baked data to empty. */
 	void ResetBaked();
 #endif
@@ -57,7 +60,7 @@ public:
 #if WITH_EDITORONLY_DATA
 	// Edit time data for the StateTree, instance of UStateTreeEditorData
 	UPROPERTY()
-	UObject* EditorData;
+	TObjectPtr<UObject> EditorData;
 #endif
 
 protected:
@@ -70,7 +73,7 @@ protected:
 #endif
 
 	/** Initializes the types and default values related to StateTree runtime storage */
-	void InitInstanceStorage();
+	void InitInstanceStorageType();
 
 	/** Resolved references between data in the StateTree. */
 	void Link();
@@ -80,7 +83,7 @@ private:
 	// Properties
 
 	UPROPERTY(EditDefaultsOnly, Category = Common, Instanced)
-	UStateTreeSchema* Schema = nullptr;
+	TObjectPtr<UStateTreeSchema> Schema = nullptr;
 
 	/** Evaluators, Tasks, and Condition items */
 	UPROPERTY()
@@ -90,9 +93,13 @@ private:
 	UPROPERTY()
 	TArray<FInstancedStruct> Instances;
 
+	/** Blueprint based Evaluators, Tasks, and Conditions runtime data. */
+	UPROPERTY()
+	TArray<TObjectPtr<UObject>> InstanceObjects;
+
 	/** Script Struct that can be used to instantiate the runtime storage */
 	UPROPERTY()
-	UScriptStruct* InstanceStorageStruct;
+	TObjectPtr<UScriptStruct> InstanceStorageStruct;
 
 	/** Offsets into the runtime type to quickly get a struct view to a specific Task or Evaluator */
 	TArray<FStateTreeInstanceStorageOffset> InstanceStorageOffsets;
diff --git a/Engine/Plugins/Runtime/StateTree/Source/StateTreeModule/Public/StateTreeConditionBase.h b/Engine/Plugins/Runtime/StateTree/Source/StateTreeModule/Public/StateTreeConditionBase.h
--- a/Engine/Plugins/Runtime/StateTree/Source/StateTreeModule/Public/StateTreeConditionBase.h
+++ b/Engine/Plugins/Runtime/StateTree/Source/StateTreeModule/Public/StateTreeConditionBase.h
@@ -57,6 +57,12 @@ struct STATETREEMODULE_API FStateTreeConditionBase
 	/** The runtime data's data view index in the StateTreeExecutionContext, and source struct index in property binding. */
 	UPROPERTY()
 	uint16 DataViewIndex = 0;
+
+	UPROPERTY()
+	uint16 InstanceIndex = 0;
+
+	UPROPERTY()
+	uint8 bInstanceIsObject : 1;
 };
 
 template<>
diff --git a/Engine/Plugins/Runtime/StateTree/Source/StateTreeModule/Public/StateTreeEvaluatorBase.h b/Engine/Plugins/Runtime/StateTree/Source/StateTreeModule/Public/StateTreeEvaluatorBase.h
--- a/Engine/Plugins/Runtime/StateTree/Source/StateTreeModule/Public/StateTreeEvaluatorBase.h
+++ b/Engine/Plugins/Runtime/StateTree/Source/StateTreeModule/Public/StateTreeEvaluatorBase.h
@@ -15,7 +15,7 @@ struct FStateTreeExecutionContext;
  * Evaluators calculate and expose data to be used for decision making in a StateTree.
  */
 USTRUCT()
-struct STATETREEMODULE_API FStateTreeEvaluatorBase // TODO: change to FStateTreeEvaluatorBase once UStateTreeEvaluatorBase is removed
+struct STATETREEMODULE_API FStateTreeEvaluatorBase
 {
 	GENERATED_BODY()
 
@@ -88,6 +88,12 @@ struct STATETREEMODULE_API FStateTreeEvaluatorBase // TODO: change to FStateTree
 	/** The runtime data's data view index in the StateTreeExecutionContext, and source struct index in property binding. */
 	UPROPERTY()
 	uint16 DataViewIndex = 0;
+
+	UPROPERTY()
+	uint16 InstanceIndex = 0;
+
+	UPROPERTY()
+	uint8 bInstanceIsObject : 1;
 };
 
 template<> struct TStructOpsTypeTraits<FStateTreeEvaluatorBase> : public TStructOpsTypeTraitsBase2<FStateTreeEvaluatorBase> { enum { WithPureVirtual = true, }; };
diff --git a/Engine/Plugins/Runtime/StateTree/Source/StateTreeModule/Public/StateTreeExecutionContext.h b/Engine/Plugins/Runtime/StateTree/Source/StateTreeModule/Public/StateTreeExecutionContext.h
--- a/Engine/Plugins/Runtime/StateTree/Source/StateTreeModule/Public/StateTreeExecutionContext.h
+++ b/Engine/Plugins/Runtime/StateTree/Source/StateTreeModule/Public/StateTreeExecutionContext.h
@@ -31,12 +31,15 @@ struct STATETREEMODULE_API FStateTreeExecutionState
 
 	/** Running time of the delayed transition */
 	float GatedTransitionTime = 0.0f;
+
+	/** Object instances, ref counting handled manually. */
+	TArray<UObject*> InstanceObjects;
 };
 
 UENUM()
 enum class EStateTreeStorage : uint8
 {
-	/** Execution context has internal storage */
+	/** Execution context has internal storage */ 
 	Internal,
 	/** Execution context assumes external storage */
 	External,
@@ -58,13 +61,20 @@ public:
 	/** Initializes the StateTree instance to be used with specific owner and StateTree asset. */
 	bool Init(UObject& InOwner, const UStateTree& InStateTree, const EStateTreeStorage InStorageType);
 
+	/**  Initializes external storage (no need to call, if using internal storage). */
+	bool InitInstanceData(UObject& InOwner, FStateTreeDataView ExternalStorage) const;
+
 	/** Returns the StateTree asset in use. */
 	const UStateTree* GetStateTree() const { return StateTree; }
 
+	/** @return The owner of the context */
 	UObject* GetOwner() const { return Owner; }
-	UWorld* GetWorld() const { return World; };
-	void SetWorld(UWorld* InWorld) { World = InWorld; };
+	/** @return The world of the owner or nullptr if the owner is not set. */ 
+	UWorld* GetWorld() const { return Owner ? Owner->GetWorld() : nullptr; };
 
+	/** @return True of the the execution context is valid and initialized. */ 
+	bool IsValid() const { return Owner != nullptr && StateTree != nullptr; }
+	
 	/** Start executing. */
 	EStateTreeRunStatus Start(FStateTreeDataView ExternalStorage = FStateTreeDataView());
 	/** Stop executing. */
@@ -98,7 +108,7 @@ public:
 			if (DataDesc.Requirement == EStateTreeExternalDataRequirement::Required)
 			{
 				// Required items must have valid pointer and expected type.  
-				if (DataView.GetMemory() == nullptr || DataView.GetStruct() != DataDesc.Struct)
+				if (!DataView.IsValid() || !DataView.GetStruct()->IsChildOf(DataDesc.Struct))
 				{
 					bResult = false;
 					break;
@@ -107,7 +117,7 @@ public:
 			else
 			{
 				// Optional items must have same type if they are set.
-				if (DataView.IsValid() && DataView.GetStruct() != DataDesc.Struct)
+				if (DataView.IsValid() && !DataView.GetStruct()->IsChildOf(DataDesc.Struct))
 				{
 					bResult = false;
 					break;
@@ -160,11 +170,21 @@ public:
 		return Handle.IsValid() ? DataViews[Handle.DataViewIndex].template GetMutablePtr<typename T::DataType>() : nullptr;
 	}
 
+	FStateTreeDataView GetExternalDataView(const FStateTreeExternalDataHandle Handle)
+	{
+		check(StateTree);
+		if (Handle.IsValid())
+		{
+			return DataViews[Handle.DataViewIndex];
+		}
+		return FStateTreeDataView();
+	}
+
 	/**
-	* Returns reference to instance data property based on provided handle. The return type is deduced from the handle's template type.
-	* @param Handle Valid FStateTreeInstanceDataPropertyHandle<> handle.
-	* @return reference to instance data property based on handle.
-	*/ 
+	 * Returns reference to instance data property based on provided handle. The return type is deduced from the handle's template type.
+	 * @param Handle Valid FStateTreeInstanceDataPropertyHandle<> handle.
+	 * @return reference to instance data property based on handle.
+	 */ 
 	template <typename T>
 	typename T::DataType& GetInstanceData(const T Handle) const
 	{
@@ -174,10 +194,10 @@ public:
 	}
 
 	/**
-	* Returns pointer to instance data property based on provided handle. The return type is deduced from the handle's template type.
-	* @param Handle Valid FStateTreeInstanceDataPropertyHandle<> handle.
-	* @return pointer to instance data property based on handle or null if item is not set or handle is invalid.
-	*/ 
+	 * Returns pointer to instance data property based on provided handle. The return type is deduced from the handle's template type.
+	 * @param Handle Valid FStateTreeInstanceDataPropertyHandle<> handle.
+	 * @return pointer to instance data property based on handle or null if item is not set or handle is invalid.
+	 */ 
 	template <typename T>
 	typename T::DataType* GetInstanceDataPtr(const T Handle) const
 	{
@@ -185,6 +205,21 @@ public:
 		return Handle.IsValid() ? (typename T::DataType*)(DataViews[Handle.DataViewIndex].GetMemory() + Handle.PropertyOffset) : nullptr;
 	}
 
+	/**
+	 * Used internally by the Blueprint wrappers to get wrapped instance objects. 
+	 * @param DataViewIndex Index to a data view
+	 * @return Pointer to an instance object based.
+	 */
+	template <typename T>
+	T* GetInstanceObjectInternal(const int32 DataViewIndex) const
+	{
+		const UStruct* Struct = DataViews[DataViewIndex].GetStruct();
+		if (Struct != nullptr && Struct->IsChildOf<T>())
+		{
+			return DataViews[DataViewIndex].template GetMutablePtr<T>();
+		}
+		return nullptr;
+	}
 	
 	EStateTreeRunStatus GetLastTickStatus(FStateTreeDataView ExternalStorage = FStateTreeDataView()) const;
 	EStateTreeRunStatus GetEnterStateStatus() const { return EnterStateStatus; }
@@ -199,7 +234,10 @@ public:
 	
 	void DebugPrintInternalLayout(FStateTreeDataView ExternalStorage);
 #endif
-	
+
+	void AddStructReferencedObjects(class FReferenceCollector& Collector);
+	void AddStructReferencedObjects(const FStateTreeDataView ExternalStorage, class FReferenceCollector& Collector);
+
 protected:
 
 	/** @return Prefix that will be used by STATETREE_LOG and STATETREE_CLOG, empty by default. */
@@ -290,10 +328,18 @@ protected:
 	}
 
 	/** @return View to an Evaluator, a Task, or a Condition instance data. */
-	FStateTreeDataView GetInstanceData(FStateTreeDataView Storage, const int32 Index) const
+	FStateTreeDataView GetInstanceData(FStateTreeDataView Storage, const bool bIsObject, const int32 Index) const
 	{
-		const FStateTreeInstanceStorageOffset& ItemOffset = StateTree->InstanceStorageOffsets[Index];
-		return FStateTreeDataView(ItemOffset.Struct, Storage.GetMutableMemory() + ItemOffset.Offset);
+		if (UNLIKELY(bIsObject == true))
+		{
+			const FStateTreeExecutionState& Exec = GetExecState(Storage);
+			return FStateTreeDataView(Exec.InstanceObjects[Index]);
+		}
+		else
+		{
+			const FStateTreeInstanceStorageOffset& ItemOffset = StateTree->InstanceStorageOffsets[Index];
+			return FStateTreeDataView(ItemOffset.Struct, Storage.GetMutableMemory() + ItemOffset.Offset);
+		}
 	}
 
 	/** @return StateTree execution state from the instance storage. */
@@ -327,9 +373,6 @@ protected:
 	UPROPERTY()
 	UObject* Owner = nullptr;
 	
-	UPROPERTY()
-	UWorld* World = nullptr;
-
 	/** States visited during a tick while updating evaluators. Initialized to match the number of states in the asset. */ 
 	TArray<bool> VisitedStates;
 
@@ -349,3 +392,12 @@ protected:
 	 */
 	EStateTreeRunStatus EnterStateStatus = EStateTreeRunStatus::Unset;
 };
+
+template<>
+struct TStructOpsTypeTraits<FStateTreeExecutionContext> : public TStructOpsTypeTraitsBase2<FStateTreeExecutionContext>
+{
+	enum
+	{
+		WithAddStructReferencedObjects = true,
+	};
+};
diff --git a/Engine/Plugins/Runtime/StateTree/Source/StateTreeModule/Public/StateTreeSchema.h b/Engine/Plugins/Runtime/StateTree/Source/StateTreeModule/Public/StateTreeSchema.h
--- a/Engine/Plugins/Runtime/StateTree/Source/StateTreeModule/Public/StateTreeSchema.h
+++ b/Engine/Plugins/Runtime/StateTree/Source/StateTreeModule/Public/StateTreeSchema.h
@@ -25,8 +25,18 @@ public:
 	virtual UScriptStruct* GetStorageSuperStruct() const { return nullptr; }
 
 	/** @return True if specified struct is supported */
-	virtual bool IsStructAllowed(const UScriptStruct* InScriptStruct) const PURE_VIRTUAL(UStateTreeSchema::IsStructAllowed, return false; );
+	virtual bool IsStructAllowed(const UScriptStruct* InScriptStruct) const { return false; }
+
+	/** @return True if specified class is supported */
+	virtual bool IsClassAllowed(const UClass* InScriptStruct) const { return false; };
 
 	/** @return True if specified struct/class is supported as external data */
 	virtual bool IsExternalItemAllowed(const UStruct& InStruct) const { return false; };
+
+	/**
+	 * Helper function to check if a class is any of the Blueprint extendable item classes (Eval, Task, Condition).
+	 * Can be used to quickly accept all of those classes in IsClassAllowed().
+	 * @return True if the class is a StateTree item Blueprint base class.
+	 */
+	bool IsChildOfBlueprintBase(const UClass* InClass) const;
 };
diff --git a/Engine/Plugins/Runtime/StateTree/Source/StateTreeModule/Public/StateTreeTaskBase.h b/Engine/Plugins/Runtime/StateTree/Source/StateTreeModule/Public/StateTreeTaskBase.h
--- a/Engine/Plugins/Runtime/StateTree/Source/StateTreeModule/Public/StateTreeTaskBase.h
+++ b/Engine/Plugins/Runtime/StateTree/Source/StateTreeModule/Public/StateTreeTaskBase.h
@@ -86,6 +86,12 @@ struct STATETREEMODULE_API FStateTreeTaskBase
 	/** The runtime data's data view index in the StateTreeExecutionContext, and source struct index in property binding. */
 	UPROPERTY()
 	uint16 DataViewIndex = 0;
+
+	UPROPERTY()
+	uint16 InstanceIndex = 0;
+
+	UPROPERTY()
+	uint8 bInstanceIsObject : 1;
 };
 
 template<> struct TStructOpsTypeTraits<FStateTreeTaskBase> : public TStructOpsTypeTraitsBase2<FStateTreeTaskBase> { enum { WithPureVirtual = true, }; };
diff --git a/Engine/Plugins/Runtime/StateTree/Source/StateTreeModule/Public/StateTreeTypes.h b/Engine/Plugins/Runtime/StateTree/Source/StateTreeModule/Public/StateTreeTypes.h
--- a/Engine/Plugins/Runtime/StateTree/Source/StateTreeModule/Public/StateTreeTypes.h
+++ b/Engine/Plugins/Runtime/StateTree/Source/StateTreeModule/Public/StateTreeTypes.h
@@ -11,28 +11,16 @@ STATETREEMODULE_API DECLARE_LOG_CATEGORY_EXTERN(LogStateTree, Warning, All);
 #define WITH_STATETREE_DEBUG (!(UE_BUILD_SHIPPING || UE_BUILD_SHIPPING_WITH_EDITOR || UE_BUILD_TEST) && 1)
 #endif // WITH_STATETREE_DEBUG
 
-/** Types of supported variables */
-UENUM()
-enum class EStateTreeVariableType : uint8
-{
-	Void,
-	Float,
-	Int,
-	Bool,
-	Vector,
-	Object,
-};
-
 /**
  * Status describing current ticking state. 
  */
 UENUM()
 enum class EStateTreeRunStatus : uint8
 {
-	Unset,				/** Status not set. */
+	Running,			/** Tree is still running. */
 	Failed,				/** Tree execution has stopped on failure. */
 	Succeeded,			/** Tree execution has stopped on success. */
-	Running,			/** Tree is still running. */
+	Unset,				/** Status not set. */
 };
 
 /**  Evaluator evaluation type. */
@@ -79,7 +67,7 @@ ENUM_CLASS_FLAGS(EStateTreeTransitionEvent)
 /**
  * Handle that is used to refer baked state tree data.
  */
-USTRUCT()
+USTRUCT(BlueprintType)
 struct STATETREEMODULE_API FStateTreeHandle
 {
 	GENERATED_BODY()
@@ -119,7 +107,7 @@ struct STATETREEMODULE_API FStateTreeHandle
 /**
  * Describes current status of a running state or desired state.
  */
-USTRUCT()
+USTRUCT(BlueprintType)
 struct STATETREEMODULE_API FStateTreeStateStatus
 {
 	GENERATED_BODY()
@@ -131,10 +119,10 @@ struct STATETREEMODULE_API FStateTreeStateStatus
 
 	bool IsSet() const { return RunStatus != EStateTreeRunStatus::Unset; }
 	
-	UPROPERTY()
+	UPROPERTY(EditDefaultsOnly, Category = Default, BlueprintReadOnly)
 	FStateTreeHandle State = FStateTreeHandle::Invalid;
 
-	UPROPERTY()
+	UPROPERTY(EditDefaultsOnly, Category = Default, BlueprintReadOnly)
 	EStateTreeRunStatus RunStatus = EStateTreeRunStatus::Unset;
 };
 
@@ -143,7 +131,7 @@ struct STATETREEMODULE_API FStateTreeStateStatus
  * and Next describes the selected state. The reason Transition and Next are different is that Transition state can be a selector state,
  * in which case the children will be visited until a leaf state is found, which will be the next state.
  */
-USTRUCT()
+USTRUCT(BlueprintType)
 struct STATETREEMODULE_API FStateTreeTransitionResult
 {
 	GENERATED_BODY()
@@ -153,19 +141,19 @@ struct STATETREEMODULE_API FStateTreeTransitionResult
 	FStateTreeTransitionResult(const FStateTreeStateStatus InSource, const FStateTreeHandle InTransition, const FStateTreeHandle InNext) : Source(InSource), Target(InTransition), Next(InNext) {}
 
 	/** State where the transition started. */
-	UPROPERTY()
+	UPROPERTY(EditDefaultsOnly, Category = Default, BlueprintReadOnly)
 	FStateTreeStateStatus Source = FStateTreeStateStatus();
 
 	/** Transition target state */
-	UPROPERTY()
+	UPROPERTY(EditDefaultsOnly, Category = Default, BlueprintReadOnly)
 	FStateTreeHandle Target = FStateTreeHandle::Invalid;
 
 	/** Selected state, can be different from Transition, if Transition is a selector state. */
-	UPROPERTY()
+	UPROPERTY(EditDefaultsOnly, Category = Default, BlueprintReadOnly)
 	FStateTreeHandle Next = FStateTreeHandle::Invalid;
 
 	/** Current state, update as we execute the tree. */
-	UPROPERTY()
+	UPROPERTY(EditDefaultsOnly, Category = Default, BlueprintReadOnly)
 	FStateTreeHandle Current = FStateTreeHandle::Invalid;
 };
 
@@ -421,7 +409,7 @@ struct FStateTreeLinker
 	template <typename T>
 	typename TEnableIf<TIsDerivedFrom<typename T::DataType, UObject>::IsDerived, void>::Type LinkExternalData(T& Handle)
 	{
-		LinkExternalDataInternal(Handle, T::DataType::StaticClass(), T::DataRequirement);
+		LinkExternalData(Handle, T::DataType::StaticClass(), T::DataRequirement);
 	}
 
 	/**
@@ -431,9 +419,28 @@ struct FStateTreeLinker
 	template <typename T>
 	typename TEnableIf<!TIsDerivedFrom<typename T::DataType, UObject>::IsDerived, void>::Type LinkExternalData(T& Handle)
 	{
-		LinkExternalDataInternal(Handle, T::DataType::StaticStruct(), T::DataRequirement);
+		LinkExternalData(Handle, T::DataType::StaticStruct(), T::DataRequirement);
 	}
 
+	/**
+	 * Links reference to an external Object or Struct.
+	 * This function should only be used when TStateTreeExternalDataHandle<> cannot be used, i.e. the Struct is based on some data.
+	 * @param Handle Reference to link to.
+	 * @param Struct Expected type of the Object or Struct to link to.
+	 * @param Requirement Describes if the external data is expected to be required or optional.
+	 */
+	void LinkExternalData(FStateTreeExternalDataHandle& Handle, const UStruct* Struct, const EStateTreeExternalDataRequirement Requirement)
+	{
+		const FStateTreeExternalDataDesc Desc(Struct, Requirement);
+		int32 Index = ExternalDataDescs.Find(Desc);
+		if (Index == INDEX_NONE)
+		{
+			Index = ExternalDataDescs.Add(Desc);
+			check(FStateTreeExternalDataHandle::IsValidIndex(Index + ExternalDataBaseIndex));
+			ExternalDataDescs[Index].Handle.DataViewIndex = (uint8)(Index + ExternalDataBaseIndex);
+		}
+		Handle.DataViewIndex = (uint8)(Index + ExternalDataBaseIndex);
+	}
 	/**
 	 * Links reference to a property in instance data.
 	 * Usage:
@@ -456,18 +463,6 @@ struct FStateTreeLinker
 
 protected:
 
-	void LinkExternalDataInternal(FStateTreeExternalDataHandle& Handle, const UStruct* Struct, const EStateTreeExternalDataRequirement Requirement)
-	{
-		const FStateTreeExternalDataDesc Desc(Struct, Requirement);
-		int32 Index = ExternalDataDescs.Find(Desc);
-		if (Index == INDEX_NONE)
-		{
-			Index = ExternalDataDescs.Add(Desc);
-			check(FStateTreeExternalDataHandle::IsValidIndex(Index + ExternalDataBaseIndex));
-			ExternalDataDescs[Index].Handle.DataViewIndex = (uint8)(Index + ExternalDataBaseIndex);
-		}
-		Handle.DataViewIndex = (uint8)(Index + ExternalDataBaseIndex);
-	}
 
 	void LinkInstanceDataPropertyInternal(FStateTreeInstanceDataPropertyHandle& Handle, const UScriptStruct* ScriptStruct, const TCHAR* PropertyName)
 	{
diff --git a/Engine/Plugins/Runtime/StateTree/Source/StateTreeModule/StateTreeModule.Build.cs b/Engine/Plugins/Runtime/StateTree/Source/StateTreeModule/StateTreeModule.Build.cs
--- a/Engine/Plugins/Runtime/StateTree/Source/StateTreeModule/StateTreeModule.Build.cs
+++ b/Engine/Plugins/Runtime/StateTree/Source/StateTreeModule/StateTreeModule.Build.cs
@@ -7,30 +7,41 @@ namespace UnrealBuildTool.Rules
 		public StateTreeModule(ReadOnlyTargetRules Target) : base(Target)
 		{
 			PublicIncludePaths.AddRange(
-			new string[] {
-			}
+				new string[] {
+				}
 			);
 
 			PublicDependencyModuleNames.AddRange(
-			new string[] {
-				"Core",
-				"CoreUObject",
-				"Engine",
-				"AIModule",
-				"GameplayBehaviorsModule",
-				"GameplayTasks",
-				"GameplayTags",
-				"GameplayAbilities",
-				"SmartObjectsModule",
-				"StructUtils",
-			}
+				new string[] {
+					"Core",
+					"CoreUObject",
+					"Engine",
+					"AIModule",
+					"GameplayBehaviorsModule",
+					"GameplayTasks",
+					"GameplayTags",
+					"GameplayAbilities",
+					"SmartObjectsModule",
+					"StructUtils",
+				}
 			);
 
 			PrivateDependencyModuleNames.AddRange(
-			new string[] {
-				"RenderCore",
-			}
+				new string[] {
+					"RenderCore",
+				}
 			);
+			
+			if (Target.Type == TargetType.Editor)
+			{
+				PrivateDependencyModuleNames.AddRange(
+					new string[]
+					{
+						"UnrealEd"	// Editor callbacks
+					}
+				);
+			}
+
 		}
 
 	}
diff --git a/Engine/Plugins/Runtime/StateTree/Source/StateTreeModule/StateTreeSchema.cpp b/Engine/Plugins/Runtime/StateTree/Source/StateTreeModule/StateTreeSchema.cpp
new file mode 100644
--- /dev/null
+++ b/Engine/Plugins/Runtime/StateTree/Source/StateTreeModule/StateTreeSchema.cpp
@@ -0,0 +1,12 @@
+﻿// Copyright Epic Games, Inc. All Rights Reserved.
+
+#include "StateTreeSchema.h"
+#include "CoreMinimal.h"
+#include "Blueprint/StateTreeEvaluatorBlueprintBase.h"
+#include "Blueprint/StateTreeTaskBlueprintBase.h"
+#include "Blueprint/StateTreeConditionBlueprintBase.h"
+
+bool UStateTreeSchema::IsChildOfBlueprintBase(const UClass* InClass) const
+{
+	return InClass->IsChildOf<UStateTreeItemBlueprintBase>();
+}
