Fix NoPCH build errors.

diff --git a/Engine/Plugins/Bridge/Source/Bridge/Private/NodeProcessRunnableThread.cpp b/Engine/Plugins/Bridge/Source/Bridge/Private/NodeProcessRunnableThread.cpp
--- a/Engine/Plugins/Bridge/Source/Bridge/Private/NodeProcessRunnableThread.cpp
+++ b/Engine/Plugins/Bridge/Source/Bridge/Private/NodeProcessRunnableThread.cpp
@@ -1,5 +1,6 @@
 // Copyright Epic Games, Inc. All Rights Reserved.
 #include "NodeProcessRunnableThread.h"
+#include "HAL/RunnableThread.h"
 #include "Misc/Paths.h"
 
 NodeProcessRunnableThread::NodeProcessRunnableThread()
diff --git a/Engine/Plugins/Runtime/Metasound/Source/MetasoundEditor/Private/MetasoundDetailCustomization.cpp b/Engine/Plugins/Runtime/Metasound/Source/MetasoundEditor/Private/MetasoundDetailCustomization.cpp
--- a/Engine/Plugins/Runtime/Metasound/Source/MetasoundEditor/Private/MetasoundDetailCustomization.cpp
+++ b/Engine/Plugins/Runtime/Metasound/Source/MetasoundEditor/Private/MetasoundDetailCustomization.cpp
@@ -38,6 +38,7 @@
 #include "Styling/SlateTypes.h"
 #include "Styling/CoreStyle.h"
 #include "Widgets/Input/SCheckBox.h"
+#include "ScopedTransaction.h"
 
 
 #define LOCTEXT_NAMESPACE "MetaSoundEditor"
diff --git a/Engine/Plugins/Runtime/StateTree/Source/StateTreeModule/Public/Blueprint/StateTreeConditionBlueprintBase.h b/Engine/Plugins/Runtime/StateTree/Source/StateTreeModule/Public/Blueprint/StateTreeConditionBlueprintBase.h
--- a/Engine/Plugins/Runtime/StateTree/Source/StateTreeModule/Public/Blueprint/StateTreeConditionBlueprintBase.h
+++ b/Engine/Plugins/Runtime/StateTree/Source/StateTreeModule/Public/Blueprint/StateTreeConditionBlueprintBase.h
@@ -1,8 +1,9 @@
-﻿// Copyright Epic Games, Inc. All Rights Reserved.
+// Copyright Epic Games, Inc. All Rights Reserved.
 
 #pragma once
 
 #include "CoreMinimal.h"
+#include "Templates/SubclassOf.h"
 #include "StateTreeTypes.h"
 #include "StateTreeConditionBase.h"
 #include "StateTreePropertyBindings.h"
diff --git a/Engine/Plugins/Runtime/StateTree/Source/StateTreeModule/Public/Blueprint/StateTreeEvaluatorBlueprintBase.h b/Engine/Plugins/Runtime/StateTree/Source/StateTreeModule/Public/Blueprint/StateTreeEvaluatorBlueprintBase.h
--- a/Engine/Plugins/Runtime/StateTree/Source/StateTreeModule/Public/Blueprint/StateTreeEvaluatorBlueprintBase.h
+++ b/Engine/Plugins/Runtime/StateTree/Source/StateTreeModule/Public/Blueprint/StateTreeEvaluatorBlueprintBase.h
@@ -1,8 +1,9 @@
-﻿// Copyright Epic Games, Inc. All Rights Reserved.
+// Copyright Epic Games, Inc. All Rights Reserved.
 
 #pragma once
 
 #include "CoreMinimal.h"
+#include "Templates/SubclassOf.h"
 #include "StateTreeTypes.h"
 #include "StateTreeEvaluatorBase.h"
 #include "StateTreePropertyBindings.h"
