StateTree: fixed class filter in StateTree bindable item bindings

diff --git a/Engine/Plugins/Runtime/StateTree/Source/StateTreeEditorModule/Private/Customizations/StateTreeBindableItemDetails.cpp b/Engine/Plugins/Runtime/StateTree/Source/StateTreeEditorModule/Private/Customizations/StateTreeBindableItemDetails.cpp
--- a/Engine/Plugins/Runtime/StateTree/Source/StateTreeEditorModule/Private/Customizations/StateTreeBindableItemDetails.cpp
+++ b/Engine/Plugins/Runtime/StateTree/Source/StateTreeEditorModule/Private/Customizations/StateTreeBindableItemDetails.cpp
@@ -244,19 +244,19 @@ public:
 			return false;
 		}
 
-		const UClass* Class = InUnloadedClassData->GetClassWithin();
-		
-		if (!Schema->IsClassAllowed(Class))
+		if (InUnloadedClassData->HasAnyClassFlags(CLASS_Hidden | CLASS_HideDropDown | CLASS_Deprecated | CLASS_Abstract))
 		{
 			return false;
 		}
-		
-		if (Class == BaseClass)
+
+		const UClass* NativeParentClass = InUnloadedClassData->GetNativeParent();
+
+		if (!Schema->IsClassAllowed(NativeParentClass))
 		{
-			return bAllowBaseClass;
+			return false;
 		}
-
-		return !BaseClass || Class->IsChildOf(BaseClass);
+		
+		return NativeParentClass->IsChildOf(BaseClass);
 	}
 };
 
