StateTree tweaks
- removed pure virtuals from StateTree nodes
- cosmetic change when StateTree current state is set

diff --git a/Engine/Plugins/AI/MassAI/Source/MassAIBehavior/Private/MassStateTreeSchema.cpp b/Engine/Plugins/AI/MassAI/Source/MassAIBehavior/Private/MassStateTreeSchema.cpp
--- a/Engine/Plugins/AI/MassAI/Source/MassAIBehavior/Private/MassStateTreeSchema.cpp
+++ b/Engine/Plugins/AI/MassAI/Source/MassAIBehavior/Private/MassStateTreeSchema.cpp
@@ -14,7 +14,7 @@ UScriptStruct* UMassStateTreeSchema::GetStorageSuperStruct() const
 
 bool UMassStateTreeSchema::IsStructAllowed(const UScriptStruct* InScriptStruct) const
 {
-	// Only allow Mass evals and tasks.
+	// Only allow Mass evals and tasks,and common conditions.
 	return InScriptStruct->IsChildOf(FMassStateTreeEvaluatorBase::StaticStruct())
 			|| InScriptStruct->IsChildOf(FMassStateTreeTaskBase::StaticStruct())
 			|| InScriptStruct->IsChildOf(FStateTreeConditionBase::StaticStruct());
diff --git a/Engine/Plugins/AI/MassAI/Source/MassAIBehavior/Public/MassStateTreeTypes.h b/Engine/Plugins/AI/MassAI/Source/MassAIBehavior/Public/MassStateTreeTypes.h
--- a/Engine/Plugins/AI/MassAI/Source/MassAIBehavior/Public/MassStateTreeTypes.h
+++ b/Engine/Plugins/AI/MassAI/Source/MassAIBehavior/Public/MassStateTreeTypes.h
@@ -28,11 +28,7 @@ USTRUCT(meta = (DisplayName = "Mass Evaluator Base"))
 struct MASSAIBEHAVIOR_API FMassStateTreeEvaluatorBase : public FStateTreeEvaluatorBase
 {
 	GENERATED_BODY()
-
-	FMassStateTreeEvaluatorBase() {}
-	virtual ~FMassStateTreeEvaluatorBase() {}
 };
-template<> struct TStructOpsTypeTraits<FMassStateTreeEvaluatorBase> : public TStructOpsTypeTraitsBase2<FMassStateTreeEvaluatorBase> { enum { WithPureVirtual = true, }; };
 
 /**
  * Base struct for all Mass StateTree Tasks.
@@ -41,11 +37,7 @@ USTRUCT(meta = (DisplayName = "Mass Task Base"))
 struct MASSAIBEHAVIOR_API FMassStateTreeTaskBase : public FStateTreeTaskBase
 {
 	GENERATED_BODY()
-
-	FMassStateTreeTaskBase() {}
-	virtual ~FMassStateTreeTaskBase() {}
 };
-template<> struct TStructOpsTypeTraits<FMassStateTreeTaskBase> : public TStructOpsTypeTraitsBase2<FMassStateTreeTaskBase> { enum { WithPureVirtual = true, }; };
 
 /**
 * A handle pointing to a registered StateTree asset in UMassStateTreeSubsystem.
diff --git a/Engine/Plugins/Runtime/StateTree/Source/StateTreeModule/Public/Conditions/StateTreeCondition_Common.h b/Engine/Plugins/Runtime/StateTree/Source/StateTreeModule/Public/Conditions/StateTreeCondition_Common.h
--- a/Engine/Plugins/Runtime/StateTree/Source/StateTreeModule/Public/Conditions/StateTreeCondition_Common.h
+++ b/Engine/Plugins/Runtime/StateTree/Source/StateTreeModule/Public/Conditions/StateTreeCondition_Common.h
@@ -22,11 +22,7 @@ USTRUCT()
 struct STATETREEMODULE_API FStateTreeConditionCommonBase : public FStateTreeConditionBase
 {
 	GENERATED_BODY()
-
-	FStateTreeConditionCommonBase() {}
-	virtual ~FStateTreeConditionCommonBase() override {}
 };
-template<> struct TStructOpsTypeTraits<FStateTreeConditionCommonBase> : public TStructOpsTypeTraitsBase2<FStateTreeConditionCommonBase> { enum { WithPureVirtual = true, }; };
 
 /**
  * Condition comparing two integers.
diff --git a/Engine/Plugins/Runtime/StateTree/Source/StateTreeModule/Public/StateTreeConditionBase.h b/Engine/Plugins/Runtime/StateTree/Source/StateTreeModule/Public/StateTreeConditionBase.h
--- a/Engine/Plugins/Runtime/StateTree/Source/StateTreeModule/Public/StateTreeConditionBase.h
+++ b/Engine/Plugins/Runtime/StateTree/Source/StateTreeModule/Public/StateTreeConditionBase.h
@@ -3,13 +3,16 @@
 #pragma once
 
 #include "CoreMinimal.h"
-#include "EngineDefines.h"
 #include "StateTreeTypes.h"
+#if WITH_EDITOR
 #include "InstancedStruct.h"
+#endif
 #include "StateTreeConditionBase.generated.h"
 
+#if WITH_EDITOR
 struct IStateTreeBindingLookup;
 struct FStateTreeEditorPropertyPath;
+#endif
 struct FStateTreeExecutionContext;
 
 /**
@@ -22,11 +25,9 @@ struct STATETREEMODULE_API FStateTreeConditionBase
 
 	virtual ~FStateTreeConditionBase() {}
 
-	/**
-	* @return Struct that represents the runtime data of the condition.
-	*/
-	virtual const UStruct* GetInstanceDataType() const PURE_VIRTUAL(FStateTreeConditionBase::GetInstanceDataType(), return nullptr;);
-	
+	/** @return Struct that represents the runtime data of the evaluator. */
+	virtual const UStruct* GetInstanceDataType() const { return nullptr; };
+
 	/**
 	 * Called when the StateTree asset is linked. Allows to resolve references to other StateTree data.
 	 * @see TStateTreeExternalDataHandle
@@ -64,12 +65,3 @@ struct STATETREEMODULE_API FStateTreeConditionBase
 	UPROPERTY()
 	uint8 bInstanceIsObject : 1;
 };
-
-template<>
-struct TStructOpsTypeTraits<FStateTreeConditionBase> : public TStructOpsTypeTraitsBase2<FStateTreeConditionBase>
-{
-	enum
-	{
-		WithPureVirtual = true,
-	};
-};
diff --git a/Engine/Plugins/Runtime/StateTree/Source/StateTreeModule/Public/StateTreeEvaluatorBase.h b/Engine/Plugins/Runtime/StateTree/Source/StateTreeModule/Public/StateTreeEvaluatorBase.h
--- a/Engine/Plugins/Runtime/StateTree/Source/StateTreeModule/Public/StateTreeEvaluatorBase.h
+++ b/Engine/Plugins/Runtime/StateTree/Source/StateTreeModule/Public/StateTreeEvaluatorBase.h
@@ -3,9 +3,7 @@
 #pragma once
 
 #include "CoreMinimal.h"
-#include "Misc/Guid.h"
 #include "StateTreeTypes.h"
-#include "StateTreePropertyBindings.h"
 #include "StateTreeEvaluatorBase.generated.h"
 
 struct FStateTreeExecutionContext;
@@ -23,10 +21,8 @@ struct STATETREEMODULE_API FStateTreeEvaluatorBase
 
 	virtual ~FStateTreeEvaluatorBase() {}
 
-	/**
-	* @return Struct that represents the runtime data of the evaluator.
-	*/
-	virtual const UStruct* GetInstanceDataType() const PURE_VIRTUAL(FStateTreeEvaluatorBase::GetInstanceDataType(), return nullptr;);
+	/** @return Struct that represents the runtime data of the evaluator. */
+	virtual const UStruct* GetInstanceDataType() const { return nullptr; };
 
 	/**
 	 * Called when the StateTree asset is linked. Allows to resolve references to other StateTree data.
@@ -95,5 +91,3 @@ struct STATETREEMODULE_API FStateTreeEvaluatorBase
 	UPROPERTY()
 	uint8 bInstanceIsObject : 1;
 };
-
-template<> struct TStructOpsTypeTraits<FStateTreeEvaluatorBase> : public TStructOpsTypeTraitsBase2<FStateTreeEvaluatorBase> { enum { WithPureVirtual = true, }; };
diff --git a/Engine/Plugins/Runtime/StateTree/Source/StateTreeModule/Public/StateTreeTaskBase.h b/Engine/Plugins/Runtime/StateTree/Source/StateTreeModule/Public/StateTreeTaskBase.h
--- a/Engine/Plugins/Runtime/StateTree/Source/StateTreeModule/Public/StateTreeTaskBase.h
+++ b/Engine/Plugins/Runtime/StateTree/Source/StateTreeModule/Public/StateTreeTaskBase.h
@@ -3,7 +3,6 @@
 #pragma once
 
 #include "CoreMinimal.h"
-#include "Misc/Guid.h"
 #include "StateTreeTypes.h"
 #include "StateTreeTaskBase.generated.h"
 
@@ -22,10 +21,8 @@ struct STATETREEMODULE_API FStateTreeTaskBase
 
 	virtual ~FStateTreeTaskBase() {}
 
-	/**
-	* @return Struct that represents the runtime data of the evaluator.
-	*/
-	virtual const UStruct* GetInstanceDataType() const PURE_VIRTUAL(FStateTreeTaskBase::GetInstanceDataType(), return nullptr;);
+	/** @return Struct that represents the runtime data of the evaluator. */
+	virtual const UStruct* GetInstanceDataType() const { return nullptr; };
 
 	/**
 	 * Called when the StateTree asset is linked. Allows to resolve references to other StateTree data.
@@ -65,11 +62,11 @@ struct STATETREEMODULE_API FStateTreeTaskBase
 	virtual void StateCompleted(FStateTreeExecutionContext& Context, const EStateTreeRunStatus CompletionStatus, const FStateTreeHandle CompletedState) const {}
 
 	/**
-	* Called during state tree tick when the task is on active state.
-	* @param Context Reference to current execution context.
-	* @param DeltaTime Time since last StateTree tick.
-	* @return Running status of the state: Running if still in progress, Succeeded if execution is done and succeeded, Failed if execution is done and failed.
-	*/
+	 * Called during state tree tick when the task is on active state.
+	 * @param Context Reference to current execution context.
+	 * @param DeltaTime Time since last StateTree tick.
+	 * @return Running status of the state: Running if still in progress, Succeeded if execution is done and succeeded, Failed if execution is done and failed.
+	 */
 	virtual EStateTreeRunStatus Tick(FStateTreeExecutionContext& Context, const float DeltaTime) const { return EStateTreeRunStatus::Running; };
 
 #if WITH_GAMEPLAY_DEBUGGER
@@ -93,5 +90,3 @@ struct STATETREEMODULE_API FStateTreeTaskBase
 	UPROPERTY()
 	uint8 bInstanceIsObject : 1;
 };
-
-template<> struct TStructOpsTypeTraits<FStateTreeTaskBase> : public TStructOpsTypeTraitsBase2<FStateTreeTaskBase> { enum { WithPureVirtual = true, }; };
