StateTree UI improvements
- fixed description display on BP conditions
- added class cache that contains also BP classes and is kept up to date as the assets and BPs change
- the condition/evaluator/task selector now shows struct and class based items in same menu
- added meta = (Hidden) support for structs to prevent them displaying in the item selector

diff --git a/Engine/Plugins/AI/MassAI/Source/MassAIBehavior/Private/Conditions/ZoneGraphTagConditions.cpp b/Engine/Plugins/AI/MassAI/Source/MassAIBehavior/Private/Conditions/ZoneGraphTagConditions.cpp
--- a/Engine/Plugins/AI/MassAI/Source/MassAIBehavior/Private/Conditions/ZoneGraphTagConditions.cpp
+++ b/Engine/Plugins/AI/MassAI/Source/MassAIBehavior/Private/Conditions/ZoneGraphTagConditions.cpp
@@ -98,7 +98,7 @@ bool FZoneGraphTagFilterCondition::TestCondition(FStateTreeExecutionContext& Con
 }
 
 #if WITH_EDITOR
-FText FZoneGraphTagFilterCondition::GetDescription(const FGuid& ID, FConstStructView InstanceData, const IStateTreeBindingLookup& BindingLookup) const
+FText FZoneGraphTagFilterCondition::GetDescription(const FGuid& ID, FStateTreeDataView InstanceData, const IStateTreeBindingLookup& BindingLookup) const
 {
 	const FStateTreeEditorPropertyPath LeftPath(ID, GET_MEMBER_NAME_STRING_CHECKED(FZoneGraphTagFilterConditionInstanceData, Tags));
 
@@ -164,7 +164,7 @@ bool FZoneGraphTagMaskCondition::TestCondition(FStateTreeExecutionContext& Conte
 }
 
 #if WITH_EDITOR
-FText FZoneGraphTagMaskCondition::GetDescription(const FGuid& ID, FConstStructView InstanceData, const IStateTreeBindingLookup& BindingLookup) const
+FText FZoneGraphTagMaskCondition::GetDescription(const FGuid& ID, FStateTreeDataView InstanceData, const IStateTreeBindingLookup& BindingLookup) const
 {
 	const FZoneGraphTagMaskConditionInstanceData& Instance = InstanceData.Get<FZoneGraphTagMaskConditionInstanceData>();
 	const FStateTreeEditorPropertyPath LeftPath(ID, GET_MEMBER_NAME_STRING_CHECKED(FZoneGraphTagMaskConditionInstanceData, Left));
@@ -224,7 +224,7 @@ bool FZoneGraphTagCondition::TestCondition(FStateTreeExecutionContext& Context)
 }
 
 #if WITH_EDITOR
-FText FZoneGraphTagCondition::GetDescription(const FGuid& ID, FConstStructView InstanceData, const IStateTreeBindingLookup& BindingLookup) const
+FText FZoneGraphTagCondition::GetDescription(const FGuid& ID, FStateTreeDataView InstanceData, const IStateTreeBindingLookup& BindingLookup) const
 {
 	const FZoneGraphTagConditionInstanceData& Instance = InstanceData.Get<FZoneGraphTagConditionInstanceData>();
 	const FStateTreeEditorPropertyPath LeftPath(ID, GET_MEMBER_NAME_STRING_CHECKED(FZoneGraphTagConditionInstanceData, Left));
diff --git a/Engine/Plugins/AI/MassAI/Source/MassAIBehavior/Public/Conditions/ZoneGraphTagConditions.h b/Engine/Plugins/AI/MassAI/Source/MassAIBehavior/Public/Conditions/ZoneGraphTagConditions.h
--- a/Engine/Plugins/AI/MassAI/Source/MassAIBehavior/Public/Conditions/ZoneGraphTagConditions.h
+++ b/Engine/Plugins/AI/MassAI/Source/MassAIBehavior/Public/Conditions/ZoneGraphTagConditions.h
@@ -32,7 +32,7 @@ struct MASSAIBEHAVIOR_API FZoneGraphTagFilterCondition : public FStateTreeCondit
 	virtual bool TestCondition(FStateTreeExecutionContext& Context) const override;
 
 #if WITH_EDITOR
-	virtual FText GetDescription(const FGuid& ID, FConstStructView InstanceData, const IStateTreeBindingLookup& BindingLookup) const override;
+	virtual FText GetDescription(const FGuid& ID, FStateTreeDataView InstanceData, const IStateTreeBindingLookup& BindingLookup) const override;
 #endif
 
 	TStateTreeInstanceDataPropertyHandle<FZoneGraphTagMask> TagsHandle;
@@ -72,7 +72,7 @@ struct MASSAIBEHAVIOR_API FZoneGraphTagMaskCondition : public FStateTreeConditio
 	virtual bool TestCondition(FStateTreeExecutionContext& Context) const override;
 
 #if WITH_EDITOR
-	virtual FText GetDescription(const FGuid& ID, FConstStructView InstanceData, const IStateTreeBindingLookup& BindingLookup) const override;
+	virtual FText GetDescription(const FGuid& ID, FStateTreeDataView InstanceData, const IStateTreeBindingLookup& BindingLookup) const override;
 #endif
 
 	TStateTreeInstanceDataPropertyHandle<FZoneGraphTagMask> LeftHandle;
@@ -113,7 +113,7 @@ struct MASSAIBEHAVIOR_API FZoneGraphTagCondition : public FStateTreeConditionBas
 	virtual bool TestCondition(FStateTreeExecutionContext& Context) const override;
 
 #if WITH_EDITOR
-	virtual FText GetDescription(const FGuid& ID, FConstStructView InstanceData, const IStateTreeBindingLookup& BindingLookup) const override;
+	virtual FText GetDescription(const FGuid& ID, FStateTreeDataView InstanceData, const IStateTreeBindingLookup& BindingLookup) const override;
 #endif
 
 	TStateTreeInstanceDataPropertyHandle<FZoneGraphTag> LeftHandle;
diff --git a/Engine/Plugins/Runtime/StateTree/Source/StateTreeEditorModule/Private/Customizations/StateTreeBindableItemDetails.cpp b/Engine/Plugins/Runtime/StateTree/Source/StateTreeEditorModule/Private/Customizations/StateTreeBindableItemDetails.cpp
--- a/Engine/Plugins/Runtime/StateTree/Source/StateTreeEditorModule/Private/Customizations/StateTreeBindableItemDetails.cpp
+++ b/Engine/Plugins/Runtime/StateTree/Source/StateTreeEditorModule/Private/Customizations/StateTreeBindableItemDetails.cpp
@@ -18,10 +18,6 @@
 #include "Styling/SlateIconFinder.h"
 #include "Editor.h"
 #include "Modules/ModuleManager.h"
-#include "StructViewerModule.h"
-#include "StructViewerFilter.h"
-#include "ClassViewerModule.h"
-#include "ClassViewerFilter.h"
 #include "InstancedStructDetails.h"
 #include "Engine/UserDefinedStruct.h"
 #include "StateTreeBindingExtension.h"
@@ -36,6 +32,8 @@
 #include "Blueprint/StateTreeEvaluatorBlueprintBase.h"
 #include "Blueprint/StateTreeTaskBlueprintBase.h"
 #include "Blueprint/StateTreeConditionBlueprintBase.h"
+#include "StateTreeEditorModule.h"
+#include "StateTreeNodeClassCache.h"
 
 #define LOCTEXT_NAMESPACE "StateTreeEditor"
 
@@ -154,114 +152,6 @@ public:
 
 ////////////////////////////////////
 
-class FBindableItemStructFilter : public IStructViewerFilter
-{
-public:
-	/** The base struct for the property that classes must be a child-of. */
-	const UScriptStruct* BaseStruct = nullptr;
-
-	// A flag controlling whether we allow UserDefinedStructs
-	bool bAllowUserDefinedStructs = false;
-
-	// A flag controlling whether we allow to select the BaseStruct
-	bool bAllowBaseStruct = true;
-
-	// Schema to filter which structs are allowed.
-	const UStateTreeSchema* Schema = nullptr;
-
-	virtual bool IsStructAllowed(const FStructViewerInitializationOptions& InInitOptions, const UScriptStruct* InStruct, TSharedRef<FStructViewerFilterFuncs> InFilterFuncs) override
-	{
-		if (!Schema)
-		{
-			return false;
-		}
-
-		if (!Schema->IsStructAllowed(InStruct))
-		{
-			return false;
-		}
-		
-		if (InStruct->IsA<UUserDefinedStruct>())
-		{
-			return bAllowUserDefinedStructs;
-		}
-
-		if (InStruct == BaseStruct)
-		{
-			return bAllowBaseStruct;
-		}
-
-		// Query the native struct to see if it has the correct parent type (if any)
-		return BaseStruct != nullptr && InStruct->IsChildOf(BaseStruct);
-	}
-
-	virtual bool IsUnloadedStructAllowed(const FStructViewerInitializationOptions& InInitOptions, const FName InStructPath, TSharedRef<FStructViewerFilterFuncs> InFilterFuncs) override
-	{
-		// User Defined Structs don't support inheritance, so only include them requested
-		return bAllowUserDefinedStructs;
-	}
-};
-
-////////////////////////////////////
-
-class FBindableItemClassFilter : public IClassViewerFilter
-{
-public:
-	/** The base struct for the property that classes must be a child-of. */
-	const UClass* BaseClass = nullptr;
-
-	// A flag controlling whether we allow to select the BaseStruct
-	bool bAllowBaseClass = true;
-
-	// Schema to filter which structs are allowed.
-	const UStateTreeSchema* Schema = nullptr;
-
-	virtual bool IsClassAllowed(const FClassViewerInitializationOptions& InInitOptions, const UClass* InClass, TSharedRef<FClassViewerFilterFuncs> InFilterFuncs ) override
-	{
-		if (!Schema || !InClass)
-		{
-			return false;
-		}
-
-		if (!Schema->IsClassAllowed(InClass))
-		{
-			return false;
-		}
-		
-		if (InClass == BaseClass)
-		{
-			return bAllowBaseClass;
-		}
-
-		// Query the native struct to see if it has the correct parent type (if any)
-		return BaseClass != nullptr && InClass->IsChildOf(BaseClass);
-	}
-
-	virtual bool IsUnloadedClassAllowed(const FClassViewerInitializationOptions& InInitOptions, const TSharedRef<const IUnloadedBlueprintData> InUnloadedClassData, TSharedRef<FClassViewerFilterFuncs> InFilterFuncs) override
-	{
-		if (!Schema)
-		{
-			return false;
-		}
-
-		if (InUnloadedClassData->HasAnyClassFlags(CLASS_Hidden | CLASS_HideDropDown | CLASS_Deprecated | CLASS_Abstract))
-		{
-			return false;
-		}
-
-		const UClass* NativeParentClass = InUnloadedClassData->GetNativeParent();
-
-		if (!Schema->IsClassAllowed(NativeParentClass))
-		{
-			return false;
-		}
-		
-		return NativeParentClass->IsChildOf(BaseClass);
-	}
-};
-
-////////////////////////////////////
-
 TSharedRef<IPropertyTypeCustomization> FStateTreeBindableItemDetails::MakeInstance()
 {
 	return MakeShareable(new FStateTreeBindableItemDetails);
@@ -521,7 +411,7 @@ void FStateTreeBindableItemDetails::OnBindingChanged(const FStateTreeEditorPrope
 					const FStateTreeBindingLookup BindingLookup(EditorData);
 
 					OuterObject->Modify();
-					Condition->OnBindingChanged(Item->ID, Item->Instance, SourcePath, TargetPath, BindingLookup);
+					Condition->OnBindingChanged(Item->ID, FStateTreeDataView(Item->Instance), SourcePath, TargetPath, BindingLookup);
 				}
 			}
 		}
@@ -564,12 +454,23 @@ FText FStateTreeBindableItemDetails::GetDescription() const
 	}
 
 	FStateTreeItem* Item = static_cast<FStateTreeItem*>(RawItemData[0]);
-	if (Item != nullptr && EditorData != nullptr && Item->Item.IsValid() && Item->Instance.IsValid())
+	if (Item != nullptr && EditorData != nullptr && Item->Item.IsValid())
 	{
 		if (const FStateTreeConditionBase* Condition = Item->Item.GetPtr<const FStateTreeConditionBase>())
 		{
 			const FStateTreeBindingLookup BindingLookup(EditorData);
-			return Condition->GetDescription(Item->ID, Item->Instance, BindingLookup);
+			if (Item->Instance.IsValid())
+			{
+				return Condition->GetDescription(Item->ID, FStateTreeDataView(Item->Instance), BindingLookup);
+			}
+			else if (Item->InstanceObject != nullptr)
+			{
+				return Condition->GetDescription(Item->ID, FStateTreeDataView(Item->InstanceObject), BindingLookup);
+			}
+			else
+			{
+				return Item->Item.GetScriptStruct()->GetDisplayNameText();
+			}
 		}
 	}
 	
@@ -578,7 +479,12 @@ FText FStateTreeBindableItemDetails::GetDescription() const
 
 EVisibility FStateTreeBindableItemDetails::IsDescriptionVisible() const
 {
-	const UScriptStruct* ScriptStruct = GetCommonItemScriptStruct();
+	const UScriptStruct* ScriptStruct = nullptr;
+	if (const FStateTreeItem* Item = GetCommonItem())
+	{
+		ScriptStruct = Item->Item.GetScriptStruct();
+	}
+
 	return ScriptStruct != nullptr && ScriptStruct->IsChildOf(FStateTreeConditionBase::StaticStruct()) ? EVisibility::Visible : EVisibility::Collapsed;
 }
 
@@ -616,7 +522,12 @@ FText FStateTreeBindableItemDetails::GetName() const
 
 EVisibility FStateTreeBindableItemDetails::IsNameVisible() const
 {
-	const UScriptStruct* ScriptStruct = GetCommonItemScriptStruct();
+	const UScriptStruct* ScriptStruct = nullptr;
+	if (const FStateTreeItem* Item = GetCommonItem())
+	{
+		ScriptStruct = Item->Item.GetScriptStruct();
+	}
+	
 	return ScriptStruct != nullptr && (ScriptStruct->IsChildOf(FStateTreeTaskBase::StaticStruct()) || ScriptStruct->IsChildOf(FStateTreeEvaluatorBase::StaticStruct())) ? EVisibility::Visible : EVisibility::Collapsed;
 }
 
@@ -677,10 +588,11 @@ void FStateTreeBindableItemDetails::OnNameCommitted(const FText& NewText, ETextC
 bool FStateTreeBindableItemDetails::IsNameEnabled() const
 {
 	// Can only edit if we have valid instantiated type.
-	return GetCommonItemScriptStruct() != nullptr;
+	const FStateTreeItem* Item = GetCommonItem();
+	return Item && Item->Item.IsValid();
 }
 
-const UScriptStruct* FStateTreeBindableItemDetails::GetCommonItemScriptStruct() const
+const FStateTreeItem* FStateTreeBindableItemDetails::GetCommonItem() const
 {
 	check(StructProperty);
 
@@ -689,38 +601,70 @@ const UScriptStruct* FStateTreeBindableItemDetails::GetCommonItemScriptStruct()
 	TArray<void*> RawItemData;
 	StructProperty->AccessRawData(RawItemData);
 
+	const FStateTreeItem* CommonItem = nullptr;
+	
 	for (void* Data : RawItemData)
 	{
-		if (FStateTreeItem* Item = static_cast<FStateTreeItem*>(Data))
+		if (const FStateTreeItem* Item = static_cast<FStateTreeItem*>(Data))
 		{
-			const UScriptStruct* ScriptStruct = Item->Item.GetScriptStruct();
-			if (!bMultipleValues && !CommonScriptStruct)
+			if (!bMultipleValues && !CommonItem)
 			{
-				CommonScriptStruct = ScriptStruct;
+				CommonItem = Item;
 			}
-			else if (ScriptStruct != CommonScriptStruct)
+			else if (CommonItem != Item)
 			{
-				CommonScriptStruct = nullptr;
+				CommonItem = nullptr;
 				bMultipleValues = true;
 			}
 		}
 	}
 
-	return CommonScriptStruct;
+	return CommonItem;
 }
 
 FText FStateTreeBindableItemDetails::GetDisplayValueString() const
 {
-	const UScriptStruct* ScriptStruct = GetCommonItemScriptStruct();
-	if (ScriptStruct)
+	if (const FStateTreeItem* Item = GetCommonItem())
 	{
-		return ScriptStruct->GetDisplayNameText();
+		if (const UScriptStruct* ScriptStruct = Item->Item.GetScriptStruct())
+		{
+			if (ScriptStruct->IsChildOf(FStateTreeBlueprintEvaluatorWrapper::StaticStruct())
+				|| ScriptStruct->IsChildOf(FStateTreeBlueprintTaskWrapper::StaticStruct())
+				|| ScriptStruct->IsChildOf(FStateTreeBlueprintConditionWrapper::StaticStruct()))
+			{
+				if (Item->InstanceObject != nullptr
+					&& Item->InstanceObject->GetClass() != nullptr)
+				{
+					return Item->InstanceObject->GetClass()->GetDisplayNameText();
+				}
+			}
+			else
+			{
+				return ScriptStruct->GetDisplayNameText();
+			}
+		}
 	}
 	return FText();
 }
 
 const FSlateBrush* FStateTreeBindableItemDetails::GetDisplayValueIcon() const
 {
+	if (const FStateTreeItem* Item = GetCommonItem())
+	{
+		if (const UScriptStruct* ScriptStruct = Item->Item.GetScriptStruct())
+		{
+			if (ScriptStruct->IsChildOf(FStateTreeBlueprintEvaluatorWrapper::StaticStruct())
+				|| ScriptStruct->IsChildOf(FStateTreeBlueprintTaskWrapper::StaticStruct())
+				|| ScriptStruct->IsChildOf(FStateTreeBlueprintConditionWrapper::StaticStruct()))
+			{
+				if (Item->InstanceObject != nullptr)
+				{
+					return FSlateIconFinder::FindIconBrushForClass(Item->InstanceObject->GetClass());
+				}
+			}
+		}
+	}
+	
 	return FSlateIconFinder::FindIconBrushForClass(UScriptStruct::StaticClass());
 }
 
@@ -728,61 +672,80 @@ TSharedRef<SWidget> FStateTreeBindableItemDetails::GeneratePicker()
 {
 	FMenuBuilder MenuBuilder(true, nullptr);
 
-	// @todo: Find a way to combine the menus. Now both are displayed even if one may be empty.
-	
-	MenuBuilder.AddWrapperSubMenu(
-		LOCTEXT("Items", "Items"),
-		LOCTEXT("Items_ToolTip", "Items"),
-		FOnGetContent::CreateSP(this, &FStateTreeBindableItemDetails::GenerateStructPicker),
-		FSlateIcon()
-		);
-
-	MenuBuilder.AddWrapperSubMenu(
-		LOCTEXT("BlueprintItems", "Blueprint Items"),
-		LOCTEXT("BlueprintItems_ToolTip", "Blueprint Items"),
-		FOnGetContent::CreateSP(this, &FStateTreeBindableItemDetails::GenerateClassPicker),
-		FSlateIcon()
-	);
+	FStateTreeEditorModule& EditorModule = FModuleManager::GetModuleChecked<FStateTreeEditorModule>(TEXT("StateTreeEditorModule"));
+	FStateTreeNodeClassCache* ClassCache = EditorModule.GetNodeClassCache().Get();
+	check(ClassCache);
+
+	FUIAction ClearAction(FExecuteAction::CreateSP(this, &FStateTreeBindableItemDetails::OnStructPicked, (const UScriptStruct*)nullptr));
+	MenuBuilder.AddMenuEntry(LOCTEXT("ClearItem", "Clear"), TAttribute<FText>(), FSlateIcon(FEditorStyle::GetStyleSetName(), "Cross"), ClearAction);
 
 	MenuBuilder.AddMenuSeparator();
 
-	FUIAction ClearAction(FExecuteAction::CreateSP(this, &FStateTreeBindableItemDetails::OnStructPicked, (const UScriptStruct*)nullptr));
-	MenuBuilder.AddMenuEntry(LOCTEXT("ClearItem", "Clear Item"), TAttribute<FText>(), FSlateIcon(FEditorStyle::GetStyleSetName(), "Cross"), ClearAction);
+	TArray<TSharedPtr<FStateTreeNodeClassData>> StructItems;
+	TArray<TSharedPtr<FStateTreeNodeClassData>> ObjectItems;
 	
-	return MenuBuilder.MakeWidget();
-}
-
+	ClassCache->GetScripStructs(BaseScriptStruct, StructItems);
+	ClassCache->GetClasses(BaseClass, ObjectItems);
 
+	const FSlateIcon ScripStructIcon = FSlateIconFinder::FindIconForClass(UScriptStruct::StaticClass());
+	const UStateTreeSchema* Schema = StateTree ? StateTree->GetSchema() : nullptr;
+	
+	for (const TSharedPtr<FStateTreeNodeClassData>& Data : StructItems)
+	{
+		if (Data->GetScriptStruct() != nullptr)
+		{
+			const UScriptStruct* ScriptStruct = Data->GetScriptStruct();
+			if (ScriptStruct == BaseScriptStruct)
+			{
+				continue;
+			}
+			if (ScriptStruct->HasMetaData(TEXT("Hidden")))
+			{
+				continue;				
+			}
+			if (Schema && !Schema->IsStructAllowed(ScriptStruct))
+			{
+				continue;				
+			}
+			
+			FUIAction ItemAction(FExecuteAction::CreateSP(this, &FStateTreeBindableItemDetails::OnStructPicked, ScriptStruct));
+			MenuBuilder.AddMenuEntry(ScriptStruct->GetDisplayNameText(), TAttribute<FText>(), ScripStructIcon, ItemAction);
+		}
+	}
 
-TSharedRef<SWidget> FStateTreeBindableItemDetails::GenerateStructPicker()
-{
-	TSharedRef<FBindableItemStructFilter> StructFilter = MakeShared<FBindableItemStructFilter>();
-	StructFilter->BaseStruct = BaseScriptStruct;
-	StructFilter->bAllowUserDefinedStructs = false;
-	StructFilter->bAllowBaseStruct = false;
-	StructFilter->Schema = StateTree ? StateTree->GetSchema() : nullptr;
-
-	FStructViewerInitializationOptions Options;
-	Options.bShowUnloadedStructs = true;
-	Options.bShowNoneOption = true;
-	Options.StructFilter = StructFilter;
-	Options.NameTypeToDisplay = EStructViewerNameTypeToDisplay::DisplayName;
-	Options.DisplayMode = EStructViewerDisplayMode::ListView;
-	Options.bAllowViewOptions = true;
-
-	FOnStructPicked OnPicked(FOnStructPicked::CreateRaw(this, &FStateTreeBindableItemDetails::OnStructPicked));
+	if (StructItems.Num() > 0 && ObjectItems.Num() > 0)
+	{
+		MenuBuilder.AddMenuSeparator();
+	}
 	
-	return SNew(SBox)
-		.WidthOverride(280.f)
-		[
-			SNew(SVerticalBox)
-			+ SVerticalBox::Slot()
-			.AutoHeight()
-			.MaxHeight(500.f)
-			[
-				FModuleManager::LoadModuleChecked<FStructViewerModule>("StructViewer").CreateStructViewer(Options, OnPicked)
-			]
-		];
+	for (const TSharedPtr<FStateTreeNodeClassData>& Data : ObjectItems)
+	{
+		if (Data->GetClass() != nullptr)
+		{
+			UClass* Class = Data->GetClass();
+			if (Class == BaseClass)
+			{
+				continue;
+			}
+			if (Class->HasAnyClassFlags(CLASS_Hidden | CLASS_HideDropDown))
+			{
+				continue;				
+			}
+			if (Class->HasMetaData(TEXT("Hidden")))
+			{
+				continue;				
+			}
+			if (Schema && !Schema->IsClassAllowed(Class))
+			{
+				continue;				
+			}
+
+			FUIAction ItemAction(FExecuteAction::CreateSP(this, &FStateTreeBindableItemDetails::OnClassPicked, Class));
+			MenuBuilder.AddMenuEntry(Class->GetDisplayNameText(), TAttribute<FText>(), FSlateIconFinder::FindIconForClass(Data->GetClass()), ItemAction);
+		}
+	}
+
+	return MenuBuilder.MakeWidget();
 }
 
 void FStateTreeBindableItemDetails::OnStructPicked(const UScriptStruct* InStruct)
@@ -801,7 +764,6 @@ void FStateTreeBindableItemDetails::OnStructPicked(const UScriptStruct* InStruct
 	{
 		if (FStateTreeItem* Item = static_cast<FStateTreeItem*>(Data))
 		{
-		
 			Item->Reset();
 			
 			if (InStruct)
@@ -871,37 +833,6 @@ void FStateTreeBindableItemDetails::OnStructPicked(const UScriptStruct* InStruct
 	}
 }
 
-
-TSharedRef<SWidget> FStateTreeBindableItemDetails::GenerateClassPicker()
-{
-	TSharedRef<FBindableItemClassFilter> ClassFilter = MakeShared<FBindableItemClassFilter>();
-	ClassFilter->BaseClass = BaseClass;
-	ClassFilter->bAllowBaseClass = false;
-	ClassFilter->Schema = StateTree ? StateTree->GetSchema() : nullptr;
-
-	FClassViewerInitializationOptions Options;
-	Options.bShowUnloadedBlueprints = true;
-	Options.bShowNoneOption = true;
-	Options.ClassFilters.Add(ClassFilter);
-	Options.NameTypeToDisplay = EClassViewerNameTypeToDisplay::DisplayName;
-	Options.DisplayMode = EClassViewerDisplayMode::ListView;
-	Options.bAllowViewOptions = true;
-
-	FOnClassPicked OnPicked(FOnClassPicked::CreateRaw(this, &FStateTreeBindableItemDetails::OnClassPicked));
-
-	return SNew(SBox)
-		.WidthOverride(280.f)
-		[
-			SNew(SVerticalBox)
-			+ SVerticalBox::Slot()
-			.AutoHeight()
-			.MaxHeight(500.f)
-			[
-				FModuleManager::LoadModuleChecked<FClassViewerModule>("ClassViewer").CreateClassViewer(Options, OnPicked)
-			]
-		];
-}
-
 void FStateTreeBindableItemDetails::OnClassPicked(UClass* InClass)
 {
 	check(StructProperty);
diff --git a/Engine/Plugins/Runtime/StateTree/Source/StateTreeEditorModule/Private/Customizations/StateTreeBindableItemDetails.h b/Engine/Plugins/Runtime/StateTree/Source/StateTreeEditorModule/Private/Customizations/StateTreeBindableItemDetails.h
--- a/Engine/Plugins/Runtime/StateTree/Source/StateTreeEditorModule/Private/Customizations/StateTreeBindableItemDetails.h
+++ b/Engine/Plugins/Runtime/StateTree/Source/StateTreeEditorModule/Private/Customizations/StateTreeBindableItemDetails.h
@@ -41,16 +41,13 @@ private:
 	void OnNameCommitted(const FText& NewText, ETextCommit::Type InTextCommit);
 	bool IsNameEnabled() const;
 
-	const UScriptStruct* GetCommonItemScriptStruct() const;
+	const struct FStateTreeItem* GetCommonItem() const;
 	FText GetDisplayValueString() const;
 	const FSlateBrush* GetDisplayValueIcon() const;
 
 	TSharedRef<SWidget> GeneratePicker();
 
-	TSharedRef<SWidget> GenerateStructPicker();
 	void OnStructPicked(const UScriptStruct* InStruct);
-
-	TSharedRef<SWidget> GenerateClassPicker();
 	void OnClassPicked(UClass* InClass);
 
 	void OnIdentifierChanged(const UStateTree& StateTree);
diff --git a/Engine/Plugins/Runtime/StateTree/Source/StateTreeEditorModule/Private/StateTreeEditorModule.cpp b/Engine/Plugins/Runtime/StateTree/Source/StateTreeEditorModule/Private/StateTreeEditorModule.cpp
--- a/Engine/Plugins/Runtime/StateTree/Source/StateTreeEditorModule/Private/StateTreeEditorModule.cpp
+++ b/Engine/Plugins/Runtime/StateTree/Source/StateTreeEditorModule/Private/StateTreeEditorModule.cpp
@@ -13,6 +13,13 @@
 #include "StateTreeEditor.h"
 #include "StateTree.h"
 #include "StateTreeEditorStyle.h"
+#include "StateTreeEvaluatorBase.h"
+#include "StateTreeTaskBase.h"
+#include "StateTreeConditionBase.h"
+#include "StateTreeNodeClassCache.h"
+#include "Blueprint/StateTreeEvaluatorBlueprintBase.h"
+#include "Blueprint/StateTreeTaskBlueprintBase.h"
+#include "Blueprint/StateTreeConditionBlueprintBase.h"
 
 
 #define LOCTEXT_NAMESPACE "StateTreeEditor"
@@ -84,6 +91,17 @@ void FStateTreeEditorModule::ShutdownModule()
 
 TSharedRef<IStateTreeEditor> FStateTreeEditorModule::CreateStateTreeEditor(const EToolkitMode::Type Mode, const TSharedPtr< class IToolkitHost >& InitToolkitHost, UStateTree* StateTree)
 {
+	if (!NodeClassCache.IsValid())
+	{
+		NodeClassCache = MakeShareable(new FStateTreeNodeClassCache());
+		NodeClassCache->AddRootScriptStruct(FStateTreeEvaluatorBase::StaticStruct());
+		NodeClassCache->AddRootScriptStruct(FStateTreeTaskBase::StaticStruct());
+		NodeClassCache->AddRootScriptStruct(FStateTreeConditionBase::StaticStruct());
+		NodeClassCache->AddRootClass(UStateTreeEvaluatorBlueprintBase::StaticClass());
+		NodeClassCache->AddRootClass(UStateTreeTaskBlueprintBase::StaticClass());
+		NodeClassCache->AddRootClass(UStateTreeConditionBlueprintBase::StaticClass());
+	}
+	
 	TSharedRef<FStateTreeEditor> NewEditor(new FStateTreeEditor());
 	NewEditor->InitEditor(Mode, InitToolkitHost, StateTree);
 	return NewEditor;
diff --git a/Engine/Plugins/Runtime/StateTree/Source/StateTreeEditorModule/Private/StateTreeNodeClassCache.cpp b/Engine/Plugins/Runtime/StateTree/Source/StateTreeEditorModule/Private/StateTreeNodeClassCache.cpp
new file mode 100644
--- /dev/null
+++ b/Engine/Plugins/Runtime/StateTree/Source/StateTreeEditorModule/Private/StateTreeNodeClassCache.cpp
@@ -0,0 +1,275 @@
+// Copyright Epic Games, Inc. All Rights Reserved.
+
+#include "StateTreeNodeClassCache.h"
+#include "UObject/Object.h"
+#include "UObject/Class.h"
+#include "Misc/FeedbackContext.h"
+#include "Modules/ModuleManager.h"
+#include "UObject/UObjectIterator.h"
+#include "Misc/PackageName.h"
+#include "UObject/ConstructorHelpers.h"
+#include "Engine/Blueprint.h"
+#include "Editor.h"
+#include "ARFilter.h"
+#include "ObjectEditorUtils.h"
+#include "Logging/MessageLog.h"
+#include "AssetRegistryModule.h"
+
+#define LOCTEXT_NAMESPACE "StateTreeEditor"
+
+FStateTreeNodeClassData::FStateTreeNodeClassData(UStruct* InStruct) :
+	Struct(InStruct)
+{
+	if (InStruct)
+	{
+		StructName = InStruct->GetFName();
+	}
+}
+
+FStateTreeNodeClassData::FStateTreeNodeClassData(const FString& InClassAssetName, const FString& InClassPackageName, const FName InStructName, UStruct* InStruct) :
+	Struct(InStruct),
+	StructName(InStructName),
+	ClassAssetName(InClassAssetName),
+	ClassPackageName(InClassPackageName)
+{
+}
+
+UStruct* FStateTreeNodeClassData::GetStruct(bool bSilent)
+{
+	UStruct* Ret = Struct.Get();
+	
+	if (Ret == nullptr && ClassPackageName.Len())
+	{
+		GWarn->BeginSlowTask(LOCTEXT("LoadPackage", "Loading Package..."), true);
+
+		UPackage* Package = LoadPackage(nullptr, *ClassPackageName, LOAD_NoRedirects);
+		if (Package)
+		{
+			Package->FullyLoad();
+
+			UObject* Object = FindObject<UObject>(Package, *ClassAssetName);
+
+			GWarn->EndSlowTask();
+
+			if (UBlueprint* BlueprintOb = Cast<UBlueprint>(Object))
+			{
+				Ret = *BlueprintOb->GeneratedClass;
+			}
+			else if (Object != nullptr)
+			{
+				Ret = Object->GetClass();
+			}
+
+			Struct = Ret;
+		}
+		else
+		{
+			GWarn->EndSlowTask();
+
+			if (!bSilent)
+			{
+				FMessageLog EditorErrors("EditorErrors");
+				EditorErrors.Error(LOCTEXT("PackageLoadFail", "Package Load Failed"));
+				EditorErrors.Info(FText::FromString(ClassPackageName));
+				EditorErrors.Notify(LOCTEXT("PackageLoadFail", "Package Load Failed"));
+			}
+		}
+	}
+
+	return Ret;
+}
+
+
+//----------------------------------------------------------------------//
+//  FStateTreeNodeClassCache
+//----------------------------------------------------------------------//
+FStateTreeNodeClassCache::FStateTreeNodeClassCache()
+{
+	// Register with the Asset Registry to be informed when it is done loading up files.
+	FAssetRegistryModule& AssetRegistryModule = FModuleManager::GetModuleChecked<FAssetRegistryModule>(TEXT("AssetRegistry"));
+	AssetRegistryModule.Get().OnFilesLoaded().AddRaw(this, &FStateTreeNodeClassCache::InvalidateCache);
+	AssetRegistryModule.Get().OnAssetAdded().AddRaw(this, &FStateTreeNodeClassCache::OnAssetAdded);
+	AssetRegistryModule.Get().OnAssetRemoved().AddRaw(this, &FStateTreeNodeClassCache::OnAssetRemoved);
+
+	// Register to have Populate called when doing a Reload.
+	FCoreUObjectDelegates::ReloadCompleteDelegate.AddRaw(this, &FStateTreeNodeClassCache::OnReloadComplete);
+
+	// Register to have Populate called when a Blueprint is compiled.
+	check(GEditor);
+	GEditor->OnBlueprintCompiled().AddRaw(this, &FStateTreeNodeClassCache::InvalidateCache);
+	GEditor->OnClassPackageLoadedOrUnloaded().AddRaw(this, &FStateTreeNodeClassCache::InvalidateCache);
+}
+
+FStateTreeNodeClassCache::~FStateTreeNodeClassCache()
+{
+	// Unregister with the Asset Registry to be informed when it is done loading up files.
+	if (FModuleManager::Get().IsModuleLoaded(TEXT("AssetRegistry")))
+	{
+		FAssetRegistryModule& AssetRegistryModule = FModuleManager::GetModuleChecked<FAssetRegistryModule>(TEXT("AssetRegistry"));
+		AssetRegistryModule.Get().OnFilesLoaded().RemoveAll(this);
+		AssetRegistryModule.Get().OnAssetAdded().RemoveAll(this);
+		AssetRegistryModule.Get().OnAssetRemoved().RemoveAll(this);
+
+		// Unregister to have Populate called when doing a Reload.
+		FCoreUObjectDelegates::ReloadCompleteDelegate.RemoveAll(this);
+
+		// Unregister to have Populate called when a Blueprint is compiled.
+		if (GEditor != nullptr)
+		{
+			// GEditor can't have been destructed before we call this or we'll crash.
+			GEditor->OnBlueprintCompiled().RemoveAll(this);
+			GEditor->OnClassPackageLoadedOrUnloaded().RemoveAll(this);
+		}
+	}
+}
+
+void FStateTreeNodeClassCache::AddRootStruct(UStruct* RootStruct)
+{
+	if (RootClasses.ContainsByPredicate([RootStruct](const FRootClassContainer& RootClass){ return RootClass.BaseStruct == RootStruct; }))
+	{
+		return;
+		return;
+	}
+	
+	RootClassNameToIndex.Add(RootStruct->GetName(), RootClasses.Num());
+	RootClasses.Emplace(RootStruct);
+
+	InvalidateCache();
+}
+
+void FStateTreeNodeClassCache::GetStructs(UStruct* BaseStruct, TArray<TSharedPtr<FStateTreeNodeClassData>>& AvailableClasses)
+{
+	AvailableClasses.Reset();
+	
+	const int32 RootClassIndex = RootClasses.IndexOfByPredicate([BaseStruct](const FRootClassContainer& RootClass){ return RootClass.BaseStruct == BaseStruct; });
+	if (RootClassIndex != INDEX_NONE)
+	{
+		const FRootClassContainer& RootClass = RootClasses[RootClassIndex];
+		
+		if (!RootClass.bUpdated)
+		{
+			CacheClasses();
+		}
+
+		AvailableClasses.Append(RootClass.ClassData);
+	}
+}
+
+void FStateTreeNodeClassCache::OnAssetAdded(const FAssetData& AssetData)
+{
+	UpdateBlueprintClass(AssetData);
+}
+
+void FStateTreeNodeClassCache::OnAssetRemoved(const FAssetData& AssetData)
+{
+	FString AssetClassName;
+	FString AssetNativeParentClassName;
+
+	if (AssetData.GetTagValue(FBlueprintTags::GeneratedClassPath, AssetClassName) && AssetData.GetTagValue(FBlueprintTags::NativeParentClassPath, AssetNativeParentClassName))
+	{
+		ConstructorHelpers::StripObjectClass(AssetClassName);
+		AssetClassName = FPackageName::ObjectPathToObjectName(AssetClassName);
+
+		if (const int32* RootClassIndex = RootClassNameToIndex.Find(AssetNativeParentClassName))
+		{
+			FRootClassContainer& RootClass = RootClasses[*RootClassIndex];
+			const FName ClassName(AssetClassName);
+			RootClass.ClassData.RemoveAll([&ClassName](const TSharedPtr<FStateTreeNodeClassData>& ClassData) { return ClassData->GetStructName() == ClassName; });
+		}
+	}
+}
+
+void FStateTreeNodeClassCache::InvalidateCache()
+{
+	for (FRootClassContainer& RootClass : RootClasses)
+	{
+		RootClass.ClassData.Reset();
+		RootClass.bUpdated = false;
+	}
+}
+
+void FStateTreeNodeClassCache::OnReloadComplete(EReloadCompleteReason Reason)
+{
+	InvalidateCache();
+}
+
+void FStateTreeNodeClassCache::UpdateBlueprintClass(const FAssetData& AssetData)
+{
+	FString AssetClassName;
+	FString AssetNativeParentClassName;
+
+	if (AssetData.GetTagValue(FBlueprintTags::GeneratedClassPath, AssetClassName) && AssetData.GetTagValue(FBlueprintTags::NativeParentClassPath, AssetNativeParentClassName))
+	{
+		UObject* Outer1 = nullptr;
+		ResolveName(Outer1, AssetClassName, false, false);
+
+		UObject* Outer2 = nullptr;
+		ResolveName(Outer2, AssetNativeParentClassName, false, false);
+
+		if (const int32* RootClassIndex = RootClassNameToIndex.Find(AssetNativeParentClassName))
+		{
+			FRootClassContainer& RootClass = RootClasses[*RootClassIndex];
+
+			const FName ClassName(AssetClassName);
+			const int32 ClassDataIndex = RootClass.ClassData.IndexOfByPredicate([&ClassName](const TSharedPtr<FStateTreeNodeClassData>& ClassData) { return ClassData->GetStructName() == ClassName; });
+			
+			if (ClassDataIndex == INDEX_NONE)
+			{
+				UObject* AssetOb = AssetData.IsAssetLoaded() ? AssetData.GetAsset() : nullptr;
+				UBlueprint* AssetBP = Cast<UBlueprint>(AssetOb);
+				UClass* AssetClass = AssetBP ? *AssetBP->GeneratedClass : AssetOb ? AssetOb->GetClass() : nullptr;
+
+				TSharedPtr<FStateTreeNodeClassData> NewData = MakeShareable(new FStateTreeNodeClassData(AssetData.AssetName.ToString(), AssetData.PackageName.ToString(), ClassName, AssetClass));
+				RootClass.ClassData.Add(NewData);
+			}
+		}
+	}
+}
+
+void FStateTreeNodeClassCache::CacheClasses()
+{
+	for (FRootClassContainer& RootClass : RootClasses)
+	{
+		RootClass.ClassData.Reset();
+		RootClass.bUpdated = true;
+		
+		// gather all native classes
+		if (const UClass* Class = Cast<UClass>(RootClass.BaseStruct))
+		{
+			for (TObjectIterator<UClass> It; It; ++It)
+			{
+				UClass* TestClass = *It;
+				if (TestClass->HasAnyClassFlags(CLASS_Native) && TestClass->IsChildOf(Class))
+				{
+					RootClass.ClassData.Add(MakeShareable(new FStateTreeNodeClassData(TestClass)));
+				}
+			}
+		}
+		else if (const UScriptStruct* ScriptStruct = Cast<UScriptStruct>(RootClass.BaseStruct))
+		{
+			for (TObjectIterator<UScriptStruct> It; It; ++It)
+			{
+				UScriptStruct* TestStruct = *It;
+				if (TestStruct->IsChildOf(ScriptStruct))
+				{
+					RootClass.ClassData.Add(MakeShareable(new FStateTreeNodeClassData(TestStruct)));
+				}
+			}
+		}
+	}
+
+	// gather all blueprints
+	FAssetRegistryModule& AssetRegistryModule = FModuleManager::LoadModuleChecked<FAssetRegistryModule>(TEXT("AssetRegistry"));
+	TArray<FAssetData> BlueprintList;
+
+	FARFilter Filter;
+	Filter.ClassNames.Add(UBlueprint::StaticClass()->GetFName());
+	AssetRegistryModule.Get().GetAssets(Filter, BlueprintList);
+
+	for (int32 i = 0; i < BlueprintList.Num(); i++)
+	{
+		UpdateBlueprintClass(BlueprintList[i]);
+	}
+}
+
+#undef LOCTEXT_NAMESPACE
diff --git a/Engine/Plugins/Runtime/StateTree/Source/StateTreeEditorModule/Public/StateTreeEditorModule.h b/Engine/Plugins/Runtime/StateTree/Source/StateTreeEditorModule/Public/StateTreeEditorModule.h
--- a/Engine/Plugins/Runtime/StateTree/Source/StateTreeEditorModule/Public/StateTreeEditorModule.h
+++ b/Engine/Plugins/Runtime/StateTree/Source/StateTreeEditorModule/Public/StateTreeEditorModule.h
@@ -9,6 +9,7 @@
 class UStateTree;
 class IStateTreeEditor;
 class FAssetTypeActions_Base;
+struct FStateTreeNodeClassCache;
 
 /**
 * The public interface to this module
@@ -26,9 +27,13 @@ public:
 	virtual TSharedPtr<FExtensibilityManager> GetMenuExtensibilityManager() override { return MenuExtensibilityManager; }
 	virtual TSharedPtr<FExtensibilityManager> GetToolBarExtensibilityManager() override { return ToolBarExtensibilityManager; }
 
+	TSharedPtr<FStateTreeNodeClassCache> GetNodeClassCache() const { return NodeClassCache; }
+	
 protected:
 	TSharedPtr<FExtensibilityManager> MenuExtensibilityManager;
 	TSharedPtr<FExtensibilityManager> ToolBarExtensibilityManager;
 
 	TArray<TSharedPtr<FAssetTypeActions_Base>> ItemDataAssetTypeActions;
+
+	TSharedPtr<FStateTreeNodeClassCache> NodeClassCache;
 };
diff --git a/Engine/Plugins/Runtime/StateTree/Source/StateTreeEditorModule/Public/StateTreeNodeClassCache.h b/Engine/Plugins/Runtime/StateTree/Source/StateTreeEditorModule/Public/StateTreeNodeClassCache.h
new file mode 100644
--- /dev/null
+++ b/Engine/Plugins/Runtime/StateTree/Source/StateTreeEditorModule/Public/StateTreeNodeClassCache.h
@@ -0,0 +1,118 @@
+// Copyright Epic Games, Inc. All Rights Reserved.
+
+#pragma once
+
+#include "CoreMinimal.h"
+#include "UObject/ObjectMacros.h"
+#include "UObject/WeakObjectPtr.h"
+#include "StateTreeNodeClassCache.generated.h"
+
+struct FAssetData;
+
+/**
+ * Describes a class or struct.
+ * If the class or struct is from a package that is not yet loaded, the data will update on GetStruct/Class/Scripstruct()
+ */
+USTRUCT()
+struct STATETREEEDITORMODULE_API FStateTreeNodeClassData
+{
+	GENERATED_BODY()
+
+	FStateTreeNodeClassData() {}
+	FStateTreeNodeClassData(UStruct* InStruct);
+	FStateTreeNodeClassData(const FString& InClassAssetName, const FString& InClassPackage, const FName InStructName, UStruct* InStruct);
+
+	FName GetStructName() const { return StructName; }
+
+	UStruct* GetStruct(bool bSilent = false);
+
+	UClass* GetClass(bool bSilent = false)
+	{
+		return Cast<UClass>(GetStruct(bSilent));
+	}
+	UScriptStruct* GetScriptStruct(bool bSilent = false)
+	{
+		return Cast<UScriptStruct>(GetStruct(bSilent));
+	}
+
+private:
+
+	/** Pointer to described struct or class. */
+	TWeakObjectPtr<UStruct> Struct;
+
+	/** Resolved name of struct or class. */
+	FName StructName;
+
+	/** Path to class if it's not loaded yet. */
+	FString ClassAssetName;
+
+	/** Package of the asset if it's not loaded yet. */
+	FString ClassPackageName;
+};
+
+/**
+ * Caches specified classes or structs and reacts to engine events to keep the lists always up to date.
+ * All the derived classes or structs are kept in the cache.
+ */
+struct STATETREEEDITORMODULE_API FStateTreeNodeClassCache
+{
+	FStateTreeNodeClassCache();
+	~FStateTreeNodeClassCache();
+
+	/** Adds a Struct to keep track of */
+	void AddRootStruct(UStruct* RootStruct);
+	
+	/** Adds a Class to keep track of */
+	void AddRootClass(UClass* RootClass)
+	{
+		AddRootStruct(RootClass);
+	}
+	
+	/** Adds a ScriptStruct to keep track of */
+	void AddRootScriptStruct(UScriptStruct* RootStruct)
+	{
+		AddRootStruct(RootStruct);
+	}
+
+	/** Returns know derived Structs based on provided base. If the base Struct is not added as root Struct, nothing is returned. */
+	void GetStructs(UStruct* BaseStruct, TArray<TSharedPtr<FStateTreeNodeClassData>>& AvailableClasses);
+	
+	/** Returns know derived Classes based on provided base. If the base Class is not added as root Class, nothing is returned. */
+	void GetClasses(UStruct* BaseClass, TArray<TSharedPtr<FStateTreeNodeClassData>>& AvailableClasses)
+	{
+		GetStructs(BaseClass, AvailableClasses);
+	}
+	
+	/** Returns know derived ScriptStructs based on provided base. If the base struct is not added as root ScriptStruct, nothing is returned. */
+	void GetScripStructs(UScriptStruct* BaseStruct, TArray<TSharedPtr<FStateTreeNodeClassData>>& AvailableClasses)
+	{
+		GetStructs(BaseStruct, AvailableClasses);
+	}
+
+	/** Invalidates the cache, it will be rebuild on next access. */
+	void InvalidateCache();
+
+protected:
+	void OnAssetAdded(const FAssetData& AssetData);
+	void OnAssetRemoved(const FAssetData& AssetData);
+	void OnReloadComplete(EReloadCompleteReason Reason);
+	
+private:
+	void UpdateBlueprintClass(const FAssetData& AssetData);
+	void CacheClasses();
+
+	struct FRootClassContainer
+	{
+		FRootClassContainer() = default;
+		FRootClassContainer(UStruct* InBaseStruct)
+		{
+			BaseStruct = InBaseStruct;
+		} 
+		TWeakObjectPtr<UStruct> BaseStruct;
+		TArray<TSharedPtr<FStateTreeNodeClassData>> ClassData;
+		bool bUpdated = false;
+	};
+
+	TArray<FRootClassContainer> RootClasses;
+	TMap<FString, int32> RootClassNameToIndex;
+};
diff --git a/Engine/Plugins/Runtime/StateTree/Source/StateTreeModule/Private/Blueprint/StateTreeConditionBlueprintBase.cpp b/Engine/Plugins/Runtime/StateTree/Source/StateTreeModule/Private/Blueprint/StateTreeConditionBlueprintBase.cpp
--- a/Engine/Plugins/Runtime/StateTree/Source/StateTreeModule/Private/Blueprint/StateTreeConditionBlueprintBase.cpp
+++ b/Engine/Plugins/Runtime/StateTree/Source/StateTreeModule/Private/Blueprint/StateTreeConditionBlueprintBase.cpp
@@ -38,3 +38,19 @@ bool FStateTreeBlueprintConditionWrapper::TestCondition(FStateTreeExecutionConte
 	}
 	return false;
 }
+
+#if WITH_EDITOR
+FText FStateTreeBlueprintConditionWrapper::GetDescription(const FGuid& ID, FStateTreeDataView InstanceData, const IStateTreeBindingLookup& BindingLookup) const
+{
+	const UClass* Class = Cast<const UClass>(InstanceData.GetStruct());
+	if (UBlueprint* ClassBP = Cast<UBlueprint>(Class->ClassGeneratedBy))
+	{
+		if (!ClassBP->BlueprintDescription.IsEmpty())
+		{
+			return FText::FromString(ClassBP->BlueprintDescription);
+		}
+	}
+	
+	return Class != nullptr ? Class->GetDisplayNameText() : FText::GetEmpty();
+}
+#endif
diff --git a/Engine/Plugins/Runtime/StateTree/Source/StateTreeModule/Private/Conditions/StateTreeCondition_Common.cpp b/Engine/Plugins/Runtime/StateTree/Source/StateTreeModule/Private/Conditions/StateTreeCondition_Common.cpp
--- a/Engine/Plugins/Runtime/StateTree/Source/StateTreeModule/Private/Conditions/StateTreeCondition_Common.cpp
+++ b/Engine/Plugins/Runtime/StateTree/Source/StateTreeModule/Private/Conditions/StateTreeCondition_Common.cpp
@@ -96,7 +96,7 @@ bool FStateTreeCondition_CompareInt::TestCondition(FStateTreeExecutionContext& C
 }
 
 #if WITH_EDITOR
-FText FStateTreeCondition_CompareInt::GetDescription(const FGuid& ID, FConstStructView InstanceData, const IStateTreeBindingLookup& BindingLookup) const
+FText FStateTreeCondition_CompareInt::GetDescription(const FGuid& ID, FStateTreeDataView InstanceData, const IStateTreeBindingLookup& BindingLookup) const
 {
 	const FStateTreeCondition_CompareIntInstanceData& Instance = InstanceData.Get<FStateTreeCondition_CompareIntInstanceData>();
 	const FStateTreeEditorPropertyPath LeftPath(ID, GET_MEMBER_NAME_STRING_CHECKED(FStateTreeCondition_CompareIntInstanceData, Left));
@@ -157,7 +157,7 @@ bool FStateTreeCondition_CompareFloat::TestCondition(FStateTreeExecutionContext&
 
 
 #if WITH_EDITOR
-FText FStateTreeCondition_CompareFloat::GetDescription(const FGuid& ID, FConstStructView InstanceData, const IStateTreeBindingLookup& BindingLookup) const
+FText FStateTreeCondition_CompareFloat::GetDescription(const FGuid& ID, FStateTreeDataView InstanceData, const IStateTreeBindingLookup& BindingLookup) const
 {
 	const FStateTreeCondition_CompareFloatInstanceData& Instance = InstanceData.Get<FStateTreeCondition_CompareFloatInstanceData>();
 	const FStateTreeEditorPropertyPath LeftPath(ID, GET_MEMBER_NAME_STRING_CHECKED(FStateTreeCondition_CompareFloatInstanceData, Left));
@@ -216,7 +216,7 @@ bool FStateTreeCondition_CompareBool::TestCondition(FStateTreeExecutionContext&
 }
 
 #if WITH_EDITOR
-FText FStateTreeCondition_CompareBool::GetDescription(const FGuid& ID, FConstStructView InstanceData, const IStateTreeBindingLookup& BindingLookup) const
+FText FStateTreeCondition_CompareBool::GetDescription(const FGuid& ID, FStateTreeDataView InstanceData, const IStateTreeBindingLookup& BindingLookup) const
 {
 	const FStateTreeCondition_CompareBoolInstanceData& Instance = InstanceData.Get<FStateTreeCondition_CompareBoolInstanceData>();
 	const FStateTreeEditorPropertyPath LeftPath(ID, GET_MEMBER_NAME_STRING_CHECKED(FStateTreeCondition_CompareBoolInstanceData, bLeft));
@@ -273,7 +273,7 @@ bool FStateTreeCondition_CompareEnum::TestCondition(FStateTreeExecutionContext&
 }
 
 #if WITH_EDITOR
-FText FStateTreeCondition_CompareEnum::GetDescription(const FGuid& ID, FConstStructView InstanceData, const IStateTreeBindingLookup& BindingLookup) const
+FText FStateTreeCondition_CompareEnum::GetDescription(const FGuid& ID, FStateTreeDataView InstanceData, const IStateTreeBindingLookup& BindingLookup) const
 {
 	const FStateTreeCondition_CompareEnumInstanceData& Instance = InstanceData.Get<FStateTreeCondition_CompareEnumInstanceData>();
 	const FStateTreeEditorPropertyPath LeftPath(ID, GET_MEMBER_NAME_STRING_CHECKED(FStateTreeCondition_CompareEnumInstanceData, Left));
@@ -315,7 +315,7 @@ FText FStateTreeCondition_CompareEnum::GetDescription(const FGuid& ID, FConstStr
 	return FText::Format(LOCTEXT("EnumEqualsDesc", "{0} <Details.Bold>{1}</> is <Details.Bold>{2}</>"), InvertText, LeftText, RightText);
 }
 
-void FStateTreeCondition_CompareEnum::OnBindingChanged(const FGuid& ID, FStructView InstanceData, const FStateTreeEditorPropertyPath& SourcePath, const FStateTreeEditorPropertyPath& TargetPath, const IStateTreeBindingLookup& BindingLookup)
+void FStateTreeCondition_CompareEnum::OnBindingChanged(const FGuid& ID, FStateTreeDataView InstanceData, const FStateTreeEditorPropertyPath& SourcePath, const FStateTreeEditorPropertyPath& TargetPath, const IStateTreeBindingLookup& BindingLookup)
 {
 	if (!TargetPath.IsValid())
 	{
@@ -386,7 +386,7 @@ bool FStateTreeCondition_CompareDistance::TestCondition(FStateTreeExecutionConte
 }
 
 #if WITH_EDITOR
-FText FStateTreeCondition_CompareDistance::GetDescription(const FGuid& ID, FConstStructView InstanceData, const IStateTreeBindingLookup& BindingLookup) const
+FText FStateTreeCondition_CompareDistance::GetDescription(const FGuid& ID, FStateTreeDataView InstanceData, const IStateTreeBindingLookup& BindingLookup) const
 {
 	const FStateTreeCondition_CompareDistanceInstanceData& Instance = InstanceData.Get<FStateTreeCondition_CompareDistanceInstanceData>();
 	const FStateTreeEditorPropertyPath SourcePath(ID, GET_MEMBER_NAME_STRING_CHECKED(FStateTreeCondition_CompareDistanceInstanceData, Source));
@@ -453,7 +453,7 @@ bool FStateTreeCondition_Random::TestCondition(FStateTreeExecutionContext& Conte
 }
 
 #if WITH_EDITOR
-FText FStateTreeCondition_Random::GetDescription(const FGuid& ID, FConstStructView InstanceData, const IStateTreeBindingLookup& BindingLookup) const
+FText FStateTreeCondition_Random::GetDescription(const FGuid& ID, FStateTreeDataView InstanceData, const IStateTreeBindingLookup& BindingLookup) const
 {
 	const FStateTreeCondition_RandomInstanceData& Instance = InstanceData.Get<FStateTreeCondition_RandomInstanceData>();
 	const FStateTreeEditorPropertyPath ThresholdPath(ID, GET_MEMBER_NAME_STRING_CHECKED(FStateTreeCondition_RandomInstanceData, Threshold));
diff --git a/Engine/Plugins/Runtime/StateTree/Source/StateTreeModule/Public/Blueprint/StateTreeConditionBlueprintBase.h b/Engine/Plugins/Runtime/StateTree/Source/StateTreeModule/Public/Blueprint/StateTreeConditionBlueprintBase.h
--- a/Engine/Plugins/Runtime/StateTree/Source/StateTreeModule/Public/Blueprint/StateTreeConditionBlueprintBase.h
+++ b/Engine/Plugins/Runtime/StateTree/Source/StateTreeModule/Public/Blueprint/StateTreeConditionBlueprintBase.h
@@ -39,6 +39,11 @@ struct STATETREEMODULE_API FStateTreeBlueprintConditionWrapper : public FStateTr
 	virtual bool Link(FStateTreeLinker& Linker) override;
 	virtual bool TestCondition(FStateTreeExecutionContext& Context) const override;
 
+#if WITH_EDITOR
+	/** @return Rich text description of the condition. */
+	virtual FText GetDescription(const FGuid& ID, FStateTreeDataView InstanceData, const IStateTreeBindingLookup& BindingLookup) const;
+#endif
+	
 	UPROPERTY()
 	TSubclassOf<UStateTreeConditionBlueprintBase> ConditionClass = nullptr;
 
diff --git a/Engine/Plugins/Runtime/StateTree/Source/StateTreeModule/Public/Conditions/StateTreeCondition_Common.h b/Engine/Plugins/Runtime/StateTree/Source/StateTreeModule/Public/Conditions/StateTreeCondition_Common.h
--- a/Engine/Plugins/Runtime/StateTree/Source/StateTreeModule/Public/Conditions/StateTreeCondition_Common.h
+++ b/Engine/Plugins/Runtime/StateTree/Source/StateTreeModule/Public/Conditions/StateTreeCondition_Common.h
@@ -18,7 +18,7 @@ enum class EStateTreeCompare : uint8
 /**
  * Base class for all common conditions so that they are easy to single out in the schema.
  */
-USTRUCT()
+USTRUCT(meta = (Hidden))
 struct STATETREEMODULE_API FStateTreeConditionCommonBase : public FStateTreeConditionBase
 {
 	GENERATED_BODY()
@@ -57,7 +57,7 @@ struct STATETREEMODULE_API FStateTreeCondition_CompareInt : public FStateTreeCon
 	virtual bool TestCondition(FStateTreeExecutionContext& Context) const override;
 
 #if WITH_EDITOR
-	virtual FText GetDescription(const FGuid& ID, FConstStructView InstanceData, const IStateTreeBindingLookup& BindingLookup) const override;
+	virtual FText GetDescription(const FGuid& ID, FStateTreeDataView InstanceData, const IStateTreeBindingLookup& BindingLookup) const override;
 #endif
 
 	UPROPERTY(EditAnywhere, Category = Parameter)
@@ -103,7 +103,7 @@ struct STATETREEMODULE_API FStateTreeCondition_CompareFloat : public FStateTreeC
 	virtual bool TestCondition(FStateTreeExecutionContext& Context) const override;
 
 #if WITH_EDITOR
-	virtual FText GetDescription(const FGuid& ID, FConstStructView InstanceData, const IStateTreeBindingLookup& BindingLookup) const override;
+	virtual FText GetDescription(const FGuid& ID, FStateTreeDataView InstanceData, const IStateTreeBindingLookup& BindingLookup) const override;
 #endif
 
 	UPROPERTY(EditAnywhere, Category = Parameter)
@@ -152,7 +152,7 @@ struct STATETREEMODULE_API FStateTreeCondition_CompareBool : public FStateTreeCo
 	virtual bool TestCondition(FStateTreeExecutionContext& Context) const override;
 
 #if WITH_EDITOR
-	virtual FText GetDescription(const FGuid& ID, FConstStructView InstanceData, const IStateTreeBindingLookup& BindingLookup) const override;
+	virtual FText GetDescription(const FGuid& ID, FStateTreeDataView InstanceData, const IStateTreeBindingLookup& BindingLookup) const override;
 #endif
 
 	UPROPERTY(EditAnywhere, Category = Parameter)
@@ -198,8 +198,8 @@ struct STATETREEMODULE_API FStateTreeCondition_CompareEnum : public FStateTreeCo
 	virtual bool TestCondition(FStateTreeExecutionContext& Context) const override;
 
 #if WITH_EDITOR
-	virtual FText GetDescription(const FGuid& ID, FConstStructView InstanceData, const IStateTreeBindingLookup& BindingLookup) const override;
-	virtual void OnBindingChanged(const FGuid& ID, FStructView InstanceData, const FStateTreeEditorPropertyPath& SourcePath, const FStateTreeEditorPropertyPath& TargetPath, const IStateTreeBindingLookup& BindingLookup) override;
+	virtual FText GetDescription(const FGuid& ID, FStateTreeDataView InstanceData, const IStateTreeBindingLookup& BindingLookup) const override;
+	virtual void OnBindingChanged(const FGuid& ID, FStateTreeDataView InstanceData, const FStateTreeEditorPropertyPath& SourcePath, const FStateTreeEditorPropertyPath& TargetPath, const IStateTreeBindingLookup& BindingLookup) override;
 #endif
 
 	UPROPERTY(EditAnywhere, Category = Parameter)
@@ -245,7 +245,7 @@ struct STATETREEMODULE_API FStateTreeCondition_CompareDistance : public FStateTr
 	virtual bool TestCondition(FStateTreeExecutionContext& Context) const override;
 
 #if WITH_EDITOR
-	virtual FText GetDescription(const FGuid& ID, FConstStructView InstanceData, const IStateTreeBindingLookup& BindingLookup) const override;
+	virtual FText GetDescription(const FGuid& ID, FStateTreeDataView InstanceData, const IStateTreeBindingLookup& BindingLookup) const override;
 #endif
 
 	UPROPERTY(EditAnywhere, Category = Condition)
@@ -285,7 +285,7 @@ struct STATETREEMODULE_API FStateTreeCondition_Random : public FStateTreeConditi
 	virtual bool TestCondition(FStateTreeExecutionContext& Context) const override;
 
 #if WITH_EDITOR
-	virtual FText GetDescription(const FGuid& ID, FConstStructView InstanceData, const IStateTreeBindingLookup& BindingLookup) const override;
+	virtual FText GetDescription(const FGuid& ID, FStateTreeDataView InstanceData, const IStateTreeBindingLookup& BindingLookup) const override;
 #endif
 
 	TStateTreeInstanceDataPropertyHandle<float> ThresholdHandle;
diff --git a/Engine/Plugins/Runtime/StateTree/Source/StateTreeModule/Public/StateTreeConditionBase.h b/Engine/Plugins/Runtime/StateTree/Source/StateTreeModule/Public/StateTreeConditionBase.h
--- a/Engine/Plugins/Runtime/StateTree/Source/StateTreeModule/Public/StateTreeConditionBase.h
+++ b/Engine/Plugins/Runtime/StateTree/Source/StateTreeModule/Public/StateTreeConditionBase.h
@@ -4,8 +4,9 @@
 
 #include "CoreMinimal.h"
 #include "StateTreeTypes.h"
+#include "StateTreeExecutionContext.h"
 #if WITH_EDITOR
-#include "InstancedStruct.h"
+#include "StateTreePropertyBindings.h"
 #endif
 #include "StateTreeConditionBase.generated.h"
 
@@ -13,7 +14,6 @@
 struct IStateTreeBindingLookup;
 struct FStateTreeEditorPropertyPath;
 #endif
-struct FStateTreeExecutionContext;
 
 /**
  * Base struct for all conditions.
@@ -25,7 +25,7 @@ struct STATETREEMODULE_API FStateTreeConditionBase
 
 	virtual ~FStateTreeConditionBase() {}
 
-	/** @return Struct that represents the runtime data of the evaluator. */
+	/** @return Struct that represents the runtime data of the condition. */
 	virtual const UStruct* GetInstanceDataType() const { return nullptr; };
 
 	/**
@@ -39,14 +39,16 @@ struct STATETREEMODULE_API FStateTreeConditionBase
 
 #if WITH_EDITOR
 	/** @return Rich text description of the condition. */
-	virtual FText GetDescription(const FGuid& ID, FConstStructView LayoutData, const IStateTreeBindingLookup& BindingLookup) const { return FText::GetEmpty(); }
+	virtual FText GetDescription(const FGuid& ID, FStateTreeDataView InstanceData, const IStateTreeBindingLookup& BindingLookup) const { return FText::GetEmpty(); }
 	/**
 	 * Called when binding of any of the properties in the condition changes.
+	 * @param ID ID of the item, can be used make property paths to this item.
+	 * @param InstanceData view to the instance data, can be struct or class. 
 	 * @param SourcePath Source path of the new binding.
 	 * @param TargetPath Target path of the new binding (the property in the condition).
 	 * @param BindingLookup Reference to binding lookup which can be used to reason about property paths.
 	 */
-	virtual void OnBindingChanged(const FGuid& ID, FStructView LayoutData, const FStateTreeEditorPropertyPath& SourcePath, const FStateTreeEditorPropertyPath& TargetPath, const IStateTreeBindingLookup& BindingLookup) {}
+	virtual void OnBindingChanged(const FGuid& ID, FStateTreeDataView InstanceData, const FStateTreeEditorPropertyPath& SourcePath, const FStateTreeEditorPropertyPath& TargetPath, const IStateTreeBindingLookup& BindingLookup) {}
 #endif
 	/** @return True if the condition passes. */
 	virtual bool TestCondition(FStateTreeExecutionContext& Context) const { return false; }
diff --git a/Engine/Plugins/Runtime/StateTree/Source/StateTreeModule/Public/StateTreeTaskBase.h b/Engine/Plugins/Runtime/StateTree/Source/StateTreeModule/Public/StateTreeTaskBase.h
--- a/Engine/Plugins/Runtime/StateTree/Source/StateTreeModule/Public/StateTreeTaskBase.h
+++ b/Engine/Plugins/Runtime/StateTree/Source/StateTreeModule/Public/StateTreeTaskBase.h
@@ -21,7 +21,7 @@ struct STATETREEMODULE_API FStateTreeTaskBase
 
 	virtual ~FStateTreeTaskBase() {}
 
-	/** @return Struct that represents the runtime data of the evaluator. */
+	/** @return Struct that represents the runtime data of the task. */
 	virtual const UStruct* GetInstanceDataType() const { return nullptr; };
 
 	/**
