StateTree: Gameplay Tag conditions

diff --git a/Engine/Plugins/Runtime/StateTree/Source/StateTreeModule/Private/Conditions/GameplayTagConditions.cpp b/Engine/Plugins/Runtime/StateTree/Source/StateTreeModule/Private/Conditions/GameplayTagConditions.cpp
new file mode 100644
--- /dev/null
+++ b/Engine/Plugins/Runtime/StateTree/Source/StateTreeModule/Private/Conditions/GameplayTagConditions.cpp
@@ -0,0 +1,227 @@
+﻿// Copyright Epic Games, Inc. All Rights Reserved.
+
+#include "Conditions/GameplayTagConditions.h"
+#include "GameplayTagContainer.h"
+#include "StateTreeExecutionContext.h"
+#if WITH_EDITOR
+#include "StateTreePropertyBindings.h"
+#endif// WITH_EDITOR
+
+#if WITH_EDITOR
+#define LOCTEXT_NAMESPACE "StateTreeEditor"
+
+namespace UE::StateTree::Conditions
+{
+	FText GetContainerAsText(const FGameplayTagContainer& TagContainer, const int ApproxMaxLength = 60)
+	{
+		FString Combined;
+		for (const FGameplayTag& Tag : TagContainer)
+		{
+			FString TagString = Tag.ToString();
+
+			if (Combined.Len() > 0)
+			{
+				Combined += TEXT(", ");
+			}
+			
+			if (Combined.Len() + TagString.Len() > ApproxMaxLength)
+			{
+				// Overflow
+				if (Combined.Len() == 0)
+				{
+					Combined += TagString.Left(ApproxMaxLength);
+				}
+				Combined += TEXT("...");
+				break;
+			}
+
+			Combined += TagString;
+		}
+
+		return FText::FromString(Combined);
+	}
+	
+}
+
+#endif// WITH_EDITOR
+
+
+//----------------------------------------------------------------------//
+//  FGameplayTagMatchCondition
+//----------------------------------------------------------------------//
+
+bool FGameplayTagMatchCondition::Link(FStateTreeLinker& Linker)
+{
+	Linker.LinkInstanceDataProperty(TagContainerHandle, STATETREE_INSTANCEDATA_PROPERTY(FGameplayTagMatchConditionInstanceData, TagContainer));
+	Linker.LinkInstanceDataProperty(TagHandle, STATETREE_INSTANCEDATA_PROPERTY(FGameplayTagMatchConditionInstanceData, Tag));
+
+	return true;
+}
+
+bool FGameplayTagMatchCondition::TestCondition(FStateTreeExecutionContext& Context) const
+{
+	const FGameplayTagContainer& TagContainer = Context.GetInstanceData(TagContainerHandle);
+	const FGameplayTag& Tag = Context.GetInstanceData(TagHandle);
+
+	return (bExactMatch ? TagContainer.HasTagExact(Tag) : TagContainer.HasTag(Tag)) ^ bInvert;
+}
+
+#if WITH_EDITOR
+FText FGameplayTagMatchCondition::GetDescription(const FGuid& ID, FStateTreeDataView InstanceData, const IStateTreeBindingLookup& BindingLookup) const
+{
+	const FGameplayTagMatchConditionInstanceData& Instance = InstanceData.Get<FGameplayTagMatchConditionInstanceData>();
+	const FStateTreeEditorPropertyPath TagContainerPath(ID, GET_MEMBER_NAME_STRING_CHECKED(FGameplayTagMatchConditionInstanceData, TagContainer));
+	const FStateTreeEditorPropertyPath TagPath(ID, GET_MEMBER_NAME_STRING_CHECKED(FGameplayTagMatchConditionInstanceData, Tag));
+
+	FText InvertText;
+	if (bInvert)
+	{
+		InvertText = LOCTEXT("Not", "Not");
+	}
+
+	FText TagContainerText;
+	if (const FStateTreeEditorPropertyPath* Binding = BindingLookup.GetPropertyBindingSource(TagContainerPath))
+	{
+		TagContainerText = BindingLookup.GetPropertyPathDisplayName(*Binding);
+	}
+	else
+	{
+		TagContainerText = UE::StateTree::Conditions::GetContainerAsText(Instance.TagContainer);
+	}
+
+	FText TagText;
+	if (const FStateTreeEditorPropertyPath* Binding = BindingLookup.GetPropertyBindingSource(TagContainerPath))
+	{
+		TagText = BindingLookup.GetPropertyPathDisplayName(*Binding);
+	}
+	else
+	{
+		TagText = FText::FromString(Instance.Tag.ToString());
+	}
+
+	return FText::Format(LOCTEXT("GameplayTagMatchDesc", "{0} <Details.Bold>{1}</> has <Details.Bold>{2}</>"),
+		InvertText, TagContainerText, TagText);
+}
+#endif// WITH_EDITOR
+
+//----------------------------------------------------------------------//
+//  FGameplayTagContainerMatchCondition
+//----------------------------------------------------------------------//
+
+bool FGameplayTagContainerMatchCondition::Link(FStateTreeLinker& Linker)
+{
+	Linker.LinkInstanceDataProperty(TagContainerHandle, STATETREE_INSTANCEDATA_PROPERTY(FGameplayTagContainerMatchConditionInstanceData, TagContainer));
+	Linker.LinkInstanceDataProperty(OtherContainerHandle, STATETREE_INSTANCEDATA_PROPERTY(FGameplayTagContainerMatchConditionInstanceData, OtherContainer));
+
+	return true;
+}
+
+bool FGameplayTagContainerMatchCondition::TestCondition(FStateTreeExecutionContext& Context) const
+{
+	const FGameplayTagContainer& TagContainer = Context.GetInstanceData(TagContainerHandle);
+	const FGameplayTagContainer& OtherContainer = Context.GetInstanceData(OtherContainerHandle);
+
+	bool bResult = false;
+	switch (MatchType)
+	{
+	case EGameplayContainerMatchType::Any:
+		bResult = bExactMatch ? TagContainer.HasAnyExact(OtherContainer) : TagContainer.HasAny(OtherContainer);
+		break;
+	case EGameplayContainerMatchType::All:
+		bResult = bExactMatch ? TagContainer.HasAllExact(OtherContainer) : TagContainer.HasAll(OtherContainer);
+		break;
+	default:
+		ensureMsgf(false, TEXT("Unhandled match type %s."), *UEnum::GetValueAsString(MatchType));
+	}
+	
+	return bResult ^ bInvert;
+}
+
+#if WITH_EDITOR
+FText FGameplayTagContainerMatchCondition::GetDescription(const FGuid& ID, FStateTreeDataView InstanceData, const IStateTreeBindingLookup& BindingLookup) const
+{
+	const FGameplayTagContainerMatchConditionInstanceData& Instance = InstanceData.Get<FGameplayTagContainerMatchConditionInstanceData>();
+	const FStateTreeEditorPropertyPath TagContainerPath(ID, GET_MEMBER_NAME_STRING_CHECKED(FGameplayTagContainerMatchConditionInstanceData, TagContainer));
+	const FStateTreeEditorPropertyPath OtherContainerPath(ID, GET_MEMBER_NAME_STRING_CHECKED(FGameplayTagContainerMatchConditionInstanceData, OtherContainer));
+
+	FText InvertText;
+	if (bInvert)
+	{
+		InvertText = LOCTEXT("Not", "Not");
+	}
+
+	FText TagContainerText;
+	if (const FStateTreeEditorPropertyPath* Binding = BindingLookup.GetPropertyBindingSource(TagContainerPath))
+	{
+		TagContainerText = BindingLookup.GetPropertyPathDisplayName(*Binding);
+	}
+	else
+	{
+		TagContainerText = LOCTEXT("NotBound", "Not Bound");
+	}
+
+	FText MatchTypeText = UEnum::GetDisplayValueAsText(MatchType);
+
+	FText OtherContainerText;
+	if (const FStateTreeEditorPropertyPath* Binding = BindingLookup.GetPropertyBindingSource(TagContainerPath))
+	{
+		OtherContainerText = BindingLookup.GetPropertyPathDisplayName(*Binding);
+	}
+	else
+	{
+		OtherContainerText = UE::StateTree::Conditions::GetContainerAsText(Instance.TagContainer);
+	}
+
+	return FText::Format(LOCTEXT("GameplayTagMatchDesc", "{0} <Details.Bold>{1}</> {2} <Details.Bold>{3}</>"),
+		InvertText, TagContainerText, MatchTypeText, OtherContainerText);
+}
+#endif// WITH_EDITOR
+
+//----------------------------------------------------------------------//
+//  FGameplayTagQueryCondition
+//----------------------------------------------------------------------//
+
+bool FGameplayTagQueryCondition::Link(FStateTreeLinker& Linker)
+{
+	Linker.LinkInstanceDataProperty(TagContainerHandle, STATETREE_INSTANCEDATA_PROPERTY(FGameplayTagQueryConditionInstanceData, TagContainer));
+
+	return true;
+}
+
+bool FGameplayTagQueryCondition::TestCondition(FStateTreeExecutionContext& Context) const
+{
+	const FGameplayTagContainer& TagContainer = Context.GetInstanceData(TagContainerHandle);
+	return TagQuery.Matches(TagContainer) ^ bInvert;
+}
+
+#if WITH_EDITOR
+FText FGameplayTagQueryCondition::GetDescription(const FGuid& ID, FStateTreeDataView InstanceData, const IStateTreeBindingLookup& BindingLookup) const
+{
+	const FGameplayTagQueryConditionInstanceData& Instance = InstanceData.Get<FGameplayTagQueryConditionInstanceData>();
+	const FStateTreeEditorPropertyPath TagContainerPath(ID, GET_MEMBER_NAME_STRING_CHECKED(FGameplayTagQueryConditionInstanceData, TagContainer));
+
+	FText InvertText;
+	if (bInvert)
+	{
+		InvertText = LOCTEXT("Not", "Not");
+	}
+
+	FText TagContainerText;
+	if (const FStateTreeEditorPropertyPath* Binding = BindingLookup.GetPropertyBindingSource(TagContainerPath))
+	{
+		TagContainerText = BindingLookup.GetPropertyPathDisplayName(*Binding);
+	}
+	else
+	{
+		TagContainerText = LOCTEXT("NotBound", "Not Bound");
+	}
+
+	return FText::Format(LOCTEXT("GameplayTagMatchDesc", "{0} <Details.Bold>{1}</> matches query</>"),
+		InvertText, TagContainerText);
+}
+#endif// WITH_EDITOR
+
+
+#if WITH_EDITOR
+#undef LOCTEXT_NAMESPACE
+#endif // WITH_EDITOR
diff --git a/Engine/Plugins/Runtime/StateTree/Source/StateTreeModule/Public/Conditions/GameplayTagConditions.h b/Engine/Plugins/Runtime/StateTree/Source/StateTreeModule/Public/Conditions/GameplayTagConditions.h
new file mode 100644
--- /dev/null
+++ b/Engine/Plugins/Runtime/StateTree/Source/StateTreeModule/Public/Conditions/GameplayTagConditions.h
@@ -0,0 +1,128 @@
+﻿// Copyright Epic Games, Inc. All Rights Reserved.
+
+#pragma once
+
+#include "CoreMinimal.h"
+#include "StateTreeConditionBase.h"
+#include "Conditions/StateTreeCondition_Common.h"
+#include "GameplayTagContainer.h"
+#include "GameplayTagConditions.generated.h"
+
+/**
+ * Gameplay Tag match condition.
+ */
+USTRUCT()
+struct STATETREEMODULE_API FGameplayTagMatchConditionInstanceData
+{
+	GENERATED_BODY()
+	
+	UPROPERTY(EditAnywhere, Category = Parameter)
+	FGameplayTagContainer TagContainer;
+
+	UPROPERTY(EditAnywhere, Category = Parameter)
+	FGameplayTag Tag;
+};
+
+USTRUCT(DisplayName="Gameplay Tag Match")
+struct STATETREEMODULE_API FGameplayTagMatchCondition : public FStateTreeConditionCommonBase
+{
+	GENERATED_BODY()
+
+	FGameplayTagMatchCondition() = default;
+	
+	virtual bool Link(FStateTreeLinker& Linker) override;
+	virtual const UStruct* GetInstanceDataType() const override { return FGameplayTagMatchConditionInstanceData::StaticStruct(); }
+	virtual bool TestCondition(FStateTreeExecutionContext& Context) const override;
+
+#if WITH_EDITOR
+	virtual FText GetDescription(const FGuid& ID, FStateTreeDataView InstanceData, const IStateTreeBindingLookup& BindingLookup) const override;
+#endif
+
+	TStateTreeInstanceDataPropertyHandle<FGameplayTagContainer> TagContainerHandle;
+	TStateTreeInstanceDataPropertyHandle<FGameplayTag> TagHandle;
+	
+	UPROPERTY(EditAnywhere, Category = Condition)
+	bool bExactMatch = false;
+
+	UPROPERTY(EditAnywhere, Category = Condition)
+	bool bInvert = false;
+};
+
+/**
+* Gameplay Tag Container match condition.
+*/
+USTRUCT()
+struct STATETREEMODULE_API FGameplayTagContainerMatchConditionInstanceData
+{
+	GENERATED_BODY()
+	
+	UPROPERTY(EditAnywhere, Category = Input)
+	FGameplayTagContainer TagContainer;
+
+	UPROPERTY(EditAnywhere, Category = Parameter)
+	FGameplayTagContainer OtherContainer;
+};
+
+USTRUCT(DisplayName="Gameplay Tag Container Match")
+struct STATETREEMODULE_API FGameplayTagContainerMatchCondition : public FStateTreeConditionCommonBase
+{
+	GENERATED_BODY()
+
+	FGameplayTagContainerMatchCondition() = default;
+	
+	virtual bool Link(FStateTreeLinker& Linker) override;
+	virtual const UStruct* GetInstanceDataType() const override { return FGameplayTagContainerMatchConditionInstanceData::StaticStruct(); }
+	virtual bool TestCondition(FStateTreeExecutionContext& Context) const override;
+
+#if WITH_EDITOR
+	virtual FText GetDescription(const FGuid& ID, FStateTreeDataView InstanceData, const IStateTreeBindingLookup& BindingLookup) const override;
+#endif
+
+	TStateTreeInstanceDataPropertyHandle<FGameplayTagContainer> TagContainerHandle;
+	TStateTreeInstanceDataPropertyHandle<FGameplayTagContainer> OtherContainerHandle;
+
+	UPROPERTY(EditAnywhere, Category = Condition)
+	EGameplayContainerMatchType MatchType = EGameplayContainerMatchType::Any;
+
+	UPROPERTY(EditAnywhere, Category = Condition)
+	bool bExactMatch = false;
+
+	UPROPERTY(EditAnywhere, Category = Condition)
+	bool bInvert = false;
+};
+
+/**
+* Gameplay Tag Query match condition.
+*/
+USTRUCT()
+struct STATETREEMODULE_API FGameplayTagQueryConditionInstanceData
+{
+	GENERATED_BODY()
+	
+	UPROPERTY(EditAnywhere, Category = Input)
+	FGameplayTagContainer TagContainer;
+};
+
+USTRUCT(DisplayName="Gameplay Tag Query")
+struct STATETREEMODULE_API FGameplayTagQueryCondition : public FStateTreeConditionCommonBase
+{
+	GENERATED_BODY()
+
+	FGameplayTagQueryCondition() = default;
+	
+	virtual bool Link(FStateTreeLinker& Linker) override;
+	virtual const UStruct* GetInstanceDataType() const override { return FGameplayTagQueryConditionInstanceData::StaticStruct(); }
+	virtual bool TestCondition(FStateTreeExecutionContext& Context) const override;
+
+#if WITH_EDITOR
+	virtual FText GetDescription(const FGuid& ID, FStateTreeDataView InstanceData, const IStateTreeBindingLookup& BindingLookup) const override;
+#endif
+
+	TStateTreeInstanceDataPropertyHandle<FGameplayTagContainer> TagContainerHandle;
+
+	UPROPERTY(EditAnywhere, Category = Condition)
+	FGameplayTagQuery TagQuery;
+
+	UPROPERTY(EditAnywhere, Category = Condition)
+	bool bInvert = false;
+};
\ No newline at end of file
