More Non-Unity fixes..

diff --git a/Engine/Plugins/Animation/ControlRig/Source/ControlRigEditor/Private/EditMode/SControlRigDetails.cpp b/Engine/Plugins/Animation/ControlRig/Source/ControlRigEditor/Private/EditMode/SControlRigDetails.cpp
--- a/Engine/Plugins/Animation/ControlRig/Source/ControlRigEditor/Private/EditMode/SControlRigDetails.cpp
+++ b/Engine/Plugins/Animation/ControlRig/Source/ControlRigEditor/Private/EditMode/SControlRigDetails.cpp
@@ -25,6 +25,7 @@
 #include "Widgets/Layout/SScrollBox.h"
 #include "ControlRigControlsProxy.h"
 #include "ControlRigEditModeSettings.h"
+#include "Modules/ModuleManager.h"
 #include "TimerManager.h"
 
 #define LOCTEXT_NAMESPACE "ControlRigDetails"
diff --git a/Engine/Plugins/Runtime/StateTree/Source/StateTreeModule/Private/Blueprint/StateTreeConditionBlueprintBase.cpp b/Engine/Plugins/Runtime/StateTree/Source/StateTreeModule/Private/Blueprint/StateTreeConditionBlueprintBase.cpp
--- a/Engine/Plugins/Runtime/StateTree/Source/StateTreeModule/Private/Blueprint/StateTreeConditionBlueprintBase.cpp
+++ b/Engine/Plugins/Runtime/StateTree/Source/StateTreeModule/Private/Blueprint/StateTreeConditionBlueprintBase.cpp
@@ -1,8 +1,9 @@
-﻿// Copyright Epic Games, Inc. All Rights Reserved.
+// Copyright Epic Games, Inc. All Rights Reserved.
 
 #include "Blueprint/StateTreeConditionBlueprintBase.h"
 #include "CoreMinimal.h"
 #include "StateTreeExecutionContext.h"
+#include "Engine/Blueprint.h"
 
 //----------------------------------------------------------------------//
 //  UStateTreeConditionBlueprintBase
