StateTree: allow schema to configure editor
- allow to disable enter conditions and evaluators
- allow to configure state to have just one task

diff --git a/Engine/Plugins/Runtime/StateTree/Source/StateTreeEditorModule/Private/Customizations/StateTreeBindableItemDetails.cpp b/Engine/Plugins/Runtime/StateTree/Source/StateTreeEditorModule/Private/Customizations/StateTreeBindableItemDetails.cpp
--- a/Engine/Plugins/Runtime/StateTree/Source/StateTreeEditorModule/Private/Customizations/StateTreeBindableItemDetails.cpp
+++ b/Engine/Plugins/Runtime/StateTree/Source/StateTreeEditorModule/Private/Customizations/StateTreeBindableItemDetails.cpp
@@ -522,6 +522,13 @@ FText FStateTreeBindableItemDetails::GetName() const
 
 EVisibility FStateTreeBindableItemDetails::IsNameVisible() const
 {
+	const UStateTreeSchema* Schema = StateTree ? StateTree->GetSchema() : nullptr;
+	if (Schema && Schema->AllowMultipleTasks() == false)
+	{
+		// Single task states use the state name as task name.
+		return EVisibility::Collapsed;
+	}
+	
 	const UScriptStruct* ScriptStruct = nullptr;
 	if (const FStateTreeItem* Item = GetCommonItem())
 	{
diff --git a/Engine/Plugins/Runtime/StateTree/Source/StateTreeEditorModule/Private/Customizations/StateTreeStateDetails.cpp b/Engine/Plugins/Runtime/StateTree/Source/StateTreeEditorModule/Private/Customizations/StateTreeStateDetails.cpp
--- a/Engine/Plugins/Runtime/StateTree/Source/StateTreeEditorModule/Private/Customizations/StateTreeStateDetails.cpp
+++ b/Engine/Plugins/Runtime/StateTree/Source/StateTreeEditorModule/Private/Customizations/StateTreeStateDetails.cpp
@@ -44,7 +44,25 @@ TSharedRef<IDetailCustomization> FStateTreeStateDetails::MakeInstance()
 
 void FStateTreeStateDetails::CustomizeDetails(IDetailLayoutBuilder& DetailBuilder)
 {
+	// Find StateTree associated with this panel.
+	UStateTree* StateTree = nullptr;
+	TArray<TWeakObjectPtr<UObject>> Objects;
+	DetailBuilder.GetObjectsBeingCustomized(Objects);
+	for (TWeakObjectPtr<UObject>& WeakObject : Objects)
+	{
+		if (UObject* Object = WeakObject.Get())
+		{
+			if (UStateTree* OuterStateTree = Object->GetTypedOuter<UStateTree>())
+			{
+				StateTree = OuterStateTree;
+				break;
+			}
+		}
+	}
+	const UStateTreeSchema* Schema = StateTree ? StateTree->GetSchema() : nullptr;
+
 	TSharedPtr<IPropertyHandle> TasksProperty = DetailBuilder.GetProperty(TEXT("Tasks"));
+	TSharedPtr<IPropertyHandle> SingleTaskProperty = DetailBuilder.GetProperty(TEXT("SingleTask"));
 	TSharedPtr<IPropertyHandle> EnterConditionsProperty = DetailBuilder.GetProperty(TEXT("EnterConditions"));
 	TSharedPtr<IPropertyHandle> EvaluatorsProperty = DetailBuilder.GetProperty(TEXT("Evaluators"));
 	TSharedPtr<IPropertyHandle> TransitionsProperty = DetailBuilder.GetProperty(TEXT("Transitions"));
@@ -52,11 +70,43 @@ void FStateTreeStateDetails::CustomizeDetails(IDetailLayoutBuilder& DetailBuilde
 	IDetailCategoryBuilder& StateCategory = DetailBuilder.EditCategory(TEXT("State"), LOCTEXT("StateDetailsState", "State"));
 	StateCategory.SetSortOrder(0);
 
-	MakeArrayCategory(DetailBuilder, "Evaluators", LOCTEXT("StateDetailsEvaluators", "Evaluators"), 1, EvaluatorsProperty);
-
-	MakeArrayCategory(DetailBuilder, "Enter Conditions", LOCTEXT("StateDetailsEnterConditions", "Enter Conditions"), 2, EnterConditionsProperty);
-
-	MakeArrayCategory(DetailBuilder, "Tasks", LOCTEXT("StateDetailsTasks", "Tasks"), 3, TasksProperty);
+	const FName EvalCategoryName(TEXT("Evaluators"));
+	if (Schema && Schema->AllowEvaluators())
+	{
+		MakeArrayCategory(DetailBuilder, EvalCategoryName, LOCTEXT("StateDetailsEvaluators", "Evaluators"), 1, EvaluatorsProperty);
+	}
+	else
+	{
+		DetailBuilder.EditCategory(EvalCategoryName).SetCategoryVisibility(false);
+	}
+
+	const FName EnterConditionsCategoryName(TEXT("Enter Conditions"));
+	if (Schema && Schema->AllowEnterConditions())
+	{
+		MakeArrayCategory(DetailBuilder, EnterConditionsCategoryName, LOCTEXT("StateDetailsEnterConditions", "Enter Conditions"), 2, EnterConditionsProperty);
+	}
+	else
+	{
+		DetailBuilder.EditCategory(EnterConditionsCategoryName).SetCategoryVisibility(false);
+	}
+
+	if (Schema && Schema->AllowMultipleTasks())
+	{
+		const FName TasksCategoryName(TEXT("Tasks"));
+		MakeArrayCategory(DetailBuilder, TasksCategoryName, LOCTEXT("StateDetailsTasks", "Tasks"), 3, TasksProperty);
+		SingleTaskProperty->MarkHiddenByCustomization();
+	}
+	else
+	{
+		const FName TaskCategoryName(TEXT("Task"));
+		IDetailCategoryBuilder& Category = DetailBuilder.EditCategory(TaskCategoryName);
+		Category.SetSortOrder(3);
+		
+		IDetailPropertyRow& Row = Category.AddProperty(SingleTaskProperty);
+		Row.ShouldAutoExpand(true);
+		
+		TasksProperty->MarkHiddenByCustomization();
+	}
 
 	MakeArrayCategory(DetailBuilder, "Transitions", LOCTEXT("StateDetailsTransitions", "Transitions"), 4, TransitionsProperty);
 }
@@ -86,5 +136,4 @@ void FStateTreeStateDetails::GenerateArrayElementWidget(TSharedRef<IPropertyHand
 	ChildrenBuilder.AddProperty(PropertyHandle);
 };
 
-
 #undef LOCTEXT_NAMESPACE
diff --git a/Engine/Plugins/Runtime/StateTree/Source/StateTreeEditorModule/Private/Customizations/StateTreeStateDetails.h b/Engine/Plugins/Runtime/StateTree/Source/StateTreeEditorModule/Private/Customizations/StateTreeStateDetails.h
--- a/Engine/Plugins/Runtime/StateTree/Source/StateTreeEditorModule/Private/Customizations/StateTreeStateDetails.h
+++ b/Engine/Plugins/Runtime/StateTree/Source/StateTreeEditorModule/Private/Customizations/StateTreeStateDetails.h
@@ -4,7 +4,6 @@
 
 #include "CoreMinimal.h"
 #include "IDetailCustomization.h"
-#include "Input/Reply.h"
 
 class IDetailLayoutBuilder;
 class IDetailChildrenBuilder;
diff --git a/Engine/Plugins/Runtime/StateTree/Source/StateTreeEditorModule/Private/StateTreeBaker.cpp b/Engine/Plugins/Runtime/StateTree/Source/StateTreeEditorModule/Private/StateTreeBaker.cpp
--- a/Engine/Plugins/Runtime/StateTree/Source/StateTreeEditorModule/Private/StateTreeBaker.cpp
+++ b/Engine/Plugins/Runtime/StateTree/Source/StateTreeEditorModule/Private/StateTreeBaker.cpp
@@ -348,6 +348,182 @@ bool FStateTreeBaker::GetAndValidateBindings(const FStateTreeBindableStructDesc&
 	return true;
 }
 
+bool FStateTreeBaker::CreateTask(const FStateTreeTaskItem& TaskItem)
+{
+	// Silently ignore empty items.
+	if (!TaskItem.Item.IsValid())
+	{
+		return true;
+	}
+	
+	// Create binding source struct descriptor.
+	FStateTreeBindableStructDesc StructDesc;
+	StructDesc.ID = TaskItem.ID;
+	StructDesc.Name = TaskItem.Item.GetScriptStruct()->GetFName();
+
+	// Check that item has valid instance initialized.
+	if (!TaskItem.Instance.IsValid() && TaskItem.InstanceObject == nullptr)
+	{
+		Log.Reportf(EMessageSeverity::Error, StructDesc,
+			TEXT("Malformed evaluator, missing instance value."));
+		return false;
+	}
+
+	// Copy the task
+	FInstancedStruct& Item = StateTree->Items.AddDefaulted_GetRef();
+	Item = TaskItem.Item;
+	
+	FStateTreeTaskBase& Task = Item.GetMutable<FStateTreeTaskBase>();
+
+	if (TaskItem.Instance.IsValid())
+	{
+		// Struct Instance
+		FInstancedStruct& Instance = StateTree->Instances.AddDefaulted_GetRef();
+		const int32 InstanceIndex = StateTree->Instances.Num() - 1;
+
+		Instance = TaskItem.Instance;
+
+		// Create binding source struct descriptor.
+		StructDesc.Struct = Instance.GetScriptStruct();
+		StructDesc.Name = Task.Name;
+
+		check(InstanceIndex <= int32(MAX_uint16));
+		Task.InstanceIndex = uint16(InstanceIndex);
+		Task.bInstanceIsObject = false;
+	}
+	else
+	{
+		// Object Instance
+		check(TaskItem.InstanceObject != nullptr);
+
+		UObject* Instance = DuplicateObject(TaskItem.InstanceObject, StateTree);
+
+		StateTree->InstanceObjects.Add(Instance);
+		const int32 InstanceIndex = StateTree->InstanceObjects.Num() - 1;
+
+		// Create binding source struct descriptor.
+		StructDesc.Struct = Instance->GetClass();
+		StructDesc.Name = Task.Name;
+
+		check(InstanceIndex <= int32(MAX_uint16));
+		Task.InstanceIndex = uint16(InstanceIndex);
+		Task.bInstanceIsObject = true;
+	}
+
+	// Mark the instance as binding source.
+	const int32 SourceStructIndex = BindingsCompiler.AddSourceStruct(StructDesc);
+
+	// Check that the bindings for this struct are still all valid.
+	TArray<FStateTreeEditorPropertyBinding> Bindings;
+	if (!GetAndValidateBindings(StructDesc, Bindings))
+	{
+		return false;
+	}
+
+	// Compile batch copy for this struct, we pass in all the bindings, the compiler will pick up the ones for the target structs.
+	int32 BatchIndex = INDEX_NONE;
+	if (!BindingsCompiler.CompileBatch(StructDesc, Bindings, BatchIndex))
+	{
+		return false;
+	}
+	
+	check(BatchIndex < int32(MAX_uint16));
+	Task.BindingsBatch = BatchIndex == INDEX_NONE ? FStateTreeHandle::Invalid : FStateTreeHandle(uint16(BatchIndex));
+	
+	check(SourceStructIndex <= int32(MAX_uint16));
+	Task.DataViewIndex = uint16(SourceStructIndex);
+
+	return true;
+}
+
+bool FStateTreeBaker::CreateEvaluator(const FStateTreeEvaluatorItem& EvaluatorItem)
+{
+	// Silently ignore empty items.
+	if (!EvaluatorItem.Item.IsValid())
+	{
+		return true;
+	}
+	
+	// Create binding source struct descriptor.
+	FStateTreeBindableStructDesc StructDesc;
+    StructDesc.ID = EvaluatorItem.ID;
+    StructDesc.Name = EvaluatorItem.Item.GetScriptStruct()->GetFName();
+
+    // Check that item has valid instance initialized.
+    if (!EvaluatorItem.Instance.IsValid() && EvaluatorItem.InstanceObject == nullptr)
+    {
+        Log.Reportf(EMessageSeverity::Error, StructDesc,
+        	TEXT("Malformed evaluator, missing instance value."));
+        return false;
+    }
+
+	// Copy the evaluator
+	FInstancedStruct& Item = StateTree->Items.AddDefaulted_GetRef();
+	Item = EvaluatorItem.Item;
+
+	FStateTreeEvaluatorBase& Eval = Item.GetMutable<FStateTreeEvaluatorBase>();
+
+	if (EvaluatorItem.Instance.IsValid())
+	{
+		// Struct Instance
+		FInstancedStruct& Instance = StateTree->Instances.AddDefaulted_GetRef();
+		const int32 InstanceIndex = StateTree->Instances.Num() - 1;
+
+		Instance = EvaluatorItem.Instance;
+
+		// Create binding source struct descriptor.
+		StructDesc.Struct = Instance.GetScriptStruct();
+		StructDesc.Name = Eval.Name;
+
+		check(InstanceIndex <= int32(MAX_uint16));
+		Eval.InstanceIndex = uint16(InstanceIndex);
+		Eval.bInstanceIsObject = false;
+	}
+	else
+	{
+		// Object Instance
+		check(EvaluatorItem.InstanceObject != nullptr);
+
+		UObject* Instance = DuplicateObject(EvaluatorItem.InstanceObject, StateTree);
+
+		StateTree->InstanceObjects.Add(Instance);
+		const int32 InstanceIndex = StateTree->InstanceObjects.Num() - 1;
+
+		// Create binding source struct descriptor.
+		StructDesc.Struct = Instance->GetClass();
+		StructDesc.Name = Eval.Name;
+
+		check(InstanceIndex <= int32(MAX_uint16));
+		Eval.InstanceIndex = uint16(InstanceIndex);
+		Eval.bInstanceIsObject = true;
+	}
+		
+	// Mark the instance as binding source.
+	const int32 SourceStructIndex = BindingsCompiler.AddSourceStruct(StructDesc);
+
+	// Check that the bindings for this struct are still all valid.
+	TArray<FStateTreeEditorPropertyBinding> Bindings;
+	if (!GetAndValidateBindings(StructDesc, Bindings))
+	{
+		return false;
+	}
+
+	// Compile batch copy for this struct, we pass in all the bindings, the compiler will pick up the ones for the target structs.
+	int32 BatchIndex = INDEX_NONE;
+	if (!BindingsCompiler.CompileBatch(StructDesc, Bindings, BatchIndex))
+	{
+		return false;
+	}
+	
+	check(BatchIndex < int32(MAX_uint16));
+	Eval.BindingsBatch = BatchIndex == INDEX_NONE ? FStateTreeHandle::Invalid : FStateTreeHandle(uint16(BatchIndex));
+	
+	check(SourceStructIndex <= int32(MAX_uint16));
+	Eval.DataViewIndex = uint16(SourceStructIndex);
+
+	return true;
+}
+
 bool FStateTreeBaker::CreateStateRecursive(UStateTreeState& State, const FStateTreeHandle Parent)
 {
 	FStateTreeCompilerLogStateScope LogStateScope(&State, Log);
@@ -360,7 +536,6 @@ bool FStateTreeBaker::CreateStateRecursive(UStateTreeState& State, const FStateT
 	SourceStates.Add(&State);
 	IDToState.Add(State.ID, StateIdx);
 
-
 	check(StateTree->Instances.Num() <= int32(MAX_uint16));
 
 	// Collect evaluators
@@ -369,88 +544,10 @@ bool FStateTreeBaker::CreateStateRecursive(UStateTreeState& State, const FStateT
 
 	for (FStateTreeEvaluatorItem& EvaluatorItem : State.Evaluators)
 	{
-		// Silently ignore empty items.
-		if (!EvaluatorItem.Item.IsValid())
-		{
-			continue;
-		}
-		
-		// Create binding source struct descriptor.
-		FStateTreeBindableStructDesc StructDesc;
-        StructDesc.ID = EvaluatorItem.ID;
-        StructDesc.Name = EvaluatorItem.Item.GetScriptStruct()->GetFName();
-    
-        // Check that item has valid instance initialized.
-        if (!EvaluatorItem.Instance.IsValid() && EvaluatorItem.InstanceObject == nullptr)
-        {
-        	Log.Reportf(EMessageSeverity::Error, StructDesc,
-        		TEXT("Malformed evaluator, missing instance value."));
-        	return false;
-        }
-
-		// Copy the evaluator
-		FInstancedStruct& Item = StateTree->Items.AddDefaulted_GetRef();
-		Item = EvaluatorItem.Item;
-
-		FStateTreeEvaluatorBase& Eval = Item.GetMutable<FStateTreeEvaluatorBase>();
-
-		if (EvaluatorItem.Instance.IsValid())
-		{
-			// Struct Instance
-			FInstancedStruct& Instance = StateTree->Instances.AddDefaulted_GetRef();
-			const int32 InstanceIndex = StateTree->Instances.Num() - 1;
-	
-			Instance = EvaluatorItem.Instance;
-
-			// Create binding source struct descriptor.
-			StructDesc.Struct = Instance.GetScriptStruct();
-			StructDesc.Name = Eval.Name;
-
-			check(InstanceIndex <= int32(MAX_uint16));
-			Eval.InstanceIndex = uint16(InstanceIndex);
-			Eval.bInstanceIsObject = false;
-		}
-		else
-		{
-			// Object Instance
-			check(EvaluatorItem.InstanceObject != nullptr);
-	
-			UObject* Instance = DuplicateObject(EvaluatorItem.InstanceObject, StateTree);
-	
-			StateTree->InstanceObjects.Add(Instance);
-			const int32 InstanceIndex = StateTree->InstanceObjects.Num() - 1;
-	
-			// Create binding source struct descriptor.
-			StructDesc.Struct = Instance->GetClass();
-			StructDesc.Name = Eval.Name;
-
-			check(InstanceIndex <= int32(MAX_uint16));
-			Eval.InstanceIndex = uint16(InstanceIndex);
-			Eval.bInstanceIsObject = true;
-		}
-			
-		// Mark the instance as binding source.
-		const int32 SourceStructIndex = BindingsCompiler.AddSourceStruct(StructDesc);
-
-		// Check that the bindings for this struct are still all valid.
-		TArray<FStateTreeEditorPropertyBinding> Bindings;
-		if (!GetAndValidateBindings(StructDesc, Bindings))
+		if (!CreateEvaluator(EvaluatorItem))
 		{
 			return false;
 		}
-
-		// Compile batch copy for this struct, we pass in all the bindings, the compiler will pick up the ones for the target structs.
-		int32 BatchIndex = INDEX_NONE;
-		if (!BindingsCompiler.CompileBatch(StructDesc, Bindings, BatchIndex))
-		{
-			return false;
-		}
-		
-		check(BatchIndex < int32(MAX_uint16));
-		Eval.BindingsBatch = BatchIndex == INDEX_NONE ? FStateTreeHandle::Invalid : FStateTreeHandle(uint16(BatchIndex));
-		
-		check(SourceStructIndex <= int32(MAX_uint16));
-		Eval.DataViewIndex = uint16(SourceStructIndex);
 	}
 	
 	const int32 EvaluatorsNum = StateTree->Items.Num() - int32(BakedState.EvaluatorsBegin);
@@ -463,88 +560,15 @@ bool FStateTreeBaker::CreateStateRecursive(UStateTreeState& State, const FStateT
 
 	for (FStateTreeTaskItem& TaskItem : State.Tasks)
 	{
-		// Silently ignore empty items.
-		if (!TaskItem.Item.IsValid())
-		{
-			continue;
-		}
-		
-		// Create binding source struct descriptor.
-		FStateTreeBindableStructDesc StructDesc;
-		StructDesc.ID = TaskItem.ID;
-		StructDesc.Name = TaskItem.Item.GetScriptStruct()->GetFName();
-    
-		// Check that item has valid instance initialized.
-		if (!TaskItem.Instance.IsValid() && TaskItem.InstanceObject == nullptr)
-		{
-			Log.Reportf(EMessageSeverity::Error, StructDesc,
-				TEXT("Malformed evaluator, missing instance value."));
-			return false;
-		}
-
-		// Copy the task
-		FInstancedStruct& Item = StateTree->Items.AddDefaulted_GetRef();
-		Item = TaskItem.Item;
-		
-		FStateTreeTaskBase& Task = Item.GetMutable<FStateTreeTaskBase>();
-
-		if (TaskItem.Instance.IsValid())
-		{
-			// Struct Instance
-			FInstancedStruct& Instance = StateTree->Instances.AddDefaulted_GetRef();
-			const int32 InstanceIndex = StateTree->Instances.Num() - 1;
-	
-			Instance = TaskItem.Instance;
-
-			// Create binding source struct descriptor.
-			StructDesc.Struct = Instance.GetScriptStruct();
-			StructDesc.Name = Task.Name;
-
-			check(InstanceIndex <= int32(MAX_uint16));
-			Task.InstanceIndex = uint16(InstanceIndex);
-			Task.bInstanceIsObject = false;
-		}
-		else
-		{
-			// Object Instance
-			check(TaskItem.InstanceObject != nullptr);
-	
-			UObject* Instance = DuplicateObject(TaskItem.InstanceObject, StateTree);
-	
-			StateTree->InstanceObjects.Add(Instance);
-			const int32 InstanceIndex = StateTree->InstanceObjects.Num() - 1;
-	
-			// Create binding source struct descriptor.
-			StructDesc.Struct = Instance->GetClass();
-			StructDesc.Name = Task.Name;
-
-			check(InstanceIndex <= int32(MAX_uint16));
-			Task.InstanceIndex = uint16(InstanceIndex);
-			Task.bInstanceIsObject = true;
-		}
-
-		// Mark the instance as binding source.
-		const int32 SourceStructIndex = BindingsCompiler.AddSourceStruct(StructDesc);
-
-		// Check that the bindings for this struct are still all valid.
-		TArray<FStateTreeEditorPropertyBinding> Bindings;
-		if (!GetAndValidateBindings(StructDesc, Bindings))
+		if (!CreateTask(TaskItem))
 		{
 			return false;
 		}
+	}
 
-		// Compile batch copy for this struct, we pass in all the bindings, the compiler will pick up the ones for the target structs.
-		int32 BatchIndex = INDEX_NONE;
-		if (!BindingsCompiler.CompileBatch(StructDesc, Bindings, BatchIndex))
-		{
-			return false;
-		}
-		
-		check(BatchIndex < int32(MAX_uint16));
-		Task.BindingsBatch = BatchIndex == INDEX_NONE ? FStateTreeHandle::Invalid : FStateTreeHandle(uint16(BatchIndex));
-		
-		check(SourceStructIndex <= int32(MAX_uint16));
-		Task.DataViewIndex = uint16(SourceStructIndex);
+	if (!CreateTask(State.SingleTask))
+	{
+		return false;
 	}
 	
 	const int32 TasksNum = StateTree->Items.Num() - int32(BakedState.TasksBegin);
diff --git a/Engine/Plugins/Runtime/StateTree/Source/StateTreeEditorModule/Private/StateTreeEditor.cpp b/Engine/Plugins/Runtime/StateTree/Source/StateTreeEditorModule/Private/StateTreeEditor.cpp
--- a/Engine/Plugins/Runtime/StateTree/Source/StateTreeEditorModule/Private/StateTreeEditor.cpp
+++ b/Engine/Plugins/Runtime/StateTree/Source/StateTreeEditorModule/Private/StateTreeEditor.cpp
@@ -10,14 +10,13 @@
 #include "StateTreeEditorData.h"
 #include "StateTreeDelegates.h"
 #include "StateTreeState.h"
-#include "StateTreeEvaluatorBase.h"
 #include "StateTreeTaskBase.h"
 #include "StateTreeBaker.h"
+#include "StateTreeTypes.h"
 #include "IDetailsView.h"
 #include "Widgets/Docking/SDockTab.h"
 #include "Widgets/Input/SMultiLineEditableTextBox.h"
 #include "Customizations/StateTreeBindingExtension.h"
-#include "Logging/MessageLog.h"
 #include "Misc/UObjectToken.h"
 
 #include "Developer/MessageLog/Public/IMessageLogListing.h"
@@ -375,6 +374,8 @@ void FStateTreeEditor::OnSchemaChanged(const UStateTree& InStateTree)
 {
 	if (StateTree == &InStateTree)
 	{
+		UpdateAsset();
+		
 		if (StateTreeViewModel)
 		{
 			StateTreeViewModel->NotifyAssetChangedExternally();
@@ -533,6 +534,83 @@ namespace UE::StateTree::Editor::Internal
 		}
 	}
 
+	void ApplySchema(UStateTree& StateTree)
+	{
+		UStateTreeEditorData* TreeData = Cast<UStateTreeEditorData>(StateTree.EditorData);
+		if (!TreeData)
+		{
+			return;
+		}
+		
+		const UStateTreeSchema* Schema = StateTree.GetSchema();
+		if (!Schema)
+		{
+			return;
+		}
+
+		TArray<UStateTreeState*> Stack;
+
+		for (UStateTreeState* Routine : TreeData->Routines)
+		{
+			if (Routine)
+			{
+				Stack.Reset();
+				Stack.Add(Routine);
+				while (Stack.Num() > 0)
+				{
+					UStateTreeState* CurState = Stack.Pop();
+
+					// Clear enter conditions if not allowed.
+					if (Schema->AllowEnterConditions() == false && CurState->EnterConditions.Num() > 0)
+					{
+						UE_LOG(LogStateTree, Warning, TEXT("%s: Resetting Enter Conditions in state %s due to current schema restrictions."), *GetNameSafe(&StateTree), *GetNameSafe(CurState));
+						CurState->EnterConditions.Reset();
+					}
+
+					// Clear evaluators if not allowed.
+					if (Schema->AllowEvaluators() == false && CurState->Evaluators.Num() > 0)
+					{
+						UE_LOG(LogStateTree, Warning, TEXT("%s: Resetting Evaluators in state %s due to current schema restrictions."), *GetNameSafe(&StateTree), *GetNameSafe(CurState));
+						CurState->Evaluators.Reset();
+					}
+
+					// Keep single and many tasks based on what is allowed.
+					if (Schema->AllowMultipleTasks() == false)
+					{
+						if (CurState->Tasks.Num() > 0)
+						{
+							CurState->Tasks.Reset();
+							UE_LOG(LogStateTree, Warning, TEXT("%s: Resetting Tasks in state %s due to current schema restrictions."), *GetNameSafe(&StateTree), *GetNameSafe(CurState));
+						}
+						
+						// Task name is the same as state name.
+						if (FStateTreeTaskBase* Task = CurState->SingleTask.Item.GetMutablePtr<FStateTreeTaskBase>())
+						{
+							Task->Name = CurState->Name;
+						}
+					}
+					else
+					{
+						if (CurState->SingleTask.Item.IsValid())
+						{
+							CurState->SingleTask.Reset();
+							UE_LOG(LogStateTree, Warning, TEXT("%s: Resetting Single Task in state %s due to current schema restrictions."), *GetNameSafe(&StateTree), *GetNameSafe(CurState));
+						}
+					}
+					
+					for (UStateTreeState* ChildState : CurState->Children)
+					{
+						if (ChildState)
+						{
+							Stack.Append(CurState->Children);
+						}
+					}
+				}
+			}
+		}
+
+	}
+
 	void RemoveUnusedBindings(UStateTree& StateTree)
 	{
 		UStateTreeEditorData* TreeData = Cast<UStateTreeEditorData>(StateTree.EditorData);
@@ -555,6 +633,7 @@ void FStateTreeEditor::UpdateAsset()
 	}
 
 	UE::StateTree::Editor::Internal::UpdateParents(*StateTree);
+	UE::StateTree::Editor::Internal::ApplySchema(*StateTree);
 	UE::StateTree::Editor::Internal::RemoveUnusedBindings(*StateTree);
 	UE::StateTree::Editor::Internal::ValidateLinkedStates(*StateTree);
 
diff --git a/Engine/Plugins/Runtime/StateTree/Source/StateTreeEditorModule/Private/StateTreeEditorData.cpp b/Engine/Plugins/Runtime/StateTree/Source/StateTreeEditorModule/Private/StateTreeEditorData.cpp
--- a/Engine/Plugins/Runtime/StateTree/Source/StateTreeEditorModule/Private/StateTreeEditorData.cpp
+++ b/Engine/Plugins/Runtime/StateTree/Source/StateTreeEditorModule/Private/StateTreeEditorData.cpp
@@ -185,6 +185,18 @@ void UStateTreeEditorData::VisitHierarchy(TFunctionRef<bool(const UStateTreeStat
 				}
 			}
 			if (bContinue)
+			{
+				if (const FStateTreeTaskBase* Task = State->SingleTask.Item.GetPtr<FStateTreeTaskBase>())
+				{
+					if (!InFunc(*State, State->SingleTask.ID, Task->Name, EStateTreeItemType::Task, State->SingleTask.Item.GetScriptStruct(), Task->GetInstanceDataType()))
+					{
+						bContinue = false;
+						break;
+					}
+				}
+
+			}
+			if (bContinue)
 			{
 				// Transitions
 				for (const FStateTreeTransition& Transition : State->Transitions)
diff --git a/Engine/Plugins/Runtime/StateTree/Source/StateTreeEditorModule/Public/StateTreeBaker.h b/Engine/Plugins/Runtime/StateTree/Source/StateTreeEditorModule/Public/StateTreeBaker.h
--- a/Engine/Plugins/Runtime/StateTree/Source/StateTreeEditorModule/Public/StateTreeBaker.h
+++ b/Engine/Plugins/Runtime/StateTree/Source/StateTreeEditorModule/Public/StateTreeBaker.h
@@ -12,6 +12,8 @@ class UStateTreeState;
 class UStateTreeEditorData;
 struct FStateTreeCondition;
 struct FStateTreeConditionItem;
+struct FStateTreeTaskItem;
+struct FStateTreeEvaluatorItem;
 struct FStateTreeStateLink;
 
 /**
@@ -40,7 +42,9 @@ private:
 	bool CreateCondition(const FStateTreeConditionItem& CondItem);
 	bool GetAndValidateBindings(const FStateTreeBindableStructDesc& TargetStruct, TArray<FStateTreeEditorPropertyBinding>& OutBindings) const;
 	bool IsPropertyAnyEnum(const FStateTreeBindableStructDesc& Struct, FStateTreeEditorPropertyPath Path) const;
-
+	bool CreateTask(const FStateTreeTaskItem& TaskItem);
+	bool CreateEvaluator(const FStateTreeEvaluatorItem& EvalItem);
+	
 	FStateTreeCompilerLog& Log;
 	UStateTree* StateTree = nullptr;
 	UStateTreeEditorData* TreeData = nullptr;
diff --git a/Engine/Plugins/Runtime/StateTree/Source/StateTreeEditorModule/Public/StateTreeState.h b/Engine/Plugins/Runtime/StateTree/Source/StateTreeEditorModule/Public/StateTreeState.h
--- a/Engine/Plugins/Runtime/StateTree/Source/StateTreeEditorModule/Public/StateTreeState.h
+++ b/Engine/Plugins/Runtime/StateTree/Source/StateTreeEditorModule/Public/StateTreeState.h
@@ -285,6 +285,10 @@ public:
 	UPROPERTY(EditDefaultsOnly, Category = "Tasks", meta = (BaseStruct = "StateTreeTaskBase", BaseClass = "StateTreeTaskBlueprintBase"))
 	TArray<FStateTreeTaskItem> Tasks;
 
+	// Single item used when schema calls for single task per state.
+	UPROPERTY(EditDefaultsOnly, Category = "Task", meta = (BaseStruct = "StateTreeTaskBase", BaseClass = "StateTreeTaskBlueprintBase"))
+	FStateTreeTaskItem SingleTask;
+
 	UPROPERTY(EditDefaultsOnly, Category = "Transitions")
 	TArray<FStateTreeTransition> Transitions;
 
diff --git a/Engine/Plugins/Runtime/StateTree/Source/StateTreeModule/Public/StateTreeSchema.h b/Engine/Plugins/Runtime/StateTree/Source/StateTreeModule/Public/StateTreeSchema.h
--- a/Engine/Plugins/Runtime/StateTree/Source/StateTreeModule/Public/StateTreeSchema.h
+++ b/Engine/Plugins/Runtime/StateTree/Source/StateTreeModule/Public/StateTreeSchema.h
@@ -39,4 +39,17 @@ public:
 	 * @return True if the class is a StateTree item Blueprint base class.
 	 */
 	bool IsChildOfBlueprintBase(const UClass* InClass) const;
+
+#if WITH_EDITOR
+	
+	/** @return True if enter conditions are allowed. */
+	virtual bool AllowEnterConditions() const { return true; }
+
+	/** @return True if evaluators are allowed. */
+	virtual bool AllowEvaluators() const { return true; }
+
+	/** @return True if multiple tasks are allowed. */
+	virtual bool AllowMultipleTasks() const { return true; }
+	
+#endif // WITH_EDITOR
 };
