StateTree improvements:
- EnterState() stops on failure, ExitState() handled symmetrically up to the last failed task
- Allow to run evaluators when reselecting a state after failed EnterState()

diff --git a/Engine/Plugins/AI/MassAI/Source/MassAIBehavior/Private/Tasks/MassZoneGraphFindSmartObjectTarget.cpp b/Engine/Plugins/AI/MassAI/Source/MassAIBehavior/Private/Tasks/MassZoneGraphFindSmartObjectTarget.cpp
--- a/Engine/Plugins/AI/MassAI/Source/MassAIBehavior/Private/Tasks/MassZoneGraphFindSmartObjectTarget.cpp
+++ b/Engine/Plugins/AI/MassAI/Source/MassAIBehavior/Private/Tasks/MassZoneGraphFindSmartObjectTarget.cpp
@@ -30,12 +30,6 @@ EStateTreeRunStatus FMassZoneGraphFindSmartObjectTarget::EnterState(FStateTreeEx
 	FMassZoneGraphTargetLocation& SmartObjectLocation = Context.GetInstanceData(SmartObjectLocationHandle);
 	SmartObjectLocation.Reset();
 
-	if (Context.GetEnterStateStatus() == EStateTreeRunStatus::Failed)
-	{
-		MASSBEHAVIOR_LOG(Log, TEXT("Transition to the state has been denied by other tasks. Nothing to do."));
-		return EStateTreeRunStatus::Failed;
-	}
-
 	if (!SOUser.ClaimHandle.SmartObjectID.IsValid())
 	{
 		MASSBEHAVIOR_LOG(Error, TEXT("Invalid claimed smart object ID."));
diff --git a/Engine/Plugins/AI/MassAI/Source/MassAIBehavior/Private/Tasks/MassZoneGraphPathFollowTask.cpp b/Engine/Plugins/AI/MassAI/Source/MassAIBehavior/Private/Tasks/MassZoneGraphPathFollowTask.cpp
--- a/Engine/Plugins/AI/MassAI/Source/MassAIBehavior/Private/Tasks/MassZoneGraphPathFollowTask.cpp
+++ b/Engine/Plugins/AI/MassAI/Source/MassAIBehavior/Private/Tasks/MassZoneGraphPathFollowTask.cpp
@@ -116,12 +116,6 @@ EStateTreeRunStatus FMassZoneGraphPathFollowTask::EnterState(FStateTreeExecution
 		MASSBEHAVIOR_LOG(Verbose, TEXT("enterstate."));
 	}
 
-	if (Context.GetEnterStateStatus() == EStateTreeRunStatus::Failed)
-	{
-		MASSBEHAVIOR_LOG(Log, TEXT("Transition to the state has been denied by other tasks. Nothing to do."));
-		return EStateTreeRunStatus::Failed;
-	}
-
 	const FMassZoneGraphLaneLocationFragment& LaneLocation = Context.GetExternalData(LocationHandle);
 	const FMassZoneGraphTargetLocation& TargetLocation = Context.GetInstanceData(TargetLocationHandle);
 	
diff --git a/Engine/Plugins/Runtime/StateTree/Source/StateTreeModule/Private/StateTreeExecutionContext.cpp b/Engine/Plugins/Runtime/StateTree/Source/StateTreeModule/Private/StateTreeExecutionContext.cpp
--- a/Engine/Plugins/Runtime/StateTree/Source/StateTreeModule/Private/StateTreeExecutionContext.cpp
+++ b/Engine/Plugins/Runtime/StateTree/Source/StateTreeModule/Private/StateTreeExecutionContext.cpp
@@ -128,7 +128,7 @@ EStateTreeRunStatus FStateTreeExecutionContext::Start(FStateTreeDataView Externa
 		{
 			// Enter state tasks can fail/succeed, treat it same as tick.
 			Exec.LastTickStatus = EnterState(Storage, Transition);
-			Exec.CurrentState = Transition.Next;
+			
 			// Report state completed immediately.
 			if (Exec.LastTickStatus != EStateTreeRunStatus::Running)
 			{
@@ -158,15 +158,16 @@ EStateTreeRunStatus FStateTreeExecutionContext::Stop(FStateTreeDataView External
 	FStateTreeDataView Storage = SelectMutableStorage(ExternalStorage);
 	FStateTreeExecutionState& Exec = GetExecState(Storage);
 
-	// Stop if still running.
-	if (Exec.TreeRunStatus == EStateTreeRunStatus::Running)
+	// Exit states if still in some valid state.
+	if (Exec.CurrentState.IsValid() && (Exec.CurrentState != FStateTreeHandle::Succeeded || Exec.CurrentState != FStateTreeHandle::Failed))
 	{
+		// Transition to Succeeded state.
 		const FStateTreeTransitionResult Transition(FStateTreeStateStatus(Exec.CurrentState, Exec.LastTickStatus), FStateTreeHandle::Succeeded);
 		ExitState(Storage, Transition);
+		Exec.CurrentState = FStateTreeHandle::Succeeded;
 	}
-	
-	Exec.TreeRunStatus = EStateTreeRunStatus::Succeeded;
-	Exec.CurrentState = FStateTreeHandle::Succeeded;
+
+	Exec.TreeRunStatus = Exec.CurrentState == FStateTreeHandle::Succeeded ? EStateTreeRunStatus::Succeeded : EStateTreeRunStatus::Failed;
 	Exec.LastTickStatus = EStateTreeRunStatus::Unset;
 
 	return Exec.TreeRunStatus;
@@ -236,7 +237,7 @@ EStateTreeRunStatus FStateTreeExecutionContext::Tick(const float DeltaTime, FSta
 			
 			// Enter state tasks can fail/succeed, treat it same as tick.
 			Exec.LastTickStatus = EnterState(Storage, Transition);
-			Exec.CurrentState = Transition.Next;
+
 			// Report state completed immediately.
 			if (Exec.LastTickStatus != EStateTreeRunStatus::Running)
 			{
@@ -249,6 +250,9 @@ EStateTreeRunStatus FStateTreeExecutionContext::Tick(const float DeltaTime, FSta
 		{
 			break;
 		}
+		
+		// Reset visited states to allow evaluators to update during the next transition selection.
+		VisitedStates.Init(false, StateTree->States.Num());
 	}
 	
 	if (!Exec.CurrentState.IsValid())
@@ -280,6 +284,8 @@ EStateTreeRunStatus FStateTreeExecutionContext::EnterState(FStateTreeDataView St
 		return EStateTreeRunStatus::Failed;
 	}
 
+	FStateTreeExecutionState& Exec = GetExecState(Storage);
+
 	// Activate evaluators along the active branch.
 	TStaticArray<FStateTreeHandle, 32> States;
 	const int32 NumStates = GetActiveStates(Transition.Next, States);
@@ -299,9 +305,11 @@ EStateTreeRunStatus FStateTreeExecutionContext::EnterState(FStateTreeDataView St
 	FStateTreeTransitionResult CurrentTransition = Transition;
 	
 	EStateTreeRunStatus Result = EStateTreeRunStatus::Running;
-	EnterStateStatus = Result;
 
-	for (int32 Index = 0; Index < NumStates; Index++)
+	Exec.EnterStateFailedTaskIndex = INDEX_NONE;
+	Exec.CurrentState = FStateTreeHandle::Failed;
+	
+	for (int32 Index = 0; Index < NumStates && Result != EStateTreeRunStatus::Failed; Index++)
 	{
 		const FStateTreeHandle CurrentHandle = States[Index];
 		const FBakedStateTreeState& State = StateTree->States[CurrentHandle.Index];
@@ -309,7 +317,7 @@ EStateTreeRunStatus FStateTreeExecutionContext::EnterState(FStateTreeDataView St
 		const bool bWasActive = PrevStates[Index] == CurrentHandle;
 		if (bWasActive && !bOnTargetBranch)
 		{
-			// States which will keep on begin active and were not part of the transition will not get enter/exit state.
+			// States which will keep on being active and were not part of the transition will not get enter/exit state.
 			// Must update data views.
 			for (int32 EvalIndex = 0; EvalIndex < int32(State.EvaluatorsNum); EvalIndex++)
 			{
@@ -330,9 +338,9 @@ EStateTreeRunStatus FStateTreeExecutionContext::EnterState(FStateTreeDataView St
 
 		STATETREE_LOG(Log, TEXT("%*sEnter state '%s' %s"), Index*UE::StateTree::DebugIndentSize, TEXT(""), *DebugGetStatePath(States, Index), *StaticEnum<EStateTreeStateChangeType>()->GetValueAsString(ChangeType));
 		
-		for (int32 EvalIndex = 0; EvalIndex < int32(State.EvaluatorsNum); EvalIndex++)
+		for (int32 EvalIndex = State.EvaluatorsBegin; EvalIndex < (int32)(State.EvaluatorsBegin + State.EvaluatorsNum); EvalIndex++)
 		{
-			const FStateTreeEvaluatorBase& Eval = GetItem<FStateTreeEvaluatorBase>(int32(State.EvaluatorsBegin) + EvalIndex);
+			const FStateTreeEvaluatorBase& Eval = GetItem<FStateTreeEvaluatorBase>(EvalIndex);
 			const FStateTreeDataView EvalInstanceView = GetInstanceData(Storage, Eval.bInstanceIsObject, Eval.InstanceIndex);
 			DataViews[Eval.DataViewIndex] = EvalInstanceView;
 
@@ -346,9 +354,9 @@ EStateTreeRunStatus FStateTreeExecutionContext::EnterState(FStateTreeDataView St
 		}
 		
 		// Activate tasks on current state.
-		for (int32 TaskIndex = 0; TaskIndex < int32(State.TasksNum); TaskIndex++)
+		for (int32 TaskIndex = State.TasksBegin; TaskIndex < (int32)(State.TasksBegin + State.TasksNum); TaskIndex++)
 		{
-			const FStateTreeTaskBase& Task = GetItem<FStateTreeTaskBase>(int32(State.TasksBegin) + TaskIndex);
+			const FStateTreeTaskBase& Task = GetItem<FStateTreeTaskBase>(TaskIndex);
 			const FStateTreeDataView TaskInstanceView = GetInstanceData(Storage, Task.bInstanceIsObject, Task.InstanceIndex);
 			DataViews[Task.DataViewIndex] = TaskInstanceView;
 
@@ -360,17 +368,23 @@ EStateTreeRunStatus FStateTreeExecutionContext::EnterState(FStateTreeDataView St
 		
 			STATETREE_LOG(Verbose, TEXT("%*s  Notify Task '%s'"), Index*UE::StateTree::DebugIndentSize, TEXT(""), *Task.Name.ToString());
 			const EStateTreeRunStatus Status = Task.EnterState(*this, ChangeType, CurrentTransition);
+			
 			if (Status == EStateTreeRunStatus::Failed)
 			{
+				// Store how far in the enter state we got. This will be used to match the ExitState() calls.
+				Exec.EnterStateFailedTaskIndex = TaskIndex;
+				Exec.CurrentState = CurrentHandle;
 				Result = Status;
-				// @todo StateTree: Ideally we should break here, but it is commented out for now in order to keep symmetrical enter/exit calls.
-				// break;
-				// As a workaround, we mark EnterStateStatus as 'failed' in context so remaining calls to EnterState could be aware of the failure and act accordingly.
-				EnterStateStatus = Status;
+				break;
 			}
 		}
 	}
 
+	if (Result != EStateTreeRunStatus::Failed)
+	{
+		Exec.CurrentState = Transition.Next;
+	}
+
 	return Result;
 }
 
@@ -415,7 +429,7 @@ void FStateTreeExecutionContext::ExitState(FStateTreeDataView Storage, const FSt
 		bOnTargetBranch = bOnTargetBranch || NextHandle == Transition.Target;
 		if (bRemainsActive && !bOnTargetBranch)
 		{
-			// States which will keep on begin active and were not part of the transition will not get enter/exit state.
+			// States which will keep on being active and were not part of the transition will not get enter/exit state.
 			// Must update item views.
 			for (int32 EvalIndex = 0; EvalIndex < int32(State.EvaluatorsNum); EvalIndex++)
 			{
@@ -436,9 +450,9 @@ void FStateTreeExecutionContext::ExitState(FStateTreeDataView Storage, const FSt
 
 		STATETREE_LOG(Log, TEXT("%*sExit state '%s' %s"), Index*UE::StateTree::DebugIndentSize, TEXT(""), *DebugGetStatePath(States, Index), *StaticEnum<EStateTreeStateChangeType>()->GetValueAsString(ChangeType));
 
-		for (int32 EvalIndex = 0; EvalIndex < int32(State.EvaluatorsNum); EvalIndex++)
+		for (int32 EvalIndex = State.EvaluatorsBegin; EvalIndex < (int32)(State.EvaluatorsBegin + State.EvaluatorsNum); EvalIndex++)
 		{
-			const FStateTreeEvaluatorBase& Eval = GetItem<FStateTreeEvaluatorBase>(int32(State.EvaluatorsBegin) + EvalIndex);
+			const FStateTreeEvaluatorBase& Eval = GetItem<FStateTreeEvaluatorBase>(EvalIndex);
 			const FStateTreeDataView EvalInstanceView = GetInstanceData(Storage, Eval.bInstanceIsObject, Eval.InstanceIndex);
 			DataViews[Eval.DataViewIndex] = EvalInstanceView;
 
@@ -452,20 +466,25 @@ void FStateTreeExecutionContext::ExitState(FStateTreeDataView Storage, const FSt
 		}
 
 		// Deactivate tasks on current State
-		for (int32 TaskIndex = 0; TaskIndex < int32(State.TasksNum); TaskIndex++)
+		for (int32 TaskIndex = State.TasksBegin; TaskIndex < (int32)(State.TasksBegin + State.TasksNum); TaskIndex++)
 		{
-			const FStateTreeTaskBase& Task = GetItem<FStateTreeTaskBase>(int32(State.TasksBegin) + TaskIndex);
-			const FStateTreeDataView TaskInstanceView = GetInstanceData(Storage, Task.bInstanceIsObject, Task.InstanceIndex);
-			DataViews[Task.DataViewIndex] = TaskInstanceView;
-
-			// Copy bound properties.
-			if (Task.BindingsBatch.IsValid())
+			// Call task completed only if EnterState() was called. This ensures that we have matching calls to Enter/ExitState.
+			// The task order in the tree (BF) allows us to use the comparison.
+			if (TaskIndex <= (int32)Exec.EnterStateFailedTaskIndex)
 			{
-				StateTree->PropertyBindings.CopyTo(DataViews, Task.BindingsBatch.Index, TaskInstanceView);
-			}
+				const FStateTreeTaskBase& Task = GetItem<FStateTreeTaskBase>(TaskIndex);
+				const FStateTreeDataView TaskInstanceView = GetInstanceData(Storage, Task.bInstanceIsObject, Task.InstanceIndex);
+				DataViews[Task.DataViewIndex] = TaskInstanceView;
 
-			STATETREE_LOG(Verbose, TEXT("%*s  Notify Task '%s'"), Index*UE::StateTree::DebugIndentSize, TEXT(""), *Task.Name.ToString());
-			Task.ExitState(*this, ChangeType, CurrentTransition);
+				// Copy bound properties.
+				if (Task.BindingsBatch.IsValid())
+				{
+					StateTree->PropertyBindings.CopyTo(DataViews, Task.BindingsBatch.Index, TaskInstanceView);
+				}
+
+				STATETREE_LOG(Verbose, TEXT("%*s  Notify Task '%s'"), Index*UE::StateTree::DebugIndentSize, TEXT(""), *Task.Name.ToString());
+				Task.ExitState(*this, ChangeType, CurrentTransition);
+			}
 		}
 	}
 }
@@ -480,6 +499,8 @@ void FStateTreeExecutionContext::StateCompleted(FStateTreeDataView Storage, cons
 	TStaticArray<FStateTreeHandle, 32> States;
 	const int32 NumStates = GetActiveStates(CurrentState, States);
 
+	FStateTreeExecutionState& Exec = GetExecState(Storage);
+
 	// Call from child towards root to allow to pass results back.
 	// Note: Completed is assumed to be called immediately after tick or enter state, so there's no property copying.
 	for (int32 Index = NumStates - 1; Index >= 0; Index--)
@@ -490,18 +511,23 @@ void FStateTreeExecutionContext::StateCompleted(FStateTreeDataView Storage, cons
 		STATETREE_LOG(Verbose, TEXT("%*sState Completed '%s' %s"), Index*UE::StateTree::DebugIndentSize, TEXT(""), *DebugGetStatePath(States, Index), *StaticEnum<EStateTreeRunStatus>()->GetValueAsString(CompletionStatus));
 
 		// Notify Tasks
-		for (int32 TaskIndex = int32(State.TasksNum) - 1; TaskIndex >= 0; TaskIndex--)
+		for (int32 TaskIndex = (int32)(State.TasksBegin + State.TasksNum) - 1; TaskIndex >= State.TasksBegin; TaskIndex--)
 		{
-			const FStateTreeTaskBase& Task = GetItem<FStateTreeTaskBase>(int32(State.TasksBegin) + TaskIndex);
+			// Call task completed only if EnterState() was called.
+			// The task order in the tree (BF) allows us to use the comparison.
+			if (TaskIndex <= (int32)Exec.EnterStateFailedTaskIndex)
+			{
+				const FStateTreeTaskBase& Task = GetItem<FStateTreeTaskBase>(TaskIndex);
 
-			STATETREE_LOG(Verbose, TEXT("%*s  Notify Task '%s'"), Index*UE::StateTree::DebugIndentSize, TEXT(""), *Task.Name.ToString());
-			Task.StateCompleted(*this, CompletionStatus, CurrentState);
+				STATETREE_LOG(Verbose, TEXT("%*s  Notify Task '%s'"), Index*UE::StateTree::DebugIndentSize, TEXT(""), *Task.Name.ToString());
+				Task.StateCompleted(*this, CompletionStatus, CurrentState);
+			}
 		}
 
 		// Notify evaluators
-		for (int32 EvalIndex = int32(State.EvaluatorsNum) - 1; EvalIndex >= 0; EvalIndex--)
+		for (int32 EvalIndex = (int32)(State.EvaluatorsBegin + State.EvaluatorsNum) - 1; EvalIndex >= State.EvaluatorsBegin; EvalIndex--)
 		{
-			const FStateTreeEvaluatorBase& Eval = GetItem<FStateTreeEvaluatorBase>(int32(State.EvaluatorsBegin) + EvalIndex);
+			const FStateTreeEvaluatorBase& Eval = GetItem<FStateTreeEvaluatorBase>(EvalIndex);
 
 			STATETREE_LOG(Verbose, TEXT("%*s  Notify Evaluator '%s'"), Index*UE::StateTree::DebugIndentSize, TEXT(""), *Eval.Name.ToString());
 			Eval.StateCompleted(*this, CompletionStatus, CurrentState);
diff --git a/Engine/Plugins/Runtime/StateTree/Source/StateTreeModule/Public/StateTreeExecutionContext.h b/Engine/Plugins/Runtime/StateTree/Source/StateTreeModule/Public/StateTreeExecutionContext.h
--- a/Engine/Plugins/Runtime/StateTree/Source/StateTreeModule/Public/StateTreeExecutionContext.h
+++ b/Engine/Plugins/Runtime/StateTree/Source/StateTreeModule/Public/StateTreeExecutionContext.h
@@ -20,6 +20,9 @@ struct STATETREEMODULE_API FStateTreeExecutionState
 	/** Currently active state */
 	FStateTreeHandle CurrentState = FStateTreeHandle::Invalid;
 
+	/** The index of the task that failed during enter state. Exit state uses it to call ExitState() symmetrically. */
+	uint16 EnterStateFailedTaskIndex = INDEX_NONE;
+
 	/** Result of last tick */
 	EStateTreeRunStatus LastTickStatus = EStateTreeRunStatus::Failed;
 
@@ -222,7 +225,6 @@ public:
 	}
 	
 	EStateTreeRunStatus GetLastTickStatus(FStateTreeDataView ExternalStorage = FStateTreeDataView()) const;
-	EStateTreeRunStatus GetEnterStateStatus() const { return EnterStateStatus; }
 
 #if WITH_GAMEPLAY_DEBUGGER
 	/** @return Debug string describing the current state of the execution */
@@ -384,13 +386,6 @@ protected:
 
 	/** Storage type of the context */
 	EStateTreeStorage StorageType = EStateTreeStorage::Internal;
-
-	/**
-	 * Temporary status held within the context when calling 'EnterState' on multiple tasks.
-	 * Since it is called on all tasks even if a failed status was returned, this allow other tasks to act accordingly.
-	 * Note. This should be replaced by symmetrical unrolling of tasks on failure.
-	 */
-	EStateTreeRunStatus EnterStateStatus = EStateTreeRunStatus::Unset;
 };
 
 template<>
diff --git a/Engine/Plugins/Runtime/StateTree/Source/StateTreeTestSuite/Private/StateTreeTest.cpp b/Engine/Plugins/Runtime/StateTree/Source/StateTreeTestSuite/Private/StateTreeTest.cpp
--- a/Engine/Plugins/Runtime/StateTree/Source/StateTreeTestSuite/Private/StateTreeTest.cpp
+++ b/Engine/Plugins/Runtime/StateTree/Source/StateTreeTestSuite/Private/StateTreeTest.cpp
@@ -78,7 +78,7 @@ struct FStateTreeTest_MakeAndBakeStateTree : FAITestBase
 		return true;
 	}
 };
-IMPLEMENT_AI_INSTANT_TEST(FStateTreeTest_MakeAndBakeStateTree, "System.AI.StateTree.MakeAndBakeStateTree");
+IMPLEMENT_AI_INSTANT_TEST(FStateTreeTest_MakeAndBakeStateTree, "System.StateTree.MakeAndBakeStateTree");
 
 #if 0
 // @todo: fix this test
@@ -212,7 +212,7 @@ struct FStateTreeTest_WanderLoop : FAITestBase
 		return true;
 	}
 };
-IMPLEMENT_AI_INSTANT_TEST(FStateTreeTest_WanderLoop, "System.AI.StateTree.WanderLoop");
+IMPLEMENT_AI_INSTANT_TEST(FStateTreeTest_WanderLoop, "System.StateTree.WanderLoop");
 #endif
 
 struct FStateTreeTest_Sequence : FAITestBase
@@ -272,7 +272,7 @@ struct FStateTreeTest_Sequence : FAITestBase
 		return true;
 	}
 };
-IMPLEMENT_AI_INSTANT_TEST(FStateTreeTest_Sequence, "System.AI.StateTree.Sequence");
+IMPLEMENT_AI_INSTANT_TEST(FStateTreeTest_Sequence, "System.StateTree.Sequence");
 
 struct FStateTreeTest_Select : FAITestBase
 {
@@ -343,7 +343,68 @@ struct FStateTreeTest_Select : FAITestBase
 		return true;
 	}
 };
-IMPLEMENT_AI_INSTANT_TEST(FStateTreeTest_Select, "System.AI.StateTree.Select");
+IMPLEMENT_AI_INSTANT_TEST(FStateTreeTest_Select, "System.StateTree.Select");
+
+
+struct FStateTreeTest_FailEnterState : FAITestBase
+{
+	virtual bool InstantTest() override
+	{
+		UStateTree& StateTree = UE::StateTree::Tests::NewStateTree(&GetWorld());
+		UStateTreeEditorData& EditorData = *Cast<UStateTreeEditorData>(StateTree.EditorData);
+		
+		UStateTreeState& Root = EditorData.AddSubTree(FName(TEXT("Root")));
+		UStateTreeState& State1 = Root.AddChildState(FName(TEXT("State1")));
+		UStateTreeState& State1A = State1.AddChildState(FName(TEXT("State1A")));
+
+		auto& TaskRoot = Root.AddTask<FTestTask_Stand>(FName(TEXT("TaskRoot")));
+
+		auto& Task1 = State1.AddTask<FTestTask_Stand>(FName(TEXT("Task1")));
+		auto& Task2 = State1.AddTask<FTestTask_Stand>(FName(TEXT("Task2")));
+		Task2.GetItem().EnterStateResult = EStateTreeRunStatus::Failed;
+		auto& Task3 = State1.AddTask<FTestTask_Stand>(FName(TEXT("Task3")));
+
+		auto& Task1A = State1A.AddTask<FTestTask_Stand>(FName(TEXT("Task1A")));
+		State1A.AddTransition(EStateTreeTransitionEvent::OnCompleted, EStateTreeTransitionType::GotoState, &State1);
+
+		FStateTreeCompilerLog Log;
+		FStateTreeBaker Baker(Log);
+		const bool bResult = Baker.Bake(StateTree);
+		AITEST_TRUE("StateTree should get baked", bResult);
+
+		EStateTreeRunStatus Status = EStateTreeRunStatus::Unset;
+		FTestStateTreeExecutionContext Exec;
+		const bool bInitSucceeded = Exec.Init(StateTree, StateTree, EStateTreeStorage::Internal);
+		AITEST_TRUE("StateTree should init", bInitSucceeded);
+
+		const FString TickStr(TEXT("Tick"));
+		const FString EnterStateStr(TEXT("EnterState"));
+		const FString ExitStateStr(TEXT("ExitState"));
+		const FString StateCompletedStr(TEXT("StateCompleted"));
+
+		// Start and enter state
+		Status = Exec.Start();
+		AITEST_TRUE("StateTree TaskRoot should enter state", Exec.Expect(TaskRoot.GetName(), EnterStateStr));
+		AITEST_TRUE("StateTree Task1 should enter state", Exec.Expect(Task1.GetName(), EnterStateStr));
+		AITEST_TRUE("StateTree Task2 should enter state", Exec.Expect(Task2.GetName(), EnterStateStr));
+		AITEST_FALSE("StateTree Task3 should not enter state", Exec.Expect(Task3.GetName(), EnterStateStr));
+		AITEST_TRUE("StateTree Should execute StateCompleted in reverse order", Exec.Expect(Task2.GetName(), StateCompletedStr).Then(Task1.GetName(), StateCompletedStr).Then(TaskRoot.GetName(), StateCompletedStr));
+		AITEST_FALSE("StateTree Task3 should not state complete", Exec.Expect(Task3.GetName(), StateCompletedStr));
+		AITEST_TRUE("StateTree exec status should be failed", Exec.GetLastTickStatus() == EStateTreeRunStatus::Failed);
+		Exec.LogClear();
+
+		// Stop and exit state
+		Status = Exec.Stop();
+		AITEST_TRUE("StateTree TaskRoot should exit state", Exec.Expect(TaskRoot.GetName(), ExitStateStr));
+		AITEST_TRUE("StateTree Task1 should exit state", Exec.Expect(Task1.GetName(), ExitStateStr));
+		AITEST_TRUE("StateTree Task2 should exit state", Exec.Expect(Task2.GetName(), ExitStateStr));
+		AITEST_FALSE("StateTree Task3 should not exit state", Exec.Expect(Task3.GetName(), ExitStateStr));
+		Exec.LogClear();
+
+		return true;
+	}
+};
+IMPLEMENT_AI_INSTANT_TEST(FStateTreeTest_FailEnterState, "System.StateTree.FailEnterState");
 
 PRAGMA_ENABLE_OPTIMIZATION
 
diff --git a/Engine/Plugins/Runtime/StateTree/Source/StateTreeTestSuite/Private/StateTreeTestTypes.h b/Engine/Plugins/Runtime/StateTree/Source/StateTreeTestSuite/Private/StateTreeTestTypes.h
--- a/Engine/Plugins/Runtime/StateTree/Source/StateTreeTestSuite/Private/StateTreeTestTypes.h
+++ b/Engine/Plugins/Runtime/StateTree/Source/StateTreeTestSuite/Private/StateTreeTestTypes.h
@@ -680,6 +680,12 @@ struct FTestTask_Stand : public FStateTreeTaskBase
 		TestContext.Log(Name, TEXT("ExitState"));
 	}
 
+	virtual void StateCompleted(FStateTreeExecutionContext& Context, const EStateTreeRunStatus CompletionStatus, const FStateTreeHandle CompletedState) const override
+	{
+		FTestStateTreeExecutionContext& TestContext = static_cast<FTestStateTreeExecutionContext&>(Context);
+		TestContext.Log(Name, TEXT("StateCompleted"));
+	}
+	
 	virtual EStateTreeRunStatus Tick(FStateTreeExecutionContext& Context, const float DeltaTime) const override
 	{
 		FTestStateTreeExecutionContext& TestContext = static_cast<FTestStateTreeExecutionContext&>(Context);
