StateTree: control flow evaluators
- added base class for common evaluators
- added cooldown guard
- added selection guard

diff --git a/Engine/Plugins/AI/MassAI/Source/MassAIBehavior/Private/MassStateTreeSchema.cpp b/Engine/Plugins/AI/MassAI/Source/MassAIBehavior/Private/MassStateTreeSchema.cpp
--- a/Engine/Plugins/AI/MassAI/Source/MassAIBehavior/Private/MassStateTreeSchema.cpp
+++ b/Engine/Plugins/AI/MassAI/Source/MassAIBehavior/Private/MassStateTreeSchema.cpp
@@ -16,6 +16,7 @@ bool UMassStateTreeSchema::IsStructAllowed(const UScriptStruct* InScriptStruct)
 {
 	// Only allow Mass evals and tasks,and common conditions.
 	return InScriptStruct->IsChildOf(FMassStateTreeEvaluatorBase::StaticStruct())
+			|| InScriptStruct->IsChildOf(FStateTreeEvaluatorCommonBase::StaticStruct())
 			|| InScriptStruct->IsChildOf(FMassStateTreeTaskBase::StaticStruct())
 			|| InScriptStruct->IsChildOf(FStateTreeConditionBase::StaticStruct());
 }
diff --git a/Engine/Plugins/Runtime/StateTree/Source/StateTreeModule/Private/Components/StateTreeBrainComponent.cpp b/Engine/Plugins/Runtime/StateTree/Source/StateTreeModule/Private/Components/StateTreeBrainComponent.cpp
--- a/Engine/Plugins/Runtime/StateTree/Source/StateTreeModule/Private/Components/StateTreeBrainComponent.cpp
+++ b/Engine/Plugins/Runtime/StateTree/Source/StateTreeModule/Private/Components/StateTreeBrainComponent.cpp
@@ -20,7 +20,8 @@
 
 bool UBrainComponentStateTreeSchema::IsStructAllowed(const UScriptStruct* InScriptStruct) const
 {
-	return InScriptStruct->IsChildOf(FStateTreeConditionCommonBase::StaticStruct());
+	return InScriptStruct->IsChildOf(FStateTreeConditionCommonBase::StaticStruct())
+			|| InScriptStruct->IsChildOf(FStateTreeEvaluatorCommonBase::StaticStruct());
 }
 
 bool UBrainComponentStateTreeSchema::IsClassAllowed(const UClass* InClass) const
diff --git a/Engine/Plugins/Runtime/StateTree/Source/StateTreeModule/Private/Evaluators/CooldownGuardEvaluator.cpp b/Engine/Plugins/Runtime/StateTree/Source/StateTreeModule/Private/Evaluators/CooldownGuardEvaluator.cpp
new file mode 100644
--- /dev/null
+++ b/Engine/Plugins/Runtime/StateTree/Source/StateTreeModule/Private/Evaluators/CooldownGuardEvaluator.cpp
@@ -0,0 +1,61 @@
+﻿// Copyright Epic Games, Inc. All Rights Reserved.
+
+#include "Evaluators/CooldownGuardEvaluator.h"
+#include "StateTreeExecutionContext.h"
+
+
+bool FCooldownGuardEvaluator::Link(FStateTreeLinker& Linker)
+{
+	Linker.LinkInstanceDataProperty(DurationHandle, STATETREE_INSTANCEDATA_PROPERTY(FCooldownGuardEvaluatorInstanceData, Duration));
+	Linker.LinkInstanceDataProperty(EndTimeHandle, STATETREE_INSTANCEDATA_PROPERTY(FCooldownGuardEvaluatorInstanceData, EndTime));
+	Linker.LinkInstanceDataProperty(ActiveHandle, STATETREE_INSTANCEDATA_PROPERTY(FCooldownGuardEvaluatorInstanceData, bActive));
+
+	return true;
+}
+
+void FCooldownGuardEvaluator::SetCoolDown(FStateTreeExecutionContext& Context) const
+{
+	const float Duration = Context.GetInstanceData(DurationHandle);
+	float& EndTime = Context.GetInstanceData(EndTimeHandle);
+	bool& bActive = Context.GetInstanceData(ActiveHandle);
+
+	const float CurrentTime = Context.GetWorld()->GetTimeSeconds();
+
+	EndTime = CurrentTime + FMath::Max(0.0f, Duration + FMath::RandRange(-RandomDeviation, RandomDeviation));
+	bActive = CurrentTime < EndTime;
+}
+
+void FCooldownGuardEvaluator::EnterState(FStateTreeExecutionContext& Context, const EStateTreeStateChangeType ChangeType, const FStateTreeTransitionResult& Transition) const
+{
+	if (Activation == ECooldownGuardActivation::OnEnterState)
+	{
+		SetCoolDown(Context);
+	}
+}
+ 
+void FCooldownGuardEvaluator::StateCompleted(FStateTreeExecutionContext& Context, const EStateTreeRunStatus CompletionStatus, const FStateTreeHandle CompletedState) const
+{
+	if (CompletionStatus == EStateTreeRunStatus::Failed
+		&& (Activation == ECooldownGuardActivation::OnCompleted || Activation == ECooldownGuardActivation::OnFailed))
+	{
+		SetCoolDown(Context);
+	}
+
+	if (CompletionStatus == EStateTreeRunStatus::Succeeded
+		&& (Activation == ECooldownGuardActivation::OnCompleted || Activation == ECooldownGuardActivation::OnSucceeded))
+	{
+		SetCoolDown(Context);
+	}
+}
+
+void FCooldownGuardEvaluator::Evaluate(FStateTreeExecutionContext& Context, const EStateTreeEvaluationType EvalType, const float DeltaTime) const
+{
+	const float EndTime = Context.GetInstanceData(EndTimeHandle);
+	bool& bActive = Context.GetInstanceData(ActiveHandle);
+
+	if (bActive)
+	{
+		const float CurrentTime = Context.GetWorld()->GetTimeSeconds();
+		bActive = CurrentTime < EndTime;
+	}
+}
diff --git a/Engine/Plugins/Runtime/StateTree/Source/StateTreeModule/Private/Evaluators/SelectionGuardEvaluator.cpp b/Engine/Plugins/Runtime/StateTree/Source/StateTreeModule/Private/Evaluators/SelectionGuardEvaluator.cpp
new file mode 100644
--- /dev/null
+++ b/Engine/Plugins/Runtime/StateTree/Source/StateTreeModule/Private/Evaluators/SelectionGuardEvaluator.cpp
@@ -0,0 +1,48 @@
+﻿// Copyright Epic Games, Inc. All Rights Reserved.
+
+#include "Evaluators/SelectionGuardEvaluator.h"
+#include "StateTreeExecutionContext.h"
+
+
+bool FSelectionGuardEvaluator::Link(FStateTreeLinker& Linker)
+{
+	Linker.LinkInstanceDataProperty(ActiveHandle, STATETREE_INSTANCEDATA_PROPERTY(FSelectionGuardEvaluatorInstanceData, bActive));
+
+	return true;
+}
+
+void FSelectionGuardEvaluator::EnterState(FStateTreeExecutionContext& Context, const EStateTreeStateChangeType ChangeType, const FStateTreeTransitionResult& Transition) const
+{
+	if (Activation == ESelectionGuardActivation::OnEnterState)
+	{
+		bool& bActive = Context.GetInstanceData(ActiveHandle);
+		bActive = true;
+	}
+}
+
+void FSelectionGuardEvaluator::ExitState(FStateTreeExecutionContext& Context, const EStateTreeStateChangeType ChangeType, const FStateTreeTransitionResult& Transition) const
+{
+	bool& bActive = Context.GetInstanceData(ActiveHandle);
+	
+	// Since new state selection happen _before_ ExitState() is called,
+	// setting the bActive on ExitState() can be used to check
+	// if the state is currently active and prevent re-selecting it.
+	bActive = false;
+}
+
+void FSelectionGuardEvaluator::StateCompleted(FStateTreeExecutionContext& Context, const EStateTreeRunStatus CompletionStatus, const FStateTreeHandle CompletedState) const
+{
+	bool& bActive = Context.GetInstanceData(ActiveHandle);
+
+	if (CompletionStatus == EStateTreeRunStatus::Failed
+		&& (Activation == ESelectionGuardActivation::OnCompleted || Activation == ESelectionGuardActivation::OnFailed))
+	{
+		bActive = true;
+	}
+
+	if (CompletionStatus == EStateTreeRunStatus::Succeeded
+		&& (Activation == ESelectionGuardActivation::OnCompleted || Activation == ESelectionGuardActivation::OnSucceeded))
+	{
+		bActive = true;
+	}
+}
\ No newline at end of file
diff --git a/Engine/Plugins/Runtime/StateTree/Source/StateTreeModule/Public/Conditions/StateTreeCondition_Common.h b/Engine/Plugins/Runtime/StateTree/Source/StateTreeModule/Public/Conditions/StateTreeCondition_Common.h
--- a/Engine/Plugins/Runtime/StateTree/Source/StateTreeModule/Public/Conditions/StateTreeCondition_Common.h
+++ b/Engine/Plugins/Runtime/StateTree/Source/StateTreeModule/Public/Conditions/StateTreeCondition_Common.h
@@ -15,15 +15,6 @@ enum class EStateTreeCompare : uint8
 	Invert,
 };
 
-/**
- * Base class for all common conditions so that they are easy to single out in the schema.
- */
-USTRUCT(meta = (Hidden))
-struct STATETREEMODULE_API FStateTreeConditionCommonBase : public FStateTreeConditionBase
-{
-	GENERATED_BODY()
-};
-
 /**
  * Condition comparing two integers.
  */
diff --git a/Engine/Plugins/Runtime/StateTree/Source/StateTreeModule/Public/Evaluators/CooldownGuardEvaluator.h b/Engine/Plugins/Runtime/StateTree/Source/StateTreeModule/Public/Evaluators/CooldownGuardEvaluator.h
new file mode 100644
--- /dev/null
+++ b/Engine/Plugins/Runtime/StateTree/Source/StateTreeModule/Public/Evaluators/CooldownGuardEvaluator.h
@@ -0,0 +1,66 @@
+﻿// Copyright Epic Games, Inc. All Rights Reserved.
+
+#pragma once
+
+#include "StateTreeTypes.h"
+#include "StateTreeEvaluatorBase.h"
+#include "CooldownGuardEvaluator.generated.h"
+
+struct FStateTreeExecutionContext;
+
+/**
+ * General purpose evaluator for selection cool-downs in the StateTree.
+ * - Add the cool-down evaluator to the state (and sub-state) that you want to guard
+ * - Use the bActive in a condition to guard entering the state.
+ */
+
+UENUM()
+enum class ECooldownGuardActivation : uint8
+{
+	OnEnterState,	// Activate cooldown when entering the state.
+	OnCompleted,	// Activate cooldown when state is completed (either success or failure).
+	OnFailed,		// Activate cooldown when state failed.
+	OnSucceeded,	// Activate cooldown when state succeeded.
+};
+
+USTRUCT()
+struct STATETREEMODULE_API FCooldownGuardEvaluatorInstanceData
+{
+	GENERATED_BODY()
+
+	UPROPERTY(EditAnywhere, Category="Parameter", Meta=(Tooltip = "Duration of the cool down."))
+	float Duration = 5.0f;
+
+	UPROPERTY()
+	float EndTime = 0.0f;
+
+	UPROPERTY(EditAnywhere, Category="Output", Meta=(Tooltip = "Cooldown is active."))
+	bool bActive = false;
+};
+
+USTRUCT(DisplayName="Cooldown Guard")
+struct STATETREEMODULE_API FCooldownGuardEvaluator : public FStateTreeEvaluatorCommonBase
+{
+	GENERATED_BODY()
+
+	FCooldownGuardEvaluator() = default;
+	virtual ~FCooldownGuardEvaluator() override {}
+
+	virtual const UStruct* GetInstanceDataType() const override { return FCooldownGuardEvaluatorInstanceData::StaticStruct(); };
+	virtual bool Link(FStateTreeLinker& Linker) override;
+	virtual void EnterState(FStateTreeExecutionContext& Context, const EStateTreeStateChangeType ChangeType, const FStateTreeTransitionResult& Transition) const override;
+	virtual void StateCompleted(FStateTreeExecutionContext& Context, const EStateTreeRunStatus CompletionStatus, const FStateTreeHandle CompletedState) const override;
+	virtual void Evaluate(FStateTreeExecutionContext& Context, const EStateTreeEvaluationType EvalType, const float DeltaTime) const override;
+
+	void SetCoolDown(FStateTreeExecutionContext& Context) const;
+	
+	TStateTreeInstanceDataPropertyHandle<float> DurationHandle;
+	TStateTreeInstanceDataPropertyHandle<float> EndTimeHandle;
+	TStateTreeInstanceDataPropertyHandle<bool> ActiveHandle;
+
+	UPROPERTY(EditAnywhere, Category="Default", Meta=(Tooltip = "Random deviation added to the duration."))
+	float RandomDeviation = 0.0f;
+
+	UPROPERTY(EditAnywhere, Category="Default", Meta=(Tooltip = "When to trigger the cooldown."))
+	ECooldownGuardActivation Activation = ECooldownGuardActivation::OnCompleted;
+};
diff --git a/Engine/Plugins/Runtime/StateTree/Source/StateTreeModule/Public/Evaluators/SelectionGuardEvaluator.h b/Engine/Plugins/Runtime/StateTree/Source/StateTreeModule/Public/Evaluators/SelectionGuardEvaluator.h
new file mode 100644
--- /dev/null
+++ b/Engine/Plugins/Runtime/StateTree/Source/StateTreeModule/Public/Evaluators/SelectionGuardEvaluator.h
@@ -0,0 +1,53 @@
+﻿// Copyright Epic Games, Inc. All Rights Reserved.
+
+#pragma once
+
+#include "StateTreeTypes.h"
+#include "StateTreeEvaluatorBase.h"
+#include "SelectionGuardEvaluator.generated.h"
+
+struct FStateTreeExecutionContext;
+
+/**
+ * General purpose evaluator for guarding against selecting an active state.
+ * - Add the guard evaluator to the state (and sub-state) that you want to prevent re-selecting
+ * - Use the bActive in a condition to guard entering the state.
+ */
+
+UENUM()
+enum class ESelectionGuardActivation : uint8
+{
+	OnEnterState,	// Activate guard when entering the state.
+	OnCompleted,	// Activate guard when state is completed (either success or failure).
+	OnFailed,		// Activate guard when state failed.
+	OnSucceeded,	// Activate guard when state succeeded.
+};
+
+USTRUCT()
+struct STATETREEMODULE_API FSelectionGuardEvaluatorInstanceData
+{
+	GENERATED_BODY()
+
+	UPROPERTY(EditAnywhere, Category="Output", Meta=(Tooltip = "Cooldown is active."))
+	bool bActive = false;
+};
+
+USTRUCT(DisplayName="Selection Guard")
+struct STATETREEMODULE_API FSelectionGuardEvaluator : public FStateTreeEvaluatorCommonBase
+{
+	GENERATED_BODY()
+
+	FSelectionGuardEvaluator() = default;
+	virtual ~FSelectionGuardEvaluator() override {}
+
+	virtual const UStruct* GetInstanceDataType() const override { return FSelectionGuardEvaluatorInstanceData::StaticStruct(); };
+	virtual bool Link(FStateTreeLinker& Linker) override;
+	virtual void EnterState(FStateTreeExecutionContext& Context, const EStateTreeStateChangeType ChangeType, const FStateTreeTransitionResult& Transition) const override;
+	virtual void StateCompleted(FStateTreeExecutionContext& Context, const EStateTreeRunStatus CompletionStatus, const FStateTreeHandle CompletedState) const override;
+	virtual void ExitState(FStateTreeExecutionContext& Context, const EStateTreeStateChangeType ChangeType, const FStateTreeTransitionResult& Transition) const override;
+	
+	TStateTreeInstanceDataPropertyHandle<bool> ActiveHandle;
+
+	UPROPERTY(EditAnywhere, Category="Default", Meta=(Tooltip = "Describes when to turn on the guard."))
+	ESelectionGuardActivation Activation = ESelectionGuardActivation::OnEnterState;
+};
diff --git a/Engine/Plugins/Runtime/StateTree/Source/StateTreeModule/Public/StateTreeConditionBase.h b/Engine/Plugins/Runtime/StateTree/Source/StateTreeModule/Public/StateTreeConditionBase.h
--- a/Engine/Plugins/Runtime/StateTree/Source/StateTreeModule/Public/StateTreeConditionBase.h
+++ b/Engine/Plugins/Runtime/StateTree/Source/StateTreeModule/Public/StateTreeConditionBase.h
@@ -67,3 +67,15 @@ struct STATETREEMODULE_API FStateTreeConditionBase
 	UPROPERTY()
 	uint8 bInstanceIsObject : 1;
 };
+
+
+/**
+ * Base class (namespace) for all common Conditions that are generally applicable.
+ * This allows schemas to safely include all conditions child of this struct. 
+ */
+USTRUCT(meta = (Hidden))
+struct STATETREEMODULE_API FStateTreeConditionCommonBase : public FStateTreeConditionBase
+{
+	GENERATED_BODY()
+};
+
diff --git a/Engine/Plugins/Runtime/StateTree/Source/StateTreeModule/Public/StateTreeEvaluatorBase.h b/Engine/Plugins/Runtime/StateTree/Source/StateTreeModule/Public/StateTreeEvaluatorBase.h
--- a/Engine/Plugins/Runtime/StateTree/Source/StateTreeModule/Public/StateTreeEvaluatorBase.h
+++ b/Engine/Plugins/Runtime/StateTree/Source/StateTreeModule/Public/StateTreeEvaluatorBase.h
@@ -91,3 +91,13 @@ struct STATETREEMODULE_API FStateTreeEvaluatorBase
 	UPROPERTY()
 	uint8 bInstanceIsObject : 1;
 };
+
+/**
+* Base class (namespace) for all common Evaluators that are generally applicable.
+* This allows schemas to safely include all Evaluators child of this struct. 
+*/
+USTRUCT(Meta=(Hidden))
+struct STATETREEMODULE_API FStateTreeEvaluatorCommonBase : public FStateTreeEvaluatorBase
+{
+	GENERATED_BODY()
+};
\ No newline at end of file
