StateTree CIS fix

diff --git a/Engine/Plugins/Runtime/StateTree/Source/StateTreeModule/Private/Evaluators/CooldownGuardEvaluator.cpp b/Engine/Plugins/Runtime/StateTree/Source/StateTreeModule/Private/Evaluators/CooldownGuardEvaluator.cpp
--- a/Engine/Plugins/Runtime/StateTree/Source/StateTreeModule/Private/Evaluators/CooldownGuardEvaluator.cpp
+++ b/Engine/Plugins/Runtime/StateTree/Source/StateTreeModule/Private/Evaluators/CooldownGuardEvaluator.cpp
@@ -2,6 +2,7 @@
 
 #include "Evaluators/CooldownGuardEvaluator.h"
 #include "StateTreeExecutionContext.h"
+#include "Engine/World.h"
 
 
 bool FCooldownGuardEvaluator::Link(FStateTreeLinker& Linker)
