StateTree: small editor improvements
- Removed State Succeeded/Failed from transitions as it currently does not make sense
- Change "Add Routine" button to "Add State"
-Rrenamed routines to SubTrees, generally support just one
- Add a root node by default

diff --git a/Engine/Plugins/Runtime/StateTree/Config/BaseStateTree.ini b/Engine/Plugins/Runtime/StateTree/Config/BaseStateTree.ini
--- a/Engine/Plugins/Runtime/StateTree/Config/BaseStateTree.ini
+++ b/Engine/Plugins/Runtime/StateTree/Config/BaseStateTree.ini
@@ -4,6 +4,7 @@
 +PropertyRedirects=(OldName="BakedStateTreeState.DefaultTransitionState",NewName="StateDoneTransitionState")
 +PropertyRedirects=(OldName="BakedStateTreeState.DefaultTransitionType",NewName="StateDoneTransitionType")
 +PropertyRedirects=(OldName="StateTreeBindableStructDesc.ScriptStruct",NewName="StateTreeBindableStructDesc.Struct")
++PropertyRedirects=(OldName="StateTreeEditorData.Routines",NewName="StateTreeEditorData.SubTrees")
 
 +StructRedirects=(OldName="StateTreeTransition2",NewName="StateTreeTransition")
 +StructRedirects=(OldName="StateTreeTask2Base",NewName="StateTreeTaskBase")
diff --git a/Engine/Plugins/Runtime/StateTree/Source/StateTreeEditorModule/Private/AssetTypeActions_StateTree.cpp b/Engine/Plugins/Runtime/StateTree/Source/StateTreeEditorModule/Private/AssetTypeActions_StateTree.cpp
--- a/Engine/Plugins/Runtime/StateTree/Source/StateTreeEditorModule/Private/AssetTypeActions_StateTree.cpp
+++ b/Engine/Plugins/Runtime/StateTree/Source/StateTreeEditorModule/Private/AssetTypeActions_StateTree.cpp
@@ -4,10 +4,14 @@
 #include "StateTree.h"
 #include "IStateTreeEditor.h"
 #include "StateTreeEditorModule.h"
-#include "AIModule.h"
 
 #define LOCTEXT_NAMESPACE "AssetTypeActions"
 
+FAssetTypeActions_StateTree::FAssetTypeActions_StateTree(const uint32 InAssetCategory)
+	: AssetCategory(InAssetCategory)
+{
+}
+
 void FAssetTypeActions_StateTree::OpenAssetEditor(const TArray<UObject*>& InObjects, TSharedPtr<IToolkitHost> EditWithinLevelEditor)
 {
 	EToolkitMode::Type Mode = EditWithinLevelEditor.IsValid() ? EToolkitMode::WorldCentric : EToolkitMode::Standalone;
@@ -28,9 +32,8 @@ UClass* FAssetTypeActions_StateTree::GetSupportedClass() const
 }
 
 uint32 FAssetTypeActions_StateTree::GetCategories()
-{ 
-	IAIModule& AIModule = FModuleManager::GetModuleChecked<IAIModule>("AIModule").Get();
-	return AIModule.GetAIAssetCategoryBit();
+{
+	return AssetCategory;
 }
 
 #undef LOCTEXT_NAMESPACE
diff --git a/Engine/Plugins/Runtime/StateTree/Source/StateTreeEditorModule/Private/AssetTypeActions_StateTree.h b/Engine/Plugins/Runtime/StateTree/Source/StateTreeEditorModule/Private/AssetTypeActions_StateTree.h
--- a/Engine/Plugins/Runtime/StateTree/Source/StateTreeEditorModule/Private/AssetTypeActions_StateTree.h
+++ b/Engine/Plugins/Runtime/StateTree/Source/StateTreeEditorModule/Private/AssetTypeActions_StateTree.h
@@ -12,10 +12,16 @@
 class FAssetTypeActions_StateTree : public FAssetTypeActions_Base
 {
 public:
+	explicit FAssetTypeActions_StateTree(const uint32 InAssetCategory);
+
+protected:
 	// IAssetTypeActions Implementation
 	virtual FText GetName() const override { return NSLOCTEXT("AssetTypeActions", "FAssetTypeActions_StateTree", "StateTree"); }
 	virtual FColor GetTypeColor() const override { return FColor(201, 185, 29); }
 	virtual UClass* GetSupportedClass() const override;
 	virtual void OpenAssetEditor(const TArray<UObject*>& InObjects, TSharedPtr<class IToolkitHost> EditWithinLevelEditor = TSharedPtr<IToolkitHost>()) override;
 	virtual uint32 GetCategories() override;
+
+private:
+	uint32 AssetCategory = EAssetTypeCategories::Gameplay;
 };
diff --git a/Engine/Plugins/Runtime/StateTree/Source/StateTreeEditorModule/Private/Customizations/StateTreeStateLinkDetails.cpp b/Engine/Plugins/Runtime/StateTree/Source/StateTreeEditorModule/Private/Customizations/StateTreeStateLinkDetails.cpp
--- a/Engine/Plugins/Runtime/StateTree/Source/StateTreeEditorModule/Private/Customizations/StateTreeStateLinkDetails.cpp
+++ b/Engine/Plugins/Runtime/StateTree/Source/StateTreeEditorModule/Private/Customizations/StateTreeStateLinkDetails.cpp
@@ -92,9 +92,9 @@ void FStateTreeStateLinkDetails::CacheStates()
 		{
 			if (UStateTreeEditorData* TreeData = Cast<UStateTreeEditorData>(OuterStateTree->EditorData))
 			{
-				for (const UStateTreeState* Routine : TreeData->Routines)
+				for (const UStateTreeState* SubTree : TreeData->SubTrees)
 				{
-					CacheStates(Routine);
+					CacheStates(SubTree);
 				}
 			}
 			break;
@@ -143,22 +143,12 @@ TSharedRef<SWidget> FStateTreeStateLinkDetails::OnGetStateContent() const
 {
 	FMenuBuilder MenuBuilder(true, NULL);
 
-	FUIAction NotSetItemAction(FExecuteAction::CreateSP(const_cast<FStateTreeStateLinkDetails*>(this), &FStateTreeStateLinkDetails::OnStateComboChange, ComboNotSet));
-	MenuBuilder.AddMenuEntry(LOCTEXT("TransitionNotSet", "Not Set"), LOCTEXT("TransitionNotSetTooltip", "Transition not set."), FSlateIcon(), NotSetItemAction);
-
-	MenuBuilder.AddSeparator();
-
-	FUIAction SucceededItemAction(FExecuteAction::CreateSP(const_cast<FStateTreeStateLinkDetails*>(this), &FStateTreeStateLinkDetails::OnStateComboChange, ComboSucceeded));
-	MenuBuilder.AddMenuEntry(LOCTEXT("TransitionTreeSucceeded", "Tree Succeeded"), LOCTEXT("TransitionSucceededTooltip", "Stop StateTree and mark it Succeeded."), FSlateIcon(), SucceededItemAction);
-
-	FUIAction FailedItemAction(FExecuteAction::CreateSP(const_cast<FStateTreeStateLinkDetails*>(this), &FStateTreeStateLinkDetails::OnStateComboChange, ComboFailed));
-	MenuBuilder.AddMenuEntry(LOCTEXT("TransitionTreeFailed", "Tree Failed"), LOCTEXT("TransitionFailedTooltip", "Stop StateTree and mark it Failed."), FSlateIcon(), FailedItemAction);
-
-	MenuBuilder.AddSeparator();
-
 	FUIAction NextItemAction(FExecuteAction::CreateSP(const_cast<FStateTreeStateLinkDetails*>(this), &FStateTreeStateLinkDetails::OnStateComboChange, ComboNextState));
 	MenuBuilder.AddMenuEntry(LOCTEXT("TransitionNextState", "Next State"), LOCTEXT("TransitionNextTooltip", "Goto next sibling State."), FSlateIcon(), NextItemAction);
 
+	FUIAction NotSetItemAction(FExecuteAction::CreateSP(const_cast<FStateTreeStateLinkDetails*>(this), &FStateTreeStateLinkDetails::OnStateComboChange, ComboNotSet));
+	MenuBuilder.AddMenuEntry(LOCTEXT("TransitionBlock", "Block Transition"), LOCTEXT("TransitionBlockTooltip", "Will not transition to any state, but will block other transitions to trigger if the condition passes."), FSlateIcon(), NotSetItemAction);
+
 	if (CachedNames.Num() > 0)
 	{
 		MenuBuilder.BeginSection(FName(), LOCTEXT("TransitionGotoState", "Goto State"));
@@ -186,14 +176,6 @@ FText FStateTreeStateLinkDetails::GetCurrentStateDesc() const
 	{
 		return LOCTEXT("TransitionNotSet", "Not Set");
 	}
-	else if (TransitionType == EStateTreeTransitionType::Succeeded)
-	{
-		return LOCTEXT("TransitionTreeSucceeded", "Tree Succeeded");
-	}
-	else if (TransitionType == EStateTreeTransitionType::Failed)
-	{
-		return LOCTEXT("TransitionTreeFailed", "Tree Failed");
-	}
 	else if (TransitionType == EStateTreeTransitionType::NextState)
 	{
 		return LOCTEXT("TransitionNextState", "Next State");
diff --git a/Engine/Plugins/Runtime/StateTree/Source/StateTreeEditorModule/Private/SStateTreeView.cpp b/Engine/Plugins/Runtime/StateTree/Source/StateTreeEditorModule/Private/SStateTreeView.cpp
--- a/Engine/Plugins/Runtime/StateTree/Source/StateTreeEditorModule/Private/SStateTreeView.cpp
+++ b/Engine/Plugins/Runtime/StateTree/Source/StateTreeEditorModule/Private/SStateTreeView.cpp
@@ -73,7 +73,7 @@ void SStateTreeView::Construct(const FArguments& InArgs, TSharedRef<FStateTreeVi
 	TreeView = SNew(STreeView<UStateTreeState*>)
 		.OnGenerateRow(this, &SStateTreeView::HandleGenerateRow)
 		.OnGetChildren(this, &SStateTreeView::HandleGetChildren)
-		.TreeItemsSource(StateTreeViewModel->GetRoutines())
+		.TreeItemsSource(StateTreeViewModel->GetSubTrees())
 		.ItemHeight(32)
 		.OnSelectionChanged(this, &SStateTreeView::HandleTreeSelectionChanged)
 		.OnContextMenuOpening(this, &SStateTreeView::HandleContextMenuOpening)
@@ -102,8 +102,8 @@ void SStateTreeView::Construct(const FArguments& InArgs, TSharedRef<FStateTreeVi
 				[
 					SNew(SButton)
 					.ButtonStyle(FEditorStyle::Get(), "FlatButton.Success")
-					.ToolTipText(LOCTEXT("AddRoutineToolTip", "Add New Routine"))
-					.OnClicked(this, &SStateTreeView::HandleAddRoutine)
+					.ToolTipText(LOCTEXT("AddStateToolTip", "Add New State"))
+					.OnClicked(this, &SStateTreeView::HandleAddStateButton)
 					[
 						SNew(SHorizontalBox)
 						+ SHorizontalBox::Slot()
@@ -121,7 +121,7 @@ void SStateTreeView::Construct(const FArguments& InArgs, TSharedRef<FStateTreeVi
 						.AutoWidth()
 						[
 							SNew(STextBlock)
-							.Text(LOCTEXT("AddRoutine", "Add Routine"))
+							.Text(LOCTEXT("AddState", "Add State"))
 							.TextStyle(FEditorStyle::Get(), "FlatButton.DefaultTextStyle")
 						]
 					]
@@ -216,7 +216,7 @@ void SStateTreeView::UpdateTree(bool bExpandPersistent)
 	StateTreeViewModel->GetSelectedStates(SelectedStates);
 
 	// Regenerate items
-	TreeView->SetTreeItemsSource(StateTreeViewModel->GetRoutines());
+	TreeView->SetTreeItemsSource(StateTreeViewModel->GetSubTrees());
 
 	// Restore expanded state
 	for (UStateTreeState* State : ExpandedStates)
@@ -371,22 +371,58 @@ TSharedPtr<SWidget> SStateTreeView::HandleContextMenuOpening()
 }
 
 
-FReply SStateTreeView::HandleAddRoutine()
+FReply SStateTreeView::HandleAddStateButton()
 {
-	if (StateTreeViewModel)
+	if (StateTreeViewModel == nullptr)
+	{
+		return FReply::Handled();
+	}
+	
+	TArray<UStateTreeState*> SelectedStates;
+	StateTreeViewModel->GetSelectedStates(SelectedStates);
+	UStateTreeState* FirstSelectedState = SelectedStates.Num() > 0 ? SelectedStates[0] : nullptr;
+
+	if (FirstSelectedState != nullptr)
 	{
+		// If the state is root, add child state, else sibling.
+		if (FirstSelectedState->Parent == nullptr)
+		{
+			StateTreeViewModel->AddChildState(FirstSelectedState);
+			TreeView->SetItemExpansion(FirstSelectedState, true);
+		}
+		else
+		{
+			StateTreeViewModel->AddState(FirstSelectedState);
+		}
+	}
+	else if (StateTreeViewModel->GetSubTreeCount() == 0)
+	{
+		// If no root state, create new.
 		StateTreeViewModel->AddState(nullptr);
 	}
+	else
+	{
+		// If no state has been selected, add to root
+		TArray<UStateTreeState*>* SubTrees = StateTreeViewModel->GetSubTrees();
+		UStateTreeState* RootState = (SubTrees != nullptr && SubTrees->Num() > 0) ? SubTrees->GetData()[0] : nullptr;
+		if (RootState != nullptr)
+		{
+			StateTreeViewModel->AddChildState(RootState);
+			TreeView->SetItemExpansion(RootState, true);
+		}
+	}
 
 	return FReply::Handled();
 }
 
 void SStateTreeView::HandleAddState(UStateTreeState* AfterState)
 {
-	if (StateTreeViewModel)
+	if (StateTreeViewModel == nullptr)
 	{
-		StateTreeViewModel->AddState(AfterState);
+		return;
 	}
+
+	StateTreeViewModel->AddState(AfterState);
 }
 
 void SStateTreeView::HandleRenameState(UStateTreeState* State)
@@ -396,7 +432,12 @@ void SStateTreeView::HandleRenameState(UStateTreeState* State)
 
 void SStateTreeView::HandleAddChildState(UStateTreeState* ParentState)
 {
-	if (StateTreeViewModel && ParentState)
+	if (StateTreeViewModel == nullptr)
+	{
+		return;
+	}
+
+	if (ParentState)
 	{
 		StateTreeViewModel->AddChildState(ParentState);
 		TreeView->SetItemExpansion(ParentState, true);
@@ -405,10 +446,12 @@ void SStateTreeView::HandleAddChildState(UStateTreeState* ParentState)
 
 void SStateTreeView::HandleDeleteItems()
 {
-	if (StateTreeViewModel)
+	if (StateTreeViewModel == nullptr)
 	{
-		StateTreeViewModel->RemoveSelectedStates();
+		return;
 	}
+
+	StateTreeViewModel->RemoveSelectedStates();
 }
 
 
diff --git a/Engine/Plugins/Runtime/StateTree/Source/StateTreeEditorModule/Private/SStateTreeView.h b/Engine/Plugins/Runtime/StateTree/Source/StateTreeEditorModule/Private/SStateTreeView.h
--- a/Engine/Plugins/Runtime/StateTree/Source/StateTreeEditorModule/Private/SStateTreeView.h
+++ b/Engine/Plugins/Runtime/StateTree/Source/StateTreeEditorModule/Private/SStateTreeView.h
@@ -65,7 +65,7 @@ private:
 	TSharedPtr<SWidget> HandleContextMenuOpening();
 
 	// Action handlers
-	FReply HandleAddRoutine();
+	FReply HandleAddStateButton();
 	void HandleRenameState(UStateTreeState* State);
 	void HandleAddState(UStateTreeState* AfterItem);
 	void HandleAddChildState(UStateTreeState* ParentItem);
diff --git a/Engine/Plugins/Runtime/StateTree/Source/StateTreeEditorModule/Private/StateTreeBaker.cpp b/Engine/Plugins/Runtime/StateTree/Source/StateTreeEditorModule/Private/StateTreeBaker.cpp
--- a/Engine/Plugins/Runtime/StateTree/Source/StateTreeEditorModule/Private/StateTreeBaker.cpp
+++ b/Engine/Plugins/Runtime/StateTree/Source/StateTreeEditorModule/Private/StateTreeBaker.cpp
@@ -68,11 +68,11 @@ bool FStateTreeBaker::CreateStates()
 	// Create item for the runtime execution state
 	StateTree->Instances.Add(FInstancedStruct::Make<FStateTreeExecutionState>());
 
-	for (UStateTreeState* Routine : TreeData->Routines)
+	for (UStateTreeState* SubTree : TreeData->SubTrees)
 	{
-		if (Routine)
+		if (SubTree != nullptr)
 		{
-			if (!CreateStateRecursive(*Routine, FStateTreeHandle::Invalid))
+			if (!CreateStateRecursive(*SubTree, FStateTreeHandle::Invalid))
 			{
 				return false;
 			}
diff --git a/Engine/Plugins/Runtime/StateTree/Source/StateTreeEditorModule/Private/StateTreeEditor.cpp b/Engine/Plugins/Runtime/StateTree/Source/StateTreeEditorModule/Private/StateTreeEditor.cpp
--- a/Engine/Plugins/Runtime/StateTree/Source/StateTreeEditorModule/Private/StateTreeEditor.cpp
+++ b/Engine/Plugins/Runtime/StateTree/Source/StateTreeEditorModule/Private/StateTreeEditor.cpp
@@ -380,6 +380,11 @@ void FStateTreeEditor::OnSchemaChanged(const UStateTree& InStateTree)
 		{
 			StateTreeViewModel->NotifyAssetChangedExternally();
 		}
+
+		if (SelectionDetailsView.IsValid())
+		{
+			SelectionDetailsView->ForceRefresh();
+		}
 	}
 }
 
@@ -452,12 +457,12 @@ namespace UE::StateTree::Editor::Internal
 		TMap<FGuid, FName> IDToName;
 		TArray<UStateTreeState*> Stack;
 
-		for (UStateTreeState* Routine : TreeData->Routines)
+		for (UStateTreeState* SubTree : TreeData->SubTrees)
 		{
-			if (Routine)
+			if (SubTree)
 			{
 				Stack.Reset();
-				Stack.Add(Routine);
+				Stack.Add(SubTree);
 				while (Stack.Num() > 0)
 				{
 					UStateTreeState* CurState = Stack.Pop();
@@ -474,12 +479,12 @@ namespace UE::StateTree::Editor::Internal
 		}
 
 		// Fix changed names.
-		for (UStateTreeState* Routine : TreeData->Routines)
+		for (UStateTreeState* SubTree : TreeData->SubTrees)
 		{
-			if (Routine)
+			if (SubTree)
 			{
 				Stack.Reset();
-				Stack.Add(Routine);
+				Stack.Add(SubTree);
 				while (Stack.Num() > 0)
 				{
 					UStateTreeState* CurState = Stack.Pop();
@@ -511,13 +516,13 @@ namespace UE::StateTree::Editor::Internal
 
 		TArray<UStateTreeState*> Stack;
 
-		for (UStateTreeState* Routine : TreeData->Routines)
+		for (UStateTreeState* SubTree : TreeData->SubTrees)
 		{
-			if (Routine)
+			if (SubTree)
 			{
-				Routine->Parent = nullptr;
+				SubTree->Parent = nullptr;
 				Stack.Reset();
-				Stack.Add(Routine);
+				Stack.Add(SubTree);
 				while (Stack.Num() > 0)
 				{
 					UStateTreeState* CurState = Stack.Pop();
@@ -550,12 +555,12 @@ namespace UE::StateTree::Editor::Internal
 
 		TArray<UStateTreeState*> Stack;
 
-		for (UStateTreeState* Routine : TreeData->Routines)
+		for (UStateTreeState* SubTree : TreeData->SubTrees)
 		{
-			if (Routine)
+			if (SubTree)
 			{
 				Stack.Reset();
-				Stack.Add(Routine);
+				Stack.Add(SubTree);
 				while (Stack.Num() > 0)
 				{
 					UStateTreeState* CurState = Stack.Pop();
diff --git a/Engine/Plugins/Runtime/StateTree/Source/StateTreeEditorModule/Private/StateTreeEditorData.cpp b/Engine/Plugins/Runtime/StateTree/Source/StateTreeEditorModule/Private/StateTreeEditorData.cpp
--- a/Engine/Plugins/Runtime/StateTree/Source/StateTreeEditorModule/Private/StateTreeEditorData.cpp
+++ b/Engine/Plugins/Runtime/StateTree/Source/StateTreeEditorModule/Private/StateTreeEditorData.cpp
@@ -4,6 +4,18 @@
 #include "StateTreeConditionBase.h"
 #include "CoreMinimal.h"
 
+
+void UStateTreeEditorData::PostInitProperties()
+{
+	Super::PostInitProperties();
+	
+	if (!HasAnyFlags(RF_ClassDefaultObject) && !(GetOuter() && GetOuter()->HasAnyFlags(RF_ClassDefaultObject|RF_ArchetypeObject)))
+	{
+		// Add root node by default.
+		AddRootState();
+	}
+}
+
 void UStateTreeEditorData::GetAccessibleStructs(const FGuid TargetStructID, TArray<FStateTreeBindableStructDesc>& OutStructDescs) const
 {
 	// Find the states that are updated before the current state.
@@ -126,14 +138,14 @@ void UStateTreeEditorData::VisitHierarchy(TFunctionRef<bool(const UStateTreeStat
 	TArray<const UStateTreeState*> Stack;
 	bool bContinue = true;
 
-	for (const UStateTreeState* Routine : Routines)
+	for (const UStateTreeState* SubTree : SubTrees)
 	{
-		if (!Routine)
+		if (!SubTree)
 		{
 			continue;
 		}
 
-		Stack.Add(Routine);
+		Stack.Add(SubTree);
 
 		while (!Stack.IsEmpty() && bContinue)
 		{
diff --git a/Engine/Plugins/Runtime/StateTree/Source/StateTreeEditorModule/Private/StateTreeEditorModule.cpp b/Engine/Plugins/Runtime/StateTree/Source/StateTreeEditorModule/Private/StateTreeEditorModule.cpp
--- a/Engine/Plugins/Runtime/StateTree/Source/StateTreeEditorModule/Private/StateTreeEditorModule.cpp
+++ b/Engine/Plugins/Runtime/StateTree/Source/StateTreeEditorModule/Private/StateTreeEditorModule.cpp
@@ -36,7 +36,10 @@ void FStateTreeEditorModule::StartupModule()
 	// Register asset types
 	IAssetTools& AssetTools = FModuleManager::LoadModuleChecked<FAssetToolsModule>("AssetTools").Get();
 
-	TSharedPtr<FAssetTypeActions_StateTree> StateTreeAssetTypeAction = MakeShareable(new FAssetTypeActions_StateTree);
+	// Create or find existing category
+	const EAssetTypeCategories::Type AICategory = AssetTools.RegisterAdvancedAssetCategory(FName(TEXT("AI")), LOCTEXT("AIAssetCategory", "Artificial Intelligence"));
+
+	TSharedPtr<FAssetTypeActions_StateTree> StateTreeAssetTypeAction = MakeShareable(new FAssetTypeActions_StateTree(AICategory | EAssetTypeCategories::Gameplay));
 	ItemDataAssetTypeActions.Add(StateTreeAssetTypeAction);
 	AssetTools.RegisterAssetTypeActions(StateTreeAssetTypeAction.ToSharedRef());
 
diff --git a/Engine/Plugins/Runtime/StateTree/Source/StateTreeEditorModule/Private/StateTreeFactory.cpp b/Engine/Plugins/Runtime/StateTree/Source/StateTreeEditorModule/Private/StateTreeFactory.cpp
--- a/Engine/Plugins/Runtime/StateTree/Source/StateTreeEditorModule/Private/StateTreeFactory.cpp
+++ b/Engine/Plugins/Runtime/StateTree/Source/StateTreeEditorModule/Private/StateTreeFactory.cpp
@@ -2,8 +2,6 @@
 
 #include "StateTreeFactory.h"
 #include "StateTree.h"
-#include "StateTreeEditorModule.h"
-#include "AIModule.h"
 
 #define LOCTEXT_NAMESPACE "StateTreeEditor"
 
@@ -18,12 +16,6 @@ UStateTreeFactory::UStateTreeFactory(const FObjectInitializer& ObjectInitializer
 	SupportedClass = UStateTree::StaticClass();
 }
 
-uint32 UStateTreeFactory::GetMenuCategories() const
-{
-	IAIModule& AIModule = FModuleManager::GetModuleChecked<IAIModule>("AIModule").Get();
-	return AIModule.GetAIAssetCategoryBit();
-}
-
 UObject* UStateTreeFactory::FactoryCreateNew(UClass* Class, UObject* InParent, FName Name, EObjectFlags Flags, UObject* Context, FFeedbackContext* Warn)
 {
 	UStateTree* NewStateTree = NewObject<UStateTree>(InParent, Class, Name, Flags);
diff --git a/Engine/Plugins/Runtime/StateTree/Source/StateTreeEditorModule/Private/StateTreeFactory.h b/Engine/Plugins/Runtime/StateTree/Source/StateTreeEditorModule/Private/StateTreeFactory.h
--- a/Engine/Plugins/Runtime/StateTree/Source/StateTreeEditorModule/Private/StateTreeFactory.h
+++ b/Engine/Plugins/Runtime/StateTree/Source/StateTreeEditorModule/Private/StateTreeFactory.h
@@ -17,7 +17,6 @@ class STATETREEEDITORMODULE_API UStateTreeFactory : public UFactory
 	GENERATED_UCLASS_BODY()
 
 	// UFactory interface
-	virtual uint32 GetMenuCategories() const override;
 	virtual UObject* FactoryCreateNew(UClass* Class, UObject* InParent, FName Name, EObjectFlags Flags, UObject* Context, FFeedbackContext* Warn) override;
 	// End of UFactory interface
 };
diff --git a/Engine/Plugins/Runtime/StateTree/Source/StateTreeEditorModule/Private/StateTreeViewModel.cpp b/Engine/Plugins/Runtime/StateTree/Source/StateTreeEditorModule/Private/StateTreeViewModel.cpp
--- a/Engine/Plugins/Runtime/StateTree/Source/StateTreeEditorModule/Private/StateTreeViewModel.cpp
+++ b/Engine/Plugins/Runtime/StateTree/Source/StateTreeEditorModule/Private/StateTreeViewModel.cpp
@@ -11,7 +11,6 @@
 #include "EditorSupportDelegates.h"
 #include "ScopedTransaction.h"
 #include "Modules/ModuleManager.h"
-#include "AIGraphTypes.h" // Class cache
 #include "StateTreeEditorModule.h"
 #include "Framework/MultiBox/MultiBoxBuilder.h"
 
@@ -188,9 +187,14 @@ void FStateTreeViewModel::NotifyStatesChangedExternally(const TSet<UStateTreeSta
 	OnStatesChanged.Broadcast(ChangedStates, PropertyChangedEvent);
 }
 
-TArray<UStateTreeState*>* FStateTreeViewModel::GetRoutines()
+TArray<UStateTreeState*>* FStateTreeViewModel::GetSubTrees()
 {
-	return TreeData ? &TreeData->Routines : nullptr;
+	return TreeData != nullptr ? &TreeData->SubTrees : nullptr;
+}
+
+int32 FStateTreeViewModel::GetSubTreeCount() const
+{
+	return TreeData != nullptr ? TreeData->SubTrees.Num() : 0;
 }
 
 void FStateTreeViewModel::PostUndo(bool bSuccess)
@@ -278,6 +282,11 @@ void FStateTreeViewModel::GetSelectedStates(TArray<UStateTreeState*>& OutSelecte
 	}
 }
 
+bool FStateTreeViewModel::HasSelection() const
+{
+	return SelectedStates.Num() > 0;
+}
+
 void FStateTreeViewModel::GetPersistentExpandedStates(TSet<UStateTreeState*>& OutExpandedStates)
 {
 	OutExpandedStates.Reset();
@@ -285,9 +294,9 @@ void FStateTreeViewModel::GetPersistentExpandedStates(TSet<UStateTreeState*>& Ou
 	{
 		return;
 	}
-	for (UStateTreeState* Routine : TreeData->Routines)
+	for (UStateTreeState* SubTree : TreeData->SubTrees)
 	{
-		GetExpandedStatesRecursive(Routine, OutExpandedStates);
+		GetExpandedStatesRecursive(SubTree, OutExpandedStates);
 	}
 }
 
@@ -348,14 +357,14 @@ void FStateTreeViewModel::AddState(UStateTreeState* AfterState)
 			NewState->Parent = nullptr;
 		}
 
-		TArray<UStateTreeState*>& ArrayToAddTo = ParentState ? ParentState->Children : TreeData->Routines;
+		TArray<UStateTreeState*>& ArrayToAddTo = ParentState ? ParentState->Children : TreeData->SubTrees;
 		FStateTreeViewUtilities::InsertRecursive(ParentState, ArrayToAddTo, AfterState, NewState, 1); // Insert after
 	}
 	else
 	{
 		TreeData->Modify();
 		NewState->Parent = nullptr;
-		TreeData->Routines.Add(NewState);
+		TreeData->SubTrees.Add(NewState);
 	}
 
 	OnStateAdded.Broadcast(ParentState, NewState);
@@ -432,7 +441,7 @@ void FStateTreeViewModel::RemoveSelectedStates()
 				{
 					AffectedParents.Add(nullptr);
 					TreeData->Modify();
-					FStateTreeViewUtilities::RemoveRecursive(TreeData->Routines, State);
+					FStateTreeViewUtilities::RemoveRecursive(TreeData->SubTrees, State);
 				}
 			}
 		}
@@ -510,8 +519,8 @@ void FStateTreeViewModel::MoveSelectedStates(UStateTreeState* TargetState, int32
 				UStateTreeState* SelectedParent = SelectedState->Parent;
 				AffectedStates.Add(SelectedState);
 
-				TArray<UStateTreeState*>& ArrayToRemoveFrom = SelectedParent ? SelectedParent->Children : TreeData->Routines;
-				TArray<UStateTreeState*>& ArrayToMoveTo = TargetParent ? TargetParent->Children : TreeData->Routines;
+				TArray<UStateTreeState*>& ArrayToRemoveFrom = SelectedParent ? SelectedParent->Children : TreeData->SubTrees;
+				TArray<UStateTreeState*>& ArrayToMoveTo = TargetParent ? TargetParent->Children : TreeData->SubTrees;
 				FStateTreeViewUtilities::RemoveRecursive(ArrayToRemoveFrom, SelectedState);
 				FStateTreeViewUtilities::InsertRecursive(TargetParent, ArrayToMoveTo, TargetState, SelectedState, RelativeLocation);
 			}
diff --git a/Engine/Plugins/Runtime/StateTree/Source/StateTreeEditorModule/Private/StateTreeViewModel.h b/Engine/Plugins/Runtime/StateTree/Source/StateTreeEditorModule/Private/StateTreeViewModel.h
--- a/Engine/Plugins/Runtime/StateTree/Source/StateTreeEditorModule/Private/StateTreeViewModel.h
+++ b/Engine/Plugins/Runtime/StateTree/Source/StateTreeEditorModule/Private/StateTreeViewModel.h
@@ -42,9 +42,11 @@ public:
 	bool IsSelected(const UStateTreeState* State) const;
 	bool IsChildOfSelection(const UStateTreeState* State) const;
 	void GetSelectedStates(TArray<UStateTreeState*>& OutSelectedStates);
+	bool HasSelection() const;
 
-	// Returns array of routines to edit.
-	TArray<UStateTreeState*>* GetRoutines();
+	// Returns array of subtrees to edit.
+	TArray<UStateTreeState*>* GetSubTrees();
+	int32 GetSubTreeCount() const;
 
 	// Gets and sets StateTree view expansion state store in the asset.
 	void GetPersistentExpandedStates(TSet<UStateTreeState*>& OutExpandedStates);
diff --git a/Engine/Plugins/Runtime/StateTree/Source/StateTreeEditorModule/Public/StateTreeEditorData.h b/Engine/Plugins/Runtime/StateTree/Source/StateTreeEditorModule/Public/StateTreeEditorData.h
--- a/Engine/Plugins/Runtime/StateTree/Source/StateTreeEditorModule/Public/StateTreeEditorData.h
+++ b/Engine/Plugins/Runtime/StateTree/Source/StateTreeEditorModule/Public/StateTreeEditorData.h
@@ -22,7 +22,10 @@ UCLASS(BlueprintType, EditInlineNew, CollapseCategories)
 class STATETREEEDITORMODULE_API UStateTreeEditorData : public UObject, public IStateTreeEditorPropertyBindingsOwner
 {
 	GENERATED_BODY()
+	
 public:
+	virtual void PostInitProperties() override;
+	
 	// IStateTreeEditorPropertyBindingsOwner
 	virtual void GetAccessibleStructs(const FGuid TargetStructID, TArray<FStateTreeBindableStructDesc>& OutStructDescs) const override;
 	virtual bool GetStructByID(const FGuid StructID, FStateTreeBindableStructDesc& OutStructDesc) const override;
@@ -52,10 +55,19 @@ public:
 		UStateTreeState* SubTreeState = NewObject<UStateTreeState>(this);
 		check(SubTreeState);
 		SubTreeState->Name = Name;
-		Routines.Add(SubTreeState);
+		SubTrees.Add(SubTreeState);
 		return *SubTreeState;
 	}
 
+	/**
+	 * Adds new Subtree named "Root".
+	 * @return Pointer to the new Subtree.
+	 */
+	UStateTreeState& AddRootState()
+	{
+		return AddSubTree(FName(TEXT("Root")));
+	}
+	
 	/**
 	 * Adds property binding between two structs.
 	 */
@@ -64,12 +76,12 @@ public:
 		EditorBindings.AddPropertyBinding(SourcePath, TargetPath);
 	}
 	// ~StateTree Builder API
-
+	
 	
 	UPROPERTY()
 	FStateTreeEditorPropertyBindings EditorBindings;
 
 	// Top level States aka Routines.
 	UPROPERTY()
-	TArray<UStateTreeState*> Routines;
+	TArray<UStateTreeState*> SubTrees;
 };
diff --git a/Engine/Plugins/Runtime/StateTree/Source/StateTreeEditorModule/StateTreeEditorModule.Build.cs b/Engine/Plugins/Runtime/StateTree/Source/StateTreeEditorModule/StateTreeEditorModule.Build.cs
--- a/Engine/Plugins/Runtime/StateTree/Source/StateTreeEditorModule/StateTreeEditorModule.Build.cs
+++ b/Engine/Plugins/Runtime/StateTree/Source/StateTreeEditorModule/StateTreeEditorModule.Build.cs
@@ -24,7 +24,6 @@ namespace UnrealBuildTool.Rules
 				"SlateCore",
 				"EditorStyle",
 				"PropertyEditor",
-				"AIModule",
 				"StateTreeModule",
 				"Projects",
 				"BlueprintGraph",
@@ -41,7 +40,6 @@ namespace UnrealBuildTool.Rules
 				"KismetWidgets",
 				"PropertyPath",
 				"PropertyEditor",
-				"AIGraph",
 				"ToolMenus",
 			}
 			);
diff --git a/Engine/Plugins/Runtime/StateTree/Source/StateTreeModule/Private/StateTreeSchema.cpp b/Engine/Plugins/Runtime/StateTree/Source/StateTreeModule/Private/StateTreeSchema.cpp
new file mode 100644
--- /dev/null
+++ b/Engine/Plugins/Runtime/StateTree/Source/StateTreeModule/Private/StateTreeSchema.cpp
@@ -0,0 +1,12 @@
+﻿// Copyright Epic Games, Inc. All Rights Reserved.
+
+#include "StateTreeSchema.h"
+#include "CoreMinimal.h"
+#include "Blueprint/StateTreeEvaluatorBlueprintBase.h"
+#include "Blueprint/StateTreeTaskBlueprintBase.h"
+#include "Blueprint/StateTreeConditionBlueprintBase.h"
+
+bool UStateTreeSchema::IsChildOfBlueprintBase(const UClass* InClass) const
+{
+	return InClass->IsChildOf<UStateTreeItemBlueprintBase>();
+}
diff --git a/Engine/Plugins/Runtime/StateTree/Source/StateTreeModule/StateTreeModule.Build.cs b/Engine/Plugins/Runtime/StateTree/Source/StateTreeModule/StateTreeModule.Build.cs
--- a/Engine/Plugins/Runtime/StateTree/Source/StateTreeModule/StateTreeModule.Build.cs
+++ b/Engine/Plugins/Runtime/StateTree/Source/StateTreeModule/StateTreeModule.Build.cs
@@ -17,11 +17,8 @@ namespace UnrealBuildTool.Rules
 					"CoreUObject",
 					"Engine",
 					"AIModule",
-					"GameplayBehaviorsModule",
 					"GameplayTasks",
 					"GameplayTags",
-					"GameplayAbilities",
-					"SmartObjectsModule",
 					"StructUtils",
 				}
 			);
diff --git a/Engine/Plugins/Runtime/StateTree/Source/StateTreeModule/StateTreeSchema.cpp b/Engine/Plugins/Runtime/StateTree/Source/StateTreeModule/StateTreeSchema.cpp
deleted file mode 100644
