StateTree Editor compile improvements
- Added Compile button and compilation state indication
- Removed compile-on-save
- Prevent compiling during PIE

diff --git a/Engine/Plugins/Runtime/StateTree/Source/StateTreeEditorModule/Private/StateTreeEditor.cpp b/Engine/Plugins/Runtime/StateTree/Source/StateTreeEditorModule/Private/StateTreeEditor.cpp
--- a/Engine/Plugins/Runtime/StateTree/Source/StateTreeEditorModule/Private/StateTreeEditor.cpp
+++ b/Engine/Plugins/Runtime/StateTree/Source/StateTreeEditorModule/Private/StateTreeEditor.cpp
@@ -18,6 +18,9 @@
 #include "Widgets/Input/SMultiLineEditableTextBox.h"
 #include "Customizations/StateTreeBindingExtension.h"
 #include "Misc/UObjectToken.h"
+#include "StateTreeEditorCommands.h"
+#include "ToolMenus.h"
+#include "StateTreeObjectHash.h"
 
 #include "Developer/MessageLog/Public/IMessageLogListing.h"
 #include "Developer/MessageLog/Public/MessageLogInitializationOptions.h"
@@ -178,15 +181,15 @@ void FStateTreeEditor::InitEditor( const EToolkitMode::Type Mode, const TSharedP
 	const bool bCreateDefaultStandaloneMenu = true;
 	const bool bCreateDefaultToolbar = true;
 	FAssetEditorToolkit::InitAssetEditor(Mode, InitToolkitHost, StateTreeEditorAppName, StandaloneDefaultLayout, bCreateDefaultStandaloneMenu, bCreateDefaultToolbar, StateTree);
+
+	BindCommands();
+	RegisterToolbar();
 	
 	FStateTreeEditorModule& StateTreeEditorModule = FModuleManager::LoadModuleChecked<FStateTreeEditorModule>("StateTreeEditorModule");
 	AddMenuExtender(StateTreeEditorModule.GetMenuExtensibilityManager()->GetAllExtenders(GetToolkitCommands(), GetEditingObjects()));
 
 	RegenerateMenusAndToolbars();
 
-	// TODO: should only do validation here to not edit asset each time it is opened.
-	UpdateAsset();
-
 	UE::StateTree::Delegates::OnIdentifierChanged.AddSP(this, &FStateTreeEditor::OnIdentifierChanged);
 	UE::StateTree::Delegates::OnSchemaChanged.AddSP(this, &FStateTreeEditor::OnSchemaChanged);
 }
@@ -417,6 +420,7 @@ void FStateTreeEditor::OnSelectionFinishedChangingProperties(const FPropertyChan
 		if (ChangedStates.Num() > 0)
 		{
 			StateTreeViewModel->NotifyStatesChangedExternally(ChangedStates, PropertyChangedEvent);
+			UpdateAsset();
 		}
 	}
 }
@@ -628,20 +632,55 @@ namespace UE::StateTree::Editor::Internal
 		TreeData->GetAllStructIDs(AllStructIDs);
 		TreeData->GetPropertyEditorBindings()->RemoveUnusedBindings(AllStructIDs);
 	}
+
 }
 
-void FStateTreeEditor::UpdateAsset()
+void FStateTreeEditor::BindCommands()
+{
+	const FStateTreeEditorCommands& Commands = FStateTreeEditorCommands::Get();
+
+	ToolkitCommands->MapAction(
+		Commands.Compile,
+		FExecuteAction::CreateSP(this, &FStateTreeEditor::Compile),
+		FCanExecuteAction::CreateSP(this, &FStateTreeEditor::CanCompile));
+}
+
+void FStateTreeEditor::RegisterToolbar()
+{
+	UToolMenus* ToolMenus = UToolMenus::Get();
+	UToolMenu* ToolBar;
+	FName ParentName;
+	const FName MenuName = GetToolMenuToolbarName(ParentName);
+	if (ToolMenus->IsMenuRegistered(MenuName))
+	{
+		ToolBar = ToolMenus->ExtendMenu(MenuName);
+	}
+	else
+	{
+		ToolBar = UToolMenus::Get()->RegisterMenu(MenuName, ParentName, EMultiBoxType::ToolBar);
+	}
+
+	const FStateTreeEditorCommands& Commands = FStateTreeEditorCommands::Get();
+	FToolMenuInsert InsertAfterAssetSection("Asset", EToolMenuInsertType::After);
+	{
+		FToolMenuSection& Section = ToolBar->AddSection("Compile", TAttribute<FText>(), InsertAfterAssetSection);
+		Section.AddEntry(FToolMenuEntry::InitToolBarButton(
+			Commands.Compile,
+			TAttribute<FText>(),
+			TAttribute<FText>(),
+			TAttribute<FSlateIcon>(this, &FStateTreeEditor::GetCompileStatusImage)));
+	}
+}
+
+void FStateTreeEditor::Compile()
 {
 	if (!StateTree)
 	{
 		return;
 	}
 
-	UE::StateTree::Editor::Internal::UpdateParents(*StateTree);
-	UE::StateTree::Editor::Internal::ApplySchema(*StateTree);
-	UE::StateTree::Editor::Internal::RemoveUnusedBindings(*StateTree);
-	UE::StateTree::Editor::Internal::ValidateLinkedStates(*StateTree);
-
+	UpdateAsset();
+	
 	if (CompilerResultsListing.IsValid())
 	{
 		CompilerResultsListing->ClearMessages();
@@ -657,15 +696,83 @@ void FStateTreeEditor::UpdateAsset()
 		Log.AppendToLog(CompilerResultsListing.Get());
 	}
 
-	if (!bSuccess)
+	if (bSuccess)
+	{
+		// Success
+		StateTree->LastCompiledEditorDataHash = EditorDataHash;
+	}
+	else
 	{
 		// Make sure not to leave stale data on failed bake.
 		StateTree->ResetBaked();
+		StateTree->LastCompiledEditorDataHash = 0;
 
 		// Show log
 		TabManager->TryInvokeTab(CompilerResultsTabId);
 	}
 }
 
+bool FStateTreeEditor::CanCompile() const
+{
+	if (StateTree == nullptr)
+	{
+		return false;
+	}
+	
+	// We can't recompile while in PIE
+	if (GEditor->IsPlaySessionInProgress())
+	{
+		return false;
+	}
+
+	return true;
+}
+
+FSlateIcon FStateTreeEditor::GetCompileStatusImage() const
+{
+	static const FName CompileStatusBackground("Blueprint.CompileStatus.Background");
+	static const FName CompileStatusUnknown("Blueprint.CompileStatus.Overlay.Unknown");
+	static const FName CompileStatusError("Blueprint.CompileStatus.Overlay.Error");
+	static const FName CompileStatusGood("Blueprint.CompileStatus.Overlay.Good");
+	static const FName CompileStatusWarning("Blueprint.CompileStatus.Overlay.Warning");
+
+	if (StateTree == nullptr)
+	{
+		return FSlateIcon(FEditorStyle::GetStyleSetName(), CompileStatusBackground, NAME_None, CompileStatusUnknown);
+	}
+	
+	if (!StateTree->IsValidStateTree())
+	{
+		return FSlateIcon(FEditorStyle::GetStyleSetName(), CompileStatusBackground, NAME_None, CompileStatusError);
+	}
+	
+	if (StateTree->LastCompiledEditorDataHash != EditorDataHash)
+	{
+		return FSlateIcon(FEditorStyle::GetStyleSetName(), CompileStatusBackground, NAME_None, CompileStatusUnknown);
+	}
+	
+	return FSlateIcon(FEditorStyle::GetStyleSetName(), CompileStatusBackground, NAME_None, CompileStatusGood);
+}
+
+void FStateTreeEditor::UpdateAsset()
+{
+	if (!StateTree)
+	{
+		return;
+	}
+
+	UE::StateTree::Editor::Internal::UpdateParents(*StateTree);
+	UE::StateTree::Editor::Internal::ApplySchema(*StateTree);
+	UE::StateTree::Editor::Internal::RemoveUnusedBindings(*StateTree);
+	UE::StateTree::Editor::Internal::ValidateLinkedStates(*StateTree);
+
+	EditorDataHash = 0;
+	if (StateTree->EditorData != nullptr)
+	{
+		FStateTreeObjectCRC32 Archive;
+		EditorDataHash = Archive.Crc32(StateTree->EditorData, 0);
+	}
+}
+
 
 #undef LOCTEXT_NAMESPACE
diff --git a/Engine/Plugins/Runtime/StateTree/Source/StateTreeEditorModule/Private/StateTreeEditor.h b/Engine/Plugins/Runtime/StateTree/Source/StateTreeEditorModule/Private/StateTreeEditor.h
--- a/Engine/Plugins/Runtime/StateTree/Source/StateTreeEditorModule/Private/StateTreeEditor.h
+++ b/Engine/Plugins/Runtime/StateTree/Source/StateTreeEditorModule/Private/StateTreeEditor.h
@@ -59,6 +59,13 @@ private:
 	TSharedRef<SDockTab> SpawnTab_StateTreeStatistics(const FSpawnTabArgs& Args);
 	TSharedRef<SDockTab> SpawnTab_CompilerResults(const FSpawnTabArgs& Args);
 
+	void BindCommands();
+	void RegisterToolbar();
+
+	void Compile();
+	bool CanCompile() const;
+	FSlateIcon GetCompileStatusImage() const;
+
 	void UpdateAsset();
 
 	void HandleModelAssetChanged();
@@ -68,8 +75,10 @@ private:
 	FText GetStatisticsText() const;
 
 	/* State Tree being edited */
-	UStateTree* StateTree;
+	UStateTree* StateTree = nullptr;
 
+	uint32 EditorDataHash = 0;
+	
 	/** Selection Property View */
 	TSharedPtr<class IDetailsView> SelectionDetailsView;
 
diff --git a/Engine/Plugins/Runtime/StateTree/Source/StateTreeEditorModule/Private/StateTreeEditorCommands.cpp b/Engine/Plugins/Runtime/StateTree/Source/StateTreeEditorModule/Private/StateTreeEditorCommands.cpp
new file mode 100644
--- /dev/null
+++ b/Engine/Plugins/Runtime/StateTree/Source/StateTreeEditorModule/Private/StateTreeEditorCommands.cpp
@@ -0,0 +1,23 @@
+﻿// Copyright Epic Games, Inc. All Rights Reserved.
+
+#include "StateTreeEditorCommands.h"
+#include "EditorStyleSet.h"
+
+#define LOCTEXT_NAMESPACE "StateTreeEditor"
+
+FStateTreeEditorCommands::FStateTreeEditorCommands() 
+	: TCommands<FStateTreeEditorCommands>(
+		"StateTreeEditor", // Context name for fast lookup
+		LOCTEXT("StateTreeEditor", "StateTree Editor"), // Localized context name for displaying
+		NAME_None,
+		FEditorStyle::GetStyleSetName()
+	)
+{
+}
+
+void FStateTreeEditorCommands::RegisterCommands()
+{
+	UI_COMMAND(Compile, "Compile", "Compile the current StateTree.", EUserInterfaceActionType::Button, FInputChord(EKeys::F7));
+}
+
+#undef LOCTEXT_NAMESPACE
diff --git a/Engine/Plugins/Runtime/StateTree/Source/StateTreeEditorModule/Private/StateTreeEditorCommands.h b/Engine/Plugins/Runtime/StateTree/Source/StateTreeEditorModule/Private/StateTreeEditorCommands.h
new file mode 100644
--- /dev/null
+++ b/Engine/Plugins/Runtime/StateTree/Source/StateTreeEditorModule/Private/StateTreeEditorCommands.h
@@ -0,0 +1,20 @@
+﻿// Copyright Epic Games, Inc. All Rights Reserved.
+
+#pragma once
+
+#include "Framework/Commands/Commands.h"
+
+
+/**
+ * StateTree editor command set.
+ */
+class FStateTreeEditorCommands : public TCommands<FStateTreeEditorCommands>
+{
+public:
+	FStateTreeEditorCommands();
+
+	// TCommands<> overrides
+	virtual void RegisterCommands() override;
+
+	TSharedPtr<FUICommandInfo> Compile;
+};
\ No newline at end of file
diff --git a/Engine/Plugins/Runtime/StateTree/Source/StateTreeEditorModule/Private/StateTreeEditorModule.cpp b/Engine/Plugins/Runtime/StateTree/Source/StateTreeEditorModule/Private/StateTreeEditorModule.cpp
--- a/Engine/Plugins/Runtime/StateTree/Source/StateTreeEditorModule/Private/StateTreeEditorModule.cpp
+++ b/Engine/Plugins/Runtime/StateTree/Source/StateTreeEditorModule/Private/StateTreeEditorModule.cpp
@@ -20,6 +20,7 @@
 #include "Blueprint/StateTreeEvaluatorBlueprintBase.h"
 #include "Blueprint/StateTreeTaskBlueprintBase.h"
 #include "Blueprint/StateTreeConditionBlueprintBase.h"
+#include "StateTreeEditorCommands.h"
 
 
 #define LOCTEXT_NAMESPACE "StateTreeEditor"
@@ -32,6 +33,7 @@ void FStateTreeEditorModule::StartupModule()
 	ToolBarExtensibilityManager = MakeShareable(new FExtensibilityManager);
 
 	FStateTreeEditorStyle::Initialize();
+	FStateTreeEditorCommands::Register();
 
 	// Register asset types
 	IAssetTools& AssetTools = FModuleManager::LoadModuleChecked<FAssetToolsModule>("AssetTools").Get();
@@ -62,6 +64,7 @@ void FStateTreeEditorModule::ShutdownModule()
 	ToolBarExtensibilityManager.Reset();
 
 	FStateTreeEditorStyle::Shutdown();
+	FStateTreeEditorCommands::Unregister();
 
 	// Unregister the data asset type actions
 	if (FModuleManager::Get().IsModuleLoaded("AssetTools"))
diff --git a/Engine/Plugins/Runtime/StateTree/Source/StateTreeEditorModule/Private/StateTreeObjectHash.cpp b/Engine/Plugins/Runtime/StateTree/Source/StateTreeEditorModule/Private/StateTreeObjectHash.cpp
new file mode 100644
--- /dev/null
+++ b/Engine/Plugins/Runtime/StateTree/Source/StateTreeEditorModule/Private/StateTreeObjectHash.cpp
@@ -0,0 +1,15 @@
+﻿// Copyright Epic Games, Inc. All Rights Reserved.
+
+#include "StateTreeObjectHash.h"
+
+#if WITH_EDITORONLY_DATA
+
+bool FStateTreeObjectCRC32::ShouldSkipProperty(const FProperty* InProperty) const
+{
+	static const FName ExcludeFromHashName(TEXT("ExcludeFromHash"));
+	const bool bExclude = InProperty != nullptr && InProperty->HasMetaData(ExcludeFromHashName);
+	
+	return FArchiveObjectCrc32::ShouldSkipProperty(InProperty) || InProperty->HasAllPropertyFlags(CPF_Transient) || bExclude;
+}
+
+#endif // WITH_EDITORONLY_DATA
\ No newline at end of file
diff --git a/Engine/Plugins/Runtime/StateTree/Source/StateTreeEditorModule/Private/StateTreeObjectHash.h b/Engine/Plugins/Runtime/StateTree/Source/StateTreeEditorModule/Private/StateTreeObjectHash.h
new file mode 100644
--- /dev/null
+++ b/Engine/Plugins/Runtime/StateTree/Source/StateTreeEditorModule/Private/StateTreeObjectHash.h
@@ -0,0 +1,21 @@
+﻿// Copyright Epic Games, Inc. All Rights Reserved.
+
+#pragma once
+
+#include "Serialization/ArchiveObjectCrc32.h"
+
+#if WITH_EDITORONLY_DATA
+
+/**
+ * Archive based object hashing to be used with StateTree data calculation.
+ * If a property has "IncludeInHash" meta tag, any of its child properties will be included in the hash.
+ * Editor only since it relies on meta data.
+ */
+class FStateTreeObjectCRC32 : public FArchiveObjectCrc32
+{
+public:
+
+	virtual bool ShouldSkipProperty(const FProperty* InProperty) const override;
+};
+
+#endif // WITH_EDITORONLY_DATA
\ No newline at end of file
diff --git a/Engine/Plugins/Runtime/StateTree/Source/StateTreeEditorModule/Public/StateTreeEditorData.h b/Engine/Plugins/Runtime/StateTree/Source/StateTreeEditorModule/Public/StateTreeEditorData.h
--- a/Engine/Plugins/Runtime/StateTree/Source/StateTreeEditorModule/Public/StateTreeEditorData.h
+++ b/Engine/Plugins/Runtime/StateTree/Source/StateTreeEditorModule/Public/StateTreeEditorData.h
@@ -78,7 +78,7 @@ public:
 	// ~StateTree Builder API
 	
 	
-	UPROPERTY()
+	UPROPERTY(meta = (ExcludeFromHash))
 	FStateTreeEditorPropertyBindings EditorBindings;
 
 	// Top level States aka Routines.
diff --git a/Engine/Plugins/Runtime/StateTree/Source/StateTreeModule/Public/StateTree.h b/Engine/Plugins/Runtime/StateTree/Source/StateTreeModule/Public/StateTree.h
--- a/Engine/Plugins/Runtime/StateTree/Source/StateTreeModule/Public/StateTree.h
+++ b/Engine/Plugins/Runtime/StateTree/Source/StateTreeModule/Public/StateTree.h
@@ -61,6 +61,10 @@ public:
 	// Edit time data for the StateTree, instance of UStateTreeEditorData
 	UPROPERTY()
 	TObjectPtr<UObject> EditorData;
+
+	// Hash of the editor data from last compile.
+	UPROPERTY()
+	uint32 LastCompiledEditorDataHash = 0;
 #endif
 
 protected:
