StateTree: Renamed item to node
- Renamed uses of item to node
- Added base class for Conditions, Evals & Tasks

diff --git a/Engine/Plugins/Runtime/StateTree/Config/BaseStateTree.ini b/Engine/Plugins/Runtime/StateTree/Config/BaseStateTree.ini
--- a/Engine/Plugins/Runtime/StateTree/Config/BaseStateTree.ini
+++ b/Engine/Plugins/Runtime/StateTree/Config/BaseStateTree.ini
@@ -10,12 +10,18 @@
 +StructRedirects=(OldName="StateTreeTask2Base",NewName="StateTreeTaskBase")
 +StructRedirects=(OldName="StateTreeEvaluator2Base",NewName="StateTreeEvaluatorBase")
 
++StructRedirects=(OldName="StateTreeItem",NewName="StateTreeEditorNode")
++StructRedirects=(OldName="StateTreeConditionItem",NewName="StateTreeEditorNode")
++StructRedirects=(OldName="StateTreeEvaluatorItem",NewName="StateTreeEditorNode")
++StructRedirects=(OldName="StateTreeTaskItem",NewName="StateTreeEditorNode")
+
 +PropertyRedirects=(OldName="StateTreeState.Transitions2",NewName="StateTreeState.Transitions")
 +PropertyRedirects=(OldName="StateTreeState.EnterConditions2",NewName="StateTreeState.EnterConditions")
 +PropertyRedirects=(OldName="StateTreeState.Evaluators2",NewName="StateTreeState.Evaluators")
 +PropertyRedirects=(OldName="StateTreeState.Tasks2",NewName="StateTreeState.Tasks")
 
 +PropertyRedirects=(OldName="StateTree.Conditions2",NewName="StateTree.Conditions")
++PropertyRedirects=(OldName="StateTreeEditorNode.Item",NewName="StateTreeEditorNode.Node")
 
 +EnumRedirects=(OldName="EStateTreePropertyCopyType",ValueChanges=(("Plain","CopyPlain"),("Complex","CopyComplex"),("Bool","CopyBool"),("Struct","CopyStruct"),("Object","CopyObject"),("Name","CopyName"),("FixedArray","CopyFixedArray")))
 +EnumRedirects=(OldName="EStateTreePropertyAccessType",ValueChanges=(("Array","IndexArray")))
diff --git a/Engine/Plugins/Runtime/StateTree/Source/StateTreeEditorModule/Private/Customizations/StateTreeBindableItemDetails.cpp b/Engine/Plugins/Runtime/StateTree/Source/StateTreeEditorModule/Private/Customizations/StateTreeBindableItemDetails.cpp
deleted file mode 100644
diff --git a/Engine/Plugins/Runtime/StateTree/Source/StateTreeEditorModule/Private/Customizations/StateTreeBindableItemDetails.h b/Engine/Plugins/Runtime/StateTree/Source/StateTreeEditorModule/Private/Customizations/StateTreeBindableItemDetails.h
deleted file mode 100644
diff --git a/Engine/Plugins/Runtime/StateTree/Source/StateTreeEditorModule/Private/Customizations/StateTreeBindingExtension.cpp b/Engine/Plugins/Runtime/StateTree/Source/StateTreeEditorModule/Private/Customizations/StateTreeBindingExtension.cpp
--- a/Engine/Plugins/Runtime/StateTree/Source/StateTreeEditorModule/Private/Customizations/StateTreeBindingExtension.cpp
+++ b/Engine/Plugins/Runtime/StateTree/Source/StateTreeEditorModule/Private/Customizations/StateTreeBindingExtension.cpp
@@ -10,7 +10,6 @@
 #include "StateTreeEvaluatorBase.h"
 #include "StateTree.h"
 #include "Algo/Accumulate.h"
-#include "InstancedStruct.h"
 #include "StateTreePropertyHelpers.h"
 #include "StateTreeAnyEnum.h"
 #include "StateTreePropertyBindingCompiler.h"
@@ -20,6 +19,9 @@
 
 namespace UE::StateTree::PropertyBinding
 {
+	
+const FName StateTreeNodeIDName(TEXT("StateTreeNodeID"));
+	
 
 UObject* FindEditorBindingsOwner(UObject* InObject)
 {
@@ -47,7 +49,7 @@ void GetStructPropertyPath(TSharedPtr<IPropertyHandle> InPropertyHandle, FStateT
 		// Keep track of the path.
 		BindingChain.Insert(FBindingChainElement(CurrentProperty, InPropertyHandle->GetIndexInArray()), 0);
 		
-		if (const FString* IDString = InPropertyHandle->GetInstanceMetaData(FName(TEXT("StateTreeItemID"))))
+		if (const FString* IDString = InPropertyHandle->GetInstanceMetaData(StateTreeNodeIDName))
 		{
 			LexFromString(StructID, **IDString);
 		}
@@ -95,9 +97,8 @@ FOnStateTreeBindingChanged STATETREEEDITORMODULE_API OnStateTreeBindingChanged;
 
 bool FStateTreeBindingExtension::IsPropertyExtendable(const UClass* InObjectClass, const IPropertyHandle& PropertyHandle) const
 {
-	// Bindable property must have item ID
-	static const FName StateTreeItemIDName(TEXT("StateTreeItemID"));
-	if (PropertyHandle.GetInstanceMetaData(StateTreeItemIDName) == nullptr)
+	// Bindable property must have node ID
+	if (PropertyHandle.GetInstanceMetaData(UE::StateTree::PropertyBinding::StateTreeNodeIDName) == nullptr)
 	{
 		return false;
 	}
diff --git a/Engine/Plugins/Runtime/StateTree/Source/StateTreeEditorModule/Private/Customizations/StateTreeBindingExtension.h b/Engine/Plugins/Runtime/StateTree/Source/StateTreeEditorModule/Private/Customizations/StateTreeBindingExtension.h
--- a/Engine/Plugins/Runtime/StateTree/Source/StateTreeEditorModule/Private/Customizations/StateTreeBindingExtension.h
+++ b/Engine/Plugins/Runtime/StateTree/Source/StateTreeEditorModule/Private/Customizations/StateTreeBindingExtension.h
@@ -14,6 +14,7 @@ struct FStateTreeEditorPropertyPath;
 
 namespace UE::StateTree::PropertyBinding
 {
+	extern const FName StateTreeNodeIDName;
 
 	/**
 	 * Get nearest Outer that implements IStateTreeEditorPropertyBindingsOwner.
diff --git a/Engine/Plugins/Runtime/StateTree/Source/StateTreeEditorModule/Private/Customizations/StateTreeEditorNodeDetails.cpp b/Engine/Plugins/Runtime/StateTree/Source/StateTreeEditorModule/Private/Customizations/StateTreeEditorNodeDetails.cpp
new file mode 100644
--- /dev/null
+++ b/Engine/Plugins/Runtime/StateTree/Source/StateTreeEditorModule/Private/Customizations/StateTreeEditorNodeDetails.cpp
@@ -0,0 +1,912 @@
+// Copyright Epic Games, Inc. All Rights Reserved.
+
+#include "StateTreeEditorNodeDetails.h"
+#include "DetailWidgetRow.h"
+#include "DetailLayoutBuilder.h"
+#include "IPropertyUtilities.h"
+#include "IDetailPropertyRow.h"
+#include "IDetailChildrenBuilder.h"
+#include "StateTree.h"
+#include "StateTreeEditorData.h"
+#include "Widgets/DeclarativeSyntaxSupport.h"
+#include "Widgets/Input/SComboButton.h"
+#include "Widgets/Input/SEditableTextBox.h"
+#include "Widgets/Images/SImage.h"
+#include "Widgets/Text/SRichTextBlock.h"
+#include "Framework/MultiBox/MultiBoxBuilder.h"
+#include "InstancedStruct.h"
+#include "Styling/SlateIconFinder.h"
+#include "Editor.h"
+#include "Modules/ModuleManager.h"
+#include "InstancedStructDetails.h"
+#include "Engine/UserDefinedStruct.h"
+#include "StateTreeBindingExtension.h"
+#include "StateTreeDelegates.h"
+#include "StateTreeTaskBase.h"
+#include "StateTreeEvaluatorBase.h"
+#include "StateTreeConditionBase.h"
+#include "StateTreePropertyHelpers.h"
+#include "StateTreeEditorStyle.h"
+#include "EdGraphSchema_K2.h"
+#include "Kismet2/BlueprintEditorUtils.h"
+#include "Blueprint/StateTreeEvaluatorBlueprintBase.h"
+#include "Blueprint/StateTreeTaskBlueprintBase.h"
+#include "Blueprint/StateTreeConditionBlueprintBase.h"
+#include "StateTreeEditorModule.h"
+#include "StateTreeNodeClassCache.h"
+
+#define LOCTEXT_NAMESPACE "StateTreeEditor"
+
+namespace UE::StateTreeEditor::Internal
+{
+	void ModifyRow(IDetailPropertyRow& ChildRow, const FGuid& ID, FStateTreeEditorPropertyBindings* EditorPropBindings)
+	{
+		if (!EditorPropBindings)
+		{
+			return;
+		}
+		
+		TSharedPtr<IPropertyHandle> ChildPropHandle = ChildRow.GetPropertyHandle();
+		const EStateTreePropertyUsage Category = UE::StateTree::PropertyBinding::ParsePropertyUsage(ChildPropHandle);
+		const FProperty* Property = ChildPropHandle->GetProperty();
+		
+		// Conditionally control visibility of the value field of bound properties.
+		if (Category != EStateTreePropertyUsage::Invalid && ID.IsValid())
+		{
+			// Pass the node ID to binding extension. Since the properties are added using AddChildStructure(), we break the hierarchy and cannot access parent.
+			ChildPropHandle->SetInstanceMetaData(UE::StateTree::PropertyBinding::StateTreeNodeIDName, LexToString(ID));
+
+			FStateTreeEditorPropertyPath Path(ID, *Property->GetFName().ToString());
+			TSharedPtr<SWidget> NameWidget;
+			TSharedPtr<SWidget> ValueWidget;
+			FDetailWidgetRow Row;
+			ChildRow.GetDefaultWidgets(NameWidget, ValueWidget, Row);
+
+			auto IsValueVisible = TAttribute<EVisibility>::Create([Path, EditorPropBindings]() -> EVisibility
+				{
+					return EditorPropBindings->HasPropertyBinding(Path) ? EVisibility::Collapsed : EVisibility::Visible;
+				});
+
+			if (Category == EStateTreePropertyUsage::Input || Category == EStateTreePropertyUsage::Output)
+			{
+				FEdGraphPinType PinType;
+				const UEdGraphSchema_K2* Schema = GetDefault<UEdGraphSchema_K2>();
+				Schema->ConvertPropertyToPinType(Property, PinType);
+				const FSlateBrush* Icon = FBlueprintEditorUtils::GetIconFromPin(PinType, true);
+				FText Label = Category == EStateTreePropertyUsage::Input ? LOCTEXT("Input", "Input") : LOCTEXT("Output", "Output");
+				FText Tooltip = FText::FromString(Property->GetCPPType());
+				const FLinearColor Color = Schema->GetPinTypeColor(PinType);
+
+				ChildRow
+					.CustomWidget(/*bShowChildren*/false)
+					.NameContent()
+					[
+						NameWidget.ToSharedRef()
+					]
+					.ValueContent()
+					[
+						SNew(SHorizontalBox)
+						.Visibility(IsValueVisible)
+						+SHorizontalBox::Slot()
+						.AutoWidth()
+						.VAlign(VAlign_Center)
+						.Padding(4.0f, 0.0f)
+						[
+							SNew(SImage)
+							.Image(Icon)
+							.ColorAndOpacity(Color)
+						]
+						+SHorizontalBox::Slot()
+						.VAlign(VAlign_Center)
+						[
+							SNew(STextBlock)
+							.Font(IDetailLayoutBuilder::GetDetailFont())
+							.Text(Label)
+							.ToolTipText(Tooltip)
+						]
+					];
+			}
+			else
+			{
+				ChildRow
+					.CustomWidget(/*bShowChildren*/true)
+					.NameContent()
+					[
+						NameWidget.ToSharedRef()
+					]
+					.ValueContent()
+					[
+						SNew(SBox)
+						.Visibility(IsValueVisible)
+						[
+							ValueWidget.ToSharedRef()
+						]
+					];
+			}
+		}
+	}
+} // UE::StateTreeEditor::Internal
+
+// Customized version of FInstancedStructDataDetails used to hide bindable properties.
+class FBindableNodeInstanceDetails : public FInstancedStructDataDetails
+{
+public:
+
+	FBindableNodeInstanceDetails(TSharedPtr<IPropertyHandle> InStructProperty, FGuid InID, UStateTreeEditorData* InEditorData)
+		: FInstancedStructDataDetails(InStructProperty)
+		, EditorData(InEditorData)
+	{
+		EditorPropBindings = EditorData ? EditorData->GetPropertyEditorBindings() : nullptr;
+		ID = InID;
+	}
+
+	virtual void OnChildRowAdded(IDetailPropertyRow& ChildRow)
+	{
+		UE::StateTreeEditor::Internal::ModifyRow(ChildRow, ID, EditorPropBindings);
+	}
+
+	UStateTreeEditorData* EditorData;
+	FStateTreeEditorPropertyBindings* EditorPropBindings;
+	FGuid ID;
+};
+
+////////////////////////////////////
+
+TSharedRef<IPropertyTypeCustomization> FStateTreeEditorNodeDetails::MakeInstance()
+{
+	return MakeShareable(new FStateTreeEditorNodeDetails);
+}
+
+FStateTreeEditorNodeDetails::~FStateTreeEditorNodeDetails()
+{
+	UE::StateTree::PropertyBinding::OnStateTreeBindingChanged.Remove(OnBindingChangedHandle);
+}
+
+void FStateTreeEditorNodeDetails::CustomizeHeader(TSharedRef<class IPropertyHandle> StructPropertyHandle, class FDetailWidgetRow& HeaderRow, IPropertyTypeCustomizationUtils& StructCustomizationUtils)
+{
+	StructProperty = StructPropertyHandle;
+	PropUtils = StructCustomizationUtils.GetPropertyUtilities().Get();
+
+	NodeProperty = StructProperty->GetChildHandle(TEXT("Node"));
+	InstanceProperty = StructProperty->GetChildHandle(TEXT("Instance"));
+	InstanceObjectProperty = StructProperty->GetChildHandle(TEXT("InstanceObject"));
+	IDProperty = StructProperty->GetChildHandle(TEXT("ID"));
+
+	check(NodeProperty.IsValid());
+	check(InstanceProperty.IsValid());
+	check(IDProperty.IsValid());
+	
+	// Find base class and struct from meta data.
+	static const FName BaseStructMetaName(TEXT("BaseStruct")); // TODO: move these names into one central place.
+	static const FName BaseClassMetaName(TEXT("BaseClass")); // TODO: move these names into one central place.
+	
+	const FString BaseStructName = StructProperty->GetMetaData(BaseStructMetaName);
+	BaseScriptStruct = FindObject<UScriptStruct>(ANY_PACKAGE, *BaseStructName);
+
+	const FString BaseClassName = StructProperty->GetMetaData(BaseClassMetaName);
+	BaseClass = FindObject<UClass>(ANY_PACKAGE, *BaseClassName);
+	
+	const FIsResetToDefaultVisible IsResetVisible = FIsResetToDefaultVisible::CreateSP(this, &FStateTreeEditorNodeDetails::ShouldResetToDefault);
+	const FResetToDefaultHandler ResetHandler = FResetToDefaultHandler::CreateSP(this, &FStateTreeEditorNodeDetails::ResetToDefault);
+	const FResetToDefaultOverride ResetOverride = FResetToDefaultOverride::Create(IsResetVisible, ResetHandler);
+
+	UE::StateTree::Delegates::OnIdentifierChanged.AddSP(this, &FStateTreeEditorNodeDetails::OnIdentifierChanged);
+	OnBindingChangedHandle = UE::StateTree::PropertyBinding::OnStateTreeBindingChanged.AddRaw(this, &FStateTreeEditorNodeDetails::OnBindingChanged);
+
+	FindOuterObjects();
+
+	HeaderRow
+		.WholeRowContent()
+		.VAlign(VAlign_Center)
+		[
+			SNew(SHorizontalBox)
+			// Description (Condition)
+			+ SHorizontalBox::Slot()
+			.FillWidth(3.0f)
+			.VAlign(VAlign_Center)
+			[
+				SNew(SRichTextBlock)
+				.DecoratorStyleSet(FStateTreeEditorStyle::Get().Get())
+				.TextStyle(FStateTreeEditorStyle::Get(), "Details.Normal")
+				.Text(this, &FStateTreeEditorNodeDetails::GetDescription)
+				.Visibility(this, &FStateTreeEditorNodeDetails::IsDescriptionVisible)
+			]
+			// Name (Eval/Task)
+			+ SHorizontalBox::Slot()
+			.FillWidth(1.5f)
+			.VAlign(VAlign_Center)
+			[
+				SNew(SEditableTextBox)
+				.IsEnabled(TAttribute<bool>(this, &FStateTreeEditorNodeDetails::IsNameEnabled))
+				.Text(this, &FStateTreeEditorNodeDetails::GetName)
+				.OnTextCommitted(this, &FStateTreeEditorNodeDetails::OnNameCommitted)
+				.SelectAllTextWhenFocused(true)
+				.RevertTextOnEscape(true)
+				.Font(IDetailLayoutBuilder::GetDetailFont())
+				.Visibility(this, &FStateTreeEditorNodeDetails::IsNameVisible)
+			]
+			// Class picker
+			+ SHorizontalBox::Slot()
+			.FillWidth(1.0f)
+			.Padding(FMargin(FMargin(4.0f, 0.0f, 0.0f, 0.0f)))
+			.VAlign(VAlign_Center)
+			[
+				SAssignNew(ComboButton, SComboButton)
+				.OnGetMenuContent(this, &FStateTreeEditorNodeDetails::GeneratePicker)
+				.ContentPadding(0.f)
+				.ButtonContent()
+				[
+					SNew(SHorizontalBox)
+					+ SHorizontalBox::Slot()
+					.AutoWidth()
+					.VAlign(VAlign_Center)
+					.Padding(0.0f, 0.0f, 4.0f, 0.0f)
+					[
+						SNew(SImage)
+						.Image(this, &FStateTreeEditorNodeDetails::GetDisplayValueIcon)
+					]
+					+ SHorizontalBox::Slot()
+					.VAlign(VAlign_Center)
+					[
+						SNew(STextBlock)
+						.Text(this, &FStateTreeEditorNodeDetails::GetDisplayValueString)
+						.Font(IDetailLayoutBuilder::GetDetailFont())
+					]
+				]
+			]
+			+ SHorizontalBox::Slot()
+			.AutoWidth()
+			[
+				StructPropertyHandle->CreateDefaultPropertyButtonWidgets()
+			]
+		]
+		.OverrideResetToDefault(ResetOverride);
+}
+
+bool FStateTreeEditorNodeDetails::ShouldResetToDefault(TSharedPtr<IPropertyHandle> PropertyHandle) const
+{
+	check(StructProperty);
+	
+	bool bAnyValid = false;
+	
+	TArray<const void*> RawNodeData;
+	StructProperty->AccessRawData(RawNodeData);
+	for (const void* Data : RawNodeData)
+	{
+		if (const FStateTreeEditorNode* Node = static_cast<const FStateTreeEditorNode*>(Data))
+		{
+			if (Node->Node.IsValid())
+			{
+				bAnyValid = true;
+				break;
+			}
+		}
+	}
+	
+	// Assume that the default value is empty. Any valid means that some can be reset to empty.
+	return bAnyValid;
+}
+
+
+void FStateTreeEditorNodeDetails::ResetToDefault(TSharedPtr<IPropertyHandle> PropertyHandle)
+{
+	check(StructProperty);
+	
+	GEditor->BeginTransaction(LOCTEXT("OnResetToDefault", "Reset to default"));
+
+	StructProperty->NotifyPreChange();
+	
+	TArray<void*> RawNodeData;
+	StructProperty->AccessRawData(RawNodeData);
+	for (void* Data : RawNodeData)
+	{
+		if (FStateTreeEditorNode* Node = static_cast<FStateTreeEditorNode*>(Data))
+		{
+			Node->Reset();
+		}
+	}
+
+	StructProperty->NotifyPostChange(EPropertyChangeType::ValueSet);
+	StructProperty->NotifyFinishedChangingProperties();
+
+	GEditor->EndTransaction();
+
+	if (PropUtils)
+	{
+		PropUtils->ForceRefresh();
+	}
+}
+
+void FStateTreeEditorNodeDetails::CustomizeChildren(TSharedRef<class IPropertyHandle> StructPropertyHandle, class IDetailChildrenBuilder& StructBuilder, IPropertyTypeCustomizationUtils& StructCustomizationUtils)
+{
+	FStateTreeEditorPropertyBindings* EditorPropBindings = EditorData ? EditorData->GetPropertyEditorBindings() : nullptr;
+
+	FGuid ID;
+	UE::StateTree::PropertyHelpers::GetStructValue<FGuid>(IDProperty, ID);
+
+	TSharedRef<FBindableNodeInstanceDetails> NodeDetails = MakeShareable(new FBindableNodeInstanceDetails(NodeProperty, FGuid(), EditorData));
+	StructBuilder.AddCustomBuilder(NodeDetails);
+
+	// Instance
+	TSharedRef<FBindableNodeInstanceDetails> InstanceDetails = MakeShareable(new FBindableNodeInstanceDetails(InstanceProperty, ID, EditorData));
+	StructBuilder.AddCustomBuilder(InstanceDetails);
+
+	// InstanceObject
+	// Get the actual UObject from the pointer.
+	TSharedPtr<IPropertyHandle> InstanceObjectValueProperty = GetInstancedObjectValueHandle(InstanceObjectProperty);
+	if (InstanceObjectValueProperty.IsValid())
+	{
+		static const FName CategoryName(TEXT("Category"));
+
+		uint32 NumChildren = 0;
+		InstanceObjectValueProperty->GetNumChildren(NumChildren);
+		for (uint32 Index = 0; Index < NumChildren; Index++)
+		{
+			TSharedPtr<IPropertyHandle> Child = InstanceObjectValueProperty->GetChildHandle(Index);
+			
+			if (!Child.IsValid()
+				|| Child->GetProperty()->HasAllPropertyFlags(CPF_DisableEditOnInstance)
+				|| Child->GetMetaData(CategoryName) == TEXT("ExternalData"))
+			{
+				continue;
+			}
+			
+			IDetailPropertyRow& ChildRow = StructBuilder.AddProperty(Child.ToSharedRef());
+			UE::StateTreeEditor::Internal::ModifyRow(ChildRow, ID, EditorPropBindings);
+		}
+	}
+}
+
+TSharedPtr<IPropertyHandle> FStateTreeEditorNodeDetails::GetInstancedObjectValueHandle(TSharedPtr<IPropertyHandle> PropertyHandle)
+{
+	TSharedPtr<IPropertyHandle> ChildHandle;
+
+	uint32 NumChildren;
+	PropertyHandle->GetNumChildren(NumChildren);
+
+	if (NumChildren > 0)
+	{
+		// when the property is a (inlined) object property, the first child will be
+		// the object instance, and its properties are the children underneath that
+		ensure(NumChildren == 1);
+		ChildHandle = PropertyHandle->GetChildHandle(0);
+	}
+
+	return ChildHandle;
+}
+
+void FStateTreeEditorNodeDetails::OnIdentifierChanged(const UStateTree& InStateTree)
+{
+	if (PropUtils && StateTree == &InStateTree)
+	{
+		PropUtils->ForceRefresh();
+	}
+}
+
+void FStateTreeEditorNodeDetails::OnBindingChanged(const FStateTreeEditorPropertyPath& SourcePath, const FStateTreeEditorPropertyPath& TargetPath)
+{
+	check(StructProperty);
+
+	TArray<UObject*> OuterObjects;
+	StructProperty->GetOuterObjects(OuterObjects);
+
+	TArray<void*> RawNodeData;
+	StructProperty->AccessRawData(RawNodeData);
+
+	if (OuterObjects.Num() != RawNodeData.Num())
+	{
+		return;
+	}
+
+	for (int32 i = 0; i < OuterObjects.Num(); i++)
+	{
+		const FStateTreeEditorNode* Node = static_cast<FStateTreeEditorNode*>(RawNodeData[i]);
+		UObject* OuterObject = OuterObjects[i]; // Immediate outer, i.e StateTreeState
+		if (Node != nullptr && EditorData != nullptr && Node->Node.IsValid() && Node->Instance.IsValid())
+		{
+			if (Node->ID == TargetPath.StructID)
+			{
+				if (FStateTreeConditionBase* Condition = Node->Node.GetMutablePtr<FStateTreeConditionBase>())
+				{
+					const FStateTreeBindingLookup BindingLookup(EditorData);
+
+					OuterObject->Modify();
+					Condition->OnBindingChanged(Node->ID, FStateTreeDataView(Node->Instance), SourcePath, TargetPath, BindingLookup);
+				}
+			}
+		}
+	}
+}
+
+void FStateTreeEditorNodeDetails::FindOuterObjects()
+{
+	check(StructProperty);
+	
+	EditorData = nullptr;
+	StateTree = nullptr;
+
+	TArray<UObject*> OuterObjects;
+	StructProperty->GetOuterObjects(OuterObjects);
+	for (UObject* Outer : OuterObjects)
+	{
+		UStateTreeEditorData* OuterEditorData = Outer->GetTypedOuter<UStateTreeEditorData>();
+		UStateTree* OuterStateTree = OuterEditorData ? OuterEditorData->GetTypedOuter<UStateTree>() : nullptr;
+		if (OuterEditorData && OuterStateTree)
+		{
+			StateTree = OuterStateTree;
+			EditorData = OuterEditorData;
+			break;
+		}
+	}
+}
+
+FText FStateTreeEditorNodeDetails::GetDescription() const
+{
+	check(StructProperty);
+
+	TArray<void*> RawNodeData;
+	StructProperty->AccessRawData(RawNodeData);
+
+	// Multiple descriptions do not make sense, just if only one node is selected.
+	if (RawNodeData.Num() != 1)
+	{
+		return LOCTEXT("MultipleSelected", "<Details.Subdued>Multiple selected</>");
+	}
+
+	FStateTreeEditorNode* Node = static_cast<FStateTreeEditorNode*>(RawNodeData[0]);
+	if (Node != nullptr && EditorData != nullptr && Node->Node.IsValid())
+	{
+		if (const FStateTreeConditionBase* Condition = Node->Node.GetPtr<const FStateTreeConditionBase>())
+		{
+			const FStateTreeBindingLookup BindingLookup(EditorData);
+			if (Node->Instance.IsValid())
+			{
+				return Condition->GetDescription(Node->ID, FStateTreeDataView(Node->Instance), BindingLookup);
+			}
+			else if (Node->InstanceObject != nullptr)
+			{
+				return Condition->GetDescription(Node->ID, FStateTreeDataView(Node->InstanceObject), BindingLookup);
+			}
+			else
+			{
+				return Node->Node.GetScriptStruct()->GetDisplayNameText();
+			}
+		}
+	}
+	
+	return LOCTEXT("ConditionNotSet", "<Details.Subdued>Condition Not Set</>");
+}
+
+EVisibility FStateTreeEditorNodeDetails::IsDescriptionVisible() const
+{
+	const UScriptStruct* ScriptStruct = nullptr;
+	if (const FStateTreeEditorNode* Node = GetCommonNode())
+	{
+		ScriptStruct = Node->Node.GetScriptStruct();
+	}
+
+	return ScriptStruct != nullptr && ScriptStruct->IsChildOf(FStateTreeConditionBase::StaticStruct()) ? EVisibility::Visible : EVisibility::Collapsed;
+}
+
+FText FStateTreeEditorNodeDetails::GetName() const
+{
+	check(StructProperty);
+
+	// Multiple names do not make sense, just if only one node is selected.
+	TArray<void*> RawNodeData;
+	StructProperty->AccessRawData(RawNodeData);
+	if (RawNodeData.Num() == 1)
+	{
+		// Dig out name from the struct without knowing the type.
+		FStateTreeEditorNode* Node = static_cast<FStateTreeEditorNode*>(RawNodeData[0]);
+
+		// Make sure the associated property has valid data (e.g. while moving an node in array)
+		const UScriptStruct* ScriptStruct = Node != nullptr ? Node->Node.GetScriptStruct() : nullptr;
+		if (ScriptStruct != nullptr)
+		{
+			if (FProperty* NameProperty = ScriptStruct->FindPropertyByName(TEXT("Name")))
+			{
+				if (NameProperty->IsA(FNameProperty::StaticClass()))
+				{
+					void* Ptr = const_cast<void*>(static_cast<const void*>(Node->Node.GetMemory()));
+					const FName NameValue = *NameProperty->ContainerPtrToValuePtr<FName>(Ptr);
+					return FText::FromName(NameValue);
+				}
+			}
+		}
+		return LOCTEXT("Empty", "Empty");
+	}
+
+	return LOCTEXT("MultipleSelected", "Multiple selected");
+}
+
+EVisibility FStateTreeEditorNodeDetails::IsNameVisible() const
+{
+	const UStateTreeSchema* Schema = StateTree ? StateTree->GetSchema() : nullptr;
+	if (Schema && Schema->AllowMultipleTasks() == false)
+	{
+		// Single task states use the state name as task name.
+		return EVisibility::Collapsed;
+	}
+	
+	const UScriptStruct* ScriptStruct = nullptr;
+	if (const FStateTreeEditorNode* Node = GetCommonNode())
+	{
+		ScriptStruct = Node->Node.GetScriptStruct();
+	}
+	
+	return ScriptStruct != nullptr && (ScriptStruct->IsChildOf(FStateTreeTaskBase::StaticStruct()) || ScriptStruct->IsChildOf(FStateTreeEvaluatorBase::StaticStruct())) ? EVisibility::Visible : EVisibility::Collapsed;
+}
+
+void FStateTreeEditorNodeDetails::OnNameCommitted(const FText& NewText, ETextCommit::Type InTextCommit)
+{
+	check(StructProperty);
+
+	if (InTextCommit == ETextCommit::OnEnter || InTextCommit == ETextCommit::OnUserMovedFocus)
+	{
+		// Remove excess whitespace and prevent categories with just spaces
+		FString NewName = FText::TrimPrecedingAndTrailing(NewText).ToString();
+
+		if (GEditor)
+		{
+			GEditor->BeginTransaction(LOCTEXT("SetName", "Set Name"));
+		}
+		StructProperty->NotifyPreChange();
+
+		TArray<void*> RawNodeData;
+		StructProperty->AccessRawData(RawNodeData);
+		
+		for (void* Data : RawNodeData)
+		{
+			if (FStateTreeEditorNode* Node = static_cast<FStateTreeEditorNode*>(Data))
+			{
+				// Set Name
+				if (const UScriptStruct* ScriptStruct = Node->Node.GetScriptStruct())
+				{
+					if (FProperty* NameProperty = ScriptStruct->FindPropertyByName(TEXT("Name")))
+					{
+						if (NameProperty->IsA(FNameProperty::StaticClass()))
+						{
+							void* Ptr = const_cast<void*>(static_cast<const void*>(Node->Node.GetMemory()));
+							FName& NameValue = *NameProperty->ContainerPtrToValuePtr<FName>(Ptr);
+							NameValue = FName(NewName);
+						}
+					}
+				}
+			}
+		}
+
+		StructProperty->NotifyPostChange(EPropertyChangeType::ValueSet);
+
+		if (StateTree)
+		{
+			UE::StateTree::Delegates::OnIdentifierChanged.Broadcast(*StateTree);
+		}
+
+		if (GEditor)
+		{
+			GEditor->EndTransaction();
+		}
+
+		StructProperty->NotifyFinishedChangingProperties();
+	}
+}
+
+bool FStateTreeEditorNodeDetails::IsNameEnabled() const
+{
+	// Can only edit if we have valid instantiated type.
+	const FStateTreeEditorNode* Node = GetCommonNode();
+	return Node && Node->Node.IsValid();
+}
+
+const FStateTreeEditorNode* FStateTreeEditorNodeDetails::GetCommonNode() const
+{
+	check(StructProperty);
+
+	const UScriptStruct* CommonScriptStruct = nullptr;
+	bool bMultipleValues = false;
+	TArray<void*> RawNodeData;
+	StructProperty->AccessRawData(RawNodeData);
+
+	const FStateTreeEditorNode* CommonNode = nullptr;
+	
+	for (void* Data : RawNodeData)
+	{
+		if (const FStateTreeEditorNode* Node = static_cast<FStateTreeEditorNode*>(Data))
+		{
+			if (!bMultipleValues && !CommonNode)
+			{
+				CommonNode = Node;
+			}
+			else if (CommonNode != Node)
+			{
+				CommonNode = nullptr;
+				bMultipleValues = true;
+			}
+		}
+	}
+
+	return CommonNode;
+}
+
+FText FStateTreeEditorNodeDetails::GetDisplayValueString() const
+{
+	if (const FStateTreeEditorNode* Node = GetCommonNode())
+	{
+		if (const UScriptStruct* ScriptStruct = Node->Node.GetScriptStruct())
+		{
+			if (ScriptStruct->IsChildOf(FStateTreeBlueprintEvaluatorWrapper::StaticStruct())
+				|| ScriptStruct->IsChildOf(FStateTreeBlueprintTaskWrapper::StaticStruct())
+				|| ScriptStruct->IsChildOf(FStateTreeBlueprintConditionWrapper::StaticStruct()))
+			{
+				if (Node->InstanceObject != nullptr
+					&& Node->InstanceObject->GetClass() != nullptr)
+				{
+					return Node->InstanceObject->GetClass()->GetDisplayNameText();
+				}
+			}
+			else
+			{
+				return ScriptStruct->GetDisplayNameText();
+			}
+		}
+	}
+	return FText();
+}
+
+const FSlateBrush* FStateTreeEditorNodeDetails::GetDisplayValueIcon() const
+{
+	if (const FStateTreeEditorNode* Node = GetCommonNode())
+	{
+		if (const UScriptStruct* ScriptStruct = Node->Node.GetScriptStruct())
+		{
+			if (ScriptStruct->IsChildOf(FStateTreeBlueprintEvaluatorWrapper::StaticStruct())
+				|| ScriptStruct->IsChildOf(FStateTreeBlueprintTaskWrapper::StaticStruct())
+				|| ScriptStruct->IsChildOf(FStateTreeBlueprintConditionWrapper::StaticStruct()))
+			{
+				if (Node->InstanceObject != nullptr)
+				{
+					return FSlateIconFinder::FindIconBrushForClass(Node->InstanceObject->GetClass());
+				}
+			}
+		}
+	}
+	
+	return FSlateIconFinder::FindIconBrushForClass(UScriptStruct::StaticClass());
+}
+
+TSharedRef<SWidget> FStateTreeEditorNodeDetails::GeneratePicker()
+{
+	FMenuBuilder MenuBuilder(true, nullptr);
+
+	FStateTreeEditorModule& EditorModule = FModuleManager::GetModuleChecked<FStateTreeEditorModule>(TEXT("StateTreeEditorModule"));
+	FStateTreeNodeClassCache* ClassCache = EditorModule.GetNodeClassCache().Get();
+	check(ClassCache);
+
+	FUIAction ClearAction(FExecuteAction::CreateSP(this, &FStateTreeEditorNodeDetails::OnStructPicked, (const UScriptStruct*)nullptr));
+	MenuBuilder.AddMenuEntry(LOCTEXT("ClearNode", "Clear"), TAttribute<FText>(), FSlateIcon(FEditorStyle::GetStyleSetName(), "Cross"), ClearAction);
+
+	MenuBuilder.AddMenuSeparator();
+
+	TArray<TSharedPtr<FStateTreeNodeClassData>> StructNodes;
+	TArray<TSharedPtr<FStateTreeNodeClassData>> ObjectNodes;
+	
+	ClassCache->GetScripStructs(BaseScriptStruct, StructNodes);
+	ClassCache->GetClasses(BaseClass, ObjectNodes);
+
+	const FSlateIcon ScripStructIcon = FSlateIconFinder::FindIconForClass(UScriptStruct::StaticClass());
+	const UStateTreeSchema* Schema = StateTree ? StateTree->GetSchema() : nullptr;
+	
+	for (const TSharedPtr<FStateTreeNodeClassData>& Data : StructNodes)
+	{
+		if (Data->GetScriptStruct() != nullptr)
+		{
+			const UScriptStruct* ScriptStruct = Data->GetScriptStruct();
+			if (ScriptStruct == BaseScriptStruct)
+			{
+				continue;
+			}
+			if (ScriptStruct->HasMetaData(TEXT("Hidden")))
+			{
+				continue;				
+			}
+			if (Schema && !Schema->IsStructAllowed(ScriptStruct))
+			{
+				continue;				
+			}
+			
+			FUIAction ItemAction(FExecuteAction::CreateSP(this, &FStateTreeEditorNodeDetails::OnStructPicked, ScriptStruct));
+			MenuBuilder.AddMenuEntry(ScriptStruct->GetDisplayNameText(), TAttribute<FText>(), ScripStructIcon, ItemAction);
+		}
+	}
+
+	if (StructNodes.Num() > 0 && ObjectNodes.Num() > 0)
+	{
+		MenuBuilder.AddMenuSeparator();
+	}
+	
+	for (const TSharedPtr<FStateTreeNodeClassData>& Data : ObjectNodes)
+	{
+		if (Data->GetClass() != nullptr)
+		{
+			UClass* Class = Data->GetClass();
+			if (Class == BaseClass)
+			{
+				continue;
+			}
+			if (Class->HasAnyClassFlags(CLASS_Hidden | CLASS_HideDropDown))
+			{
+				continue;				
+			}
+			if (Class->HasMetaData(TEXT("Hidden")))
+			{
+				continue;				
+			}
+			if (Schema && !Schema->IsClassAllowed(Class))
+			{
+				continue;				
+			}
+
+			FUIAction ItemAction(FExecuteAction::CreateSP(this, &FStateTreeEditorNodeDetails::OnClassPicked, Class));
+			MenuBuilder.AddMenuEntry(Class->GetDisplayNameText(), TAttribute<FText>(), FSlateIconFinder::FindIconForClass(Data->GetClass()), ItemAction);
+		}
+	}
+
+	return MenuBuilder.MakeWidget();
+}
+
+void FStateTreeEditorNodeDetails::OnStructPicked(const UScriptStruct* InStruct)
+{
+	check(StructProperty);
+	check(StateTree);
+
+	TArray<void*> RawNodeData;
+	StructProperty->AccessRawData(RawNodeData);
+
+	GEditor->BeginTransaction(LOCTEXT("SelectNode", "Select Node"));
+
+	StructProperty->NotifyPreChange();
+
+	for (void* Data : RawNodeData)
+	{
+		if (FStateTreeEditorNode* Node = static_cast<FStateTreeEditorNode*>(Data))
+		{
+			Node->Reset();
+			
+			if (InStruct)
+			{
+				// Generate new ID.
+				Node->ID = FGuid::NewGuid();
+
+				// Initialize node
+				Node->Node.InitializeAs(InStruct);
+				
+				// Generate new name and instantiate instance data.
+				if (InStruct->IsChildOf(FStateTreeTaskBase::StaticStruct()))
+				{
+					FStateTreeTaskBase& Task = Node->Node.GetMutable<FStateTreeTaskBase>();
+					Task.Name = FName(InStruct->GetDisplayNameText().ToString());
+
+					if (const UScriptStruct* InstanceType = Cast<const UScriptStruct>(Task.GetInstanceDataType()))
+					{
+						Node->Instance.InitializeAs(InstanceType);
+					}
+					else if (const UClass* InstanceClass = Cast<const UClass>(Task.GetInstanceDataType()))
+					{
+						Node->InstanceObject = NewObject<UObject>(EditorData, InstanceClass);
+					}
+				}
+				else if (InStruct->IsChildOf(FStateTreeEvaluatorBase::StaticStruct()))
+				{
+					FStateTreeEvaluatorBase& Eval = Node->Node.GetMutable<FStateTreeEvaluatorBase>();
+					Eval.Name = FName(InStruct->GetDisplayNameText().ToString());
+
+					if (const UScriptStruct* InstanceType = Cast<const UScriptStruct>(Eval.GetInstanceDataType()))
+					{
+						Node->Instance.InitializeAs(InstanceType);
+					}
+					else if (const UClass* InstanceClass = Cast<const UClass>(Eval.GetInstanceDataType()))
+					{
+						Node->InstanceObject = NewObject<UObject>(EditorData, InstanceClass);
+					}
+				}
+				else if (InStruct->IsChildOf(FStateTreeConditionBase::StaticStruct()))
+				{
+					FStateTreeConditionBase& Cond = Node->Node.GetMutable<FStateTreeConditionBase>();
+
+					if (const UScriptStruct* InstanceType = Cast<const UScriptStruct>(Cond.GetInstanceDataType()))
+					{
+						Node->Instance.InitializeAs(InstanceType);
+					}
+					else if (const UClass* InstanceClass = Cast<const UClass>(Cond.GetInstanceDataType()))
+					{
+						Node->InstanceObject = NewObject<UObject>(EditorData, InstanceClass);
+					}
+				}
+			}
+		}
+	}
+
+	StructProperty->NotifyPostChange(EPropertyChangeType::ValueSet);
+	StructProperty->NotifyFinishedChangingProperties();
+
+	GEditor->EndTransaction();
+
+	ComboButton->SetIsOpen(false);
+
+	if (PropUtils)
+	{
+		PropUtils->ForceRefresh();
+	}
+}
+
+void FStateTreeEditorNodeDetails::OnClassPicked(UClass* InClass)
+{
+	check(StructProperty);
+	check(StateTree);
+
+	TArray<void*> RawNodeData;
+	StructProperty->AccessRawData(RawNodeData);
+
+	GEditor->BeginTransaction(LOCTEXT("SelectBlueprintNode", "Select Blueprint Node"));
+
+	StructProperty->NotifyPreChange();
+
+	for (void* Data : RawNodeData)
+	{
+		if (FStateTreeEditorNode* Node = static_cast<FStateTreeEditorNode*>(Data))
+		{
+			Node->Reset();
+
+			if (InClass && InClass->IsChildOf(UStateTreeTaskBlueprintBase::StaticClass()))
+			{
+				Node->Node.InitializeAs(FStateTreeBlueprintTaskWrapper::StaticStruct());
+				FStateTreeBlueprintTaskWrapper& Task = Node->Node.GetMutable<FStateTreeBlueprintTaskWrapper>();
+				Task.TaskClass = InClass;
+				Task.Name = FName(InClass->GetDisplayNameText().ToString());
+				
+				Node->InstanceObject = NewObject<UObject>(EditorData, InClass);
+
+				Node->ID = FGuid::NewGuid();
+			}
+			else if (InClass && InClass->IsChildOf(UStateTreeEvaluatorBlueprintBase::StaticClass()))
+			{
+				Node->Node.InitializeAs(FStateTreeBlueprintEvaluatorWrapper::StaticStruct());
+				FStateTreeBlueprintEvaluatorWrapper& Eval = Node->Node.GetMutable<FStateTreeBlueprintEvaluatorWrapper>();
+				Eval.EvaluatorClass = InClass;
+				Eval.Name = FName(InClass->GetDisplayNameText().ToString());
+				
+				Node->InstanceObject = NewObject<UObject>(EditorData, InClass);
+
+				Node->ID = FGuid::NewGuid();
+			}
+			else if (InClass && InClass->IsChildOf(UStateTreeConditionBlueprintBase::StaticClass()))
+			{
+				Node->Node.InitializeAs(FStateTreeBlueprintConditionWrapper::StaticStruct());
+				FStateTreeBlueprintConditionWrapper& Cond = Node->Node.GetMutable<FStateTreeBlueprintConditionWrapper>();
+				Cond.ConditionClass = InClass;
+				
+				Node->InstanceObject = NewObject<UObject>(EditorData, InClass);
+
+				Node->ID = FGuid::NewGuid();
+			}
+
+			
+		}
+	}
+
+	StructProperty->NotifyPostChange(EPropertyChangeType::ValueSet);
+	StructProperty->NotifyFinishedChangingProperties();
+
+	GEditor->EndTransaction();
+
+	ComboButton->SetIsOpen(false);
+
+	if (PropUtils)
+	{
+		PropUtils->ForceRefresh();
+	}
+}
+
+
+#undef LOCTEXT_NAMESPACE
diff --git a/Engine/Plugins/Runtime/StateTree/Source/StateTreeEditorModule/Private/Customizations/StateTreeEditorNodeDetails.h b/Engine/Plugins/Runtime/StateTree/Source/StateTreeEditorModule/Private/Customizations/StateTreeEditorNodeDetails.h
new file mode 100644
--- /dev/null
+++ b/Engine/Plugins/Runtime/StateTree/Source/StateTreeEditorModule/Private/Customizations/StateTreeEditorNodeDetails.h
@@ -0,0 +1,72 @@
+// Copyright Epic Games, Inc. All Rights Reserved.
+
+#pragma once
+
+#include "CoreMinimal.h"
+#include "UnrealClient.h"
+#include "IPropertyTypeCustomization.h"
+
+class IPropertyHandle;
+class UStateTree;
+class UStateTreeState;
+class UStateTreeEditorData;
+struct FStateTreeEditorPropertyPath;
+
+/**
+ * Type customization for nodes (Conditions, Evaluators and Tasks) in StateTreeState.
+ */
+class FStateTreeEditorNodeDetails : public IPropertyTypeCustomization
+{
+public:
+	/** Makes a new instance of this detail layout class for a specific detail view requesting it */
+	static TSharedRef<IPropertyTypeCustomization> MakeInstance();
+
+	virtual ~FStateTreeEditorNodeDetails();
+	
+	/** IPropertyTypeCustomization interface */
+	virtual void CustomizeHeader(TSharedRef<IPropertyHandle> StructPropertyHandle, class FDetailWidgetRow& HeaderRow, IPropertyTypeCustomizationUtils& StructCustomizationUtils) override;
+	virtual void CustomizeChildren(TSharedRef<IPropertyHandle> StructPropertyHandle, class IDetailChildrenBuilder& StructBuilder, IPropertyTypeCustomizationUtils& StructCustomizationUtils) override;
+
+private:
+
+	bool ShouldResetToDefault(TSharedPtr<IPropertyHandle> PropertyHandle) const;
+	void ResetToDefault(TSharedPtr<IPropertyHandle> PropertyHandle);
+	TSharedPtr<IPropertyHandle> GetInstancedObjectValueHandle(TSharedPtr<IPropertyHandle> PropertyHandle);
+
+	FText GetDescription() const;
+	EVisibility IsDescriptionVisible() const;
+	
+	FText GetName() const;
+	EVisibility IsNameVisible() const;
+	void OnNameCommitted(const FText& NewText, ETextCommit::Type InTextCommit);
+	bool IsNameEnabled() const;
+
+	const struct FStateTreeEditorNode* GetCommonNode() const;
+	FText GetDisplayValueString() const;
+	const FSlateBrush* GetDisplayValueIcon() const;
+
+	TSharedRef<SWidget> GeneratePicker();
+
+	void OnStructPicked(const UScriptStruct* InStruct);
+	void OnClassPicked(UClass* InClass);
+
+	void OnIdentifierChanged(const UStateTree& StateTree);
+	void OnBindingChanged(const FStateTreeEditorPropertyPath& SourcePath, const FStateTreeEditorPropertyPath& TargetPath);
+	void FindOuterObjects();
+
+	UScriptStruct* BaseScriptStruct = nullptr;
+	UClass* BaseClass = nullptr;
+	TSharedPtr<class SComboButton> ComboButton;
+
+	UStateTreeEditorData* EditorData = nullptr;
+	UStateTree* StateTree = nullptr;
+
+	class IPropertyUtilities* PropUtils = nullptr;
+	TSharedPtr<IPropertyHandle> StructProperty;
+	TSharedPtr<IPropertyHandle> NodeProperty;
+	TSharedPtr<IPropertyHandle> InstanceProperty;
+	TSharedPtr<IPropertyHandle> InstanceObjectProperty;
+	TSharedPtr<IPropertyHandle> IDProperty;
+
+	FDelegateHandle OnBindingChangedHandle;
+};
diff --git a/Engine/Plugins/Runtime/StateTree/Source/StateTreeEditorModule/Private/SStateTreeViewRow.cpp b/Engine/Plugins/Runtime/StateTree/Source/StateTreeEditorModule/Private/SStateTreeViewRow.cpp
--- a/Engine/Plugins/Runtime/StateTree/Source/StateTreeEditorModule/Private/SStateTreeViewRow.cpp
+++ b/Engine/Plugins/Runtime/StateTree/Source/StateTreeEditorModule/Private/SStateTreeViewRow.cpp
@@ -23,6 +23,9 @@
 
 #include "StateTree.h"
 #include "StateTreeState.h"
+#include "StateTreeConditionBase.h"
+#include "StateTreeEvaluatorBase.h"
+#include "StateTreeTaskBase.h"
 #include "StateTreeViewModel.h"
 #include "Algo/ForEach.h"
 
@@ -419,7 +422,7 @@ EVisibility SStateTreeViewRow::GetEvaluatorsVisibility() const
 		int32 ValidCount = 0;
 		for (int32 i = 0; i < State->Evaluators.Num(); i++)
 		{
-			if (const FStateTreeEvaluatorBase* Eval = State->Evaluators[i].Item.GetPtr<FStateTreeEvaluatorBase>())
+			if (const FStateTreeEvaluatorBase* Eval = State->Evaluators[i].Node.GetPtr<FStateTreeEvaluatorBase>())
 			{
 				ValidCount++;
 			}
@@ -441,7 +444,7 @@ FText SStateTreeViewRow::GetEvaluatorsDesc() const
 
 	for (int32 i = 0; i < State->Evaluators.Num(); i++)
 	{
-		if (const FStateTreeEvaluatorBase* Eval = State->Evaluators[i].Item.GetPtr<FStateTreeEvaluatorBase>())
+		if (const FStateTreeEvaluatorBase* Eval = State->Evaluators[i].Node.GetPtr<FStateTreeEvaluatorBase>())
 		{
 			Names.Add(FText::FromName(Eval->Name));
 		}
@@ -457,7 +460,7 @@ EVisibility SStateTreeViewRow::GetTasksVisibility() const
 		int32 ValidCount = 0;
 		for (int32 i = 0; i < State->Tasks.Num(); i++)
 		{
-			if (const FStateTreeTaskBase* Task = State->Tasks[i].Item.GetPtr<FStateTreeTaskBase>())
+			if (const FStateTreeTaskBase* Task = State->Tasks[i].Node.GetPtr<FStateTreeTaskBase>())
 			{
 				ValidCount++;
 			}
@@ -478,7 +481,7 @@ FText SStateTreeViewRow::GetTasksDesc() const
 	TArray<FText> Names;
 	for (int32 i = 0; i < State->Tasks.Num(); i++)
 	{
-		if (const FStateTreeTaskBase* Task = State->Tasks[i].Item.GetPtr<FStateTreeTaskBase>())
+		if (const FStateTreeTaskBase* Task = State->Tasks[i].Node.GetPtr<FStateTreeTaskBase>())
 		{
 			Names.Add(FText::FromName(Task->Name));
 		}
diff --git a/Engine/Plugins/Runtime/StateTree/Source/StateTreeEditorModule/Private/StateTreeBaker.cpp b/Engine/Plugins/Runtime/StateTree/Source/StateTreeEditorModule/Private/StateTreeBaker.cpp
--- a/Engine/Plugins/Runtime/StateTree/Source/StateTreeEditorModule/Private/StateTreeBaker.cpp
+++ b/Engine/Plugins/Runtime/StateTree/Source/StateTreeEditorModule/Private/StateTreeBaker.cpp
@@ -94,9 +94,9 @@
 		
 		// Enter conditions.
 		BakedState.EnterConditionsBegin = uint16(StateTree->Items.Num());
-		for (FStateTreeConditionItem& ConditionItem : SourceState->EnterConditions)
+		for (FStateTreeEditorNode& CondNode : SourceState->EnterConditions)
 		{
-			if (!CreateCondition(ConditionItem))
+			if (!CreateCondition(CondNode))
 			{
 				Log.Reportf(EMessageSeverity::Error,
 					TEXT("Failed to create state enter condition."));
@@ -121,9 +121,9 @@
 			// Note: Unset transition is allowed here. It can be used to mask a transition at parent.
 
 			BakedTransition.ConditionsBegin = uint16(StateTree->Items.Num());
-			for (FStateTreeConditionItem& ConditionItem : Transition.Conditions)
+			for (FStateTreeEditorNode& CondNode : Transition.Conditions)
 			{
-				if (!CreateCondition(ConditionItem))
+				if (!CreateCondition(CondNode))
 				{
 					Log.Reportf(EMessageSeverity::Error,
 						TEXT("Failed to create condition for transition to %s."),
@@ -177,20 +177,20 @@
 	return true;
 }
 
-bool FStateTreeBaker::CreateCondition(const FStateTreeConditionItem& CondItem)
+bool FStateTreeBaker::CreateCondition(const FStateTreeEditorNode& CondNode)
 {
-	if (!CondItem.Item.IsValid())
+	if (!CondNode.Node.IsValid())
 	{
 		// Empty line in conditions array, just silently ignore.
 		return true;
 	}
 
 	FStateTreeBindableStructDesc StructDesc;
-	StructDesc.ID = CondItem.ID;
-	StructDesc.Name = CondItem.Item.GetScriptStruct()->GetFName();
+	StructDesc.ID = CondNode.ID;
+	StructDesc.Name = CondNode.Node.GetScriptStruct()->GetFName();
 
 	// Check that item has valid instance initialized.
-	if (!CondItem.Instance.IsValid() && CondItem.InstanceObject == nullptr)
+	if (!CondNode.Instance.IsValid() && CondNode.InstanceObject == nullptr)
 	{
 		Log.Reportf(EMessageSeverity::Error, StructDesc,
 			TEXT("Malformed condition, missing instance value."));
@@ -199,17 +199,17 @@
 
 	// Copy the condition
 	FInstancedStruct& Item = StateTree->Items.AddDefaulted_GetRef();
-	Item = CondItem.Item;
+	Item = CondNode.Node;
 
 	FStateTreeConditionBase& Cond = Item.GetMutable<FStateTreeConditionBase>();
 
-	if (CondItem.Instance.IsValid())
+	if (CondNode.Instance.IsValid())
 	{
 		// Struct instance
 		FInstancedStruct& Instance = StateTree->Instances.AddDefaulted_GetRef();
 		const int32 InstanceIndex = StateTree->Instances.Num() - 1;
 		
-		Instance = CondItem.Instance;
+		Instance = CondNode.Instance;
 	
 		// Create binding source struct descriptor.
 		StructDesc.Struct = Instance.GetScriptStruct();
@@ -222,9 +222,9 @@
 	else
 	{
 		// Object Instance
-		check(CondItem.InstanceObject != nullptr);
+		check(CondNode.InstanceObject != nullptr);
 		
-		UObject* Instance = DuplicateObject(CondItem.InstanceObject, StateTree);
+		UObject* Instance = DuplicateObject(CondNode.InstanceObject, StateTree);
 		
 		StateTree->InstanceObjects.Add(Instance);
 		const int32 InstanceIndex = StateTree->InstanceObjects.Num() - 1;
@@ -263,106 +263,21 @@
 	return true;
 }
 
-bool FStateTreeBaker::IsPropertyAnyEnum(const FStateTreeBindableStructDesc& Struct, FStateTreeEditorPropertyPath Path) const
+bool FStateTreeBaker::CreateTask(const FStateTreeEditorNode& TaskNode)
 {
-	bool bIsAnyEnum = false;
-	TArray<FStateTreePropertySegment> Segments;
-	const FProperty* LeafProperty = nullptr;
-	int32 LeafArrayIndex = INDEX_NONE;
-	const bool bResolved = FStateTreePropertyBindingCompiler::ResolvePropertyPath(Struct, Path, Segments, LeafProperty, LeafArrayIndex);
-	if (bResolved && LeafProperty)
-	{
-		if (const FProperty* OwnerProperty = LeafProperty->GetOwnerProperty())
-		{
-			if (const FStructProperty* OwnerStructProperty = CastField<FStructProperty>(OwnerProperty))
-			{
-				bIsAnyEnum = OwnerStructProperty->Struct == FStateTreeAnyEnum::StaticStruct();
-			}
-		}
-	}
-	return bIsAnyEnum;
-}
-
-bool FStateTreeBaker::GetAndValidateBindings(const FStateTreeBindableStructDesc& TargetStruct, TArray<FStateTreeEditorPropertyBinding>& OutBindings) const
-{
-	OutBindings.Reset();
-	
-	for (const FStateTreeEditorPropertyBinding& Binding : TreeData->EditorBindings.GetBindings())
-	{
-		if (Binding.TargetPath.StructID != TargetStruct.ID)
-		{
-			continue;
-		}
-
-		// Source must be one of the source structs we have discovered in the tree.
-		const FGuid SourceStructID = Binding.SourcePath.StructID;
-		const int32 SourceStructIdx = BindingsCompiler.GetSourceStructIndexByID(SourceStructID);
-		if (SourceStructIdx == INDEX_NONE)
-		{
-			Log.Reportf(EMessageSeverity::Error, TargetStruct,
-						TEXT("Failed to find binding source %s:%s."),
-						*TargetStruct.Name.ToString(), *Binding.SourcePath.ToString());
-			return false;
-		}
-		const FStateTreeBindableStructDesc& SourceStruct = BindingsCompiler.GetSourceStructDesc(SourceStructIdx);
-
-		// Source must be accessible by the target struct.
-		TArray<FStateTreeBindableStructDesc> AccessibleStructs;
-		TreeData->GetAccessibleStructs(Binding.TargetPath.StructID, AccessibleStructs);
-		const bool SourceAccessible = AccessibleStructs.ContainsByPredicate([SourceStructID](const FStateTreeBindableStructDesc& Structs)
-			{
-				return (Structs.ID == SourceStructID);
-			});
-		if (!SourceAccessible)
-		{
-			Log.Reportf(EMessageSeverity::Error, TargetStruct,
-        				TEXT("%s:%s is not accessible to %s:%s."),
-        				*SourceStruct.Name.ToString(), *Binding.SourcePath.ToString(), *TargetStruct.Name.ToString(), *Binding.TargetPath.ToString());
-			return false;
-		}
-
-		// Special case fo AnyEnum. StateTreeBindingExtension allows AnyEnums to bind to other enum types.
-		// The actual copy will be done via potential type promotion copy, into the value property inside the AnyEnum.
-		// We amend the paths here to point to the 'Value' property.
-		const bool bSourceIsAnyEnum = IsPropertyAnyEnum(SourceStruct, Binding.SourcePath);
-		const bool bTargetIsAnyEnum = IsPropertyAnyEnum(TargetStruct, Binding.TargetPath);
-		if (bSourceIsAnyEnum || bTargetIsAnyEnum)
-		{
-			FStateTreeEditorPropertyBinding ModifiedBinding(Binding);
-			if (bSourceIsAnyEnum)
-			{
-				ModifiedBinding.SourcePath.Path.Add(GET_MEMBER_NAME_STRING_CHECKED(FStateTreeAnyEnum, Value));
-			}
-			if (bTargetIsAnyEnum)
-			{
-				ModifiedBinding.TargetPath.Path.Add(GET_MEMBER_NAME_STRING_CHECKED(FStateTreeAnyEnum, Value));
-			}
-			OutBindings.Add(ModifiedBinding);
-		}
-		else
-		{
-			OutBindings.Add(Binding);
-		}
-	}
-
-	return true;
-}
-
-bool FStateTreeBaker::CreateTask(const FStateTreeTaskItem& TaskItem)
-{
 	// Silently ignore empty items.
-	if (!TaskItem.Item.IsValid())
+	if (!TaskNode.Node.IsValid())
 	{
 		return true;
 	}
 	
 	// Create binding source struct descriptor.
 	FStateTreeBindableStructDesc StructDesc;
-	StructDesc.ID = TaskItem.ID;
-	StructDesc.Name = TaskItem.Item.GetScriptStruct()->GetFName();
+	StructDesc.ID = TaskNode.ID;
+	StructDesc.Name = TaskNode.Node.GetScriptStruct()->GetFName();
 
 	// Check that item has valid instance initialized.
-	if (!TaskItem.Instance.IsValid() && TaskItem.InstanceObject == nullptr)
+	if (!TaskNode.Instance.IsValid() && TaskNode.InstanceObject == nullptr)
 	{
 		Log.Reportf(EMessageSeverity::Error, StructDesc,
 			TEXT("Malformed evaluator, missing instance value."));
@@ -371,17 +286,17 @@
 
 	// Copy the task
 	FInstancedStruct& Item = StateTree->Items.AddDefaulted_GetRef();
-	Item = TaskItem.Item;
+	Item = TaskNode.Node;
 	
 	FStateTreeTaskBase& Task = Item.GetMutable<FStateTreeTaskBase>();
 
-	if (TaskItem.Instance.IsValid())
+	if (TaskNode.Instance.IsValid())
 	{
 		// Struct Instance
 		FInstancedStruct& Instance = StateTree->Instances.AddDefaulted_GetRef();
 		const int32 InstanceIndex = StateTree->Instances.Num() - 1;
 
-		Instance = TaskItem.Instance;
+		Instance = TaskNode.Instance;
 
 		// Create binding source struct descriptor.
 		StructDesc.Struct = Instance.GetScriptStruct();
@@ -394,9 +309,9 @@
 	else
 	{
 		// Object Instance
-		check(TaskItem.InstanceObject != nullptr);
+		check(TaskNode.InstanceObject != nullptr);
 
-		UObject* Instance = DuplicateObject(TaskItem.InstanceObject, StateTree);
+		UObject* Instance = DuplicateObject(TaskNode.InstanceObject, StateTree);
 
 		StateTree->InstanceObjects.Add(Instance);
 		const int32 InstanceIndex = StateTree->InstanceObjects.Num() - 1;
@@ -436,21 +351,21 @@
 	return true;
 }
 
-bool FStateTreeBaker::CreateEvaluator(const FStateTreeEvaluatorItem& EvaluatorItem)
+bool FStateTreeBaker::CreateEvaluator(const FStateTreeEditorNode& EvalNode)
 {
 	// Silently ignore empty items.
-	if (!EvaluatorItem.Item.IsValid())
+	if (!EvalNode.Node.IsValid())
 	{
 		return true;
 	}
 	
 	// Create binding source struct descriptor.
 	FStateTreeBindableStructDesc StructDesc;
-    StructDesc.ID = EvaluatorItem.ID;
-    StructDesc.Name = EvaluatorItem.Item.GetScriptStruct()->GetFName();
+    StructDesc.ID = EvalNode.ID;
+    StructDesc.Name = EvalNode.Node.GetScriptStruct()->GetFName();
 
     // Check that item has valid instance initialized.
-    if (!EvaluatorItem.Instance.IsValid() && EvaluatorItem.InstanceObject == nullptr)
+    if (!EvalNode.Instance.IsValid() && EvalNode.InstanceObject == nullptr)
     {
         Log.Reportf(EMessageSeverity::Error, StructDesc,
         	TEXT("Malformed evaluator, missing instance value."));
@@ -459,17 +374,17 @@
 
 	// Copy the evaluator
 	FInstancedStruct& Item = StateTree->Items.AddDefaulted_GetRef();
-	Item = EvaluatorItem.Item;
+	Item = EvalNode.Node;
 
 	FStateTreeEvaluatorBase& Eval = Item.GetMutable<FStateTreeEvaluatorBase>();
 
-	if (EvaluatorItem.Instance.IsValid())
+	if (EvalNode.Instance.IsValid())
 	{
 		// Struct Instance
 		FInstancedStruct& Instance = StateTree->Instances.AddDefaulted_GetRef();
 		const int32 InstanceIndex = StateTree->Instances.Num() - 1;
 
-		Instance = EvaluatorItem.Instance;
+		Instance = EvalNode.Instance;
 
 		// Create binding source struct descriptor.
 		StructDesc.Struct = Instance.GetScriptStruct();
@@ -482,9 +397,9 @@
 	else
 	{
 		// Object Instance
-		check(EvaluatorItem.InstanceObject != nullptr);
+		check(EvalNode.InstanceObject != nullptr);
 
-		UObject* Instance = DuplicateObject(EvaluatorItem.InstanceObject, StateTree);
+		UObject* Instance = DuplicateObject(EvalNode.InstanceObject, StateTree);
 
 		StateTree->InstanceObjects.Add(Instance);
 		const int32 InstanceIndex = StateTree->InstanceObjects.Num() - 1;
@@ -524,6 +439,91 @@
 	return true;
 }
 
+bool FStateTreeBaker::IsPropertyAnyEnum(const FStateTreeBindableStructDesc& Struct, FStateTreeEditorPropertyPath Path) const
+{
+	bool bIsAnyEnum = false;
+	TArray<FStateTreePropertySegment> Segments;
+	const FProperty* LeafProperty = nullptr;
+	int32 LeafArrayIndex = INDEX_NONE;
+	const bool bResolved = FStateTreePropertyBindingCompiler::ResolvePropertyPath(Struct, Path, Segments, LeafProperty, LeafArrayIndex);
+	if (bResolved && LeafProperty)
+	{
+		if (const FProperty* OwnerProperty = LeafProperty->GetOwnerProperty())
+		{
+			if (const FStructProperty* OwnerStructProperty = CastField<FStructProperty>(OwnerProperty))
+			{
+				bIsAnyEnum = OwnerStructProperty->Struct == FStateTreeAnyEnum::StaticStruct();
+			}
+		}
+	}
+	return bIsAnyEnum;
+}
+
+bool FStateTreeBaker::GetAndValidateBindings(const FStateTreeBindableStructDesc& TargetStruct, TArray<FStateTreeEditorPropertyBinding>& OutBindings) const
+{
+	OutBindings.Reset();
+	
+	for (const FStateTreeEditorPropertyBinding& Binding : TreeData->EditorBindings.GetBindings())
+	{
+		if (Binding.TargetPath.StructID != TargetStruct.ID)
+		{
+			continue;
+		}
+
+		// Source must be one of the source structs we have discovered in the tree.
+		const FGuid SourceStructID = Binding.SourcePath.StructID;
+		const int32 SourceStructIdx = BindingsCompiler.GetSourceStructIndexByID(SourceStructID);
+		if (SourceStructIdx == INDEX_NONE)
+		{
+			Log.Reportf(EMessageSeverity::Error, TargetStruct,
+						TEXT("Failed to find binding source %s:%s."),
+						*TargetStruct.Name.ToString(), *Binding.SourcePath.ToString());
+			return false;
+		}
+		const FStateTreeBindableStructDesc& SourceStruct = BindingsCompiler.GetSourceStructDesc(SourceStructIdx);
+
+		// Source must be accessible by the target struct.
+		TArray<FStateTreeBindableStructDesc> AccessibleStructs;
+		TreeData->GetAccessibleStructs(Binding.TargetPath.StructID, AccessibleStructs);
+		const bool SourceAccessible = AccessibleStructs.ContainsByPredicate([SourceStructID](const FStateTreeBindableStructDesc& Structs)
+			{
+				return (Structs.ID == SourceStructID);
+			});
+		if (!SourceAccessible)
+		{
+			Log.Reportf(EMessageSeverity::Error, TargetStruct,
+        				TEXT("%s:%s is not accessible to %s:%s."),
+        				*SourceStruct.Name.ToString(), *Binding.SourcePath.ToString(), *TargetStruct.Name.ToString(), *Binding.TargetPath.ToString());
+			return false;
+		}
+
+		// Special case fo AnyEnum. StateTreeBindingExtension allows AnyEnums to bind to other enum types.
+		// The actual copy will be done via potential type promotion copy, into the value property inside the AnyEnum.
+		// We amend the paths here to point to the 'Value' property.
+		const bool bSourceIsAnyEnum = IsPropertyAnyEnum(SourceStruct, Binding.SourcePath);
+		const bool bTargetIsAnyEnum = IsPropertyAnyEnum(TargetStruct, Binding.TargetPath);
+		if (bSourceIsAnyEnum || bTargetIsAnyEnum)
+		{
+			FStateTreeEditorPropertyBinding ModifiedBinding(Binding);
+			if (bSourceIsAnyEnum)
+			{
+				ModifiedBinding.SourcePath.Path.Add(GET_MEMBER_NAME_STRING_CHECKED(FStateTreeAnyEnum, Value));
+			}
+			if (bTargetIsAnyEnum)
+			{
+				ModifiedBinding.TargetPath.Path.Add(GET_MEMBER_NAME_STRING_CHECKED(FStateTreeAnyEnum, Value));
+			}
+			OutBindings.Add(ModifiedBinding);
+		}
+		else
+		{
+			OutBindings.Add(Binding);
+		}
+	}
+
+	return true;
+}
+
 bool FStateTreeBaker::CreateStateRecursive(UStateTreeState& State, const FStateTreeHandle Parent)
 {
 	FStateTreeCompilerLogStateScope LogStateScope(&State, Log);
@@ -542,9 +542,9 @@
 	check(StateTree->Items.Num() <= int32(MAX_uint16));
 	BakedState.EvaluatorsBegin = uint16(StateTree->Items.Num());
 
-	for (FStateTreeEvaluatorItem& EvaluatorItem : State.Evaluators)
+	for (FStateTreeEditorNode& EvalNode : State.Evaluators)
 	{
-		if (!CreateEvaluator(EvaluatorItem))
+		if (!CreateEvaluator(EvalNode))
 		{
 			return false;
 		}
@@ -558,9 +558,9 @@
 	check(StateTree->Items.Num() <= int32(MAX_uint16));
 	BakedState.TasksBegin = uint16(StateTree->Items.Num());
 
-	for (FStateTreeTaskItem& TaskItem : State.Tasks)
+	for (FStateTreeEditorNode& TaskNode : State.Tasks)
 	{
-		if (!CreateTask(TaskItem))
+		if (!CreateTask(TaskNode))
 		{
 			return false;
 		}
diff --git a/Engine/Plugins/Runtime/StateTree/Source/StateTreeEditorModule/Private/StateTreeEditor.cpp b/Engine/Plugins/Runtime/StateTree/Source/StateTreeEditorModule/Private/StateTreeEditor.cpp
--- a/Engine/Plugins/Runtime/StateTree/Source/StateTreeEditorModule/Private/StateTreeEditor.cpp
+++ b/Engine/Plugins/Runtime/StateTree/Source/StateTreeEditorModule/Private/StateTreeEditor.cpp
@@ -325,9 +325,9 @@ FText FStateTreeEditor::GetStatisticsText() const
 	if (const UScriptStruct* StorageStruct = StateTree->GetInstanceStorageStruct())
 	{
 		const FText SizeText = FText::AsMemory((uint64)StorageStruct->GetStructureSize());
-		const FText NumItemsText = FText::AsNumber(StateTree->GetNumInstances());
+		const FText NumNodesText = FText::AsNumber(StateTree->GetNumInstances());
 
-		return FText::Format(LOCTEXT("RuntimeSize", "Runtime size: {0}, {1} items"), SizeText, NumItemsText);
+		return FText::Format(LOCTEXT("RuntimeSize", "Runtime size: {0}, {1} nodes"), SizeText, NumNodesText);
 	}
 	
 	return FText::GetEmpty();
@@ -593,14 +593,14 @@ namespace UE::StateTree::Editor::Internal
 						}
 						
 						// Task name is the same as state name.
-						if (FStateTreeTaskBase* Task = CurState->SingleTask.Item.GetMutablePtr<FStateTreeTaskBase>())
+						if (FStateTreeTaskBase* Task = CurState->SingleTask.Node.GetMutablePtr<FStateTreeTaskBase>())
 						{
 							Task->Name = CurState->Name;
 						}
 					}
 					else
 					{
-						if (CurState->SingleTask.Item.IsValid())
+						if (CurState->SingleTask.Node.IsValid())
 						{
 							CurState->SingleTask.Reset();
 							UE_LOG(LogStateTree, Warning, TEXT("%s: Resetting Single Task in state %s due to current schema restrictions."), *GetNameSafe(&StateTree), *GetNameSafe(CurState));
diff --git a/Engine/Plugins/Runtime/StateTree/Source/StateTreeEditorModule/Private/StateTreeEditorData.cpp b/Engine/Plugins/Runtime/StateTree/Source/StateTreeEditorModule/Private/StateTreeEditorData.cpp
--- a/Engine/Plugins/Runtime/StateTree/Source/StateTreeEditorModule/Private/StateTreeEditorData.cpp
+++ b/Engine/Plugins/Runtime/StateTree/Source/StateTreeEditorModule/Private/StateTreeEditorData.cpp
@@ -2,6 +2,8 @@
 
 #include "StateTreeEditorData.h"
 #include "StateTreeConditionBase.h"
+#include "StateTreeEvaluatorBase.h"
+#include "StateTreeTaskBase.h"
 #include "CoreMinimal.h"
 
 
@@ -32,29 +34,29 @@ void UStateTreeEditorData::GetAccessibleStructs(const FGuid TargetStructID, TArr
 	if (ValidStates.Num() > 0)
 	{
 		VisitHierarchy([&ValidStates, &OutStructDescs, &EvalDescs, &TaskDescs, TargetStructID, CurrentState]
-			(const UStateTreeState& State, const FGuid& ID, const FName& Name, const EStateTreeItemType ItemType, const UScriptStruct* ItemStruct, const UStruct* InstanceStruct) -> bool
+			(const UStateTreeState& State, const FGuid& ID, const FName& Name, const EStateTreeNodeType NodeType, const UScriptStruct* NodeStruct, const UStruct* InstanceStruct) -> bool
 			{
 				if (ValidStates.Contains(&State))
 				{
 					if (ID == TargetStructID)
 					{
-						if (ItemType == EStateTreeItemType::EnterCondition)
+						if (NodeType == EStateTreeNodeType::EnterCondition)
 						{
 							// Enter conditions can only see evaluators
 							OutStructDescs.Append(EvalDescs);
 						}
-						else if (ItemType == EStateTreeItemType::Evaluator)
+						else if (NodeType == EStateTreeNodeType::Evaluator)
 						{
 							// Evaluators can only see other evaluators
 							OutStructDescs.Append(EvalDescs);
 						}
-						else if (ItemType == EStateTreeItemType::Task)
+						else if (NodeType == EStateTreeNodeType::Task)
 						{
 							// Tasks can see evals and tasks.
 							OutStructDescs.Append(EvalDescs);
 							OutStructDescs.Append(TaskDescs);
 						}
-						else if (ItemType == EStateTreeItemType::TransitionCondition)
+						else if (NodeType == EStateTreeNodeType::TransitionCondition)
 						{
 							// Transitions can see evals and tasks.
 							OutStructDescs.Append(EvalDescs);
@@ -65,14 +67,14 @@ void UStateTreeEditorData::GetAccessibleStructs(const FGuid TargetStructID, TArr
 					}
 
 					// Not at target yet, collect all evaluators and tasks visible so far.
-					if (ItemStruct->IsChildOf(FStateTreeEvaluatorBase::StaticStruct()))
+					if (NodeStruct->IsChildOf(FStateTreeEvaluatorBase::StaticStruct()))
 					{
 						FStateTreeBindableStructDesc& Desc = EvalDescs.AddDefaulted_GetRef();
 						Desc.Struct = InstanceStruct;
 						Desc.Name = Name;
 						Desc.ID = ID;
 					}
-					else if (ItemStruct->IsChildOf(FStateTreeTaskBase::StaticStruct()))
+					else if (NodeStruct->IsChildOf(FStateTreeTaskBase::StaticStruct()))
 					{
 						FStateTreeBindableStructDesc& Desc = TaskDescs.AddDefaulted_GetRef();
 						Desc.Struct = InstanceStruct;
@@ -89,7 +91,7 @@ bool UStateTreeEditorData::GetStructByID(const FGuid StructID, FStateTreeBindabl
 {
 	bool bResult = false;
 
-	VisitHierarchy([&bResult, &OutStructDesc, StructID](const UStateTreeState& State, const FGuid& ID, const FName& Name, const EStateTreeItemType ItemType, const UScriptStruct* ItemStruct, const UStruct* InstanceStruct) -> bool
+	VisitHierarchy([&bResult, &OutStructDesc, StructID](const UStateTreeState& State, const FGuid& ID, const FName& Name, const EStateTreeNodeType NodeType, const UScriptStruct* NodeStruct, const UStruct* InstanceStruct) -> bool
 		{
 			if (ID == StructID)
 			{
@@ -109,7 +111,7 @@ const UStateTreeState* UStateTreeEditorData::GetStateByStructID(const FGuid Targ
 {
 	const UStateTreeState* Result = nullptr;
 
-	VisitHierarchy([&Result, TargetStructID](const UStateTreeState& State, const FGuid& ID, const FName& Name, const EStateTreeItemType ItemType, const UScriptStruct* ItemStruct, const UStruct* InstanceStruct) -> bool
+	VisitHierarchy([&Result, TargetStructID](const UStateTreeState& State, const FGuid& ID, const FName& Name, const EStateTreeNodeType NodeType, const UScriptStruct* NodeStruct, const UStruct* InstanceStruct) -> bool
 		{
 			if (ID == TargetStructID)
 			{
@@ -126,14 +128,14 @@ void UStateTreeEditorData::GetAllStructIDs(TMap<FGuid, const UStruct*>& AllStruc
 {
 	AllStructs.Reset();
 
-	VisitHierarchy([&AllStructs](const UStateTreeState& State, const FGuid& ID, const FName& Name, const EStateTreeItemType ItemType, const UScriptStruct* ItemStruct, const UStruct* InstanceStruct) -> bool
+	VisitHierarchy([&AllStructs](const UStateTreeState& State, const FGuid& ID, const FName& Name, const EStateTreeNodeType NodeType, const UScriptStruct* NodeStruct, const UStruct* InstanceStruct) -> bool
 		{
 			AllStructs.Add(ID, InstanceStruct);
 			return true; // Continue
 		});
 }
 
-void UStateTreeEditorData::VisitHierarchy(TFunctionRef<bool(const UStateTreeState& State, const FGuid& ID, const FName& Name, const EStateTreeItemType ItemType, const UScriptStruct* ItemStruct, const UStruct* InstanceStruct)> InFunc) const
+void UStateTreeEditorData::VisitHierarchy(TFunctionRef<bool(const UStateTreeState& State, const FGuid& ID, const FName& Name, const EStateTreeNodeType NodeType, const UScriptStruct* NodeStruct, const UStruct* InstanceStruct)> InFunc) const
 {
 	TArray<const UStateTreeState*> Stack;
 	bool bContinue = true;
@@ -155,11 +157,11 @@ void UStateTreeEditorData::VisitHierarchy(TFunctionRef<bool(const UStateTreeStat
 			Stack.RemoveAt(0);
 
 			// Evaluators
-			for (const FStateTreeEvaluatorItem& Item : State->Evaluators)
+			for (const FStateTreeEditorNode& Node : State->Evaluators)
 			{
-				if (const FStateTreeEvaluatorBase* Evaluator = Item.Item.GetPtr<FStateTreeEvaluatorBase>())
+				if (const FStateTreeEvaluatorBase* Evaluator = Node.Node.GetPtr<FStateTreeEvaluatorBase>())
 				{
-					if (!InFunc(*State, Item.ID, Evaluator->Name, EStateTreeItemType::Evaluator, Item.Item.GetScriptStruct(), Evaluator->GetInstanceDataType()))
+					if (!InFunc(*State, Node.ID, Evaluator->Name, EStateTreeNodeType::Evaluator, Node.Node.GetScriptStruct(), Evaluator->GetInstanceDataType()))
 					{
 						bContinue = false;
 						break;
@@ -169,11 +171,11 @@ void UStateTreeEditorData::VisitHierarchy(TFunctionRef<bool(const UStateTreeStat
 			if (bContinue)
 			{
 				// Enter conditions
-				for (const FStateTreeConditionItem& Item : State->EnterConditions)
+				for (const FStateTreeEditorNode& Node : State->EnterConditions)
 				{
-					if (const FStateTreeConditionBase* Cond = Item.Item.GetPtr<FStateTreeConditionBase>())
+					if (const FStateTreeConditionBase* Cond = Node.Node.GetPtr<FStateTreeConditionBase>())
 					{
-						if (!InFunc(*State, Item.ID, Item.Item.GetScriptStruct()->GetFName(), EStateTreeItemType::EnterCondition, Item.Item.GetScriptStruct(), Cond->GetInstanceDataType()))
+						if (!InFunc(*State, Node.ID, Node.Node.GetScriptStruct()->GetFName(), EStateTreeNodeType::EnterCondition, Node.Node.GetScriptStruct(), Cond->GetInstanceDataType()))
 						{
 							bContinue = false;
 							break;
@@ -184,11 +186,11 @@ void UStateTreeEditorData::VisitHierarchy(TFunctionRef<bool(const UStateTreeStat
 			if (bContinue)
 			{
 				// Tasks
-				for (const FStateTreeTaskItem& Item : State->Tasks)
+				for (const FStateTreeEditorNode& Node : State->Tasks)
 				{
-					if (const FStateTreeTaskBase* Task = Item.Item.GetPtr<FStateTreeTaskBase>())
+					if (const FStateTreeTaskBase* Task = Node.Node.GetPtr<FStateTreeTaskBase>())
 					{
-						if (!InFunc(*State, Item.ID, Task->Name, EStateTreeItemType::Task, Item.Item.GetScriptStruct(), Task->GetInstanceDataType()))
+						if (!InFunc(*State, Node.ID, Task->Name, EStateTreeNodeType::Task, Node.Node.GetScriptStruct(), Task->GetInstanceDataType()))
 						{
 							bContinue = false;
 							break;
@@ -198,9 +200,9 @@ void UStateTreeEditorData::VisitHierarchy(TFunctionRef<bool(const UStateTreeStat
 			}
 			if (bContinue)
 			{
-				if (const FStateTreeTaskBase* Task = State->SingleTask.Item.GetPtr<FStateTreeTaskBase>())
+				if (const FStateTreeTaskBase* Task = State->SingleTask.Node.GetPtr<FStateTreeTaskBase>())
 				{
-					if (!InFunc(*State, State->SingleTask.ID, Task->Name, EStateTreeItemType::Task, State->SingleTask.Item.GetScriptStruct(), Task->GetInstanceDataType()))
+					if (!InFunc(*State, State->SingleTask.ID, Task->Name, EStateTreeNodeType::Task, State->SingleTask.Node.GetScriptStruct(), Task->GetInstanceDataType()))
 					{
 						bContinue = false;
 						break;
@@ -213,11 +215,11 @@ void UStateTreeEditorData::VisitHierarchy(TFunctionRef<bool(const UStateTreeStat
 				// Transitions
 				for (const FStateTreeTransition& Transition : State->Transitions)
 				{
-					for (const FStateTreeConditionItem& Item : Transition.Conditions)
+					for (const FStateTreeEditorNode& Node : Transition.Conditions)
 					{
-						if (const FStateTreeConditionBase* Cond = Item.Item.GetPtr<FStateTreeConditionBase>())
+						if (const FStateTreeConditionBase* Cond = Node.Node.GetPtr<FStateTreeConditionBase>())
 						{
-							if (!InFunc(*State, Item.ID, Item.Item.GetScriptStruct()->GetFName(), EStateTreeItemType::TransitionCondition, Item.Item.GetScriptStruct(), Cond->GetInstanceDataType()))
+							if (!InFunc(*State, Node.ID, Node.Node.GetScriptStruct()->GetFName(), EStateTreeNodeType::TransitionCondition, Node.Node.GetScriptStruct(), Cond->GetInstanceDataType()))
 							{
 								bContinue = false;
 								break;
diff --git a/Engine/Plugins/Runtime/StateTree/Source/StateTreeEditorModule/Private/StateTreeEditorModule.cpp b/Engine/Plugins/Runtime/StateTree/Source/StateTreeEditorModule/Private/StateTreeEditorModule.cpp
--- a/Engine/Plugins/Runtime/StateTree/Source/StateTreeEditorModule/Private/StateTreeEditorModule.cpp
+++ b/Engine/Plugins/Runtime/StateTree/Source/StateTreeEditorModule/Private/StateTreeEditorModule.cpp
@@ -8,7 +8,7 @@
 #include "Customizations/StateTreeTransitionDetails.h"
 #include "Customizations/StateTreeStateLinkDetails.h"
 #include "Customizations/StateTreeStateDetails.h"
-#include "Customizations/StateTreeBindableItemDetails.h"
+#include "Customizations/StateTreeEditorNodeDetails.h"
 #include "Customizations/StateTreeAnyEnumDetails.h"
 #include "StateTreeEditor.h"
 #include "StateTree.h"
@@ -49,9 +49,7 @@ void FStateTreeEditorModule::StartupModule()
 	FPropertyEditorModule& PropertyModule = FModuleManager::LoadModuleChecked<FPropertyEditorModule>("PropertyEditor");
 	PropertyModule.RegisterCustomPropertyTypeLayout("StateTreeTransition", FOnGetPropertyTypeCustomizationInstance::CreateStatic(&FStateTreeTransitionDetails::MakeInstance));
 	PropertyModule.RegisterCustomPropertyTypeLayout("StateTreeStateLink", FOnGetPropertyTypeCustomizationInstance::CreateStatic(&FStateTreeStateLinkDetails::MakeInstance));
-	PropertyModule.RegisterCustomPropertyTypeLayout("StateTreeEvaluatorItem", FOnGetPropertyTypeCustomizationInstance::CreateStatic(&FStateTreeBindableItemDetails::MakeInstance));
-	PropertyModule.RegisterCustomPropertyTypeLayout("StateTreeTaskItem", FOnGetPropertyTypeCustomizationInstance::CreateStatic(&FStateTreeBindableItemDetails::MakeInstance));
-	PropertyModule.RegisterCustomPropertyTypeLayout("StateTreeConditionItem", FOnGetPropertyTypeCustomizationInstance::CreateStatic(&FStateTreeBindableItemDetails::MakeInstance));
+	PropertyModule.RegisterCustomPropertyTypeLayout("StateTreeEditorNode", FOnGetPropertyTypeCustomizationInstance::CreateStatic(&FStateTreeEditorNodeDetails::MakeInstance));
 	PropertyModule.RegisterCustomPropertyTypeLayout("StateTreeAnyEnum", FOnGetPropertyTypeCustomizationInstance::CreateStatic(&FStateTreeAnyEnumDetails::MakeInstance));
 	PropertyModule.RegisterCustomClassLayout("StateTreeState", FOnGetDetailCustomizationInstance::CreateStatic(&FStateTreeStateDetails::MakeInstance));
 
diff --git a/Engine/Plugins/Runtime/StateTree/Source/StateTreeEditorModule/Private/StateTreeState.cpp b/Engine/Plugins/Runtime/StateTree/Source/StateTreeEditorModule/Private/StateTreeState.cpp
--- a/Engine/Plugins/Runtime/StateTree/Source/StateTreeEditorModule/Private/StateTreeState.cpp
+++ b/Engine/Plugins/Runtime/StateTree/Source/StateTreeEditorModule/Private/StateTreeState.cpp
@@ -3,6 +3,8 @@
 #include "StateTreeState.h"
 #include "StateTree.h"
 #include "StateTreeConditionBase.h"
+#include "StateTreeEvaluatorBase.h"
+#include "StateTreeTaskBase.h"
 #include "StateTreeDelegates.h"
 #include "CoreMinimal.h"
 #include "UObject/Field.h"
@@ -76,7 +78,7 @@ void UStateTreeState::PostEditChangeChainProperty(FPropertyChangedChainEvent& Pr
 					const int32 ArrayIndex = PropertyChangedEvent.GetArrayIndex(MemberProperty->GetFName().ToString());
 					if (Evaluators.IsValidIndex(ArrayIndex))
 					{
-						if (FStateTreeEvaluatorBase* Eval = Evaluators[ArrayIndex].Item.GetMutablePtr<FStateTreeEvaluatorBase>())
+						if (FStateTreeEvaluatorBase* Eval = Evaluators[ArrayIndex].Node.GetMutablePtr<FStateTreeEvaluatorBase>())
 						{
 							Eval->Name = FName(Eval->Name.ToString() + TEXT(" Duplicate"));
 						}
@@ -88,7 +90,7 @@ void UStateTreeState::PostEditChangeChainProperty(FPropertyChangedChainEvent& Pr
 					const int32 ArrayIndex = PropertyChangedEvent.GetArrayIndex(MemberProperty->GetFName().ToString());
 					if (Tasks.IsValidIndex(ArrayIndex))
 					{
-						if (FStateTreeTaskBase* Task = Tasks[ArrayIndex].Item.GetMutablePtr<FStateTreeTaskBase>())
+						if (FStateTreeTaskBase* Task = Tasks[ArrayIndex].Node.GetMutablePtr<FStateTreeTaskBase>())
 						{
 							Task->Name = FName(Task->Name.ToString() + TEXT(" Duplicate"));
 						}
diff --git a/Engine/Plugins/Runtime/StateTree/Source/StateTreeEditorModule/Public/StateTreeBaker.h b/Engine/Plugins/Runtime/StateTree/Source/StateTreeEditorModule/Public/StateTreeBaker.h
--- a/Engine/Plugins/Runtime/StateTree/Source/StateTreeEditorModule/Public/StateTreeBaker.h
+++ b/Engine/Plugins/Runtime/StateTree/Source/StateTreeEditorModule/Public/StateTreeBaker.h
@@ -10,10 +10,7 @@
 class UStateTree;
 class UStateTreeState;
 class UStateTreeEditorData;
-struct FStateTreeCondition;
-struct FStateTreeConditionItem;
-struct FStateTreeTaskItem;
-struct FStateTreeEvaluatorItem;
+struct FStateTreeEditorNode;
 struct FStateTreeStateLink;
 
 /**
@@ -39,11 +36,11 @@ private:
 	bool CreateStates();
 	bool CreateStateRecursive(UStateTreeState& State, const FStateTreeHandle Parent);
 	bool CreateStateTransitions();
-	bool CreateCondition(const FStateTreeConditionItem& CondItem);
+	bool CreateCondition(const FStateTreeEditorNode& CondNode);
+	bool CreateTask(const FStateTreeEditorNode& TaskNode);
+	bool CreateEvaluator(const FStateTreeEditorNode& EvalNode);
 	bool GetAndValidateBindings(const FStateTreeBindableStructDesc& TargetStruct, TArray<FStateTreeEditorPropertyBinding>& OutBindings) const;
 	bool IsPropertyAnyEnum(const FStateTreeBindableStructDesc& Struct, FStateTreeEditorPropertyPath Path) const;
-	bool CreateTask(const FStateTreeTaskItem& TaskItem);
-	bool CreateEvaluator(const FStateTreeEvaluatorItem& EvalItem);
 	
 	FStateTreeCompilerLog& Log;
 	UStateTree* StateTree = nullptr;
diff --git a/Engine/Plugins/Runtime/StateTree/Source/StateTreeEditorModule/Public/StateTreeEditorData.h b/Engine/Plugins/Runtime/StateTree/Source/StateTreeEditorModule/Public/StateTreeEditorData.h
--- a/Engine/Plugins/Runtime/StateTree/Source/StateTreeEditorModule/Public/StateTreeEditorData.h
+++ b/Engine/Plugins/Runtime/StateTree/Source/StateTreeEditorModule/Public/StateTreeEditorData.h
@@ -7,7 +7,7 @@
 #include "StateTreeEditorData.generated.h"
 
 UENUM()
-enum class EStateTreeItemType : uint8
+enum class EStateTreeNodeType : uint8
 {
 	EnterCondition,
 	Evaluator,
@@ -40,9 +40,9 @@ public:
 
 	/**
 	 * Iterates over all structs that are related to binding
-	 * @param InFunc function called at each item, should return true if visiting is continued or false to stop.
+	 * @param InFunc function called at each node, should return true if visiting is continued or false to stop.
 	 */
-	void VisitHierarchy(TFunctionRef<bool(const UStateTreeState& State, const FGuid& ID, const FName& Name, const EStateTreeItemType ItemType, const UScriptStruct* ItemStruct, const UStruct* InstanceStruct)> InFunc) const;
+	void VisitHierarchy(TFunctionRef<bool(const UStateTreeState& State, const FGuid& ID, const FName& Name, const EStateTreeNodeType NodeType, const UScriptStruct* NodeStruct, const UStruct* InstanceStruct)> InFunc) const;
 
 	// StateTree Builder API
 
diff --git a/Engine/Plugins/Runtime/StateTree/Source/StateTreeEditorModule/Public/StateTreeState.h b/Engine/Plugins/Runtime/StateTree/Source/StateTreeEditorModule/Public/StateTreeState.h
--- a/Engine/Plugins/Runtime/StateTree/Source/StateTreeEditorModule/Public/StateTreeState.h
+++ b/Engine/Plugins/Runtime/StateTree/Source/StateTreeEditorModule/Public/StateTreeState.h
@@ -2,8 +2,8 @@
 
 #pragma once
 
-#include "StateTreeEvaluatorBase.h"
-#include "StateTreeTaskBase.h"
+#include "StateTreeTypes.h"
+#include "StateTreeNodeBase.h"
 #include "Misc/Guid.h"
 #include "InstancedStruct.h"
 #include "StateTreeState.generated.h"
@@ -17,8 +17,8 @@ USTRUCT()
 struct STATETREEEDITORMODULE_API FStateTreeStateLink
 {
 	GENERATED_BODY()
-public:
-	FStateTreeStateLink() : Type(EStateTreeTransitionType::GotoState) {}
+
+	FStateTreeStateLink() = default;
 	FStateTreeStateLink(EStateTreeTransitionType InType) : Type(InType) {}
 
 	void Set(const EStateTreeTransitionType InType, const UStateTreeState* InState = nullptr);
@@ -37,103 +37,51 @@ public:
 
 
 /**
- * Base for Evaluator, Task and Condition items.
+ * Base for Evaluator, Task and Condition nodes.
  */
 USTRUCT()
-struct STATETREEEDITORMODULE_API FStateTreeItem
+struct STATETREEEDITORMODULE_API FStateTreeEditorNode
 {
 	GENERATED_BODY()
 
 	void Reset()
 	{
-		Item.Reset();
+		Node.Reset();
 		Instance.Reset();
 		InstanceObject = nullptr;
 		ID = FGuid();
 	}
-	
-	UPROPERTY(EditDefaultsOnly, Category = Item)
-	FInstancedStruct Item;
-
-	UPROPERTY(EditDefaultsOnly, Category = Item)
-	FInstancedStruct Instance;
-
-	UPROPERTY(EditDefaultsOnly, Instanced, Category = Item)
-	UObject* InstanceObject = nullptr;
-	
-	UPROPERTY(EditDefaultsOnly, Category = Item)
-	FGuid ID;
-};
-
-/**
- * Helper struct for Evaluator details customization (i.e. show summary when collapsed)
- */
-USTRUCT()
-struct STATETREEEDITORMODULE_API FStateTreeEvaluatorItem : public FStateTreeItem
-{
-	GENERATED_BODY()
 
 	FName GetName() const
 	{
-		if (const FStateTreeEvaluatorBase* Eval = Item.GetPtr<FStateTreeEvaluatorBase>())
+		if (const FStateTreeNodeBase* NodePtr = Node.GetPtr<FStateTreeNodeBase>())
 		{
-			return Eval->Name;
+			return NodePtr->Name;
 		}
 		return FName();
 	}
-};
-
-template <typename T>
-struct TStateTreeEvaluatorItem : public FStateTreeEvaluatorItem
-{
-	typedef T ItemType;
-	FORCEINLINE T& GetItem() { return Item.template GetMutable<T>(); }
-	FORCEINLINE typename T::InstanceDataType& GetInstance() { return Instance.template GetMutable<typename T::InstanceDataType>(); }
-};
 
+	UPROPERTY(EditDefaultsOnly, Category = Node)
+	FInstancedStruct Node;
 
-/**
- * Helper struct for Task details customization (i.e. show summary when collapsed)
- */
-USTRUCT()
-struct STATETREEEDITORMODULE_API FStateTreeTaskItem : public FStateTreeItem
-{
-	GENERATED_BODY()
+	UPROPERTY(EditDefaultsOnly, Category = Node)
+	FInstancedStruct Instance;
 
-	FName GetName() const
-	{
-		if (const FStateTreeTaskBase* Task = Item.GetPtr<FStateTreeTaskBase>())
-		{
-			return Task->Name;
-		}
-		return FName();
-	}
+	UPROPERTY(EditDefaultsOnly, Instanced, Category = Node)
+	UObject* InstanceObject = nullptr;
+	
+	UPROPERTY(EditDefaultsOnly, Category = Node)
+	FGuid ID;
 };
 
 template <typename T>
-struct TStateTreeTaskItem : public FStateTreeTaskItem
+struct TStateTreeEditorNode : public FStateTreeEditorNode
 {
-	typedef T ItemType;
-	FORCEINLINE T& GetItem() { return Item.template GetMutable<T>(); }
+	typedef T NodeType;
+	FORCEINLINE T& GetItem() { return Node.template GetMutable<T>(); }
 	FORCEINLINE typename T::InstanceDataType& GetInstance() { return Instance.template GetMutable<typename T::InstanceDataType>(); }
 };
 
-/**
- * Helper struct for Condition details customization (i.e. show summary when collapsed)
- */
-USTRUCT()
-struct STATETREEEDITORMODULE_API FStateTreeConditionItem : public FStateTreeItem
-{
-	GENERATED_BODY()
-};
-
-template <typename T>
-struct TStateTreeConditionItem : public FStateTreeConditionItem
-{
-	typedef T ItemType;
-	FORCEINLINE T& GetItem() { return Item.template GetMutable<T>(); }
-	FORCEINLINE typename T::InstanceDataType& GetInstance() { return Instance.template GetMutable<typename T::InstanceDataType>(); }
-};
 
 /**
  * Editor representation of a transition in StateTree
@@ -147,17 +95,17 @@ struct STATETREEEDITORMODULE_API FStateTreeTransition
 	FStateTreeTransition(const EStateTreeTransitionEvent InEvent, const EStateTreeTransitionType InType, const UStateTreeState* InState = nullptr);
 
 	template<typename T, typename... TArgs>
-	TStateTreeConditionItem<T>& AddCondition(TArgs&&... InArgs)
+	TStateTreeEditorNode<T>& AddCondition(TArgs&&... InArgs)
 	{
-		FStateTreeConditionItem& CondItem = Conditions.AddDefaulted_GetRef();
-		CondItem.ID = FGuid::NewGuid();
-		CondItem.Item.InitializeAs<T>(Forward<TArgs>(InArgs)...);
-		T& Cond = CondItem.Item.GetMutable<T>();
+		FStateTreeEditorNode& CondNode = Conditions.AddDefaulted_GetRef();
+		CondNode.ID = FGuid::NewGuid();
+		CondNode.Node.InitializeAs<T>(Forward<TArgs>(InArgs)...);
+		T& Cond = CondNode.Node.GetMutable<T>();
 		if (const UScriptStruct* InstanceType = Cast<const UScriptStruct>(Cond.GetInstanceDataType()))
 		{
-			CondItem.Instance.InitializeAs(InstanceType);
+			CondNode.Instance.InitializeAs(InstanceType);
 		}
-		return static_cast<TStateTreeConditionItem<T>&>(CondItem);
+		return static_cast<TStateTreeEditorNode<T>&>(CondNode);
 	}
 
 	UPROPERTY(EditDefaultsOnly, Category = Transition)
@@ -171,7 +119,7 @@ struct STATETREEEDITORMODULE_API FStateTreeTransition
 	float GateDelay = 0.0f;
 	
 	UPROPERTY(EditDefaultsOnly, Category = Transition, meta = (BaseStruct = "StateTreeConditionBase", BaseClass = "StateTreeConditionBlueprintBase"))
-	TArray<FStateTreeConditionItem> Conditions;
+	TArray<FStateTreeEditorNode> Conditions;
 };
 
 
@@ -210,17 +158,17 @@ public:
 	 * @return reference to the new condition. 
 	 */
 	template<typename T, typename... TArgs>
-	TStateTreeConditionItem<T>& AddEnterCondition(TArgs&&... InArgs)
+	TStateTreeEditorNode<T>& AddEnterCondition(TArgs&&... InArgs)
 	{
-		FStateTreeConditionItem& CondItem = EnterConditions.AddDefaulted_GetRef();
-		CondItem.ID = FGuid::NewGuid();
-		CondItem.Item.InitializeAs<T>(Forward<TArgs>(InArgs)...);
-		T& Cond = CondItem.Item.GetMutable<T>();
+		FStateTreeEditorNode& CondNode = EnterConditions.AddDefaulted_GetRef();
+		CondNode.ID = FGuid::NewGuid();
+		CondNode.Node.InitializeAs<T>(Forward<TArgs>(InArgs)...);
+		T& Cond = CondNode.Node.GetMutable<T>();
 		if (const UScriptStruct* InstanceType = Cast<const UScriptStruct>(Cond.GetInstanceDataType()))
 		{
-			CondItem.Instance.InitializeAs(InstanceType);
+			CondNode.Instance.InitializeAs(InstanceType);
 		}
-		return static_cast<TStateTreeConditionItem<T>&>(CondItem);
+		return static_cast<TStateTreeEditorNode<T>&>(CondNode);
 	}
 
 	/**
@@ -228,17 +176,17 @@ public:
 	 * @return reference to the new Task. 
 	 */
 	template<typename T, typename... TArgs>
-	TStateTreeTaskItem<T>& AddTask(TArgs&&... InArgs)
+	TStateTreeEditorNode<T>& AddTask(TArgs&&... InArgs)
 	{
-		FStateTreeTaskItem& TaskItem = Tasks.AddDefaulted_GetRef();
+		FStateTreeEditorNode& TaskItem = Tasks.AddDefaulted_GetRef();
 		TaskItem.ID = FGuid::NewGuid();
-		TaskItem.Item.InitializeAs<T>(Forward<TArgs>(InArgs)...);
-		T& Task = TaskItem.Item.GetMutable<T>();
+		TaskItem.Node.InitializeAs<T>(Forward<TArgs>(InArgs)...);
+		T& Task = TaskItem.Node.GetMutable<T>();
 		if (const UScriptStruct* InstanceType = Cast<const UScriptStruct>(Task.GetInstanceDataType()))
 		{
 			TaskItem.Instance.InitializeAs(InstanceType);
 		}
-		return static_cast<TStateTreeTaskItem<T>&>(TaskItem);
+		return static_cast<TStateTreeEditorNode<T>&>(TaskItem);
 	}
 
 	/**
@@ -246,17 +194,17 @@ public:
 	 * @return reference to the new Evaluator. 
 	 */
 	template<typename T, typename... TArgs>
-    TStateTreeEvaluatorItem<T>& AddEvaluator(TArgs&&... InArgs)
+    TStateTreeEditorNode<T>& AddEvaluator(TArgs&&... InArgs)
 	{
-		FStateTreeEvaluatorItem& EvalItem = Evaluators.AddDefaulted_GetRef();
+		FStateTreeEditorNode& EvalItem = Evaluators.AddDefaulted_GetRef();
 		EvalItem.ID = FGuid::NewGuid();
-		EvalItem.Item.InitializeAs<T>(Forward<TArgs>(InArgs)...);
-		T& Eval = EvalItem.Item.GetMutable<T>();
+		EvalItem.Node.InitializeAs<T>(Forward<TArgs>(InArgs)...);
+		T& Eval = EvalItem.Node.GetMutable<T>();
 		if (const UScriptStruct* InstanceType = Cast<const UScriptStruct>(Eval.GetInstanceDataType()))
 		{
 			EvalItem.Instance.InitializeAs(InstanceType);
 		}
-		return static_cast<TStateTreeEvaluatorItem<T>&>(EvalItem);
+		return static_cast<TStateTreeEditorNode<T>&>(EvalItem);
 	}
 
 	/**
@@ -277,17 +225,17 @@ public:
 	FGuid ID;
 
 	UPROPERTY(EditDefaultsOnly, Category = "Enter Conditions", meta = (BaseStruct = "StateTreeConditionBase", BaseClass = "StateTreeConditionBlueprintBase"))
-	TArray<FStateTreeConditionItem> EnterConditions;
+	TArray<FStateTreeEditorNode> EnterConditions;
 
 	UPROPERTY(EditDefaultsOnly, Category = "Evaluators", meta = (BaseStruct = "StateTreeEvaluatorBase", BaseClass = "StateTreeEvaluatorBlueprintBase"))
-	TArray<FStateTreeEvaluatorItem> Evaluators;
+	TArray<FStateTreeEditorNode> Evaluators;
 
 	UPROPERTY(EditDefaultsOnly, Category = "Tasks", meta = (BaseStruct = "StateTreeTaskBase", BaseClass = "StateTreeTaskBlueprintBase"))
-	TArray<FStateTreeTaskItem> Tasks;
+	TArray<FStateTreeEditorNode> Tasks;
 
 	// Single item used when schema calls for single task per state.
 	UPROPERTY(EditDefaultsOnly, Category = "Task", meta = (BaseStruct = "StateTreeTaskBase", BaseClass = "StateTreeTaskBlueprintBase"))
-	FStateTreeTaskItem SingleTask;
+	FStateTreeEditorNode SingleTask;
 
 	UPROPERTY(EditDefaultsOnly, Category = "Transitions")
 	TArray<FStateTreeTransition> Transitions;
@@ -295,10 +243,10 @@ public:
 	UPROPERTY()
 	TArray<UStateTreeState*> Children;
 
-	UPROPERTY()
+	UPROPERTY(meta = (ExcludeFromHash))
 	bool bExpanded;
 
-	UPROPERTY()
+	UPROPERTY(meta = (ExcludeFromHash))
 	UStateTreeState* Parent;
 
 };
diff --git a/Engine/Plugins/Runtime/StateTree/Source/StateTreeModule/Private/StateTree.cpp b/Engine/Plugins/Runtime/StateTree/Source/StateTreeModule/Private/StateTree.cpp
--- a/Engine/Plugins/Runtime/StateTree/Source/StateTreeModule/Private/StateTree.cpp
+++ b/Engine/Plugins/Runtime/StateTree/Source/StateTreeModule/Private/StateTree.cpp
@@ -127,20 +127,10 @@ void UStateTree::Link()
 	
 	for (FInstancedStruct& Item : Items)
 	{
-		if (FStateTreeEvaluatorBase* Eval = Item.GetMutablePtr<FStateTreeEvaluatorBase>())
+		if (FStateTreeNodeBase* Node = Item.GetMutablePtr<FStateTreeNodeBase>())
 		{
-			Linker.SetCurrentInstanceDataType(Eval->GetInstanceDataType(), Eval->DataViewIndex);
-			Eval->Link(Linker);
-		}
-		else if (FStateTreeTaskBase* Task = Item.GetMutablePtr<FStateTreeTaskBase>())
-		{
-			Linker.SetCurrentInstanceDataType(Task->GetInstanceDataType(), Task->DataViewIndex);
-			Task->Link(Linker);
-		}
-		else if (FStateTreeConditionBase* Cond = Item.GetMutablePtr<FStateTreeConditionBase>())
-		{
-			Linker.SetCurrentInstanceDataType(Cond->GetInstanceDataType(), Cond->DataViewIndex);
-			Cond->Link(Linker);
+			Linker.SetCurrentInstanceDataType(Node->GetInstanceDataType(), Node->DataViewIndex);
+			Node->Link(Linker);
 		}
 	}
 
diff --git a/Engine/Plugins/Runtime/StateTree/Source/StateTreeModule/Public/StateTreeConditionBase.h b/Engine/Plugins/Runtime/StateTree/Source/StateTreeModule/Public/StateTreeConditionBase.h
--- a/Engine/Plugins/Runtime/StateTree/Source/StateTreeModule/Public/StateTreeConditionBase.h
+++ b/Engine/Plugins/Runtime/StateTree/Source/StateTreeModule/Public/StateTreeConditionBase.h
@@ -2,8 +2,8 @@
 
 #pragma once
 
-#include "CoreMinimal.h"
 #include "StateTreeTypes.h"
+#include "StateTreeNodeBase.h"
 #include "StateTreeExecutionContext.h"
 #if WITH_EDITOR
 #include "StateTreePropertyBindings.h"
@@ -19,24 +19,10 @@ struct FStateTreeEditorPropertyPath;
  * Base struct for all conditions.
  */
 USTRUCT()
-struct STATETREEMODULE_API FStateTreeConditionBase
+struct STATETREEMODULE_API FStateTreeConditionBase : public FStateTreeNodeBase
 {
 	GENERATED_BODY()
 
-	virtual ~FStateTreeConditionBase() {}
-
-	/** @return Struct that represents the runtime data of the condition. */
-	virtual const UStruct* GetInstanceDataType() const { return nullptr; };
-
-	/**
-	 * Called when the StateTree asset is linked. Allows to resolve references to other StateTree data.
-	 * @see TStateTreeExternalDataHandle
-	 * @see TStateTreeInstanceDataPropertyHandle
-	 * @param Linker Reference to the linker
-	 * @return true if linking succeeded. 
-	 */
-	virtual bool Link(FStateTreeLinker& Linker) { return true; }
-
 #if WITH_EDITOR
 	/** @return Rich text description of the condition. */
 	virtual FText GetDescription(const FGuid& ID, FStateTreeDataView InstanceData, const IStateTreeBindingLookup& BindingLookup) const { return FText::GetEmpty(); }
@@ -50,22 +36,9 @@ struct STATETREEMODULE_API FStateTreeConditionBase
 	 */
 	virtual void OnBindingChanged(const FGuid& ID, FStateTreeDataView InstanceData, const FStateTreeEditorPropertyPath& SourcePath, const FStateTreeEditorPropertyPath& TargetPath, const IStateTreeBindingLookup& BindingLookup) {}
 #endif
+	
 	/** @return True if the condition passes. */
 	virtual bool TestCondition(FStateTreeExecutionContext& Context) const { return false; }
-
-	/** Property binding copy batch handle. */
-	UPROPERTY()
-	FStateTreeHandle BindingsBatch = FStateTreeHandle::Invalid;
-
-	/** The runtime data's data view index in the StateTreeExecutionContext, and source struct index in property binding. */
-	UPROPERTY()
-	uint16 DataViewIndex = 0;
-
-	UPROPERTY()
-	uint16 InstanceIndex = 0;
-
-	UPROPERTY()
-	uint8 bInstanceIsObject : 1;
 };
 
 
diff --git a/Engine/Plugins/Runtime/StateTree/Source/StateTreeModule/Public/StateTreeEvaluatorBase.h b/Engine/Plugins/Runtime/StateTree/Source/StateTreeModule/Public/StateTreeEvaluatorBase.h
--- a/Engine/Plugins/Runtime/StateTree/Source/StateTreeModule/Public/StateTreeEvaluatorBase.h
+++ b/Engine/Plugins/Runtime/StateTree/Source/StateTreeModule/Public/StateTreeEvaluatorBase.h
@@ -2,8 +2,8 @@
 
 #pragma once
 
-#include "CoreMinimal.h"
 #include "StateTreeTypes.h"
+#include "StateTreeNodeBase.h"
 #include "StateTreeEvaluatorBase.generated.h"
 
 struct FStateTreeExecutionContext;
@@ -13,26 +13,10 @@ struct FStateTreeExecutionContext;
  * Evaluators calculate and expose data to be used for decision making in a StateTree.
  */
 USTRUCT()
-struct STATETREEMODULE_API FStateTreeEvaluatorBase
+struct STATETREEMODULE_API FStateTreeEvaluatorBase : public FStateTreeNodeBase
 {
 	GENERATED_BODY()
 
-	FStateTreeEvaluatorBase() = default;
-
-	virtual ~FStateTreeEvaluatorBase() {}
-
-	/** @return Struct that represents the runtime data of the evaluator. */
-	virtual const UStruct* GetInstanceDataType() const { return nullptr; };
-
-	/**
-	 * Called when the StateTree asset is linked. Allows to resolve references to other StateTree data.
-	 * @see TStateTreeExternalDataHandle
-	 * @see TStateTreeInstanceDataPropertyHandle
-	 * @param Linker Reference to the linker
-	 * @return true if linking succeeded. 
-	 */
-	virtual bool Link(FStateTreeLinker& Linker) { return true; }
-
 	/**
 	 * Called when a new state is entered and evaluator is part of active states. The change type parameter describes if the evaluator's state
 	 * was previously part of the list of active states (Sustained), or if it just became active (Changed).
@@ -73,23 +57,6 @@ struct STATETREEMODULE_API FStateTreeEvaluatorBase
 #if WITH_GAMEPLAY_DEBUGGER
 	virtual void AppendDebugInfoString(FString& DebugString, const FStateTreeExecutionContext& Context) const;
 #endif // WITH_GAMEPLAY_DEBUGGER
-
-	UPROPERTY(EditDefaultsOnly, Category = Evaluator, meta=(EditCondition = "false", EditConditionHides))
-	FName Name;
-
-	/** Property binding copy batch handle. */
-	UPROPERTY()
-	FStateTreeHandle BindingsBatch = FStateTreeHandle::Invalid;		// Property binding copy batch handle.
-
-	/** The runtime data's data view index in the StateTreeExecutionContext, and source struct index in property binding. */
-	UPROPERTY()
-	uint16 DataViewIndex = 0;
-
-	UPROPERTY()
-	uint16 InstanceIndex = 0;
-
-	UPROPERTY()
-	uint8 bInstanceIsObject : 1;
 };
 
 /**
diff --git a/Engine/Plugins/Runtime/StateTree/Source/StateTreeModule/Public/StateTreeNodeBase.h b/Engine/Plugins/Runtime/StateTree/Source/StateTreeModule/Public/StateTreeNodeBase.h
new file mode 100644
--- /dev/null
+++ b/Engine/Plugins/Runtime/StateTree/Source/StateTreeModule/Public/StateTreeNodeBase.h
@@ -0,0 +1,51 @@
+﻿// Copyright Epic Games, Inc. All Rights Reserved.
+
+#pragma once
+
+#include "StateTreeTypes.h"
+#include "StateTreeNodeBase.generated.h"
+
+/**
+ * Base struct of StateTree Conditions, Evaluators, and Tasks.
+ */
+USTRUCT()
+struct STATETREEMODULE_API FStateTreeNodeBase
+{
+	GENERATED_BODY()
+
+	FStateTreeNodeBase() = default;
+
+	virtual ~FStateTreeNodeBase() {}
+
+	/** @return Struct that represents the runtime data of the evaluator. */
+	virtual const UStruct* GetInstanceDataType() const { return nullptr; };
+
+	/**
+	 * Called when the StateTree asset is linked. Allows to resolve references to other StateTree data.
+	 * @see TStateTreeExternalDataHandle
+	 * @see TStateTreeInstanceDataPropertyHandle
+	 * @param Linker Reference to the linker
+	 * @return true if linking succeeded. 
+	 */
+	virtual bool Link(FStateTreeLinker& Linker) { return true; }
+
+	/** Name of the node. */
+	UPROPERTY(EditDefaultsOnly, Category = "", meta=(EditCondition = "false", EditConditionHides))
+	FName Name;
+
+	/** Property binding copy batch handle. */
+	UPROPERTY()
+	FStateTreeHandle BindingsBatch = FStateTreeHandle::Invalid;
+
+	/** The runtime data's data view index in the StateTreeExecutionContext, and source struct index in property binding. */
+	UPROPERTY()
+	uint16 DataViewIndex = 0;
+
+	/** Index in runtime instance storage. */
+	UPROPERTY()
+	uint16 InstanceIndex = 0;
+
+	/** True if the instance is an UObject. */
+	UPROPERTY()
+	uint8 bInstanceIsObject : 1;
+};
diff --git a/Engine/Plugins/Runtime/StateTree/Source/StateTreeModule/Public/StateTreeTaskBase.h b/Engine/Plugins/Runtime/StateTree/Source/StateTreeModule/Public/StateTreeTaskBase.h
--- a/Engine/Plugins/Runtime/StateTree/Source/StateTreeModule/Public/StateTreeTaskBase.h
+++ b/Engine/Plugins/Runtime/StateTree/Source/StateTreeModule/Public/StateTreeTaskBase.h
@@ -2,8 +2,8 @@
 
 #pragma once
 
-#include "CoreMinimal.h"
 #include "StateTreeTypes.h"
+#include "StateTreeNodeBase.h"
 #include "StateTreeTaskBase.generated.h"
 
 struct FStateTreeExecutionContext;
@@ -13,26 +13,10 @@ struct FStateTreeExecutionContext;
  * Tasks are logic executed in an active state.
  */
 USTRUCT()
-struct STATETREEMODULE_API FStateTreeTaskBase
+struct STATETREEMODULE_API FStateTreeTaskBase : public FStateTreeNodeBase
 {
 	GENERATED_BODY()
 
-	FStateTreeTaskBase() = default;
-
-	virtual ~FStateTreeTaskBase() {}
-
-	/** @return Struct that represents the runtime data of the task. */
-	virtual const UStruct* GetInstanceDataType() const { return nullptr; };
-
-	/**
-	 * Called when the StateTree asset is linked. Allows to resolve references to other StateTree data.
-	 * @see TStateTreeExternalDataHandle
-	 * @see TStateTreeInstanceDataPropertyHandle
-	 * @param Linker Reference to the linker
-	 * @return true if linking succeeded. 
-	 */
-	virtual bool Link(FStateTreeLinker& Linker) { return true; }
-
 	/**
 	 * Called when a new state is entered and task is part of active states. The change type parameter describes if the task's state
 	 * was previously part of the list of active states (Sustained), or if it just became active (Changed).
@@ -72,21 +56,4 @@ struct STATETREEMODULE_API FStateTreeTaskBase
 #if WITH_GAMEPLAY_DEBUGGER
 	virtual void AppendDebugInfoString(FString& DebugString, const FStateTreeExecutionContext& Context) const;
 #endif
-
-	UPROPERTY(EditAnywhere, Category = Task, meta = (EditCondition = "false", EditConditionHides))
-	FName Name;
-
-	/** Property binding copy batch handle. */
-	UPROPERTY()
-	FStateTreeHandle BindingsBatch = FStateTreeHandle::Invalid;
-
-	/** The runtime data's data view index in the StateTreeExecutionContext, and source struct index in property binding. */
-	UPROPERTY()
-	uint16 DataViewIndex = 0;
-
-	UPROPERTY()
-	uint16 InstanceIndex = 0;
-
-	UPROPERTY()
-	uint8 bInstanceIsObject : 1;
 };
