Fix non-unity build and static analysis errors in StateTree object hash

diff --git a/Engine/Plugins/Runtime/StateTree/Source/StateTreeEditorModule/Private/StateTreeObjectHash.cpp b/Engine/Plugins/Runtime/StateTree/Source/StateTreeEditorModule/Private/StateTreeObjectHash.cpp
--- a/Engine/Plugins/Runtime/StateTree/Source/StateTreeEditorModule/Private/StateTreeObjectHash.cpp
+++ b/Engine/Plugins/Runtime/StateTree/Source/StateTreeEditorModule/Private/StateTreeObjectHash.cpp
@@ -1,13 +1,19 @@
 ﻿// Copyright Epic Games, Inc. All Rights Reserved.
 
 #include "StateTreeObjectHash.h"
+#include "UObject/UnrealType.h"
 
 #if WITH_EDITORONLY_DATA
 
 bool FStateTreeObjectCRC32::ShouldSkipProperty(const FProperty* InProperty) const
 {
+	if (InProperty == nullptr)
+	{
+		return false;
+	}
+	
 	static const FName ExcludeFromHashName(TEXT("ExcludeFromHash"));
-	const bool bExclude = InProperty != nullptr && InProperty->HasMetaData(ExcludeFromHashName);
+	const bool bExclude = InProperty->HasMetaData(ExcludeFromHashName);
 	
 	return FArchiveObjectCrc32::ShouldSkipProperty(InProperty) || InProperty->HasAllPropertyFlags(CPF_Transient) || bExclude;
 }
