Fix LOCTEXT conflicts
#rnx

diff --git a/Engine/Plugins/AI/MassAI/Source/MassAIBehavior/Private/Conditions/ZoneGraphTagConditions.cpp b/Engine/Plugins/AI/MassAI/Source/MassAIBehavior/Private/Conditions/ZoneGraphTagConditions.cpp
--- a/Engine/Plugins/AI/MassAI/Source/MassAIBehavior/Private/Conditions/ZoneGraphTagConditions.cpp
+++ b/Engine/Plugins/AI/MassAI/Source/MassAIBehavior/Private/Conditions/ZoneGraphTagConditions.cpp
@@ -122,19 +122,19 @@ FText FZoneGraphTagFilterCondition::GetDescription(const FGuid& ID, FStateTreeDa
 	int32 PartIndex = 0;
 	if (Filter.AnyTags != FZoneGraphTagMask::None)
 	{
-		FilterParts[PartIndex++] = LOCTEXT("ContainsAny", "any");
+		FilterParts[PartIndex++] = UE::MassBehavior::ZoneGraph::GetMaskOperatorText(EZoneLaneTagMaskComparison::Any);
 		FilterParts[PartIndex++] = UE::MassBehavior::ZoneGraph::GetTagMaskName(Filter.AnyTags);
 	}
 
 	if (Filter.AllTags != FZoneGraphTagMask::None)
 	{
-		FilterParts[PartIndex++] = LOCTEXT("ContainsAll", "all");
+		FilterParts[PartIndex++] = UE::MassBehavior::ZoneGraph::GetMaskOperatorText(EZoneLaneTagMaskComparison::All);
 		FilterParts[PartIndex++] = UE::MassBehavior::ZoneGraph::GetTagMaskName(Filter.AllTags);
 	}
 
 	if (Filter.NotTags != FZoneGraphTagMask::None)
 	{
-		FilterParts[PartIndex++] = LOCTEXT("ContainsNot", "not");
+		FilterParts[PartIndex++] = UE::MassBehavior::ZoneGraph::GetMaskOperatorText(EZoneLaneTagMaskComparison::Not);
 		FilterParts[PartIndex++] = UE::MassBehavior::ZoneGraph::GetTagMaskName(Filter.NotTags);
 	}
 
diff --git a/Engine/Plugins/Runtime/MassEntity/Source/MassEntityEditor/Private/MassEntityEditor.cpp b/Engine/Plugins/Runtime/MassEntity/Source/MassEntityEditor/Private/MassEntityEditor.cpp
--- a/Engine/Plugins/Runtime/MassEntity/Source/MassEntityEditor/Private/MassEntityEditor.cpp
+++ b/Engine/Plugins/Runtime/MassEntity/Source/MassEntityEditor/Private/MassEntityEditor.cpp
@@ -106,7 +106,7 @@ TSharedRef<SDockTab> FMassEntityEditor::SpawnTab_AssetDetails(const FSpawnTabArg
 	AssetDetailsView->OnFinishedChangingProperties().AddSP(this, &FMassEntityEditor::OnAssetFinishedChangingProperties);
 
 	TSharedRef<SDockTab> SpawnedTab = SNew(SDockTab)
-		.Label(NSLOCTEXT("MassEntityEditor", "AssetDetailsTab", "Mass"))
+		.Label(NSLOCTEXT("MassEntityEditor", "AssetDetailsTabLabel", "Mass"))
 		[
 			AssetDetailsView.ToSharedRef()
 		];
diff --git a/Engine/Plugins/Runtime/StateTree/Source/StateTreeEditorModule/Private/Customizations/StateTreeEditorNodeDetails.cpp b/Engine/Plugins/Runtime/StateTree/Source/StateTreeEditorModule/Private/Customizations/StateTreeEditorNodeDetails.cpp
--- a/Engine/Plugins/Runtime/StateTree/Source/StateTreeEditorModule/Private/Customizations/StateTreeEditorNodeDetails.cpp
+++ b/Engine/Plugins/Runtime/StateTree/Source/StateTreeEditorModule/Private/Customizations/StateTreeEditorNodeDetails.cpp
@@ -450,7 +450,7 @@ FText FStateTreeEditorNodeDetails::GetDescription() const
 	// Multiple descriptions do not make sense, just if only one node is selected.
 	if (RawNodeData.Num() != 1)
 	{
-		return LOCTEXT("MultipleSelected", "<Details.Subdued>Multiple selected</>");
+		return LOCTEXT("MultipleSelectedFormatted", "<Details.Subdued>Multiple Selected</>");
 	}
 
 	FStateTreeEditorNode* Node = static_cast<FStateTreeEditorNode*>(RawNodeData[0]);
@@ -517,7 +517,7 @@ FText FStateTreeEditorNodeDetails::GetName() const
 		return LOCTEXT("Empty", "Empty");
 	}
 
-	return LOCTEXT("MultipleSelected", "Multiple selected");
+	return LOCTEXT("MultipleSelected", "Multiple Selected");
 }
 
 EVisibility FStateTreeEditorNodeDetails::IsNameVisible() const
diff --git a/Engine/Plugins/Runtime/StateTree/Source/StateTreeEditorModule/Private/Customizations/StateTreeStateLinkDetails.cpp b/Engine/Plugins/Runtime/StateTree/Source/StateTreeEditorModule/Private/Customizations/StateTreeStateLinkDetails.cpp
--- a/Engine/Plugins/Runtime/StateTree/Source/StateTreeEditorModule/Private/Customizations/StateTreeStateLinkDetails.cpp
+++ b/Engine/Plugins/Runtime/StateTree/Source/StateTreeEditorModule/Private/Customizations/StateTreeStateLinkDetails.cpp
@@ -174,7 +174,7 @@ FText FStateTreeStateLinkDetails::GetCurrentStateDesc() const
 
 	if (TransitionType == EStateTreeTransitionType::NotSet)
 	{
-		return LOCTEXT("TransitionNotSet", "Not Set");
+		return LOCTEXT("TransitionNotSet", "None");
 	}
 	else if (TransitionType == EStateTreeTransitionType::NextState)
 	{
diff --git a/Engine/Plugins/Runtime/StateTree/Source/StateTreeEditorModule/Private/Customizations/StateTreeTransitionDetails.cpp b/Engine/Plugins/Runtime/StateTree/Source/StateTreeEditorModule/Private/Customizations/StateTreeTransitionDetails.cpp
--- a/Engine/Plugins/Runtime/StateTree/Source/StateTreeEditorModule/Private/Customizations/StateTreeTransitionDetails.cpp
+++ b/Engine/Plugins/Runtime/StateTree/Source/StateTreeEditorModule/Private/Customizations/StateTreeTransitionDetails.cpp
@@ -136,7 +136,7 @@ FText FStateTreeTransitionDetails::GetDescription() const
 		}
 		else
 		{
-			return LOCTEXT("MultipleSelected", "Multiple selected");
+			return LOCTEXT("MultipleSelected", "Multiple Selected");
 		}
 	}
 	return FText::GetEmpty();
diff --git a/Engine/Plugins/Runtime/StateTree/Source/StateTreeEditorModule/Private/StateTreeEditor.cpp b/Engine/Plugins/Runtime/StateTree/Source/StateTreeEditorModule/Private/StateTreeEditor.cpp
--- a/Engine/Plugins/Runtime/StateTree/Source/StateTreeEditorModule/Private/StateTreeEditor.cpp
+++ b/Engine/Plugins/Runtime/StateTree/Source/StateTreeEditorModule/Private/StateTreeEditor.cpp
@@ -73,10 +73,12 @@ void FStateTreeEditor::RegisterTabSpawners(const TSharedRef<class FTabManager>&
 		.SetDisplayName(NSLOCTEXT("StateTreeEditor", "StateTreeViewTab", "StateTree"))
 		.SetGroup(WorkspaceMenuCategoryRef)
 		.SetIcon(FSlateIcon(FEditorStyle::GetStyleSetName(), "LevelEditor.Tabs.Outliner"));
+
 	InTabManager->RegisterTabSpawner(StateTreeStatisticsTabId, FOnSpawnTab::CreateSP(this, &FStateTreeEditor::SpawnTab_StateTreeStatistics))
 		.SetDisplayName(NSLOCTEXT("StateTreeEditor", "StatisticsTab", "StateTree Statistics"))
 		.SetGroup(WorkspaceMenuCategoryRef)
 		.SetIcon(FSlateIcon(FEditorStyle::GetStyleSetName(), "LevelEditor.Tabs.Outliner"));
+
 	InTabManager->RegisterTabSpawner(CompilerResultsTabId, FOnSpawnTab::CreateSP(this, &FStateTreeEditor::SpawnTab_CompilerResults))
 		.SetDisplayName(NSLOCTEXT("StateTreeEditor", "CompilerResultsTab", "Compiler Results"))
 		.SetGroup(WorkspaceMenuCategoryRef)
@@ -232,7 +234,7 @@ TSharedRef<SDockTab> FStateTreeEditor::SpawnTab_StateTreeView(const FSpawnTabArg
 	check(Args.GetTabId() == StateTreeViewTabId);
 
 	return SNew(SDockTab)
-		.Label(NSLOCTEXT("StateTreeEditor", "StateTreeViewTab", "State Tree"))
+		.Label(NSLOCTEXT("StateTreeEditor", "StateTreeViewTab", "StateTree"))
 		.TabColorScale(GetTabColorScale())
 		[
 			SAssignNew(StateTreeView, SStateTreeView, StateTreeViewModel.ToSharedRef())
@@ -276,7 +278,7 @@ TSharedRef<SDockTab> FStateTreeEditor::SpawnTab_AssetDetails(const FSpawnTabArgs
 	AssetDetailsView->OnFinishedChangingProperties().AddSP(this, &FStateTreeEditor::OnAssetFinishedChangingProperties);
 
 	TSharedRef<SDockTab> SpawnedTab = SNew(SDockTab)
-		.Label(NSLOCTEXT("StateTreeEditor", "AssetDetailsTab", "StateTree"))
+		.Label(NSLOCTEXT("StateTreeEditor", "AssetDetailsTabLabel", "StateTree"))
 		[
 			AssetDetailsView.ToSharedRef()
 		];
diff --git a/Engine/Plugins/Runtime/StateTree/Source/StateTreeModule/Private/Conditions/GameplayTagConditions.cpp b/Engine/Plugins/Runtime/StateTree/Source/StateTreeModule/Private/Conditions/GameplayTagConditions.cpp
--- a/Engine/Plugins/Runtime/StateTree/Source/StateTreeModule/Private/Conditions/GameplayTagConditions.cpp
+++ b/Engine/Plugins/Runtime/StateTree/Source/StateTreeModule/Private/Conditions/GameplayTagConditions.cpp
@@ -172,7 +172,7 @@ FText FGameplayTagContainerMatchCondition::GetDescription(const FGuid& ID, FStat
 		OtherContainerText = UE::StateTree::Conditions::GetContainerAsText(Instance.TagContainer);
 	}
 
-	return FText::Format(LOCTEXT("GameplayTagMatchDesc", "{0} <Details.Bold>{1}</> {2} <Details.Bold>{3}</>"),
+	return FText::Format(LOCTEXT("GameplayTagContainerMatchDesc", "{0} <Details.Bold>{1}</> {2} <Details.Bold>{3}</>"),
 		InvertText, TagContainerText, MatchTypeText, OtherContainerText);
 }
 #endif// WITH_EDITOR
@@ -216,7 +216,7 @@ FText FGameplayTagQueryCondition::GetDescription(const FGuid& ID, FStateTreeData
 		TagContainerText = LOCTEXT("NotBound", "Not Bound");
 	}
 
-	return FText::Format(LOCTEXT("GameplayTagMatchDesc", "{0} <Details.Bold>{1}</> matches query</>"),
+	return FText::Format(LOCTEXT("GameplayTagQueryMatchDesc", "{0} <Details.Bold>{1}</> matches query</>"),
 		InvertText, TagContainerText);
 }
 #endif// WITH_EDITOR
diff --git a/Engine/Plugins/Runtime/ZoneGraph/Source/ZoneGraphEditor/Private/ZoneShapeComponentDetails.cpp b/Engine/Plugins/Runtime/ZoneGraph/Source/ZoneGraphEditor/Private/ZoneShapeComponentDetails.cpp
--- a/Engine/Plugins/Runtime/ZoneGraph/Source/ZoneGraphEditor/Private/ZoneShapeComponentDetails.cpp
+++ b/Engine/Plugins/Runtime/ZoneGraph/Source/ZoneGraphEditor/Private/ZoneShapeComponentDetails.cpp
@@ -542,7 +542,7 @@ void FZoneShapePointDetails::OnSetInControlPoint(float NewValue, ETextCommit::Ty
 		return;
 	}
 
-	const FScopedTransaction Transaction(LOCTEXT("SetControlPoint", "Set Zone control point position"));
+	const FScopedTransaction Transaction(LOCTEXT("SetControlPoint", "Set Zone Shape control point position"));
 	ShapeComp->Modify();
 
 	TArray<FZoneShapePoint>& ShapePoints = ShapeComp->GetMutablePoints();
