StateTree: Moved instance data out from framents
- Added new struct to hold instance data
- Removed on the fly created customs struct for storing the instance data
- The Mass StateTree instance data is not a fragment anymore
- Fixes crash when using StateTreeBrainComponent directly on actor
- Moved StateTree activation processor settings to Mass Behavior settings

diff --git a/Engine/Plugins/AI/MassAI/Source/MassAIBehavior/MassBehaviorSettings.cpp b/Engine/Plugins/AI/MassAI/Source/MassAIBehavior/MassBehaviorSettings.cpp
new file mode 100644
--- /dev/null
+++ b/Engine/Plugins/AI/MassAI/Source/MassAIBehavior/MassBehaviorSettings.cpp
@@ -0,0 +1,17 @@
+﻿// Copyright Epic Games, Inc. All Rights Reserved.
+
+#include "MassBehaviorSettings.h"
+
+//----------------------------------------------------------------------//
+// UMassBehaviorSettings
+//----------------------------------------------------------------------//
+
+UMassBehaviorSettings::UMassBehaviorSettings(const FObjectInitializer& ObjectInitializer)
+	: Super(ObjectInitializer)
+{
+	// Default values.
+	MaxActivationsPerLOD[EMassLOD::High] = 100;
+	MaxActivationsPerLOD[EMassLOD::Medium] = 100;
+	MaxActivationsPerLOD[EMassLOD::Low] = 100;
+	MaxActivationsPerLOD[EMassLOD::Off] = 100;
+}
diff --git a/Engine/Plugins/AI/MassAI/Source/MassAIBehavior/Private/MassStateTreeProcessors.cpp b/Engine/Plugins/AI/MassAI/Source/MassAIBehavior/Private/MassStateTreeProcessors.cpp
--- a/Engine/Plugins/AI/MassAI/Source/MassAIBehavior/Private/MassStateTreeProcessors.cpp
+++ b/Engine/Plugins/AI/MassAI/Source/MassAIBehavior/Private/MassStateTreeProcessors.cpp
@@ -13,6 +13,7 @@
 #include "MassSignals/Public/MassSignalSubsystem.h"
 #include "ProfilingDebugging/CsvProfiler.h"
 #include "Engine/World.h"
+#include "MassBehaviorSettings.h"
 
 CSV_DEFINE_CATEGORY(StateTreeProcessor, true);
 
@@ -100,18 +101,17 @@ bool SetExternalSubsystems(FMassStateTreeExecutionContext& Context)
 	return bFoundAllSubsystems;
 }
 
-void ForEachEntityInChunk(
-	FMassStateTreeExecutionContext& StateTreeContext,
-	UMassStateTreeSubsystem& MassStateTreeSubsystem,
-	const TFunctionRef<void(FMassStateTreeExecutionContext&, FStateTreeDataView)> ForEachEntityCallback)
+template<typename TFunc>
+void ForEachEntityInChunk(FMassExecutionContext& Context, FMassStateTreeExecutionContext& StateTreeContext, UMassStateTreeSubsystem& MassStateTreeSubsystem, TFunc&& Callback)
 {
-	const FMassExecutionContext& Context = StateTreeContext.GetEntitySubsystemExecutionContext();
-	const TConstArrayView<FMassStateTreeFragment> StateTreeList = Context.GetFragmentView<FMassStateTreeFragment>();
+	const TArrayView<FMassStateTreeInstanceFragment> StateTreeInstanceList = Context.GetMutableFragmentView<FMassStateTreeInstanceFragment>();
+	const FMassStateTreeSharedFragment& SharedStateTree = Context.GetConstSharedFragment<FMassStateTreeSharedFragment>();
 
 	// Assuming that all the entities share same StateTree, because they all should have the same storage fragment.
 	const int32 NumEntities = Context.GetNumEntities();
 	check(NumEntities > 0);
-	UStateTree* StateTree = MassStateTreeSubsystem.GetRegisteredStateTreeAsset(StateTreeList[0].StateTreeHandle);
+	
+	UStateTree* StateTree = SharedStateTree.StateTree;
 
 	// Initialize the execution context if changed between chunks.
 	if (StateTreeContext.GetStateTree() != StateTree)
@@ -131,12 +131,12 @@ void ForEachEntityInChunk(
 		}
 	}
 
-	const UScriptStruct* StorageScriptStruct = StateTree->GetInstanceStorageStruct();
-	for (int32 i = 0; i < NumEntities; ++i)
+	for (int32 EntityIndex = 0; EntityIndex < NumEntities; EntityIndex++)
 	{
-		const FMassEntityHandle Entity = Context.GetEntity(i);
+		const FMassEntityHandle Entity = Context.GetEntity(EntityIndex);
+		FMassStateTreeInstanceFragment& StateTreeInstance = StateTreeInstanceList[EntityIndex];
+		
 		StateTreeContext.SetEntity(Entity);
-		StateTreeContext.SetEntityIndex(i);
 
 		// Gather all required fragments.
 		{
@@ -157,7 +157,11 @@ void ForEachEntityInChunk(
 			}
 		}
 
-		ForEachEntityCallback(StateTreeContext, StateTreeContext.GetEntitySubsystem().GetFragmentDataStruct(Entity, StorageScriptStruct));
+		FStateTreeInstanceData* InstanceData = MassStateTreeSubsystem.GetInstanceData(StateTreeInstance.InstanceHandle);
+		if (InstanceData != nullptr)
+		{
+			Callback(StateTreeContext, StateTreeInstance, *InstanceData);
+		}
 	}
 }
 
@@ -170,7 +174,7 @@ void ForEachEntityInChunk(
 UMassStateTreeFragmentDestructor::UMassStateTreeFragmentDestructor()
 {
 	ExecutionFlags = (int32)(EProcessorExecutionFlags::Standalone | EProcessorExecutionFlags::Server);
-	ObservedType = FMassStateTreeFragment::StaticStruct();
+	ObservedType = FMassStateTreeInstanceFragment::StaticStruct();
 	Operation = EMassObservedOperation::Remove;
 }
 
@@ -181,7 +185,8 @@ void UMassStateTreeFragmentDestructor::Initialize(UObject& Owner)
 
 void UMassStateTreeFragmentDestructor::ConfigureQueries()
 {
-	EntityQuery.AddRequirement<FMassStateTreeFragment>(EMassFragmentAccess::ReadOnly);
+	EntityQuery.AddRequirement<FMassStateTreeInstanceFragment>(EMassFragmentAccess::ReadWrite);
+	EntityQuery.AddConstSharedRequirement<FMassStateTreeSharedFragment>();
 }
 
 void UMassStateTreeFragmentDestructor::Execute(UMassEntitySubsystem& EntitySubsystem, FMassExecutionContext& Context)
@@ -194,19 +199,31 @@ void UMassStateTreeFragmentDestructor::Execute(UMassEntitySubsystem& EntitySubsy
 	FMassStateTreeExecutionContext StateTreeContext(EntitySubsystem, *SignalSubsystem, Context);
 	UMassStateTreeSubsystem* MassStateTreeSubsystem = UWorld::GetSubsystem<UMassStateTreeSubsystem>(EntitySubsystem.GetWorld());
 
-	EntityQuery.ForEachEntityChunk(
-		EntitySubsystem,
-		Context,
-		[this, &StateTreeContext, &MassStateTreeSubsystem](FMassExecutionContext&)
+	EntityQuery.ForEachEntityChunk(EntitySubsystem,Context,
+		[this, &StateTreeContext, MassStateTreeSubsystem](FMassExecutionContext& Context)
 		{
-			UE::MassBehavior::ForEachEntityInChunk(
-				StateTreeContext,
-				*MassStateTreeSubsystem,
-				[](FMassStateTreeExecutionContext& StateTreeExecutionContext, FStateTreeDataView Storage)
+			const TArrayView<FMassStateTreeInstanceFragment> StateTreeInstanceList = Context.GetMutableFragmentView<FMassStateTreeInstanceFragment>();
+			const FMassStateTreeSharedFragment& SharedStateTree = Context.GetConstSharedFragment<FMassStateTreeSharedFragment>();
+
+			const int32 NumEntities = Context.GetNumEntities();
+	
+			UE::MassBehavior::ForEachEntityInChunk(Context, StateTreeContext, *MassStateTreeSubsystem,
+				[](FMassStateTreeExecutionContext& StateTreeExecutionContext, FMassStateTreeInstanceFragment& StateTreeInstance, FStateTreeInstanceData& InstanceData)
 				{
 					// Stop the tree instance
-					StateTreeExecutionContext.Stop(Storage);
+					StateTreeExecutionContext.Stop(&InstanceData);
 				});
+
+			// Free the StateTree instance memory
+			for (int32 EntityIndex = 0; EntityIndex < NumEntities; EntityIndex++)
+			{
+				FMassStateTreeInstanceFragment& StateTreeInstance = StateTreeInstanceList[EntityIndex];
+				if (StateTreeInstance.InstanceHandle.IsValid())
+				{
+					MassStateTreeSubsystem->FreeInstanceData(StateTreeInstance.InstanceHandle);
+					StateTreeInstance.InstanceHandle = FMassStateTreeInstanceHandle();
+				}
+			}
 		});
 }
 
@@ -217,10 +234,6 @@ UMassStateTreeActivationProcessor::UMassStateTreeActivationProcessor()
 {
 	ExecutionOrder.ExecuteAfter.Add(UE::Mass::ProcessorGroupNames::LOD);
 	ExecutionOrder.ExecuteBefore.Add(UE::Mass::ProcessorGroupNames::Behavior);
-	MaxActivationsPerLOD[EMassLOD::High] = 100;
-	MaxActivationsPerLOD[EMassLOD::Medium] = 100;
-	MaxActivationsPerLOD[EMassLOD::Low] = 100;
-	MaxActivationsPerLOD[EMassLOD::Off] = 100;
 }
 
 void UMassStateTreeActivationProcessor::Initialize(UObject& Owner)
@@ -230,7 +243,8 @@ void UMassStateTreeActivationProcessor::Initialize(UObject& Owner)
 
 void UMassStateTreeActivationProcessor::ConfigureQueries()
 {
-	EntityQuery.AddRequirement<FMassStateTreeFragment>(EMassFragmentAccess::ReadOnly);
+	EntityQuery.AddRequirement<FMassStateTreeInstanceFragment>(EMassFragmentAccess::ReadWrite);
+	EntityQuery.AddConstSharedRequirement<FMassStateTreeSharedFragment>();
 	EntityQuery.AddTagRequirement<FMassStateTreeActivatedTag>(EMassFragmentPresence::None);
 	EntityQuery.AddChunkRequirement<FMassSimulationVariableTickChunkFragment>(EMassFragmentAccess::ReadOnly, EMassFragmentPresence::Optional);
 }
@@ -242,45 +256,66 @@ void UMassStateTreeActivationProcessor::Execute(UMassEntitySubsystem& EntitySubs
 		return;
 	}
 
+	const UMassBehaviorSettings* BehaviorSettings = GetDefault<UMassBehaviorSettings>();
+	check(BehaviorSettings);
+
 	// StateTree processor relies on signals to be ticked but we need an 'initial tick' to set the tree in the proper state.
 	// The initializer provides that by sending a signal to all new entities that use StateTree.
 
 	FMassStateTreeExecutionContext StateTreeContext(EntitySubsystem, *SignalSubsystem, Context);
 	UMassStateTreeSubsystem* MassStateTreeSubsystem = UWorld::GetSubsystem<UMassStateTreeSubsystem>(EntitySubsystem.GetWorld());
 
+	const float TimeInSeconds = EntitySubsystem.GetWorld()->GetTimeSeconds();
+
 	TArray<FMassEntityHandle> EntitiesToSignal;
 	int32 ActivationCounts[EMassLOD::Max] {0,0,0,0};
-	EntityQuery.ForEachEntityChunk(EntitySubsystem, Context, [&EntitiesToSignal, &ActivationCounts, MaxActivationsPerLOD = MaxActivationsPerLOD, &StateTreeContext, MassStateTreeSubsystem](FMassExecutionContext& Context)
-	{
-		const int32 NumEntities = Context.GetNumEntities();
-		// Check if we already reached the maximum for this frame
-		const EMassLOD::Type ChunkLOD = FMassSimulationVariableTickChunkFragment::GetChunkLOD(Context);
-		if (ActivationCounts[ChunkLOD] > MaxActivationsPerLOD[ChunkLOD])
+	
+	EntityQuery.ForEachEntityChunk(EntitySubsystem, Context,
+		[&EntitiesToSignal, &ActivationCounts, MaxActivationsPerLOD = BehaviorSettings->MaxActivationsPerLOD, &StateTreeContext, TimeInSeconds, MassStateTreeSubsystem](FMassExecutionContext& Context)
 		{
-			return;
-		}
-		ActivationCounts[ChunkLOD] += NumEntities;
+			const int32 NumEntities = Context.GetNumEntities();
+			// Check if we already reached the maximum for this frame
+			const EMassLOD::Type ChunkLOD = FMassSimulationVariableTickChunkFragment::GetChunkLOD(Context);
+			if (ActivationCounts[ChunkLOD] > MaxActivationsPerLOD[ChunkLOD])
+			{
+				return;
+			}
+			ActivationCounts[ChunkLOD] += NumEntities;
+
+			const TArrayView<FMassStateTreeInstanceFragment> StateTreeInstanceList = Context.GetMutableFragmentView<FMassStateTreeInstanceFragment>();
+			const FMassStateTreeSharedFragment& SharedStateTree = Context.GetConstSharedFragment<FMassStateTreeSharedFragment>();
 
-		// Start StateTree. This may do substantial amount of work, as we select and enter the first state.
-		UE::MassBehavior::ForEachEntityInChunk(
-			StateTreeContext,
-			*MassStateTreeSubsystem,
-			[MassStateTreeSubsystem](FMassStateTreeExecutionContext& StateTreeExecutionContext, FStateTreeDataView Storage)
+			// Allocate and initialize the StateTree instance memory
+			for (int32 EntityIndex = 0; EntityIndex < NumEntities; EntityIndex++)
 			{
-				// Init object instances
-				StateTreeExecutionContext.InitInstanceData(*MassStateTreeSubsystem, Storage);
-				// Start the tree instance
-				StateTreeExecutionContext.Start(Storage);
-			});
-
-		// Append all entities of the current chunk to the consolidated list to send signal once
-		EntitiesToSignal.Append(Context.GetEntities().GetData(), Context.GetEntities().Num());
-		// Adding a tag on each entities to remember we have sent the state tree initialization signal
-		for (int32 i = 0; i < NumEntities; ++i)
-		{
-			Context.Defer().AddTag<FMassStateTreeActivatedTag>(Context.GetEntity(i));
-		}
-	});
+				FMassStateTreeInstanceFragment& StateTreeInstance = StateTreeInstanceList[EntityIndex];
+				StateTreeInstance.InstanceHandle = MassStateTreeSubsystem->AllocateInstanceData(SharedStateTree.StateTree);
+			}
+			
+			// Start StateTree. This may do substantial amount of work, as we select and enter the first state.
+			UE::MassBehavior::ForEachEntityInChunk(Context, StateTreeContext,*MassStateTreeSubsystem,
+				[TimeInSeconds, MassStateTreeSubsystem]
+				(FMassStateTreeExecutionContext& StateTreeExecutionContext, FMassStateTreeInstanceFragment& StateTreeInstance, FStateTreeInstanceData& InstanceData)
+				{
+					// Start the tree instance
+					StateTreeExecutionContext.Start(&InstanceData);
+					StateTreeInstance.LastUpdateTimeInSeconds = TimeInSeconds;
+				});
+
+			// Adding a tag on each entities to remember we have sent the state tree initialization signal
+			EntitiesToSignal.Reserve(EntitiesToSignal.Num() + NumEntities);
+			for (int32 EntityIndex = 0; EntityIndex < NumEntities; EntityIndex++)
+			{
+				const FMassStateTreeInstanceFragment& StateTreeInstance = StateTreeInstanceList[EntityIndex];
+				if (StateTreeInstance.InstanceHandle.IsValid())
+				{
+					const FMassEntityHandle Entity = Context.GetEntity(EntityIndex);
+					Context.Defer().AddTag<FMassStateTreeActivatedTag>(Entity);
+					EntitiesToSignal.Add(Entity);
+				}
+			}
+		});
+	
 	// Signal all entities inside the consolidated list
 	if (EntitiesToSignal.Num())
 	{
@@ -338,7 +373,8 @@ void UMassStateTreeProcessor::Initialize(UObject& Owner)
 
 void UMassStateTreeProcessor::ConfigureQueries()
 {
-	EntityQuery.AddRequirement<FMassStateTreeFragment>(EMassFragmentAccess::ReadWrite);
+	EntityQuery.AddRequirement<FMassStateTreeInstanceFragment>(EMassFragmentAccess::ReadWrite);
+	EntityQuery.AddConstSharedRequirement<FMassStateTreeSharedFragment>();
 }
 
 void UMassStateTreeProcessor::SignalEntities(UMassEntitySubsystem& EntitySubsystem, FMassExecutionContext& Context, FMassSignalNameLookup& EntitySignals)
@@ -350,43 +386,35 @@ void UMassStateTreeProcessor::SignalEntities(UMassEntitySubsystem& EntitySubsyst
 	QUICK_SCOPE_CYCLE_COUNTER(StateTreeProcessor_Run);
 	CSV_SCOPED_TIMING_STAT_EXCLUSIVE(StateTreeProcessorExecute);
 
-	const float TimeDelta = Context.GetDeltaTimeSeconds();
-	FMassStateTreeExecutionContext StateTreeContext(EntitySubsystem, *SignalSubsystem, Context);
 	const float TimeInSeconds = EntitySubsystem.GetWorld()->GetTimeSeconds();
+	FMassStateTreeExecutionContext StateTreeContext(EntitySubsystem, *SignalSubsystem, Context);
 
 	TArray<FMassEntityHandle> EntitiesToSignal;
 
-	EntityQuery.ForEachEntityChunk(
-		EntitySubsystem,
-		Context,
-		[this, &StateTreeContext, TimeDelta, TimeInSeconds, &EntitiesToSignal](FMassExecutionContext& Context)
+	EntityQuery.ForEachEntityChunk(EntitySubsystem, Context,
+		[this, &StateTreeContext, TimeInSeconds, &EntitiesToSignal](FMassExecutionContext& Context)
 		{
 			// Keep stats regarding the amount of tree instances ticked per frame
 			CSV_CUSTOM_STAT(StateTreeProcessor, NumTickedStateTree, Context.GetNumEntities(), ECsvCustomStatOp::Accumulate);
 
-			TArrayView<FMassStateTreeFragment> StateTreeList = Context.GetMutableFragmentView<FMassStateTreeFragment>();
-
-			UE::MassBehavior::ForEachEntityInChunk(
-				StateTreeContext,
-				*MassStateTreeSubsystem,
-				[&StateTreeList, TimeDelta, TimeInSeconds, &EntitiesToSignal](FMassStateTreeExecutionContext& StateTreeExecutionContext, const FStateTreeDataView Storage)
+			UE::MassBehavior::ForEachEntityInChunk(Context,StateTreeContext,*MassStateTreeSubsystem,
+				[TimeInSeconds, &EntitiesToSignal](FMassStateTreeExecutionContext& StateTreeExecutionContext, FMassStateTreeInstanceFragment& StateTreeInstance, FStateTreeInstanceData& InstanceData)
 				{
 					// Compute adjusted delta time
-					TOptional<float>& LastUpdate = StateTreeList[StateTreeExecutionContext.GetEntityIndex()].LastUpdateTimeInSeconds;
-					const float AdjustedTimeDelta = LastUpdate.IsSet() ? TimeDelta + (TimeInSeconds - LastUpdate.GetValue()) : TimeDelta;
-					LastUpdate = TimeInSeconds;
+					const float AdjustedDeltaTime = TimeInSeconds - StateTreeInstance.LastUpdateTimeInSeconds;
+					StateTreeInstance.LastUpdateTimeInSeconds = TimeInSeconds;
 
 					// Tick the tree instance
-					StateTreeExecutionContext.Tick(AdjustedTimeDelta, Storage);
+					StateTreeExecutionContext.Tick(AdjustedDeltaTime, &InstanceData);
 
 					// When last tick status is different than "Running", the state tree need to be tick again
 					// For performance reason, tick again to see if we could find a new state right away instead of waiting the next frame.
-					if (StateTreeExecutionContext.GetLastTickStatus(Storage) != EStateTreeRunStatus::Running)
+					if (StateTreeExecutionContext.GetLastTickStatus(&InstanceData) != EStateTreeRunStatus::Running)
 					{
-						StateTreeExecutionContext.Tick(0.0f, Storage);
+						StateTreeExecutionContext.Tick(0.0f, &InstanceData);
 
 						// Could not find new state yet, try again next frame
-						if (StateTreeExecutionContext.GetLastTickStatus(Storage) != EStateTreeRunStatus::Running)
+						if (StateTreeExecutionContext.GetLastTickStatus(&InstanceData) != EStateTreeRunStatus::Running)
 						{
 							EntitiesToSignal.Add(StateTreeExecutionContext.GetEntity());
 						}
diff --git a/Engine/Plugins/AI/MassAI/Source/MassAIBehavior/Private/MassStateTreeSubsystem.cpp b/Engine/Plugins/AI/MassAI/Source/MassAIBehavior/Private/MassStateTreeSubsystem.cpp
--- a/Engine/Plugins/AI/MassAI/Source/MassAIBehavior/Private/MassStateTreeSubsystem.cpp
+++ b/Engine/Plugins/AI/MassAI/Source/MassAIBehavior/Private/MassStateTreeSubsystem.cpp
@@ -2,26 +2,50 @@
 
 #include "MassStateTreeSubsystem.h"
 #include "StateTree.h"
+#include "StateTreeExecutionContext.h"
 #include "Engine/Engine.h"
 
 
 void UMassStateTreeSubsystem::Initialize(FSubsystemCollectionBase& Collection)
 {
 	Super::Initialize(Collection);
-	
 }
 
-FMassStateTreeHandle UMassStateTreeSubsystem::RegisterStateTreeAsset(UStateTree* StateTree)
+FMassStateTreeInstanceHandle UMassStateTreeSubsystem::AllocateInstanceData(const UStateTree* StateTree)
 {
-	// Return already registered asset if found.
-	int32 Index = RegisteredStateTrees.IndexOfByPredicate([StateTree](UStateTree* ExistingStateTree) -> bool { return ExistingStateTree == StateTree; });
-	if (Index != INDEX_NONE)
+	if (StateTree == nullptr)
+	{
+		return FMassStateTreeInstanceHandle();
+	}
+	
+	int32 Index = 0;
+	if (InstanceDataFreelist.Num() > 0)
 	{
-		return FMassStateTreeHandle::Make(uint16(Index));
+		Index = InstanceDataFreelist.Pop();
 	}
-	// Add new, check that it fits the StateTree handle.
-	Index = RegisteredStateTrees.Add(StateTree);
-	check(Index < int32(MAX_uint16));
-	return FMassStateTreeHandle::Make(uint16(Index));
+	else
+	{
+		Index = InstanceDataArray.Num();
+		InstanceDataArray.AddDefaulted();
+	}
+
+	FMassStateTreeInstanceDataItem& Item = InstanceDataArray[Index];
+	
+	FStateTreeExecutionContext::InitInstanceData(*this, *StateTree, Item.InstanceData);
+	
+	return FMassStateTreeInstanceHandle::Make(Index, Item.Generation);
 }
 
+void UMassStateTreeSubsystem::FreeInstanceData(const FMassStateTreeInstanceHandle Handle)
+{
+	if (!IsValidHandle(Handle))
+	{
+		return;
+	}
+
+	FMassStateTreeInstanceDataItem& Item = InstanceDataArray[Handle.GetIndex()];
+	Item.InstanceData.Reset();
+	Item.Generation++;
+
+	InstanceDataFreelist.Add(Handle.GetIndex());
+}
diff --git a/Engine/Plugins/AI/MassAI/Source/MassAIBehavior/Private/MassStateTreeTrait.cpp b/Engine/Plugins/AI/MassAI/Source/MassAIBehavior/Private/MassStateTreeTrait.cpp
--- a/Engine/Plugins/AI/MassAI/Source/MassAIBehavior/Private/MassStateTreeTrait.cpp
+++ b/Engine/Plugins/AI/MassAI/Source/MassAIBehavior/Private/MassStateTreeTrait.cpp
@@ -11,6 +11,9 @@
 
 void UMassStateTreeTrait::BuildTemplate(FMassEntityTemplateBuildContext& BuildContext, UWorld& World) const
 {
+	UMassEntitySubsystem* EntitySubsystem = UWorld::GetSubsystem<UMassEntitySubsystem>(&World);
+	check(EntitySubsystem);
+
 	UMassStateTreeSubsystem* MassStateTreeSubsystem = World.GetSubsystem<UMassStateTreeSubsystem>();
 	if (!MassStateTreeSubsystem)
 	{
@@ -23,20 +26,19 @@ void UMassStateTreeTrait::BuildTemplate(FMassEntityTemplateBuildContext& BuildCo
 		UE_VLOG(MassStateTreeSubsystem, LogMassBehavior, Error, TEXT("StateTree asset is not set or unavailable."));
 		return;
 	}
-	if (!StateTree->GetInstanceStorageDefaultValue().IsValid())
+	if (!StateTree->IsReadyToRun())
 	{
-		UE_VLOG(MassStateTreeSubsystem, LogMassBehavior, Error, TEXT("StateTree asset is valid but missing runtime storage type."));
+		UE_VLOG(MassStateTreeSubsystem, LogMassBehavior, Error, TEXT("StateTree asset is ready to run."));
 		return;
 	}
 
-	const FMassStateTreeHandle Handle = MassStateTreeSubsystem->RegisterStateTreeAsset(StateTree);
-
-	// Add fragment describing which StateTree to run
-	FMassStateTreeFragment& StateTreeFragment = BuildContext.AddFragment_GetRef<FMassStateTreeFragment>();
-	StateTreeFragment.StateTreeHandle = Handle;
+	FMassStateTreeSharedFragment SharedStateTree;
+	SharedStateTree.StateTree = StateTree;
+	
+	const FConstSharedStruct StateTreeFragment = EntitySubsystem->GetOrCreateConstSharedFragment(UE::StructUtils::GetStructCrc32(FConstStructView::Make(SharedStateTree)), SharedStateTree);
+	BuildContext.AddConstSharedFragment(StateTreeFragment);
 
-	// Add runtime storage as a fragment
-	BuildContext.AddFragment(StateTree->GetInstanceStorageDefaultValue());
+	BuildContext.AddFragment<FMassStateTreeInstanceFragment>();
 }
 
 void UMassStateTreeTrait::ValidateTemplate(FMassEntityTemplateBuildContext& BuildContext, UWorld& World) const
diff --git a/Engine/Plugins/AI/MassAI/Source/MassAIBehavior/Public/MassBehaviorSettings.h b/Engine/Plugins/AI/MassAI/Source/MassAIBehavior/Public/MassBehaviorSettings.h
new file mode 100644
--- /dev/null
+++ b/Engine/Plugins/AI/MassAI/Source/MassAIBehavior/Public/MassBehaviorSettings.h
@@ -0,0 +1,19 @@
+﻿// Copyright Epic Games, Inc. All Rights Reserved.
+
+#pragma once
+
+#include "MassSettings.h"
+#include "MassLODTypes.h"
+#include "MassBehaviorSettings.generated.h"
+
+UCLASS(config = Mass, defaultconfig, meta = (DisplayName = "Mass Behavior"))
+class MASSAIBEHAVIOR_API UMassBehaviorSettings : public UMassModuleSettings
+{
+	GENERATED_BODY()
+
+public:
+	UMassBehaviorSettings(const FObjectInitializer& ObjectInitializer = FObjectInitializer::Get());
+	
+	UPROPERTY(EditAnywhere, Category = "Mass|LOD", config)
+	int32 MaxActivationsPerLOD[EMassLOD::Max];
+};
diff --git a/Engine/Plugins/AI/MassAI/Source/MassAIBehavior/Public/MassStateTreeExecutionContext.h b/Engine/Plugins/AI/MassAI/Source/MassAIBehavior/Public/MassStateTreeExecutionContext.h
--- a/Engine/Plugins/AI/MassAI/Source/MassAIBehavior/Public/MassStateTreeExecutionContext.h
+++ b/Engine/Plugins/AI/MassAI/Source/MassAIBehavior/Public/MassStateTreeExecutionContext.h
@@ -28,9 +28,6 @@ public:
 	FMassEntityHandle GetEntity() const { return Entity; }
 	void SetEntity(const FMassEntityHandle InEntity) { Entity = InEntity; }
 
-	int32 GetEntityIndex() const { return EntityIndex; }
-	void SetEntityIndex(const int32 InIndex) { EntityIndex = InIndex; }
-
 protected:
 
 	/** Prefix that will be used by STATETREE_LOG and STATETREE_CLOG, using Entity description. */
@@ -42,5 +39,4 @@ protected:
 	UMassSignalSubsystem* SignalSubsystem = nullptr;
 	FMassExecutionContext* EntitySubsystemExecutionContext = nullptr;
 	FMassEntityHandle Entity;
-	int32 EntityIndex = INDEX_NONE;
 };
diff --git a/Engine/Plugins/AI/MassAI/Source/MassAIBehavior/Public/MassStateTreeFragments.h b/Engine/Plugins/AI/MassAI/Source/MassAIBehavior/Public/MassStateTreeFragments.h
--- a/Engine/Plugins/AI/MassAI/Source/MassAIBehavior/Public/MassStateTreeFragments.h
+++ b/Engine/Plugins/AI/MassAI/Source/MassAIBehavior/Public/MassStateTreeFragments.h
@@ -8,15 +8,27 @@
 
 
 USTRUCT()
-struct MASSAIBEHAVIOR_API FMassStateTreeFragment : public FMassFragment
+struct MASSAIBEHAVIOR_API FMassStateTreeInstanceFragment : public FMassFragment
 {
 	GENERATED_BODY()
 
-	FMassStateTreeFragment() = default;
+	FMassStateTreeInstanceFragment() = default;
 
-	/** Handle to a StateTree asset in MassStateTreeSubsystem */
-	FMassStateTreeHandle StateTreeHandle;
+	/** Handle to a StateTree instance data in MassStateTreeSubsystem. */
+	FMassStateTreeInstanceHandle InstanceHandle;
 
-	/** Keep track of the last update time to adjust time delta */
-	TOptional<float> LastUpdateTimeInSeconds;
+	/** The last update time use to calculate ticking delta time. */
+	float LastUpdateTimeInSeconds = 0.0f;
+};
+
+
+USTRUCT()
+struct MASSAIBEHAVIOR_API FMassStateTreeSharedFragment : public FMassSharedFragment
+{
+	GENERATED_BODY()
+
+	FMassStateTreeSharedFragment() = default;
+
+	UPROPERTY()
+	UStateTree* StateTree = nullptr;
 };
diff --git a/Engine/Plugins/AI/MassAI/Source/MassAIBehavior/Public/MassStateTreeProcessors.h b/Engine/Plugins/AI/MassAI/Source/MassAIBehavior/Public/MassStateTreeProcessors.h
--- a/Engine/Plugins/AI/MassAI/Source/MassAIBehavior/Public/MassStateTreeProcessors.h
+++ b/Engine/Plugins/AI/MassAI/Source/MassAIBehavior/Public/MassStateTreeProcessors.h
@@ -57,10 +57,6 @@ protected:
 	UMassSignalSubsystem* SignalSubsystem = nullptr;
 	
 	FMassEntityQuery EntityQuery;
-
-	/** The maximum number of entities to signal activation per LOD for each update, -1 means no limit */
-	UPROPERTY(EditAnywhere, Category = "Mass|LOD", config)
-	int32 MaxActivationsPerLOD[EMassLOD::Max];
 };
 
 /** 
diff --git a/Engine/Plugins/AI/MassAI/Source/MassAIBehavior/Public/MassStateTreeSubsystem.h b/Engine/Plugins/AI/MassAI/Source/MassAIBehavior/Public/MassStateTreeSubsystem.h
--- a/Engine/Plugins/AI/MassAI/Source/MassAIBehavior/Public/MassStateTreeSubsystem.h
+++ b/Engine/Plugins/AI/MassAI/Source/MassAIBehavior/Public/MassStateTreeSubsystem.h
@@ -4,10 +4,23 @@
 
 #include "MassStateTreeTypes.h"
 #include "Subsystems/WorldSubsystem.h"
+#include "StateTreeExecutionContext.h"
 #include "MassStateTreeSubsystem.generated.h"
 
 class UStateTree;
 
+USTRUCT()
+struct MASSAIBEHAVIOR_API FMassStateTreeInstanceDataItem
+{
+	GENERATED_BODY()
+
+	UPROPERTY()
+	FStateTreeInstanceData InstanceData;
+
+	UPROPERTY()
+	int32 Generation = 0;
+};
+
 /**
 * A subsystem managing StateTree assets in Mass
 */
@@ -22,18 +35,35 @@ public:
 	virtual void Initialize(FSubsystemCollectionBase& Collection) override;
 	// USubsystem END
 
-	/** Registers a StateTree asset to be used */
-	FMassStateTreeHandle RegisterStateTreeAsset(UStateTree* StateTree);
+	/**
+	 * Allocates new instance data for specified StateTree.
+	 * @param StateTree StateTree to allocated the data for.
+	 * @return Handle to the data.
+	 */
+	FMassStateTreeInstanceHandle AllocateInstanceData(const UStateTree* StateTree);
 
-	/** @return StateTree asset based on a handle */
-	UStateTree* GetRegisteredStateTreeAsset(FMassStateTreeHandle Handle) { return RegisteredStateTrees[Handle.GetIndex()]; }
+	/**
+	 * Frees instance data.
+	 * @param Handle Instance data handle to free.
+	 */
+	void FreeInstanceData(const FMassStateTreeInstanceHandle Handle);
 
-	/** @return Array of registered StateTree assets */
-	TConstArrayView<UStateTree*> GetRegisteredStateTreeAssets() const { return RegisteredStateTrees; }
+	/** @return Pointer to instance data held by the handle, or nullptr if handle is not valid. */
+	FStateTreeInstanceData* GetInstanceData(const FMassStateTreeInstanceHandle Handle)
+	{
+		return IsValidHandle(Handle) ? &InstanceDataArray[Handle.GetIndex()].InstanceData : nullptr;
+	}
+
+	/** @return True if the handle points to active instance data. */
+	bool IsValidHandle(const FMassStateTreeInstanceHandle Handle) const
+	{
+		return InstanceDataArray.IsValidIndex(Handle.GetIndex()) && InstanceDataArray[Handle.GetIndex()].Generation == Handle.GetGeneration();
+	}
 	
 protected:
 
-	/** Array of registered (in use) StateTrees */
+	TArray<int32> InstanceDataFreelist;
+
 	UPROPERTY(Transient)
-	TArray<UStateTree*> RegisteredStateTrees;
+	TArray<FMassStateTreeInstanceDataItem> InstanceDataArray;
 };
diff --git a/Engine/Plugins/AI/MassAI/Source/MassAIBehavior/Public/MassStateTreeTypes.h b/Engine/Plugins/AI/MassAI/Source/MassAIBehavior/Public/MassStateTreeTypes.h
--- a/Engine/Plugins/AI/MassAI/Source/MassAIBehavior/Public/MassStateTreeTypes.h
+++ b/Engine/Plugins/AI/MassAI/Source/MassAIBehavior/Public/MassStateTreeTypes.h
@@ -40,23 +40,28 @@ struct MASSAIBEHAVIOR_API FMassStateTreeTaskBase : public FStateTreeTaskBase
 };
 
 /**
-* A handle pointing to a registered StateTree asset in UMassStateTreeSubsystem.
-*/
-struct FMassStateTreeHandle
+ * A handle pointing to a StateTree instance data in UMassStateTreeSubsystem.
+ */
+struct FMassStateTreeInstanceHandle
 {
-	static constexpr uint16 Invalid = MAX_uint16;
-
-	FMassStateTreeHandle() = default;
+	FMassStateTreeInstanceHandle() = default;
 
 	/** Initializes new handle based on an index */
-	static FMassStateTreeHandle Make(const uint16 InIndex) { return FMassStateTreeHandle(InIndex); }
+	static FMassStateTreeInstanceHandle Make(const int32 InIndex, const int32 InGeneration) { return FMassStateTreeInstanceHandle(InIndex, InGeneration); }
 	
-	/** Returns index the handle points to */
-	uint16 GetIndex() const { return Index; }
+	/** @returns index the handle points to */
+	int32 GetIndex() const { return Index; }
+
+	/** @returns generation of the handle, used to identify recycled indices. */ 
+	int32 GetGeneration() const { return Generation; }
+
+	/** @returns true if the handle is valid. */
+	bool IsValid() const { return Index != INDEX_NONE; }
 
 protected:
-	FMassStateTreeHandle(const uint16 InIndex) : Index(InIndex) {}
+	FMassStateTreeInstanceHandle(const int32 InIndex, const int32 InGeneration) : Index(InIndex), Generation(InGeneration) {}
 
-	uint16 Index = Invalid;
+	int32 Index = INDEX_NONE;
+	int32 Generation = 0;
 };
 
diff --git a/Engine/Plugins/AI/MassAI/Source/MassAIDebug/Private/GameplayDebuggerCategory_Mass.cpp b/Engine/Plugins/AI/MassAI/Source/MassAIDebug/Private/GameplayDebuggerCategory_Mass.cpp
--- a/Engine/Plugins/AI/MassAI/Source/MassAIDebug/Private/GameplayDebuggerCategory_Mass.cpp
+++ b/Engine/Plugins/AI/MassAI/Source/MassAIDebug/Private/GameplayDebuggerCategory_Mass.cpp
@@ -346,7 +346,8 @@ void FGameplayDebuggerCategory_Mass::CollectData(APlayerController* OwnerPC, AAc
 		ensureMsgf(GetViewPoint(OwnerPC, ViewLocation, ViewDirection), TEXT("GetViewPoint is expected to always succeed when passing a valid controller."));
 
 		FMassEntityQuery EntityQuery;
-		EntityQuery.AddRequirement<FMassStateTreeFragment>(EMassFragmentAccess::ReadOnly);
+		EntityQuery.AddRequirement<FMassStateTreeInstanceFragment>(EMassFragmentAccess::ReadOnly);
+		EntityQuery.AddConstSharedRequirement<FMassStateTreeSharedFragment>();
 		EntityQuery.AddRequirement<FTransformFragment>(EMassFragmentAccess::ReadOnly);
 		EntityQuery.AddRequirement<FAgentRadiusFragment>(EMassFragmentAccess::ReadOnly);
 		EntityQuery.AddRequirement<FMassSteeringFragment>(EMassFragmentAccess::ReadOnly);
@@ -371,7 +372,7 @@ void FGameplayDebuggerCategory_Mass::CollectData(APlayerController* OwnerPC, AAc
 			EntityQuery.ForEachEntityChunk(*EntitySystem, Context, [this, MassStateTreeSubsystem, SignalSubsystem, EntitySystem, OwnerPC, ViewLocation, ViewDirection, CurrentTime](FMassExecutionContext& Context)
 			{
 				const int32 NumEntities = Context.GetNumEntities();
-				const TConstArrayView<FMassStateTreeFragment> StateTreeList = Context.GetFragmentView<FMassStateTreeFragment>();
+				const TConstArrayView<FMassStateTreeInstanceFragment> StateTreeInstanceList = Context.GetFragmentView<FMassStateTreeInstanceFragment>();
 				const TConstArrayView<FTransformFragment> TransformList = Context.GetFragmentView<FTransformFragment>();
 				const TConstArrayView<FAgentRadiusFragment> RadiusList = Context.GetFragmentView<FAgentRadiusFragment>();
 				const TConstArrayView<FMassSteeringFragment> SteeringList = Context.GetFragmentView<FMassSteeringFragment>();
@@ -381,16 +382,18 @@ void FGameplayDebuggerCategory_Mass::CollectData(APlayerController* OwnerPC, AAc
 				const TConstArrayView<FMassForceFragment> ForceList = Context.GetFragmentView<FMassForceFragment>();
 				const TConstArrayView<FMassMoveTargetFragment> MoveTargetList = Context.GetFragmentView<FMassMoveTargetFragment>();
 				const TConstArrayView<FMassLookAtFragment> LookAtList = Context.GetFragmentView<FMassLookAtFragment>();
-				const bool bHasLookAt = (LookAtList.Num() > 0);
 				const TConstArrayView<FMassSimulationLODFragment> SimLODList = Context.GetFragmentView<FMassSimulationLODFragment>();
-				const bool bHasLOD = (SimLODList.Num() > 0);
 				const TConstArrayView<FMassZoneGraphShortPathFragment> ShortPathList = Context.GetFragmentView<FMassZoneGraphShortPathFragment>();
+				const FMassStateTreeSharedFragment& SharedStateTree = Context.GetConstSharedFragment<FMassStateTreeSharedFragment>();
+
+				const bool bHasLOD = (SimLODList.Num() > 0);
+				const bool bHasLookAt = (LookAtList.Num() > 0);
 
 				const UGameplayDebuggerUserSettings* Settings = GetDefault<UGameplayDebuggerUserSettings>();
 				const float MaxViewDistance = Settings->MaxViewDistance;
 				const float MinViewDirDot = FMath::Cos(FMath::DegreesToRadians(Settings->MaxViewAngle));
 
-				const UStateTree* StateTree = MassStateTreeSubsystem->GetRegisteredStateTreeAsset(StateTreeList[0].StateTreeHandle);
+				const UStateTree* StateTree = SharedStateTree.StateTree;
 
 				for (int32 EntityIndex = 0; EntityIndex < NumEntities; ++EntityIndex)
 				{
@@ -419,6 +422,7 @@ void FGameplayDebuggerCategory_Mass::CollectData(APlayerController* OwnerPC, AAc
 					const FMassMoveTargetFragment& MoveTarget = MoveTargetList[EntityIndex];
 					const FMassSimulationLODFragment& SimLOD = bHasLOD ? SimLODList[EntityIndex] : FMassSimulationLODFragment();
 					const FMassZoneGraphShortPathFragment& ShortPath = ShortPathList[EntityIndex];
+					const FMassStateTreeInstanceFragment& StateTreeInstance = StateTreeInstanceList[EntityIndex];
 
 					const FVector EntityForward = Transform.GetTransform().GetRotation().GetForwardVector();
 
@@ -550,10 +554,16 @@ void FGameplayDebuggerCategory_Mass::CollectData(APlayerController* OwnerPC, AAc
 							StateTreeContext.Init(*OwnerPC, *StateTree, EStateTreeStorage::External);
 							StateTreeContext.SetEntity(Entity);
 							
-							FStructView Storage = EntitySystem->GetFragmentDataStruct(Entity, StateTree->GetInstanceStorageStruct());
-							
-							Status += StateTreeContext.GetActiveStateName(Storage);
-							Status += TEXT("\n");
+							if (FStateTreeInstanceData* InstanceData = MassStateTreeSubsystem->GetInstanceData(StateTreeInstance.InstanceHandle))
+							{
+								Status += StateTreeContext.GetActiveStateName(InstanceData);
+								Status += TEXT("\n");
+							}
+							else
+							{
+								Status += StateTreeContext.GetActiveStateName(InstanceData);
+								Status += TEXT("{red}<No StateTree instance>{white}\n");
+							}
 						}
 
 						// Movement info
diff --git a/Engine/Plugins/AI/MassAI/Source/MassAIDebug/Private/MassDebugStateTreeProcessor.cpp b/Engine/Plugins/AI/MassAI/Source/MassAIDebug/Private/MassDebugStateTreeProcessor.cpp
--- a/Engine/Plugins/AI/MassAI/Source/MassAIDebug/Private/MassDebugStateTreeProcessor.cpp
+++ b/Engine/Plugins/AI/MassAI/Source/MassAIDebug/Private/MassDebugStateTreeProcessor.cpp
@@ -20,7 +20,8 @@ UMassDebugStateTreeProcessor::UMassDebugStateTreeProcessor()
 
 void UMassDebugStateTreeProcessor::ConfigureQueries()
 {
-	EntityQuery.AddRequirement<FMassStateTreeFragment>(EMassFragmentAccess::ReadOnly);
+	EntityQuery.AddRequirement<FMassStateTreeInstanceFragment>(EMassFragmentAccess::ReadOnly);
+	EntityQuery.AddConstSharedRequirement<FMassStateTreeSharedFragment>();
 	EntityQuery.AddRequirement<FTransformFragment>(EMassFragmentAccess::ReadOnly);
 }
 
@@ -58,54 +59,67 @@ void UMassDebugStateTreeProcessor::Execute(UMassEntitySubsystem& EntitySubsystem
 	
 	QUICK_SCOPE_CYCLE_COUNTER(UMassDebugStateTreeProcessor_Run);	
 	EntityQuery.ForEachEntityChunk(EntitySubsystem, Context, [this, Debugger, MassStateTreeSubsystem, &EntitySubsystem, MassSignalSubsystem](FMassExecutionContext& Context)
+	{
+		const FMassEntityHandle SelectedEntity = Debugger->GetSelectedEntity();
+		const int32 NumEntities = Context.GetNumEntities();
+		const TConstArrayView<FMassStateTreeInstanceFragment> StateTreeInstanceList = Context.GetFragmentView<FMassStateTreeInstanceFragment>();
+		const FMassStateTreeSharedFragment& SharedStateTree = Context.GetConstSharedFragment<FMassStateTreeSharedFragment>();
+		const TConstArrayView<FTransformFragment> TransformList = Context.GetFragmentView<FTransformFragment>();
+
+		const UStateTree* StateTree = SharedStateTree.StateTree;
+
+		// Not reporting error since this processor is a debug tool 
+		if (StateTree == nullptr)
 		{
-			const FMassEntityHandle SelectedEntity = Debugger->GetSelectedEntity();
-			const int32 NumEntities = Context.GetNumEntities();
-			const TConstArrayView<FMassStateTreeFragment> StateTreeList = Context.GetFragmentView<FMassStateTreeFragment>();
-			const TConstArrayView<FTransformFragment> TransformList = Context.GetFragmentView<FTransformFragment>();
-			const UStateTree* StateTree = MassStateTreeSubsystem->GetRegisteredStateTreeAsset(StateTreeList[0].StateTreeHandle);
+			return;
+		}
+	
+		for (int32 i = 0; i < NumEntities; ++i)
+		{
+			const FMassEntityHandle Entity = Context.GetEntity(i);
 
-			// Not reporting error since this processor is a debug tool 
-			if (StateTree == nullptr)
+			if (Entity != SelectedEntity && !UE::Mass::Debug::IsDebuggingEntity(Entity))
 			{
-				return;
+				continue;
 			}
-		
-			for (int32 i = 0; i < NumEntities; ++i)
+			
+			const FMassStateTreeInstanceFragment& StateTreeInstance = StateTreeInstanceList[i];
+
+			FStateTreeInstanceData* InstanceData = MassStateTreeSubsystem->GetInstanceData(StateTreeInstance.InstanceHandle);
+			if (InstanceData == nullptr)
 			{
-				FMassEntityHandle Entity = Context.GetEntity(i);
-				if (Entity == SelectedEntity)
-				{
-					FMassStateTreeExecutionContext StateTreeContext(EntitySubsystem, *MassSignalSubsystem, Context);
-					StateTreeContext.Init(*this, *StateTree, EStateTreeStorage::External);
-					StateTreeContext.SetEntity(Entity);
-					
+				continue;
+			}
+			
+			if (Entity == SelectedEntity)
+			{
+				FMassStateTreeExecutionContext StateTreeContext(EntitySubsystem, *MassSignalSubsystem, Context);
+				StateTreeContext.Init(*this, *StateTree, EStateTreeStorage::External);
+				StateTreeContext.SetEntity(Entity);
+				
 #if WITH_GAMEPLAY_DEBUGGER
-					const FStructView Storage = EntitySubsystem.GetFragmentDataStruct(Entity, StateTree->GetInstanceStorageStruct());
-					Debugger->AppendSelectedEntityInfo(StateTreeContext.GetDebugInfoString(Storage));
+				Debugger->AppendSelectedEntityInfo(StateTreeContext.GetDebugInfoString(InstanceData));
 #endif // WITH_GAMEPLAY_DEBUGGER
-				}
-					
-				FColor EntityColor = FColor::White;
-				const bool bDisplayDebug = UE::Mass::Debug::IsDebuggingEntity(Entity, &EntityColor);
-				if (bDisplayDebug)
-				{
-					const FTransformFragment& Transform = TransformList[i];
-					
-					const FVector ZOffset(0,0,50);
-					const FVector Position = Transform.GetTransform().GetLocation() + ZOffset;
-
-					FMassStateTreeExecutionContext StateTreeContext(EntitySubsystem, *MassSignalSubsystem, Context);
-					StateTreeContext.Init(*this, *StateTree, EStateTreeStorage::External);
-					StateTreeContext.SetEntity(Entity);
+			}
+				
+			FColor EntityColor = FColor::White;
+			const bool bDisplayDebug = UE::Mass::Debug::IsDebuggingEntity(Entity, &EntityColor);
+			if (bDisplayDebug)
+			{
+				const FTransformFragment& Transform = TransformList[i];
+				
+				const FVector ZOffset(0,0,50);
+				const FVector Position = Transform.GetTransform().GetLocation() + ZOffset;
 
-					const FStructView Storage = EntitySubsystem.GetFragmentDataStruct(Entity, StateTree->GetInstanceStorageStruct());
+				FMassStateTreeExecutionContext StateTreeContext(EntitySubsystem, *MassSignalSubsystem, Context);
+				StateTreeContext.Init(*this, *StateTree, EStateTreeStorage::External);
+				StateTreeContext.SetEntity(Entity);
 
-					// State
-					UE_VLOG_SEGMENT_THICK(this, LogStateTree, Log, Position, Position + FVector(0,0,50), EntityColor, /*Thickness*/ 2, TEXT("%s %s"),
-						*Entity.DebugGetDescription(), *StateTreeContext.GetActiveStateName(Storage));
-				}
+				// State
+				UE_VLOG_SEGMENT_THICK(this, LogStateTree, Log, Position, Position + FVector(0,0,50), EntityColor, /*Thickness*/ 2, TEXT("%s %s"),
+					*Entity.DebugGetDescription(), *StateTreeContext.GetActiveStateName(InstanceData));
 			}
-		});
+		}
+	});
 	#endif // WITH_MASSGAMEPLAY_DEBUG
 }
diff --git a/Engine/Plugins/AI/MassAI/Source/MassAIDebug/Public/MassDebugStateTreeProcessor.h b/Engine/Plugins/AI/MassAI/Source/MassAIDebug/Public/MassDebugStateTreeProcessor.h
--- a/Engine/Plugins/AI/MassAI/Source/MassAIDebug/Public/MassDebugStateTreeProcessor.h
+++ b/Engine/Plugins/AI/MassAI/Source/MassAIDebug/Public/MassDebugStateTreeProcessor.h
@@ -6,7 +6,6 @@
 #include "MassDebugStateTreeProcessor.generated.h"
 
 class UMassEntitySubsystem;
-struct FMassStateTreeFragment;
 struct FMassEntityQuery;
 struct FMassExecutionContext;
 
diff --git a/Engine/Plugins/Runtime/StateTree/Config/BaseStateTree.ini b/Engine/Plugins/Runtime/StateTree/Config/BaseStateTree.ini
--- a/Engine/Plugins/Runtime/StateTree/Config/BaseStateTree.ini
+++ b/Engine/Plugins/Runtime/StateTree/Config/BaseStateTree.ini
@@ -21,6 +21,7 @@
 +PropertyRedirects=(OldName="StateTreeState.Tasks2",NewName="StateTreeState.Tasks")
 
 +PropertyRedirects=(OldName="StateTree.Conditions2",NewName="StateTree.Conditions")
++PropertyRedirects=(OldName="StateTree.Items",NewName="StateTree.Nodes")
 +PropertyRedirects=(OldName="StateTreeEditorNode.Item",NewName="StateTreeEditorNode.Node")
 
 +EnumRedirects=(OldName="EStateTreePropertyCopyType",ValueChanges=(("Plain","CopyPlain"),("Complex","CopyComplex"),("Bool","CopyBool"),("Struct","CopyStruct"),("Object","CopyObject"),("Name","CopyName"),("FixedArray","CopyFixedArray")))
diff --git a/Engine/Plugins/Runtime/StateTree/Source/StateTreeEditorModule/Private/StateTreeBaker.cpp b/Engine/Plugins/Runtime/StateTree/Source/StateTreeEditorModule/Private/StateTreeBaker.cpp
--- a/Engine/Plugins/Runtime/StateTree/Source/StateTreeEditorModule/Private/StateTreeBaker.cpp
+++ b/Engine/Plugins/Runtime/StateTree/Source/StateTreeEditorModule/Private/StateTreeBaker.cpp
@@ -42,12 +42,8 @@ bool FStateTreeBaker::Bake(UStateTree& InStateTree)
 
 	BindingsCompiler.Finalize();
 
-	// TODO: This is just for testing. It should be called just before use.
-	StateTree->ResolvePropertyPaths();
-
+	StateTree->PropertyBindings.ResolvePaths();
 	StateTree->Link();
-	
-	StateTree->InitInstanceStorageType();
 
 	return true;
 }
@@ -93,7 +89,7 @@ bool FStateTreeBaker::CreateStateTransitions()
 		FStateTreeCompilerLogStateScope LogStateScope(SourceState, Log);
 		
 		// Enter conditions.
-		BakedState.EnterConditionsBegin = uint16(StateTree->Items.Num());
+		BakedState.EnterConditionsBegin = uint16(StateTree->Nodes.Num());
 		for (FStateTreeEditorNode& CondNode : SourceState->EnterConditions)
 		{
 			if (!CreateCondition(CondNode))
@@ -103,7 +99,7 @@ bool FStateTreeBaker::CreateStateTransitions()
 				return false;
 			}
 		}
-		BakedState.EnterConditionsNum = uint8(uint16(StateTree->Items.Num()) - BakedState.EnterConditionsBegin);
+		BakedState.EnterConditionsNum = uint8(uint16(StateTree->Nodes.Num()) - BakedState.EnterConditionsBegin);
 
 		// Transitions
 		BakedState.TransitionsBegin = uint16(StateTree->Transitions.Num());
@@ -120,7 +116,7 @@ bool FStateTreeBaker::CreateStateTransitions()
 			}
 			// Note: Unset transition is allowed here. It can be used to mask a transition at parent.
 
-			BakedTransition.ConditionsBegin = uint16(StateTree->Items.Num());
+			BakedTransition.ConditionsBegin = uint16(StateTree->Nodes.Num());
 			for (FStateTreeEditorNode& CondNode : Transition.Conditions)
 			{
 				if (!CreateCondition(CondNode))
@@ -131,7 +127,7 @@ bool FStateTreeBaker::CreateStateTransitions()
 					return false;
 				}
 			}
-			BakedTransition.ConditionsNum = uint8(uint16(StateTree->Items.Num()) - BakedTransition.ConditionsBegin);
+			BakedTransition.ConditionsNum = uint8(uint16(StateTree->Nodes.Num()) - BakedTransition.ConditionsBegin);
 		}
 		BakedState.TransitionsNum = uint8(uint16(StateTree->Transitions.Num()) - BakedState.TransitionsBegin);
 	}
@@ -198,7 +194,7 @@ bool FStateTreeBaker::CreateCondition(const FStateTreeEditorNode& CondNode)
 	}
 
 	// Copy the condition
-	FInstancedStruct& Item = StateTree->Items.AddDefaulted_GetRef();
+	FInstancedStruct& Item = StateTree->Nodes.AddDefaulted_GetRef();
 	Item = CondNode.Node;
 
 	FStateTreeConditionBase& Cond = Item.GetMutable<FStateTreeConditionBase>();
@@ -285,7 +281,7 @@ bool FStateTreeBaker::CreateTask(const FStateTreeEditorNode& TaskNode)
 	}
 
 	// Copy the task
-	FInstancedStruct& Item = StateTree->Items.AddDefaulted_GetRef();
+	FInstancedStruct& Item = StateTree->Nodes.AddDefaulted_GetRef();
 	Item = TaskNode.Node;
 	
 	FStateTreeTaskBase& Task = Item.GetMutable<FStateTreeTaskBase>();
@@ -373,7 +369,7 @@ bool FStateTreeBaker::CreateEvaluator(const FStateTreeEditorNode& EvalNode)
     }
 
 	// Copy the evaluator
-	FInstancedStruct& Item = StateTree->Items.AddDefaulted_GetRef();
+	FInstancedStruct& Item = StateTree->Nodes.AddDefaulted_GetRef();
 	Item = EvalNode.Node;
 
 	FStateTreeEvaluatorBase& Eval = Item.GetMutable<FStateTreeEvaluatorBase>();
@@ -539,8 +535,8 @@ bool FStateTreeBaker::CreateStateRecursive(UStateTreeState& State, const FStateT
 	check(StateTree->Instances.Num() <= int32(MAX_uint16));
 
 	// Collect evaluators
-	check(StateTree->Items.Num() <= int32(MAX_uint16));
-	BakedState.EvaluatorsBegin = uint16(StateTree->Items.Num());
+	check(StateTree->Nodes.Num() <= int32(MAX_uint16));
+	BakedState.EvaluatorsBegin = uint16(StateTree->Nodes.Num());
 
 	for (FStateTreeEditorNode& EvalNode : State.Evaluators)
 	{
@@ -550,13 +546,13 @@ bool FStateTreeBaker::CreateStateRecursive(UStateTreeState& State, const FStateT
 		}
 	}
 	
-	const int32 EvaluatorsNum = StateTree->Items.Num() - int32(BakedState.EvaluatorsBegin);
+	const int32 EvaluatorsNum = StateTree->Nodes.Num() - int32(BakedState.EvaluatorsBegin);
 	check(EvaluatorsNum <= int32(MAX_uint8));
 	BakedState.EvaluatorsNum = uint8(EvaluatorsNum);
 
 	// Collect tasks
-	check(StateTree->Items.Num() <= int32(MAX_uint16));
-	BakedState.TasksBegin = uint16(StateTree->Items.Num());
+	check(StateTree->Nodes.Num() <= int32(MAX_uint16));
+	BakedState.TasksBegin = uint16(StateTree->Nodes.Num());
 
 	for (FStateTreeEditorNode& TaskNode : State.Tasks)
 	{
@@ -571,7 +567,7 @@ bool FStateTreeBaker::CreateStateRecursive(UStateTreeState& State, const FStateT
 		return false;
 	}
 	
-	const int32 TasksNum = StateTree->Items.Num() - int32(BakedState.TasksBegin);
+	const int32 TasksNum = StateTree->Nodes.Num() - int32(BakedState.TasksBegin);
 	check(TasksNum <= int32(MAX_uint8));
 	BakedState.TasksNum = uint8(TasksNum);
 
diff --git a/Engine/Plugins/Runtime/StateTree/Source/StateTreeEditorModule/Private/StateTreeEditor.cpp b/Engine/Plugins/Runtime/StateTree/Source/StateTreeEditorModule/Private/StateTreeEditor.cpp
--- a/Engine/Plugins/Runtime/StateTree/Source/StateTreeEditorModule/Private/StateTreeEditor.cpp
+++ b/Engine/Plugins/Runtime/StateTree/Source/StateTreeEditorModule/Private/StateTreeEditor.cpp
@@ -324,13 +324,10 @@ FText FStateTreeEditor::GetStatisticsText() const
 		return FText::GetEmpty();
 	}
 
-	if (const UScriptStruct* StorageStruct = StateTree->GetInstanceStorageStruct())
-	{
-		const FText SizeText = FText::AsMemory((uint64)StorageStruct->GetStructureSize());
-		const FText NumNodesText = FText::AsNumber(StateTree->GetNumInstances());
+	const FText SizeText = FText::AsMemory((uint64)StateTree->GetInstanceDataSize());
+	const FText NumNodesText = FText::AsNumber(StateTree->GetNumInstances());
 
-		return FText::Format(LOCTEXT("RuntimeSize", "Runtime size: {0}, {1} nodes"), SizeText, NumNodesText);
-	}
+	return FText::Format(LOCTEXT("RuntimeSize", "Runtime size: {0}, {1} nodes"), SizeText, NumNodesText);
 	
 	return FText::GetEmpty();
 }
@@ -743,7 +740,7 @@ FSlateIcon FStateTreeEditor::GetCompileStatusImage() const
 		return FSlateIcon(FEditorStyle::GetStyleSetName(), CompileStatusBackground, NAME_None, CompileStatusUnknown);
 	}
 	
-	if (!StateTree->IsValidStateTree())
+	if (!StateTree->IsReadyToRun())
 	{
 		return FSlateIcon(FEditorStyle::GetStyleSetName(), CompileStatusBackground, NAME_None, CompileStatusError);
 	}
diff --git a/Engine/Plugins/Runtime/StateTree/Source/StateTreeModule/Private/Components/StateTreeBrainComponent.cpp b/Engine/Plugins/Runtime/StateTree/Source/StateTreeModule/Private/Components/StateTreeBrainComponent.cpp
--- a/Engine/Plugins/Runtime/StateTree/Source/StateTreeModule/Private/Components/StateTreeBrainComponent.cpp
+++ b/Engine/Plugins/Runtime/StateTree/Source/StateTreeModule/Private/Components/StateTreeBrainComponent.cpp
@@ -153,9 +153,6 @@ void UStateTreeBrainComponent::StartLogic()
 {
 	STATETREE_LOG(Log, TEXT("%s: Start Logic"), ANSI_TO_TCHAR(__FUNCTION__));
 
-	// TODO: Move this to UStateTree
-	StateTree->ResolvePropertyPaths();
-
 	if (SetContextRequirements())
 	{
 		StateTreeContext.Start();
diff --git a/Engine/Plugins/Runtime/StateTree/Source/StateTreeModule/Private/StateTree.cpp b/Engine/Plugins/Runtime/StateTree/Source/StateTreeModule/Private/StateTree.cpp
--- a/Engine/Plugins/Runtime/StateTree/Source/StateTreeModule/Private/StateTree.cpp
+++ b/Engine/Plugins/Runtime/StateTree/Source/StateTreeModule/Private/StateTree.cpp
@@ -6,57 +6,29 @@
 #include "CoreMinimal.h"
 #include "StateTreeConditionBase.h"
 #include "StateTreeDelegates.h"
-#if WITH_EDITOR
-#include "Editor.h"
-#endif
-
-UStateTree::UStateTree(const FObjectInitializer& ObjectInitializer)
-	: Super(ObjectInitializer)
-{
-
-#if WITH_EDITOR
-	FEditorDelegates::PreBeginPIE.AddUObject(this, &UStateTree::OnPIEStarted);
-#endif
-}
-
-UStateTree::~UStateTree()
-{
-#if WITH_EDITOR
-	FEditorDelegates::PreBeginPIE.RemoveAll(this);
-#endif
-}
 
-bool UStateTree::IsValidStateTree() const
+bool UStateTree::IsReadyToRun() const
 {
-	// Valid tree must have at least one state.
-	return States.Num() > 0;
+	// Valid tree must have at least one state and valid instance data.
+	return States.Num() > 0 && InstanceDataDefaultValue.IsValid();
 }
 
 #if WITH_EDITOR
-
-void UStateTree::OnPIEStarted(const bool bIsSimulating)
-{
-	PropertyBindings.ResolvePaths();
-	Link();
-	InitInstanceStorageType();
-}
-
 void UStateTree::ResetBaked()
 {
 	States.Reset();
 	Transitions.Reset();
 
-	Items.Reset();
+	Nodes.Reset();
 	Instances.Reset();
 	InstanceObjects.Reset();
-	InstanceStorageStruct = nullptr;
-	InstanceStorageOffsets.Reset();
-	InstanceStorageDefaultValue.Reset();
 	ExternalDataDescs.Reset();
 	PropertyBindings.Reset();
 
 	NumDataViews = 0;
 	ExternalDataBaseIndex = 0;
+
+	InstanceDataDefaultValue.Reset();
 }
 
 void UStateTree::PostEditChangeProperty(struct FPropertyChangedEvent& PropertyChangedEvent)
@@ -88,34 +60,6 @@ void UStateTree::PostLoad()
 	PropertyBindings.ResolvePaths();
 
 	Link();
-	
-#if WITH_EDITOR
-	InitInstanceStorageType();
-#else
-	// Item offsets still need to be calculated in non editor target since the struct sizes might be different.
-	checkf(InstanceStorageOffsets.Num() == 0, TEXT("RuntimeStorageOffsets is transient and should only be computed once."));
-
-	for (TFieldIterator<FProperty> PropertyIt(InstanceStorageStruct); PropertyIt; ++PropertyIt)
-	{
-		FStructProperty* StructProperty = CastField<FStructProperty>(*PropertyIt);
-		checkf(StructProperty, TEXT("InstanceStorageStruct is expected to only contain Struct properties"));
-		InstanceStorageOffsets.Emplace(StructProperty->Struct, StructProperty->GetOffset_ForInternal());
-	}
-#endif // WITH_EDITOR
-}
-
-void UStateTree::BeginDestroy()
-{
-	Super::BeginDestroy();
-	
-	// Destroy the runtime storage before the UScriptStruct it uses gets removed.
-	InstanceStorageDefaultValue.Reset();
-}
-
-void UStateTree::ResolvePropertyPaths()
-{
-	// TODO: find better hook when to call this. Currently this gets called from StateTreeComponent, it should be called once PostLoad() or when entering SIE/PIE in editor.
-	PropertyBindings.ResolvePaths();
 }
 
 void UStateTree::Link()
@@ -125,100 +69,35 @@ void UStateTree::Link()
 	ExternalDataBaseIndex = PropertyBindings.GetSourceStructNum();
 	Linker.SetExternalDataBaseIndex(ExternalDataBaseIndex);
 	
-	for (FInstancedStruct& Item : Items)
+	for (FInstancedStruct& Node : Nodes)
 	{
-		if (FStateTreeNodeBase* Node = Item.GetMutablePtr<FStateTreeNodeBase>())
+		if (FStateTreeNodeBase* NodePtr = Node.GetMutablePtr<FStateTreeNodeBase>())
 		{
-			Linker.SetCurrentInstanceDataType(Node->GetInstanceDataType(), Node->DataViewIndex);
-			Node->Link(Linker);
+			Linker.SetCurrentInstanceDataType(NodePtr->GetInstanceDataType(), NodePtr->DataViewIndex);
+			NodePtr->Link(Linker);
 		}
 	}
 
 	ExternalDataDescs = Linker.GetExternalDataDescs();
 	
 	NumDataViews = ExternalDataBaseIndex + ExternalDataDescs.Num();
-}
-
-void UStateTree::InitInstanceStorageType()
-{
-	InstanceStorageStruct = nullptr;
-	InstanceStorageOffsets.Reset();
-	InstanceStorageDefaultValue.Reset();
-
-	// Check that the items are valid before trying to create the type.
-	// The structs can become invalid i.e. because of missing type.
-	bool bValid = true;
-	for (const FInstancedStruct& Item : Instances)
-	{
-		if (!Item.IsValid())
-		{
-			bValid = false;
-			break;
-		}
-	}
-	if (!bValid)
-	{
-		return;
-	}
-
-	const FString StructName = GetName() + TEXT("_InstanceStorage");
-
-	// Remove existing struct of same name.
-	UScriptStruct* OldStruct = FindObject<UScriptStruct>(this, *StructName);
-	if (OldStruct)
-	{
-		const FString OldStructName = MakeUniqueObjectName(OldStruct->GetOuter(), OldStruct->GetClass(), *FString::Printf(TEXT("%s_TRASH"), *OldStruct->GetName())).ToString();
-		OldStruct->SetFlags(RF_NewerVersionExists);
-		OldStruct->ClearFlags(RF_Public | RF_Standalone);
-		OldStruct->SetStructTrashed(/*bIsTrash*/true);
-		OldStruct->Rename(*OldStructName, nullptr, REN_DontCreateRedirectors|REN_ForceNoResetLoaders);
-	}
-
-	UScriptStruct* NewStruct = NewObject<UScriptStruct>(this, *StructName, RF_Public);
-
-	if (Schema)
-	{
-		NewStruct->SetSuperStruct(Schema->GetStorageSuperStruct());
-	}
-
-	// Append all evaluators and tasks.
-	// Since properties are stored in linked list, add in reverse order so that we retain the correct order in output.
-	TArray<FStructProperty*> NewProperties;
-	NewProperties.SetNumZeroed(Instances.Num());
-	for (int32 ItemIndex = Instances.Num() - 1; ItemIndex >= 0; ItemIndex--)
-	{
-		const FInstancedStruct& ItemPtr = Instances[ItemIndex];
-		UScriptStruct* ItemStruct = const_cast<UScriptStruct*>(ItemPtr.GetScriptStruct());
-		FName PropName(FString::Printf(TEXT("%s%d"), *ItemStruct->GetName(), NewProperties.Num()));
-		FStructProperty* NewStructProperty = new FStructProperty(NewStruct, PropName, RF_Public);
-		NewStructProperty->Struct = ItemStruct;
-
-		NewStruct->AddCppProperty(NewStructProperty);
-
-		NewProperties[ItemIndex] = NewStructProperty;
-	}
 
-	// Finalize the struct
-	NewStruct->Bind();
-	NewStruct->StaticLink(/*bRelinkExistingProperties*/true);
-
-	// Store item offsets for fast struct view creation at runtime.
-	TArray<FStateTreeInstanceStorageOffset> NewOffsets;
-	for (FStructProperty* StructProperty : NewProperties)
+	// Initialize the instance data default value.
+	// This data will be used to allocate runtime instance on all StateTree users.
+	InstanceDataDefaultValue.Reset();
+	if (Instances.Num() > 0)
 	{
-		NewOffsets.Emplace(StructProperty->Struct, StructProperty->GetOffset_ForInternal());
+		InstanceDataDefaultValue.Initialize(Instances);
 	}
-	check(Instances.Num() == NewOffsets.Num());
+}
 
-	// Instantiate default value
-	InstanceStorageDefaultValue.InitializeAs(NewStruct);
-	for (int i = 0; i < NewOffsets.Num(); i++)
+int32 UStateTree::GetInstanceDataSize() const
+{
+	if (!InstanceDataDefaultValue.IsValid() || !InstanceDataDefaultValue.GetLayout().IsValid())
 	{
-		const FStateTreeInstanceStorageOffset& Item = NewOffsets[i];
-		uint8* Dest = InstanceStorageDefaultValue.GetMutableMemory() + Item.Offset;
-		Item.Struct->CopyScriptStruct(Dest, Instances[i].GetMemory());
+		return 0;
 	}
 
-	InstanceStorageOffsets = NewOffsets;
-	InstanceStorageStruct = NewStruct;
+	// TODO: this should count all allocated data (arrays, objects).
+	return InstanceDataDefaultValue.GetLayout()->GetLayoutInstanceSize();
 }
diff --git a/Engine/Plugins/Runtime/StateTree/Source/StateTreeModule/Private/StateTreeExecutionContext.cpp b/Engine/Plugins/Runtime/StateTree/Source/StateTreeModule/Private/StateTreeExecutionContext.cpp
--- a/Engine/Plugins/Runtime/StateTree/Source/StateTreeModule/Private/StateTreeExecutionContext.cpp
+++ b/Engine/Plugins/Runtime/StateTree/Source/StateTreeModule/Private/StateTreeExecutionContext.cpp
@@ -11,6 +11,7 @@
 #include "ProfilingDebugging/CsvProfiler.h"
 #include "Algo/Reverse.h"
 #include "Logging/LogScopedVerbosityOverride.h"
+#include "ProfilingDebugging/CsvProfiler.h"
 
 #define STATETREE_LOG(Verbosity, Format, ...) UE_VLOG_UELOG(GetOwner(), LogStateTree, Verbosity, TEXT("%s: ") Format, *GetInstanceDescription(), ##__VA_ARGS__)
 #define STATETREE_CLOG(Condition, Verbosity, Format, ...) UE_CVLOG_UELOG((Condition), GetOwner(), LogStateTree, Verbosity, TEXT("%s: ") Format, *GetInstanceDescription(), ##__VA_ARGS__)
@@ -33,28 +34,20 @@ bool FStateTreeExecutionContext::Init(UObject& InOwner, const UStateTree& InStat
 	// Set owner first for proper logging (it will be reset in case of failure) 
 	Owner = &InOwner;
 	
-	if (!InStateTree.IsValidStateTree())
+	if (!InStateTree.IsReadyToRun())
 	{
 		STATETREE_LOG(Error, TEXT("%s: StateTree asset '%s' is not valid."), ANSI_TO_TCHAR(__FUNCTION__), *InStateTree.GetName());
 		Reset();
 		return false;
 	}
 	
-	if (InStateTree.Instances.Num() == 0)
-	{
-		STATETREE_LOG(Error, TEXT("%s: StateTree asset '%s' has no valid instances."), ANSI_TO_TCHAR(__FUNCTION__), *InStateTree.GetName());
-		Reset();
-		return false;
-	}
-	
 	StateTree = &InStateTree;
 
 	StorageType = InStorageType;
 	if (StorageType == EStateTreeStorage::Internal)
 	{
 		// Duplicate instance state.
-		StorageInstance = StateTree->InstanceStorageDefaultValue;
-		InitInstanceData(InOwner, FStateTreeDataView(StorageInstance));
+		InitInstanceData(InOwner, *StateTree, InternalInstanceData);
 	}
 
 	// Initialize data views for all possible items.
@@ -62,24 +55,29 @@ bool FStateTreeExecutionContext::Init(UObject& InOwner, const UStateTree& InStat
 
 	// Initialize array to keep track of which states have been updated.
 	VisitedStates.Init(false, StateTree->States.Num());
-	
+
 	return true;
 }
 
-bool FStateTreeExecutionContext::InitInstanceData(UObject& InOwner, FStateTreeDataView ExternalStorage) const
+void FStateTreeExecutionContext::Reset()
 {
-	if (!Owner || !StateTree)
-	{
-		STATETREE_LOG(Error, TEXT("%s: Failed to init instance data on '%s' using StateTree '%s'."),
-			ANSI_TO_TCHAR(__FUNCTION__), *GetNameSafe(Owner), *GetNameSafe(StateTree));
-		return false;
-	}
+	InternalInstanceData.Reset();
+	DataViews.Reset();
+	StorageType = EStateTreeStorage::Internal;
+	VisitedStates.Reset();
+	StateTree = nullptr;
+	Owner = nullptr;
+}
+
+bool FStateTreeExecutionContext::InitInstanceData(UObject& InOwner, const UStateTree& InStateTree, FStateTreeInstanceData& OutInstanceData)
+{
+	OutInstanceData.CopyFrom(InStateTree.InstanceDataDefaultValue);
 
-	check (ExternalStorage.IsValid());
-	FStateTreeExecutionState& Exec = GetExecState(ExternalStorage);
+	check (OutInstanceData.IsValid());
+	FStateTreeExecutionState& Exec = GetExecState(OutInstanceData);
 
 	Exec.InstanceObjects.Reset();
-	for (const UObject* Instance : StateTree->InstanceObjects)
+	for (const UObject* Instance : InStateTree.InstanceObjects)
 	{
 		Exec.InstanceObjects.Add(DuplicateObject(Instance, &InOwner));
 	}
@@ -87,21 +85,24 @@ bool FStateTreeExecutionContext::InitInstanceData(UObject& InOwner, FStateTreeDa
 	return true;
 }
 
-EStateTreeRunStatus FStateTreeExecutionContext::Start(FStateTreeDataView ExternalStorage)
+EStateTreeRunStatus FStateTreeExecutionContext::Start(FStateTreeInstanceData* ExternalInstanceData)
 {
+	CSV_SCOPED_TIMING_STAT_EXCLUSIVE(StateTree_Start);
+
 	if (!Owner || !StateTree)
 	{
 		return EStateTreeRunStatus::Failed;
 	}
 	
-	FStateTreeDataView Storage = SelectMutableStorage(ExternalStorage);
-	FStateTreeExecutionState& Exec = GetExecState(Storage);
+	FStateTreeInstanceData& InstanceData = SelectMutableInstanceData(ExternalInstanceData);
+	checkSlow(InstanceData.GetLayout() == StateTree->InstanceDataDefaultValue.GetLayout());
+	FStateTreeExecutionState& Exec = GetExecState(InstanceData);
 
 	// Stop if still running previous state.
 	if (Exec.TreeRunStatus == EStateTreeRunStatus::Running)
 	{
 		const FStateTreeTransitionResult Transition(FStateTreeStateStatus(Exec.CurrentState, Exec.LastTickStatus), FStateTreeHandle::Succeeded);
-		ExitState(Storage, Transition);
+		ExitState(InstanceData, Transition);
 	}
 	
 	// Initialize to unset running state.
@@ -112,7 +113,7 @@ EStateTreeRunStatus FStateTreeExecutionContext::Start(FStateTreeDataView Externa
 	// Select new Starting from root.
 	VisitedStates.Init(false, StateTree->States.Num());
 	static const FStateTreeHandle RootState = FStateTreeHandle(0);
-	const FStateTreeTransitionResult Transition = SelectState(Storage, FStateTreeStateStatus(FStateTreeHandle::Invalid), RootState, RootState, 0);
+	const FStateTreeTransitionResult Transition = SelectState(InstanceData, FStateTreeStateStatus(FStateTreeHandle::Invalid), RootState, RootState, 0);
 
 	// Handle potential transition.
 	if (Transition.Next.IsValid())
@@ -127,12 +128,12 @@ EStateTreeRunStatus FStateTreeExecutionContext::Start(FStateTreeDataView Externa
 		else
 		{
 			// Enter state tasks can fail/succeed, treat it same as tick.
-			Exec.LastTickStatus = EnterState(Storage, Transition);
+			Exec.LastTickStatus = EnterState(InstanceData, Transition);
 			
 			// Report state completed immediately.
 			if (Exec.LastTickStatus != EStateTreeRunStatus::Running)
 			{
-				StateCompleted(Storage, Exec.CurrentState, Exec.LastTickStatus);
+				StateCompleted(InstanceData, Exec.CurrentState, Exec.LastTickStatus);
 			}
 		}
 	}
@@ -148,22 +149,25 @@ EStateTreeRunStatus FStateTreeExecutionContext::Start(FStateTreeDataView Externa
 	return Exec.TreeRunStatus;
 }
 
-EStateTreeRunStatus FStateTreeExecutionContext::Stop(FStateTreeDataView ExternalStorage)
+EStateTreeRunStatus FStateTreeExecutionContext::Stop(FStateTreeInstanceData* ExternalInstanceData)
 {
+	CSV_SCOPED_TIMING_STAT_EXCLUSIVE(StateTree_Stop);
+
 	if (!Owner || !StateTree)
 	{
 		return EStateTreeRunStatus::Failed;
 	}
 	
-	FStateTreeDataView Storage = SelectMutableStorage(ExternalStorage);
-	FStateTreeExecutionState& Exec = GetExecState(Storage);
+	FStateTreeInstanceData& InstanceData = SelectMutableInstanceData(ExternalInstanceData);
+	checkSlow(InstanceData.GetLayout() == StateTree->InstanceDataDefaultValue.GetLayout());
+	FStateTreeExecutionState& Exec = GetExecState(InstanceData);
 
 	// Exit states if still in some valid state.
 	if (Exec.CurrentState.IsValid() && (Exec.CurrentState != FStateTreeHandle::Succeeded || Exec.CurrentState != FStateTreeHandle::Failed))
 	{
 		// Transition to Succeeded state.
 		const FStateTreeTransitionResult Transition(FStateTreeStateStatus(Exec.CurrentState, Exec.LastTickStatus), FStateTreeHandle::Succeeded);
-		ExitState(Storage, Transition);
+		ExitState(InstanceData, Transition);
 		Exec.CurrentState = FStateTreeHandle::Succeeded;
 	}
 
@@ -173,16 +177,17 @@ EStateTreeRunStatus FStateTreeExecutionContext::Stop(FStateTreeDataView External
 	return Exec.TreeRunStatus;
 }
 
-EStateTreeRunStatus FStateTreeExecutionContext::Tick(const float DeltaTime, FStateTreeDataView ExternalStorage)
+EStateTreeRunStatus FStateTreeExecutionContext::Tick(const float DeltaTime, FStateTreeInstanceData* ExternalInstanceData)
 {
-	CSV_SCOPED_TIMING_STAT_EXCLUSIVE(StateTreeCtxTick);
+	CSV_SCOPED_TIMING_STAT_EXCLUSIVE(StateTree_Tick);
 
 	if (!Owner || !StateTree)
 	{
 		return EStateTreeRunStatus::Failed;
 	}
-	FStateTreeDataView Storage = SelectMutableStorage(ExternalStorage);
-	FStateTreeExecutionState& Exec = GetExecState(Storage);
+	FStateTreeInstanceData& InstanceData = SelectMutableInstanceData(ExternalInstanceData);
+	checkSlow(InstanceData.GetLayout() == StateTree->InstanceDataDefaultValue.GetLayout());
+	FStateTreeExecutionState& Exec = GetExecState(InstanceData);
 
 	// No ticking of the tree is done or stopped.
 	if (Exec.TreeRunStatus != EStateTreeRunStatus::Running)
@@ -203,15 +208,15 @@ EStateTreeRunStatus FStateTreeExecutionContext::Tick(const float DeltaTime, FSta
 	{
 		// Tick evaluators on active states. Further evaluator may be ticked during selection as non-active states are visited.
 		// Ticked states are kept track of and evaluators are ticked just once per frame.
-		TickEvaluators(Storage, Exec.CurrentState, EStateTreeEvaluationType::Tick, DeltaTime);
+		TickEvaluators(InstanceData, Exec.CurrentState, EStateTreeEvaluationType::Tick, DeltaTime);
 
 		// Tick tasks on active states.
-		Exec.LastTickStatus = TickTasks(Storage, Exec.CurrentState, DeltaTime);
+		Exec.LastTickStatus = TickTasks(InstanceData, Exec.CurrentState, DeltaTime);
 		
 		// Report state completed immediately.
 		if (Exec.LastTickStatus != EStateTreeRunStatus::Running)
 		{
-			StateCompleted(Storage, Exec.CurrentState, Exec.LastTickStatus);
+			StateCompleted(InstanceData, Exec.CurrentState, Exec.LastTickStatus);
 		}
 	}
 
@@ -221,12 +226,12 @@ EStateTreeRunStatus FStateTreeExecutionContext::Tick(const float DeltaTime, FSta
 	for (int32 Iter = 0; Iter < MaxIterations; Iter++)
 	{
 		// Trigger conditional transitions or state succeed/failed transitions. First tick transition is handled here too.
-		FStateTreeTransitionResult Transition = TriggerTransitions(Storage, FStateTreeStateStatus(Exec.CurrentState, Exec.LastTickStatus), 0);
+		FStateTreeTransitionResult Transition = TriggerTransitions(InstanceData, FStateTreeStateStatus(Exec.CurrentState, Exec.LastTickStatus), 0);
 
 		// Handle potential transition.
 		if (Transition.Next.IsValid())
 		{
-			ExitState(Storage, Transition);
+			ExitState(InstanceData, Transition);
 			
 			if (Transition.Next == FStateTreeHandle::Succeeded || Transition.Next == FStateTreeHandle::Failed)
 			{
@@ -236,12 +241,12 @@ EStateTreeRunStatus FStateTreeExecutionContext::Tick(const float DeltaTime, FSta
 			}
 			
 			// Enter state tasks can fail/succeed, treat it same as tick.
-			Exec.LastTickStatus = EnterState(Storage, Transition);
+			Exec.LastTickStatus = EnterState(InstanceData, Transition);
 
 			// Report state completed immediately.
 			if (Exec.LastTickStatus != EStateTreeRunStatus::Running)
 			{
-				StateCompleted(Storage, Exec.CurrentState, Exec.LastTickStatus);
+				StateCompleted(InstanceData, Exec.CurrentState, Exec.LastTickStatus);
 			}
 		}
 		
@@ -267,24 +272,16 @@ EStateTreeRunStatus FStateTreeExecutionContext::Tick(const float DeltaTime, FSta
 	return Exec.TreeRunStatus;
 }
 
-void FStateTreeExecutionContext::Reset()
+EStateTreeRunStatus FStateTreeExecutionContext::EnterState(FStateTreeInstanceData& InstanceData, const FStateTreeTransitionResult& Transition)
 {
-	StateTree = nullptr;
-	Owner = nullptr;
-	DataViews.Reset();
-	StorageInstance.Reset();
-	StorageType = EStateTreeStorage::Internal;
-	VisitedStates.Reset();
-}
+	CSV_SCOPED_TIMING_STAT_EXCLUSIVE(StateTree_EnterState);
 
-EStateTreeRunStatus FStateTreeExecutionContext::EnterState(FStateTreeDataView Storage, const FStateTreeTransitionResult& Transition)
-{
 	if (!Transition.Next.IsValid())
 	{
 		return EStateTreeRunStatus::Failed;
 	}
 
-	FStateTreeExecutionState& Exec = GetExecState(Storage);
+	FStateTreeExecutionState& Exec = GetExecState(InstanceData);
 
 	// Activate evaluators along the active branch.
 	TStaticArray<FStateTreeHandle, 32> States;
@@ -321,13 +318,13 @@ EStateTreeRunStatus FStateTreeExecutionContext::EnterState(FStateTreeDataView St
 			// Must update data views.
 			for (int32 EvalIndex = 0; EvalIndex < int32(State.EvaluatorsNum); EvalIndex++)
 			{
-				const FStateTreeEvaluatorBase& Eval = GetItem<FStateTreeEvaluatorBase>(int32(State.EvaluatorsBegin) + EvalIndex);
-				DataViews[Eval.DataViewIndex] = GetInstanceData(Storage, Eval.bInstanceIsObject, Eval.InstanceIndex);
+				const FStateTreeEvaluatorBase& Eval = GetNode<FStateTreeEvaluatorBase>(int32(State.EvaluatorsBegin) + EvalIndex);
+				DataViews[Eval.DataViewIndex] = GetInstanceData(InstanceData, Eval.bInstanceIsObject, Eval.InstanceIndex);
 			}
 			for (int32 TaskIndex = 0; TaskIndex < int32(State.TasksNum); TaskIndex++)
 			{
-				const FStateTreeTaskBase& Task = GetItem<FStateTreeTaskBase>(int32(State.TasksBegin) + TaskIndex);
-				DataViews[Task.DataViewIndex] = GetInstanceData(Storage, Task.bInstanceIsObject, Task.InstanceIndex);
+				const FStateTreeTaskBase& Task = GetNode<FStateTreeTaskBase>(int32(State.TasksBegin) + TaskIndex);
+				DataViews[Task.DataViewIndex] = GetInstanceData(InstanceData, Task.bInstanceIsObject, Task.InstanceIndex);
 			}
 			continue;
 		}
@@ -340,8 +337,8 @@ EStateTreeRunStatus FStateTreeExecutionContext::EnterState(FStateTreeDataView St
 		
 		for (int32 EvalIndex = State.EvaluatorsBegin; EvalIndex < (int32)(State.EvaluatorsBegin + State.EvaluatorsNum); EvalIndex++)
 		{
-			const FStateTreeEvaluatorBase& Eval = GetItem<FStateTreeEvaluatorBase>(EvalIndex);
-			const FStateTreeDataView EvalInstanceView = GetInstanceData(Storage, Eval.bInstanceIsObject, Eval.InstanceIndex);
+			const FStateTreeEvaluatorBase& Eval = GetNode<FStateTreeEvaluatorBase>(EvalIndex);
+			const FStateTreeDataView EvalInstanceView = GetInstanceData(InstanceData, Eval.bInstanceIsObject, Eval.InstanceIndex);
 			DataViews[Eval.DataViewIndex] = EvalInstanceView;
 
 			// Copy bound properties.
@@ -350,14 +347,18 @@ EStateTreeRunStatus FStateTreeExecutionContext::EnterState(FStateTreeDataView St
 				StateTree->PropertyBindings.CopyTo(DataViews, Eval.BindingsBatch.Index, EvalInstanceView);
 			}
 			STATETREE_LOG(Verbose, TEXT("%*s  Notify Evaluator '%s'."), Index*UE::StateTree::DebugIndentSize, TEXT(""), *Eval.Name.ToString());
-			Eval.EnterState(*this, ChangeType, CurrentTransition);
+			{
+				QUICK_SCOPE_CYCLE_COUNTER(StateTree_Eval_EnterState);
+				CSV_SCOPED_TIMING_STAT_EXCLUSIVE(StateTree_Eval_EnterState);
+				Eval.EnterState(*this, ChangeType, CurrentTransition);
+			}
 		}
 		
 		// Activate tasks on current state.
 		for (int32 TaskIndex = State.TasksBegin; TaskIndex < (int32)(State.TasksBegin + State.TasksNum); TaskIndex++)
 		{
-			const FStateTreeTaskBase& Task = GetItem<FStateTreeTaskBase>(TaskIndex);
-			const FStateTreeDataView TaskInstanceView = GetInstanceData(Storage, Task.bInstanceIsObject, Task.InstanceIndex);
+			const FStateTreeTaskBase& Task = GetNode<FStateTreeTaskBase>(TaskIndex);
+			const FStateTreeDataView TaskInstanceView = GetInstanceData(InstanceData, Task.bInstanceIsObject, Task.InstanceIndex);
 			DataViews[Task.DataViewIndex] = TaskInstanceView;
 
 			// Copy bound properties.
@@ -367,15 +368,20 @@ EStateTreeRunStatus FStateTreeExecutionContext::EnterState(FStateTreeDataView St
 			}
 		
 			STATETREE_LOG(Verbose, TEXT("%*s  Notify Task '%s'"), Index*UE::StateTree::DebugIndentSize, TEXT(""), *Task.Name.ToString());
-			const EStateTreeRunStatus Status = Task.EnterState(*this, ChangeType, CurrentTransition);
-			
-			if (Status == EStateTreeRunStatus::Failed)
+
 			{
-				// Store how far in the enter state we got. This will be used to match the ExitState() calls.
-				Exec.EnterStateFailedTaskIndex = TaskIndex;
-				Exec.CurrentState = CurrentHandle;
-				Result = Status;
-				break;
+				QUICK_SCOPE_CYCLE_COUNTER(StateTree_Task_EnterState);
+				CSV_SCOPED_TIMING_STAT_EXCLUSIVE(StateTree_Task_EnterState);
+				const EStateTreeRunStatus Status = Task.EnterState(*this, ChangeType, CurrentTransition);
+				
+				if (Status == EStateTreeRunStatus::Failed)
+				{
+					// Store how far in the enter state we got. This will be used to match the ExitState() calls.
+					Exec.EnterStateFailedTaskIndex = TaskIndex;
+					Exec.CurrentState = CurrentHandle;
+					Result = Status;
+					break;
+				}
 			}
 		}
 	}
@@ -388,15 +394,17 @@ EStateTreeRunStatus FStateTreeExecutionContext::EnterState(FStateTreeDataView St
 	return Result;
 }
 
-void FStateTreeExecutionContext::ExitState(FStateTreeDataView Storage, const FStateTreeTransitionResult& Transition)
+void FStateTreeExecutionContext::ExitState(FStateTreeInstanceData& InstanceData, const FStateTreeTransitionResult& Transition)
 {
+	CSV_SCOPED_TIMING_STAT_EXCLUSIVE(StateTree_ExitState);
+
 	if (!Transition.Source.State.IsValid())
 	{
 		return;
 	}
 
 	// Reset transition delay
-	FStateTreeExecutionState& Exec = GetExecState(Storage);
+	FStateTreeExecutionState& Exec = GetExecState(InstanceData);
 	Exec.GatedTransitionIndex = INDEX_NONE;
 	Exec.GatedTransitionTime = 0.0f;
 
@@ -433,13 +441,13 @@ void FStateTreeExecutionContext::ExitState(FStateTreeDataView Storage, const FSt
 			// Must update item views.
 			for (int32 EvalIndex = 0; EvalIndex < int32(State.EvaluatorsNum); EvalIndex++)
 			{
-				FStateTreeEvaluatorBase& Eval = GetItem<FStateTreeEvaluatorBase>(int32(State.EvaluatorsBegin) + EvalIndex);
-				DataViews[Eval.DataViewIndex] = GetInstanceData(Storage, Eval.bInstanceIsObject, Eval.InstanceIndex);
+				const FStateTreeEvaluatorBase& Eval = GetNode<FStateTreeEvaluatorBase>(int32(State.EvaluatorsBegin) + EvalIndex);
+				DataViews[Eval.DataViewIndex] = GetInstanceData(InstanceData, Eval.bInstanceIsObject, Eval.InstanceIndex);
 			}
 			for (int32 TaskIndex = 0; TaskIndex < int32(State.TasksNum); TaskIndex++)
 			{
-				FStateTreeTaskBase& Task = GetItem<FStateTreeTaskBase>(int32(State.TasksBegin) + TaskIndex);
-				DataViews[Task.DataViewIndex] = GetInstanceData(Storage, Task.bInstanceIsObject, Task.InstanceIndex);
+				const FStateTreeTaskBase& Task = GetNode<FStateTreeTaskBase>(int32(State.TasksBegin) + TaskIndex);
+				DataViews[Task.DataViewIndex] = GetInstanceData(InstanceData, Task.bInstanceIsObject, Task.InstanceIndex);
 			}
 			continue;
 		}
@@ -452,8 +460,8 @@ void FStateTreeExecutionContext::ExitState(FStateTreeDataView Storage, const FSt
 
 		for (int32 EvalIndex = State.EvaluatorsBegin; EvalIndex < (int32)(State.EvaluatorsBegin + State.EvaluatorsNum); EvalIndex++)
 		{
-			const FStateTreeEvaluatorBase& Eval = GetItem<FStateTreeEvaluatorBase>(EvalIndex);
-			const FStateTreeDataView EvalInstanceView = GetInstanceData(Storage, Eval.bInstanceIsObject, Eval.InstanceIndex);
+			const FStateTreeEvaluatorBase& Eval = GetNode<FStateTreeEvaluatorBase>(EvalIndex);
+			const FStateTreeDataView EvalInstanceView = GetInstanceData(InstanceData, Eval.bInstanceIsObject, Eval.InstanceIndex);
 			DataViews[Eval.DataViewIndex] = EvalInstanceView;
 
 			// Copy bound properties.
@@ -462,7 +470,11 @@ void FStateTreeExecutionContext::ExitState(FStateTreeDataView Storage, const FSt
 				StateTree->PropertyBindings.CopyTo(DataViews, Eval.BindingsBatch.Index, EvalInstanceView);
 			}
 			STATETREE_LOG(Verbose, TEXT("%*s  Notify Evaluator '%s'."), Index*UE::StateTree::DebugIndentSize, TEXT(""), *Eval.Name.ToString());
-			Eval.ExitState(*this, ChangeType, CurrentTransition);
+			{
+				QUICK_SCOPE_CYCLE_COUNTER(StateTree_Eval_ExitState);
+				CSV_SCOPED_TIMING_STAT_EXCLUSIVE(StateTree_Eval_ExitState);
+				Eval.ExitState(*this, ChangeType, CurrentTransition);
+			}
 		}
 
 		// Deactivate tasks on current State
@@ -472,8 +484,8 @@ void FStateTreeExecutionContext::ExitState(FStateTreeDataView Storage, const FSt
 			// The task order in the tree (BF) allows us to use the comparison.
 			if (TaskIndex <= (int32)Exec.EnterStateFailedTaskIndex)
 			{
-				const FStateTreeTaskBase& Task = GetItem<FStateTreeTaskBase>(TaskIndex);
-				const FStateTreeDataView TaskInstanceView = GetInstanceData(Storage, Task.bInstanceIsObject, Task.InstanceIndex);
+				const FStateTreeTaskBase& Task = GetNode<FStateTreeTaskBase>(TaskIndex);
+				const FStateTreeDataView TaskInstanceView = GetInstanceData(InstanceData, Task.bInstanceIsObject, Task.InstanceIndex);
 				DataViews[Task.DataViewIndex] = TaskInstanceView;
 
 				// Copy bound properties.
@@ -483,14 +495,20 @@ void FStateTreeExecutionContext::ExitState(FStateTreeDataView Storage, const FSt
 				}
 
 				STATETREE_LOG(Verbose, TEXT("%*s  Notify Task '%s'"), Index*UE::StateTree::DebugIndentSize, TEXT(""), *Task.Name.ToString());
-				Task.ExitState(*this, ChangeType, CurrentTransition);
+				{
+					QUICK_SCOPE_CYCLE_COUNTER(StateTree_Task_ExitState);
+					CSV_SCOPED_TIMING_STAT_EXCLUSIVE(StateTree_Task_ExitState);
+					Task.ExitState(*this, ChangeType, CurrentTransition);
+				}
 			}
 		}
 	}
 }
 
-void FStateTreeExecutionContext::StateCompleted(FStateTreeDataView Storage, const FStateTreeHandle CurrentState, const EStateTreeRunStatus CompletionStatus)
+void FStateTreeExecutionContext::StateCompleted(FStateTreeInstanceData& InstanceData, const FStateTreeHandle CurrentState, const EStateTreeRunStatus CompletionStatus)
 {
+	CSV_SCOPED_TIMING_STAT_EXCLUSIVE(StateTree_StateCompleted);
+
 	if (!CurrentState.IsValid())
 	{
 		return;
@@ -499,7 +517,7 @@ void FStateTreeExecutionContext::StateCompleted(FStateTreeDataView Storage, cons
 	TStaticArray<FStateTreeHandle, 32> States;
 	const int32 NumStates = GetActiveStates(CurrentState, States);
 
-	FStateTreeExecutionState& Exec = GetExecState(Storage);
+	FStateTreeExecutionState& Exec = GetExecState(InstanceData);
 
 	// Call from child towards root to allow to pass results back.
 	// Note: Completed is assumed to be called immediately after tick or enter state, so there's no property copying.
@@ -517,7 +535,7 @@ void FStateTreeExecutionContext::StateCompleted(FStateTreeDataView Storage, cons
 			// The task order in the tree (BF) allows us to use the comparison.
 			if (TaskIndex <= (int32)Exec.EnterStateFailedTaskIndex)
 			{
-				const FStateTreeTaskBase& Task = GetItem<FStateTreeTaskBase>(TaskIndex);
+				const FStateTreeTaskBase& Task = GetNode<FStateTreeTaskBase>(TaskIndex);
 
 				STATETREE_LOG(Verbose, TEXT("%*s  Notify Task '%s'"), Index*UE::StateTree::DebugIndentSize, TEXT(""), *Task.Name.ToString());
 				Task.StateCompleted(*this, CompletionStatus, CurrentState);
@@ -527,7 +545,7 @@ void FStateTreeExecutionContext::StateCompleted(FStateTreeDataView Storage, cons
 		// Notify evaluators
 		for (int32 EvalIndex = (int32)(State.EvaluatorsBegin + State.EvaluatorsNum) - 1; EvalIndex >= State.EvaluatorsBegin; EvalIndex--)
 		{
-			const FStateTreeEvaluatorBase& Eval = GetItem<FStateTreeEvaluatorBase>(EvalIndex);
+			const FStateTreeEvaluatorBase& Eval = GetNode<FStateTreeEvaluatorBase>(EvalIndex);
 
 			STATETREE_LOG(Verbose, TEXT("%*s  Notify Evaluator '%s'"), Index*UE::StateTree::DebugIndentSize, TEXT(""), *Eval.Name.ToString());
 			Eval.StateCompleted(*this, CompletionStatus, CurrentState);
@@ -535,9 +553,9 @@ void FStateTreeExecutionContext::StateCompleted(FStateTreeDataView Storage, cons
 	}
 }
 
-void FStateTreeExecutionContext::TickEvaluators(FStateTreeDataView Storage, const FStateTreeHandle CurrentState, const EStateTreeEvaluationType EvalType, const float DeltaTime)
+void FStateTreeExecutionContext::TickEvaluators(FStateTreeInstanceData& InstanceData, const FStateTreeHandle CurrentState, const EStateTreeEvaluationType EvalType, const float DeltaTime)
 {
-	CSV_SCOPED_TIMING_STAT_EXCLUSIVE(StateTreeCtxTickEvaluators);
+	CSV_SCOPED_TIMING_STAT_EXCLUSIVE(StateTree_TickEvaluators);
 
 	if (!CurrentState.IsValid())
 	{
@@ -563,8 +581,8 @@ void FStateTreeExecutionContext::TickEvaluators(FStateTreeDataView Storage, cons
 		// Tick evaluators
 		for (int32 EvalIndex = 0; EvalIndex < int32(State.EvaluatorsNum); EvalIndex++)
 		{
-			const FStateTreeEvaluatorBase& Eval = GetItem<FStateTreeEvaluatorBase>(int32(State.EvaluatorsBegin) + EvalIndex);
-			const FStateTreeDataView EvalInstanceView = GetInstanceData(Storage, Eval.bInstanceIsObject, Eval.InstanceIndex);
+			const FStateTreeEvaluatorBase& Eval = GetNode<FStateTreeEvaluatorBase>(int32(State.EvaluatorsBegin) + EvalIndex);
+			const FStateTreeDataView EvalInstanceView = GetInstanceData(InstanceData, Eval.bInstanceIsObject, Eval.InstanceIndex);
 			DataViews[Eval.DataViewIndex] = EvalInstanceView;
 
 			// Copy bound properties.
@@ -573,16 +591,19 @@ void FStateTreeExecutionContext::TickEvaluators(FStateTreeDataView Storage, cons
 				StateTree->PropertyBindings.CopyTo(DataViews, Eval.BindingsBatch.Index, EvalInstanceView);
 			}
 			STATETREE_LOG(Verbose, TEXT("%*s  Evaluate: '%s'"), Index*UE::StateTree::DebugIndentSize, TEXT(""), *Eval.Name.ToString());
-			Eval.Evaluate(*this, EvalType, DeltaTime);
+			{
+				QUICK_SCOPE_CYCLE_COUNTER(StateTree_Eval_Tick);
+				Eval.Evaluate(*this, EvalType, DeltaTime);
+			}
 		}
 
 		VisitedStates[CurrentHandle.Index] = true;
 	}
 }
 
-EStateTreeRunStatus FStateTreeExecutionContext::TickTasks(FStateTreeDataView Storage, const FStateTreeHandle CurrentState, const float DeltaTime)
+EStateTreeRunStatus FStateTreeExecutionContext::TickTasks(FStateTreeInstanceData& InstanceData, const FStateTreeHandle CurrentState, const float DeltaTime)
 {
-	CSV_SCOPED_TIMING_STAT_EXCLUSIVE(StateTreeCtxTickTasks);
+	CSV_SCOPED_TIMING_STAT_EXCLUSIVE(StateTree_TickTasks);
 
 	if (!CurrentState.IsValid())
 	{
@@ -605,8 +626,8 @@ EStateTreeRunStatus FStateTreeExecutionContext::TickTasks(FStateTreeDataView Sto
 		// Tick Tasks
 		for (int32 TaskIndex = 0; TaskIndex < int32(State.TasksNum); TaskIndex++)
 		{
-			const FStateTreeTaskBase& Task = GetItem<FStateTreeTaskBase>(int32(State.TasksBegin) + TaskIndex);
-			const FStateTreeDataView TaskInstanceView = GetInstanceData(Storage, Task.bInstanceIsObject, Task.InstanceIndex);
+			const FStateTreeTaskBase& Task = GetNode<FStateTreeTaskBase>(int32(State.TasksBegin) + TaskIndex);
+			const FStateTreeDataView TaskInstanceView = GetInstanceData(InstanceData, Task.bInstanceIsObject, Task.InstanceIndex);
 			DataViews[Task.DataViewIndex] = TaskInstanceView;
 
 			// Copy bound properties.
@@ -615,16 +636,22 @@ EStateTreeRunStatus FStateTreeExecutionContext::TickTasks(FStateTreeDataView Sto
 				StateTree->PropertyBindings.CopyTo(DataViews, Task.BindingsBatch.Index, TaskInstanceView);
 			}
 			STATETREE_LOG(Verbose, TEXT("%*s  Tick: '%s'"), Index*UE::StateTree::DebugIndentSize, TEXT(""), *Task.Name.ToString());
-			const EStateTreeRunStatus TaskResult = Task.Tick(*this, DeltaTime);
-
-			// TODO: Add more control over which states can control the failed/succeeded result.
-			if (TaskResult != EStateTreeRunStatus::Running)
-			{
-				Result = TaskResult;
-			}
-			if (TaskResult == EStateTreeRunStatus::Failed)
+			
 			{
-				break;
+				QUICK_SCOPE_CYCLE_COUNTER(StateTree_Task_Tick);
+				CSV_SCOPED_TIMING_STAT_EXCLUSIVE(StateTree_Task_Tick);
+
+				const EStateTreeRunStatus TaskResult = Task.Tick(*this, DeltaTime);
+
+				// TODO: Add more control over which states can control the failed/succeeded result.
+				if (TaskResult != EStateTreeRunStatus::Running)
+				{
+					Result = TaskResult;
+				}
+				if (TaskResult == EStateTreeRunStatus::Failed)
+				{
+					break;
+				}
 			}
 		}
 		NumTotalTasks += State.TasksNum;
@@ -639,14 +666,14 @@ EStateTreeRunStatus FStateTreeExecutionContext::TickTasks(FStateTreeDataView Sto
 	return Result;
 }
 
-bool FStateTreeExecutionContext::TestAllConditions(FStateTreeDataView Storage, const uint32 ConditionsOffset, const uint32 ConditionsNum)
+bool FStateTreeExecutionContext::TestAllConditions(FStateTreeInstanceData& InstanceData, const uint32 ConditionsOffset, const uint32 ConditionsNum)
 {
-	CSV_SCOPED_TIMING_STAT_EXCLUSIVE(StateTreeCtxTestConditions);
+	CSV_SCOPED_TIMING_STAT_EXCLUSIVE(StateTree_TestConditions);
 
 	for (uint32 i = 0; i < ConditionsNum; i++)
 	{
-		const FStateTreeConditionBase& Cond = GetItem<FStateTreeConditionBase>(ConditionsOffset + i);
-		const FStateTreeDataView CondInstanceView = GetInstanceData(Storage, Cond.bInstanceIsObject, Cond.InstanceIndex);
+		const FStateTreeConditionBase& Cond = GetNode<FStateTreeConditionBase>(ConditionsOffset + i);
+		const FStateTreeDataView CondInstanceView = GetInstanceData(InstanceData, Cond.bInstanceIsObject, Cond.InstanceIndex);
 		DataViews[Cond.DataViewIndex] = CondInstanceView;
 
 		// Copy bound properties.
@@ -663,9 +690,11 @@ bool FStateTreeExecutionContext::TestAllConditions(FStateTreeDataView Storage, c
 	return true;
 }
 
-FStateTreeTransitionResult FStateTreeExecutionContext::TriggerTransitions(FStateTreeDataView Storage, const FStateTreeStateStatus CurrentStatus, const int32 Depth)
+FStateTreeTransitionResult FStateTreeExecutionContext::TriggerTransitions(FStateTreeInstanceData& InstanceData, const FStateTreeStateStatus CurrentStatus, const int32 Depth)
 {
-	FStateTreeExecutionState& Exec = GetExecState(Storage);
+	CSV_SCOPED_TIMING_STAT_EXCLUSIVE(StateTree_TriggerTransition);
+
+	FStateTreeExecutionState& Exec = GetExecState(InstanceData);
 	EStateTreeTransitionEvent Event = EStateTreeTransitionEvent::OnCondition;
 	if (CurrentStatus.RunStatus == EStateTreeRunStatus::Succeeded)
 	{
@@ -685,7 +714,7 @@ FStateTreeTransitionResult FStateTreeExecutionContext::TriggerTransitions(FState
 			// All transition conditions must pass
 			const int16 TransitionIndex = State.TransitionsBegin + i;
 			const FBakedStateTransition& Transition = StateTree->Transitions[TransitionIndex];
-			if (EnumHasAllFlags(Transition.Event, Event) && TestAllConditions(Storage, Transition.ConditionsBegin, Transition.ConditionsNum))
+			if (EnumHasAllFlags(Transition.Event, Event) && TestAllConditions(InstanceData, Transition.ConditionsBegin, Transition.ConditionsNum))
 			{
 				// If a transition has delay, we stop testing other transitions, but the transition will not pass the condition until the delay time passes.
 				if (Transition.GateDelay > 0)
@@ -709,7 +738,7 @@ FStateTreeTransitionResult FStateTreeExecutionContext::TriggerTransitions(FState
 				
 				if (Transition.Type == EStateTreeTransitionType::GotoState || Transition.Type == EStateTreeTransitionType::NextState)
 				{
-					FStateTreeTransitionResult Result = SelectState(Storage, CurrentStatus, Transition.State, Transition.State, Depth + 1);
+					FStateTreeTransitionResult Result = SelectState(InstanceData, CurrentStatus, Transition.State, Transition.State, Depth + 1);
 					if (Result.Next.IsValid())
 					{
 						STATETREE_LOG(Verbose, TEXT("Transition on state '%s' (%s) -[%s]-> state '%s'"), *GetSafeStateName(CurrentStatus.State), *State.Name.ToString(), *GetSafeStateName(Result.Target), *GetSafeStateName(Result.Next));
@@ -752,8 +781,10 @@ FStateTreeTransitionResult FStateTreeExecutionContext::TriggerTransitions(FState
 	return FStateTreeTransitionResult(CurrentStatus, FStateTreeHandle::Invalid);
 }
 
-FStateTreeTransitionResult FStateTreeExecutionContext::SelectState(FStateTreeDataView Storage, const FStateTreeStateStatus InitialStateStatus, const FStateTreeHandle InitialTargetState, const FStateTreeHandle NextState, const int Depth)
+FStateTreeTransitionResult FStateTreeExecutionContext::SelectState(FStateTreeInstanceData& InstanceData, const FStateTreeStateStatus InitialStateStatus, const FStateTreeHandle InitialTargetState, const FStateTreeHandle NextState, const int Depth)
 {
+	CSV_SCOPED_TIMING_STAT_EXCLUSIVE(StateTree_SelectState);
+
 	if (!NextState.IsValid())
 	{
 		// Trying to select non-existing state.
@@ -765,17 +796,17 @@ FStateTreeTransitionResult FStateTreeExecutionContext::SelectState(FStateTreeDat
 	const FBakedStateTreeState& State = StateTree->States[NextState.Index];
 
 	// Make sure all the evaluators for the target state are up to date.
-	TickEvaluators(Storage, NextState, EStateTreeEvaluationType::PreSelect, 0.0f);
+	TickEvaluators(InstanceData, NextState, EStateTreeEvaluationType::PreSelect, 0.0f);
 	
 	// Check that the state can be entered
-	if (TestAllConditions(Storage, State.EnterConditionsBegin, State.EnterConditionsNum))
+	if (TestAllConditions(InstanceData, State.EnterConditionsBegin, State.EnterConditionsNum))
 	{
 		// If the state has children, proceed to select children.
 		if (State.HasChildren())
 		{
 			for (uint16 ChildState = State.ChildrenBegin; ChildState < State.ChildrenEnd; ChildState = StateTree->States[ChildState].GetNextSibling())
 			{
-				const FStateTreeTransitionResult ChildResult = SelectState(Storage, InitialStateStatus, InitialTargetState, FStateTreeHandle(ChildState), Depth + 1);
+				const FStateTreeTransitionResult ChildResult = SelectState(InstanceData, InitialStateStatus, InitialTargetState, FStateTreeHandle(ChildState), Depth + 1);
 				if (ChildResult.Next.IsValid())
 				{
 					// Selection succeeded
@@ -871,39 +902,45 @@ FString FStateTreeExecutionContext::GetStateStatusString(const FStateTreeStateSt
 	}
 }
 
-EStateTreeRunStatus FStateTreeExecutionContext::GetLastTickStatus(FStateTreeDataView ExternalStorage) const
+EStateTreeRunStatus FStateTreeExecutionContext::GetLastTickStatus(const FStateTreeInstanceData* ExternalInstanceData) const
 {
-	const FStateTreeExecutionState& Exec = GetExecState(SelectStorage(ExternalStorage));
+	const FStateTreeExecutionState& Exec = GetExecState(SelectInstanceData(ExternalInstanceData));
 	return Exec.LastTickStatus;
 }
 
-void FStateTreeExecutionContext::AddStructReferencedObjects(FReferenceCollector& Collector)
+void FStateTreeExecutionContext::AddStructReferencedObjects(FReferenceCollector& Collector) const
 {
 	if (StorageType == EStateTreeStorage::Internal)
 	{
-		AddStructReferencedObjects(FStateTreeDataView(StorageInstance), Collector);
+		AddStructReferencedObjects(&InternalInstanceData, Collector);
 	}
 }
 
-void FStateTreeExecutionContext::AddStructReferencedObjects(const FStateTreeDataView Storage, FReferenceCollector& Collector)
+void FStateTreeExecutionContext::AddStructReferencedObjects(const FStateTreeInstanceData* InstanceData, FReferenceCollector& Collector) const
 {
-	if (Storage.IsValid())
+	if (InstanceData && InstanceData->IsValid())
 	{
-		FStateTreeExecutionState& Exec = GetExecState(Storage);
-		Collector.AddReferencedObjects(Exec.InstanceObjects);
+		const FStateTreeExecutionState& Exec = GetExecState(*InstanceData);
+		for (const UObject* Object : Exec.InstanceObjects)
+		{
+			Collector.AddReferencedObject(Object);
+		}
+		
+		InstanceData->AddStructReferencedObjects(Collector);
 	}
 }
 
 #if WITH_GAMEPLAY_DEBUGGER
 
-FString FStateTreeExecutionContext::GetDebugInfoString(FStateTreeDataView ExternalStorage) const
+FString FStateTreeExecutionContext::GetDebugInfoString(const FStateTreeInstanceData* ExternalInstanceData) const
 {
 	if (!StateTree)
 	{
 		return FString(TEXT("No StateTree asset."));
 	}
-	FStateTreeDataView Storage = SelectStorage(ExternalStorage);
-	FStateTreeExecutionState& Exec = GetExecState(Storage);
+	const FStateTreeInstanceData& InstanceData = SelectInstanceData(ExternalInstanceData);
+	checkSlow(InstanceData.GetLayout() == StateTree->InstanceDataDefaultValue.GetLayout());
+	const FStateTreeExecutionState& Exec = GetExecState(InstanceData);
 
 	FString DebugString = FString::Printf(TEXT("StateTree (asset: '%s')\n"), *GetNameSafe(StateTree));
 
@@ -941,7 +978,7 @@ FString FStateTreeExecutionContext::GetDebugInfoString(FStateTreeDataView Extern
 				DebugString += TEXT("\nTasks:\n");
 				for (int32 j = 0; j < int32(State.TasksNum); j++)
 				{
-					const FStateTreeTaskBase& Task = GetItem<FStateTreeTaskBase>(int32(State.TasksBegin) + j);
+					const FStateTreeTaskBase& Task = GetNode<FStateTreeTaskBase>(int32(State.TasksBegin) + j);
 					Task.AppendDebugInfoString(DebugString, *this);
 				}
 			}
@@ -950,7 +987,7 @@ FString FStateTreeExecutionContext::GetDebugInfoString(FStateTreeDataView Extern
 				DebugString += TEXT("\nEvaluators:\n");
 				for (int32 EvalIndex = 0; EvalIndex < int32(State.EvaluatorsNum); EvalIndex++)
 				{
-					const FStateTreeEvaluatorBase& Eval = GetItem<FStateTreeEvaluatorBase>(int32(State.EvaluatorsBegin) + EvalIndex);
+					const FStateTreeEvaluatorBase& Eval = GetNode<FStateTreeEvaluatorBase>(int32(State.EvaluatorsBegin) + EvalIndex);
 					Eval.AppendDebugInfoString(DebugString, *this);
 				}
 			}
@@ -962,7 +999,7 @@ FString FStateTreeExecutionContext::GetDebugInfoString(FStateTreeDataView Extern
 #endif // WITH_GAMEPLAY_DEBUGGER
 
 #if WITH_STATETREE_DEBUG
-void FStateTreeExecutionContext::DebugPrintInternalLayout(FStateTreeDataView ExternalStorage)
+void FStateTreeExecutionContext::DebugPrintInternalLayout(const FStateTreeInstanceData* ExternalInstanceData)
 {
 	LOG_SCOPE_VERBOSITY_OVERRIDE(LogStateTree, ELogVerbosity::Log);
 
@@ -975,24 +1012,30 @@ void FStateTreeExecutionContext::DebugPrintInternalLayout(FStateTreeDataView Ext
 	FString DebugString = FString::Printf(TEXT("StateTree (asset: '%s')\n"), *GetNameSafe(StateTree));
 
 	// Tree items (e.g. tasks, evaluators, conditions)
-	DebugString += FString::Printf(TEXT("\nItems(%d)\n"), StateTree->Items.Num());
-	for (const FInstancedStruct& Item : StateTree->Items)
+	DebugString += FString::Printf(TEXT("\nItems(%d)\n"), StateTree->Nodes.Num());
+	for (const FInstancedStruct& Node : StateTree->Nodes)
 	{
-		DebugString += FString::Printf(TEXT("  %s\n"), Item.IsValid() ? *Item.GetScriptStruct()->GetName() : TEXT("null"));
+		DebugString += FString::Printf(TEXT("  %s\n"), Node.IsValid() ? *Node.GetScriptStruct()->GetName() : TEXT("null"));
 	}
 
-	// Instance storage data (e.g. tasks, evaluators, conditions)
+	// Instance InstanceData data (e.g. tasks, evaluators, conditions)
 	DebugString += FString::Printf(TEXT("\nInstance Data(%d)\n"), StateTree->Instances.Num());
 	for (const FInstancedStruct& Data : StateTree->Instances)
 	{
 		DebugString += FString::Printf(TEXT("  %s\n"), Data.IsValid() ? *Data.GetScriptStruct()->GetName() : TEXT("null"));
 	}
 
-	// Instance storage offsets (e.g. tasks, evaluators, conditions)
-	DebugString += FString::Printf(TEXT("\nInstance Data Offsets(%d)\n  [ %-40s | %-6s ]\n"), StateTree->InstanceStorageOffsets.Num(), TEXT("Name"), TEXT("Offset"));
-	for (const FStateTreeInstanceStorageOffset& Offset : StateTree->InstanceStorageOffsets)
+	// Instance InstanceData offsets (e.g. tasks, evaluators, conditions)
+
+	if (const FStateTreeInstanceDataLayout* InstanceLayout = StateTree->InstanceDataDefaultValue.GetLayout().Get())
 	{
-		DebugString += FString::Printf(TEXT("  | %-40s | %6d |\n"), Offset.Struct ? *Offset.Struct->GetName() : TEXT("null"),  Offset.Offset);
+		DebugString += FString::Printf(TEXT("\nInstance Data Offsets(%d)\n  [ %-40s | %-6s ]\n"), InstanceLayout->Num(), TEXT("Name"), TEXT("Offset"));
+
+		for (int32 Index = 0; Index < InstanceLayout->Num(); Index++)
+		{
+			const FStateTreeInstanceDataLayout::FLayoutItem& Item = InstanceLayout->GetItem(Index);
+			DebugString += FString::Printf(TEXT("  | %-40s | %6d |\n"), Item.ScriptStruct ? *Item.ScriptStruct->GetName() : TEXT("null"), Item.Offset);
+		}
 	}
 
 	// External data (e.g. fragments, subsystems)
@@ -1025,7 +1068,6 @@ void FStateTreeExecutionContext::DebugPrintInternalLayout(FStateTreeDataView Ext
 	}
 
 	// States
-	FStateTreeDataView Storage = SelectMutableStorage(ExternalStorage);
 	DebugString += FString::Printf(TEXT("\nStates(%d)\n"
 		"  [ %-30s | %15s | %5s [%3s:%-3s[ | Begin Idx : %4s %4s %4s %4s | Num : %4s %4s %4s %4s | Transitions : %-16s %-40s %-16s %-40s ]\n"),
 		StateTree->States.Num(),
@@ -1053,7 +1095,7 @@ void FStateTreeExecutionContext::DebugPrintInternalLayout(FStateTreeDataView Ext
 		{
 			for (int32 EvalIndex = 0; EvalIndex < State.EvaluatorsNum; EvalIndex++)
 			{
-				const FStateTreeEvaluatorBase& Eval = GetItem<FStateTreeEvaluatorBase>(State.EvaluatorsBegin + EvalIndex);
+				const FStateTreeEvaluatorBase& Eval = GetNode<FStateTreeEvaluatorBase>(State.EvaluatorsBegin + EvalIndex);
 				DebugString += FString::Printf(TEXT("  | %-30s | %-12s | %-30s | %15s | %10d |\n"), *State.Name.ToString(),
 					TEXT("  Evaluator"), *Eval.Name.ToString(), *Eval.BindingsBatch.Describe(), Eval.DataViewIndex);
 			}
@@ -1064,7 +1106,7 @@ void FStateTreeExecutionContext::DebugPrintInternalLayout(FStateTreeDataView Ext
 		{
 			for (int32 TaskIndex = 0; TaskIndex < State.TasksNum; TaskIndex++)
 			{
-				const FStateTreeTaskBase& Task = GetItem<FStateTreeTaskBase>(State.TasksBegin + TaskIndex);
+				const FStateTreeTaskBase& Task = GetNode<FStateTreeTaskBase>(State.TasksBegin + TaskIndex);
 				DebugString += FString::Printf(TEXT("  | %-30s | %-12s | %-30s | %15s | %10d |\n"), *State.Name.ToString(),
 					TEXT("  Task"), *Task.Name.ToString(), *Task.BindingsBatch.Describe(), Task.DataViewIndex);
 			}
@@ -1074,14 +1116,15 @@ void FStateTreeExecutionContext::DebugPrintInternalLayout(FStateTreeDataView Ext
 	UE_LOG(LogStateTree, Log, TEXT("%s"), *DebugString);
 }
 
-FString FStateTreeExecutionContext::GetActiveStateName(FStateTreeDataView ExternalStorage) const
+FString FStateTreeExecutionContext::GetActiveStateName(const FStateTreeInstanceData* ExternalInstanceData) const
 {
 	if (!StateTree)
 	{
 		return FString(TEXT("<None>"));
 	}
-	const FStateTreeDataView Storage = SelectStorage(ExternalStorage);
-	FStateTreeExecutionState& Exec = GetExecState(Storage);
+	const FStateTreeInstanceData& InstanceData = SelectInstanceData(ExternalInstanceData);
+	checkSlow(InstanceData.GetLayout() == StateTree->InstanceDataDefaultValue.GetLayout());
+	const FStateTreeExecutionState& Exec = GetExecState(InstanceData);
 
 	FString FullStateName;
 	
diff --git a/Engine/Plugins/Runtime/StateTree/Source/StateTreeModule/Private/StateTreeInstanceData.cpp b/Engine/Plugins/Runtime/StateTree/Source/StateTreeModule/Private/StateTreeInstanceData.cpp
new file mode 100644
--- /dev/null
+++ b/Engine/Plugins/Runtime/StateTree/Source/StateTreeModule/Private/StateTreeInstanceData.cpp
@@ -0,0 +1,173 @@
+﻿// Copyright Epic Games, Inc. All Rights Reserved.
+
+#include "StateTreeInstanceData.h"
+#include "Algo/Transform.h"
+
+//----------------------------------------------------------------//
+//  FStateTreeInstanceDataLayout
+//----------------------------------------------------------------//
+
+TSharedPtr<FStateTreeInstanceDataLayout> FStateTreeInstanceDataLayout::Create(TConstArrayView<const UScriptStruct*> Structs)
+{
+	constexpr int32 MinAlignment = FMath::Max(alignof(FStateTreeInstanceDataLayout), alignof(FLayoutItem));
+
+	const int NumStructs = Structs.Num();
+		
+	int32 RequiredSize = sizeof(FStateTreeInstanceDataLayout);
+	RequiredSize = Align(RequiredSize, alignof(FLayoutItem)) + sizeof(FLayoutItem) * NumStructs;
+
+	FStateTreeInstanceDataLayout* Layout = new(FMemory::Malloc(RequiredSize, MinAlignment)) FStateTreeInstanceDataLayout;
+
+	Layout->NumItems = NumStructs;
+	
+	FLayoutItem* Items = Layout->GetItemsPtr();
+	int32 Offset = 0;
+	
+	for (int32 ItemIndex = 0; ItemIndex < NumStructs; ItemIndex++)
+	{
+		FLayoutItem& Item = Items[ItemIndex];
+		const UScriptStruct* ScriptStruct = Structs[ItemIndex];
+		const int32 Alignment = ScriptStruct != nullptr ? ScriptStruct->GetMinAlignment() : 0;
+		const int32 Size = ScriptStruct != nullptr ? ScriptStruct->GetStructureSize() : 0;
+
+		Offset = Align(Offset, Alignment);
+
+		Item.Offset = Offset;
+		Item.ScriptStruct = ScriptStruct;
+		
+		Offset += Size;
+	}
+	
+	return MakeShareable(Layout, FLayoutMemoryDeleter());
+}
+
+int32 FStateTreeInstanceDataLayout::GetLayoutInstanceSize() const
+{
+	if (NumItems == 0)
+	{
+		return 0;
+	}
+	
+	const FLayoutItem* Items = GetItemsPtr();
+	const FLayoutItem& Last = Items[NumItems - 1];
+
+	return Last.Offset + (Last.ScriptStruct != nullptr ? Last.ScriptStruct->GetStructureSize() : 0);
+}
+
+int32 FStateTreeInstanceDataLayout::GetLayoutInstanceMinAlignment() const
+{
+	int32 CombinedAlignment = 0;
+	
+	const FLayoutItem* Items = GetItemsPtr();
+	for (int32 Index = 0; Index < NumItems; Index++)
+	{
+		const FLayoutItem& Item = Items[Index];
+		const int32 Alignment = Item.ScriptStruct != nullptr ? Item.ScriptStruct->GetMinAlignment() : 0;
+		CombinedAlignment = FMath::Max(CombinedAlignment, Alignment);
+	}
+		
+	return CombinedAlignment;
+}
+
+//----------------------------------------------------------------//
+//  FStateTreeInstanceData
+//----------------------------------------------------------------//
+
+void FStateTreeInstanceData::AddStructReferencedObjects(class FReferenceCollector& Collector) const
+{
+	if (!Layout.IsValid())
+	{
+		return;
+	}
+	
+	for (int32 Index = 0; Index < Layout->Num(); Index++)
+	{
+		FStateTreeInstanceDataLayout::FLayoutItem Item = Layout->GetItem(Index);
+		if (Item.ScriptStruct != nullptr)
+		{
+			Collector.AddReferencedObject(Item.ScriptStruct);
+			Collector.AddReferencedObjects(Item.ScriptStruct, Memory + Item.Offset);
+		}
+	}
+}
+
+void FStateTreeInstanceData::Allocate(const TSharedPtr<FStateTreeInstanceDataLayout>& InLayout)
+{
+	Reset();
+	
+	check(InLayout.IsValid());
+	Layout = InLayout;
+
+	Memory = (uint8*)FMemory::Malloc(Layout->GetLayoutInstanceSize(), Layout->GetLayoutInstanceMinAlignment());
+
+	for (int32 Index = 0; Index < Layout->Num(); Index++)
+	{
+		const FStateTreeInstanceDataLayout::FLayoutItem& Item = Layout->GetItem(Index);
+		if (Item.ScriptStruct != nullptr)
+		{
+			Item.ScriptStruct->InitializeStruct(Memory + Item.Offset);
+		}
+	}
+}
+
+void FStateTreeInstanceData::CopyFrom(const FStateTreeInstanceData& InOther)
+{
+	if (&InOther == this)
+	{
+		return;
+	}
+	
+	Allocate(InOther.GetLayout());
+
+	// Copy values over
+	for (int32 Index = 0; Index < Layout->Num(); Index++)
+	{
+		const FStateTreeInstanceDataLayout::FLayoutItem& Item = Layout->GetItem(Index);
+		if (Item.ScriptStruct != nullptr)
+		{
+			Item.ScriptStruct->CopyScriptStruct(Memory + Item.Offset, InOther.Memory + Item.Offset);
+		}
+	}
+}
+
+void FStateTreeInstanceData::Initialize(TConstArrayView<FInstancedStruct> InValues)
+{
+	TArray<const UScriptStruct*> Structs;
+	Algo::Transform(InValues, Structs, [](const FInstancedStruct& Value) { return Value.GetScriptStruct(); });
+	
+	Allocate(FStateTreeInstanceDataLayout::Create(Structs));
+
+	// Copy values over
+	check(Layout->Num() == InValues.Num());
+	for (int32 Index = 0; Index < Layout->Num(); Index++)
+	{
+		const FStateTreeInstanceDataLayout::FLayoutItem& Item = Layout->GetItem(Index);
+		const FInstancedStruct& Value = InValues[Index];
+		check(Item.ScriptStruct == Value.GetScriptStruct());
+		if (Item.ScriptStruct != nullptr)
+		{
+			Item.ScriptStruct->CopyScriptStruct(Memory + Item.Offset, Value.GetMutableMemory());
+		}
+	}
+}
+
+void FStateTreeInstanceData::Reset()
+{
+	// Destruct items
+	if (Layout.IsValid())
+	{
+		for (int32 Index = 0; Index < Layout->Num(); Index++)
+		{
+			const FStateTreeInstanceDataLayout::FLayoutItem& Item = Layout->GetItem(Index);
+			if (Item.ScriptStruct != nullptr)
+			{
+				Item.ScriptStruct->DestroyStruct(Memory + Item.Offset);
+			}
+		}
+	}
+
+	// Free memory
+	FMemory::Free(Memory);
+	Memory = nullptr;
+	Layout = nullptr;
+}
diff --git a/Engine/Plugins/Runtime/StateTree/Source/StateTreeModule/Public/Components/StateTreeBrainComponent.h b/Engine/Plugins/Runtime/StateTree/Source/StateTreeModule/Public/Components/StateTreeBrainComponent.h
--- a/Engine/Plugins/Runtime/StateTree/Source/StateTreeModule/Public/Components/StateTreeBrainComponent.h
+++ b/Engine/Plugins/Runtime/StateTree/Source/StateTreeModule/Public/Components/StateTreeBrainComponent.h
@@ -61,8 +61,6 @@ public:
 	virtual FString GetDebugInfoString() const override;
 #endif // WITH_GAMEPLAY_DEBUGGER
 
-private:
-
 protected:
 
 	bool SetContextRequirements();
diff --git a/Engine/Plugins/Runtime/StateTree/Source/StateTreeModule/Public/StateTree.h b/Engine/Plugins/Runtime/StateTree/Source/StateTreeModule/Public/StateTree.h
--- a/Engine/Plugins/Runtime/StateTree/Source/StateTreeModule/Public/StateTree.h
+++ b/Engine/Plugins/Runtime/StateTree/Source/StateTreeModule/Public/StateTree.h
@@ -8,6 +8,7 @@
 #include "StateTreeSchema.h"
 #include "InstancedStruct.h"
 #include "StateTreePropertyBindings.h"
+#include "StateTreeInstanceData.h"
 #include "StateTree.generated.h"
 
 /**
@@ -20,24 +21,15 @@ class STATETREEMODULE_API UStateTree : public UDataAsset
 
 public:
 
-	UStateTree(const FObjectInitializer& ObjectInitializer = FObjectInitializer::Get());
-	~UStateTree();
-
-	/** @return Script Struct that can be used to instantiate the runtime storage */
-	const UScriptStruct* GetInstanceStorageStruct() const { return InstanceStorageStruct; }
-
-	/** @return Instance of the runtime storage that contains the default values */
-	const FInstancedStruct& GetInstanceStorageDefaultValue() const { return InstanceStorageDefaultValue; }
-
+	/** @return Size required for the instance data in bytes. */
+	int32 GetInstanceDataSize() const;
+	
 	/** @return Number of runtime data (Evaluators, Tasks, Conditions) in the runtime storage. */
 	int32 GetNumInstances() const { return Instances.Num(); }
 
 	/** @return Number of data views required for StateTree execution (Evaluators, Tasks, Conditions, External data). */
 	int32 GetNumDataViews() const { return NumDataViews; }
 
-	/** @return Base index in data views for external data. */
-	int32 GetExternalDataBaseIndex() const { return ExternalDataBaseIndex; }
-
 	/** @return List of external data required by the state tree */
 	TConstArrayView<FStateTreeExternalDataDesc> GetExternalDataDescs() const { return ExternalDataDescs; }
 
@@ -45,14 +37,10 @@ public:
 	const UStateTreeSchema* GetSchema() const { return Schema; }
 	void SetSchema(UStateTreeSchema* InSchema) { Schema = InSchema; }
 
-	/** @return true is the tree asset is considered valid (e.g. at least one state) */
-	bool IsValidStateTree() const;
-
-	void ResolvePropertyPaths();
+	/** @return true is the tree asset is to be used at runtime. */
+	bool IsReadyToRun() const;
 
 #if WITH_EDITOR
-	void OnPIEStarted(const bool bIsSimulating);
-	
 	/** Resets the baked data to empty. */
 	void ResetBaked();
 #endif
@@ -68,7 +56,10 @@ public:
 #endif
 
 protected:
-	virtual void BeginDestroy() override;
+	
+	/** Resolved references between data in the StateTree. */
+	void Link();
+
 	virtual void PostLoad() override;
 	
 #if WITH_EDITOR
@@ -76,12 +67,6 @@ protected:
 	virtual void GetAssetRegistryTags(TArray<FAssetRegistryTag>& OutTags) const override;
 #endif
 
-	/** Initializes the types and default values related to StateTree runtime storage */
-	void InitInstanceStorageType();
-
-	/** Resolved references between data in the StateTree. */
-	void Link();
-
 private:
 
 	// Properties
@@ -91,7 +76,7 @@ private:
 
 	/** Evaluators, Tasks, and Condition items */
 	UPROPERTY()
-	TArray<FInstancedStruct> Items;
+	TArray<FInstancedStruct> Nodes;
 
 	/** Evaluators, Tasks, and Conditions runtime data. */
 	UPROPERTY()
@@ -101,16 +86,8 @@ private:
 	UPROPERTY()
 	TArray<TObjectPtr<UObject>> InstanceObjects;
 
-	/** Script Struct that can be used to instantiate the runtime storage */
-	UPROPERTY()
-	TObjectPtr<UScriptStruct> InstanceStorageStruct;
-
-	/** Offsets into the runtime type to quickly get a struct view to a specific Task or Evaluator */
-	TArray<FStateTreeInstanceStorageOffset> InstanceStorageOffsets;
-
-	/** Instance of the runtime storage that contains the default values. */
 	UPROPERTY()
-	FInstancedStruct InstanceStorageDefaultValue;
+	FStateTreeInstanceData InstanceDataDefaultValue;
 
 	/** List of external data required by the state tree, creating during linking. */
 	UPROPERTY(Transient)
diff --git a/Engine/Plugins/Runtime/StateTree/Source/StateTreeModule/Public/StateTreeExecutionContext.h b/Engine/Plugins/Runtime/StateTree/Source/StateTreeModule/Public/StateTreeExecutionContext.h
--- a/Engine/Plugins/Runtime/StateTree/Source/StateTreeModule/Public/StateTreeExecutionContext.h
+++ b/Engine/Plugins/Runtime/StateTree/Source/StateTreeModule/Public/StateTreeExecutionContext.h
@@ -6,6 +6,7 @@
 #include "InstancedStruct.h"
 #include "Containers/StaticArray.h"
 #include "StateTreePropertyBindings.h"
+#include "StateTreeInstanceData.h"
 #include "StateTreeExecutionContext.generated.h"
 
 struct FStateTreeEvaluatorBase;
@@ -64,8 +65,11 @@ public:
 	/** Initializes the StateTree instance to be used with specific owner and StateTree asset. */
 	bool Init(UObject& InOwner, const UStateTree& InStateTree, const EStateTreeStorage InStorageType);
 
-	/**  Initializes external storage (no need to call, if using internal storage). */
-	bool InitInstanceData(UObject& InOwner, FStateTreeDataView ExternalStorage) const;
+	/** Resets the instance to initial empty state. Note: Does not call ExitState(). */
+	void Reset();
+
+	/** Initializes external instance data (no need to call, if using internal storage). */
+	static bool InitInstanceData(UObject& InOwner, const UStateTree& InStateTree, FStateTreeInstanceData& OutInstanceData);
 
 	/** Returns the StateTree asset in use. */
 	const UStateTree* GetStateTree() const { return StateTree; }
@@ -79,12 +83,12 @@ public:
 	bool IsValid() const { return Owner != nullptr && StateTree != nullptr; }
 	
 	/** Start executing. */
-	EStateTreeRunStatus Start(FStateTreeDataView ExternalStorage = FStateTreeDataView());
+	EStateTreeRunStatus Start(FStateTreeInstanceData* ExternalInstanceData = nullptr);
 	/** Stop executing. */
-	EStateTreeRunStatus Stop(FStateTreeDataView ExternalStorage = FStateTreeDataView());
+	EStateTreeRunStatus Stop(FStateTreeInstanceData* ExternalInstanceData = nullptr);
 
 	/** Tick the state tree logic. */
-	EStateTreeRunStatus Tick(const float DeltaTime, FStateTreeDataView ExternalStorage = FStateTreeDataView());
+	EStateTreeRunStatus Tick(const float DeltaTime, FStateTreeInstanceData* ExternalInstanceData = nullptr);
 
 	/** @return Pointer to a State or null if state not found */ 
 	const FBakedStateTreeState* GetStateFromHandle(const FStateTreeHandle StateHandle) const
@@ -224,21 +228,21 @@ public:
 		return nullptr;
 	}
 	
-	EStateTreeRunStatus GetLastTickStatus(FStateTreeDataView ExternalStorage = FStateTreeDataView()) const;
+	EStateTreeRunStatus GetLastTickStatus(const FStateTreeInstanceData* ExternalInstanceData = nullptr) const;
 
 #if WITH_GAMEPLAY_DEBUGGER
 	/** @return Debug string describing the current state of the execution */
-	FString GetDebugInfoString(FStateTreeDataView ExternalStorage = FStateTreeDataView()) const;
+	FString GetDebugInfoString(const FStateTreeInstanceData* ExternalInstanceData = nullptr) const;
 #endif // WITH_GAMEPLAY_DEBUGGER
 
 #if WITH_STATETREE_DEBUG
-	FString GetActiveStateName(FStateTreeDataView ExternalStorage = FStateTreeDataView()) const;
+	FString GetActiveStateName(const FStateTreeInstanceData* ExternalInstanceData = nullptr) const;
 	
-	void DebugPrintInternalLayout(FStateTreeDataView ExternalStorage);
+	void DebugPrintInternalLayout(const FStateTreeInstanceData* ExternalInstanceData = nullptr);
 #endif
 
-	void AddStructReferencedObjects(class FReferenceCollector& Collector);
-	void AddStructReferencedObjects(const FStateTreeDataView ExternalStorage, class FReferenceCollector& Collector);
+	void AddStructReferencedObjects(class FReferenceCollector& Collector) const;
+	void AddStructReferencedObjects(const FStateTreeInstanceData* ExternalInstanceData, class FReferenceCollector& Collector) const;
 
 protected:
 
@@ -248,48 +252,43 @@ protected:
 	/** Callback when gated transition is triggered. Contexts that are event based can use this to trigger a future event. */
 	virtual void BeginGatedTransition(const FStateTreeExecutionState& Exec) {};
 	
-	/**
-	 * Resets the instance to initial empty state. Note: Does not call ExitState().
-	 */
-	void Reset();
-
 	/**
 	 * Handles logic for entering State. EnterState is called on new active Evaluators and Tasks that are part of the re-planned tree.
 	 * Re-planned tree is from the transition target up to the leaf state. States that are parent to the transition target state
 	 * and still active after the transition will remain intact.
 	 * @return Run status returned by the tasks.
 	 */
-	EStateTreeRunStatus EnterState(FStateTreeDataView Storage, const FStateTreeTransitionResult& Transition);
+	EStateTreeRunStatus EnterState(FStateTreeInstanceData& InstanceData, const FStateTreeTransitionResult& Transition);
 
 	/**
 	 * Handles logic for exiting State. ExitState is called on current active Evaluators and Tasks that are part of the re-planned tree.
 	 * Re-planned tree is from the transition target up to the leaf state. States that are parent to the transition target state
 	 * and still active after the transition will remain intact.
 	 */
-	void ExitState(FStateTreeDataView Storage, const FStateTreeTransitionResult& Transition);
+	void ExitState(FStateTreeInstanceData& InstanceData, const FStateTreeTransitionResult& Transition);
 
 	/**
 	 * Handles logic for exiting State. ExitState is called on current active Evaluators and Tasks in reverse order (from leaf to root).
 	 */
-	void StateCompleted(FStateTreeDataView Storage, const FStateTreeHandle CurrentState, const EStateTreeRunStatus CompletionStatus);
+	void StateCompleted(FStateTreeInstanceData& InstanceData, const FStateTreeHandle CurrentState, const EStateTreeRunStatus CompletionStatus);
 
 	/**
 	 * Ticks evaluators of all active states starting from current state by delta time.
 	 * If TickEvaluators() is called multiple times per frame (i.e. during selection when visiting new states), each state and evaluator is ticked only once.
 	 */
-	void TickEvaluators(FStateTreeDataView Storage, const FStateTreeHandle CurrentState, const EStateTreeEvaluationType EvalType, const float DeltaTime);
+	void TickEvaluators(FStateTreeInstanceData& InstanceData, const FStateTreeHandle CurrentState, const EStateTreeEvaluationType EvalType, const float DeltaTime);
 
 	/**
 	 * Ticks tasks of all active states starting from current state by delta time.
 	 * @return Run status returned by the tasks.
 	 */
-	EStateTreeRunStatus TickTasks(FStateTreeDataView Storage, const FStateTreeHandle CurrentState, const float DeltaTime);
+	EStateTreeRunStatus TickTasks(FStateTreeInstanceData& InstanceData, const FStateTreeHandle CurrentState, const float DeltaTime);
 
 	/**
 	 * Checks all conditions at given range
 	 * @return True if all conditions pass.
 	 */
-	bool TestAllConditions(FStateTreeDataView Storage, const uint32 ConditionsOffset, const uint32 ConditionsNum);
+	bool TestAllConditions(FStateTreeInstanceData& InstanceData, const uint32 ConditionsOffset, const uint32 ConditionsNum);
 
 	/**
 	 * Triggers transitions based on current run status. CurrentStatus is used to select which transitions events are triggered.
@@ -298,7 +297,7 @@ protected:
 	 * the actual next state returned by the selector.
 	 * @return Transition result describing the source state, state transitioned to, and next selected state.
 	 */
-	FStateTreeTransitionResult TriggerTransitions(FStateTreeDataView Storage, const FStateTreeStateStatus CurrentStatus, const int Depth);
+	FStateTreeTransitionResult TriggerTransitions(FStateTreeInstanceData& InstanceData, const FStateTreeStateStatus CurrentStatus, const int Depth);
 
 	/**
 	 * Runs state selection logic starting at the specified state, walking towards the leaf states.
@@ -312,51 +311,54 @@ protected:
 	 * @param Depth Depth of recursion.
 	 * @return Transition result describing the source state, transition target state, and next selected state.
 	 */
-	FStateTreeTransitionResult SelectState(FStateTreeDataView Storage, const FStateTreeStateStatus InitialStateStatus, const FStateTreeHandle InitialTargetState, const FStateTreeHandle NextState, const int Depth);
+	FStateTreeTransitionResult SelectState(FStateTreeInstanceData& InstanceData, const FStateTreeStateStatus InitialStateStatus, const FStateTreeHandle InitialTargetState, const FStateTreeHandle NextState, const int Depth);
 
 	/** @return State handles from specified state handle back to the root, specified handle included. */
 	int32 GetActiveStates(const FStateTreeHandle StateHandle, TStaticArray<FStateTreeHandle, 32>& OutStateHandles) const;
 
 	/** @return Mutable storage based on storage settings. */
-	FStateTreeDataView SelectMutableStorage(FStateTreeDataView ExternalStorage)
+	FStateTreeInstanceData& SelectMutableInstanceData(FStateTreeInstanceData* ExternalInstanceData)
 	{
-		return StorageType == EStateTreeStorage::External ? ExternalStorage : FStateTreeDataView(StorageInstance);
+		check(StorageType != EStateTreeStorage::External || (StorageType == EStateTreeStorage::External && ExternalInstanceData != nullptr));
+		return StorageType == EStateTreeStorage::External ? *ExternalInstanceData : InternalInstanceData;
 	}
 
 	/** @return Const storage based on storage settings. */
-	const FStateTreeDataView SelectStorage(FStateTreeDataView ExternalStorage) const
+	const FStateTreeInstanceData& SelectInstanceData(const FStateTreeInstanceData* ExternalInstanceData) const
 	{
-		return StorageType == EStateTreeStorage::External ? ExternalStorage : FStateTreeDataView(StorageInstance);
+		check(StorageType != EStateTreeStorage::External || (StorageType == EStateTreeStorage::External && ExternalInstanceData != nullptr));
+		return StorageType == EStateTreeStorage::External ? *ExternalInstanceData : InternalInstanceData;
 	}
 
 	/** @return View to an Evaluator, a Task, or a Condition instance data. */
-	FStateTreeDataView GetInstanceData(FStateTreeDataView Storage, const bool bIsObject, const int32 Index) const
+	FStateTreeDataView GetInstanceData(FStateTreeInstanceData& InstanceData, const bool bIsObject, const int32 Index) const
 	{
 		if (UNLIKELY(bIsObject == true))
 		{
-			const FStateTreeExecutionState& Exec = GetExecState(Storage);
+			const FStateTreeExecutionState& Exec = GetExecState(InstanceData);
 			return FStateTreeDataView(Exec.InstanceObjects[Index]);
 		}
-		else
-		{
-			const FStateTreeInstanceStorageOffset& ItemOffset = StateTree->InstanceStorageOffsets[Index];
-			return FStateTreeDataView(ItemOffset.Struct, Storage.GetMutableMemory() + ItemOffset.Offset);
-		}
+		
+		return FStateTreeDataView(InstanceData.GetMutable(Index));
 	}
 
 	/** @return StateTree execution state from the instance storage. */
-	FStateTreeExecutionState& GetExecState(FStateTreeDataView Storage) const
+	static FStateTreeExecutionState& GetExecState(FStateTreeInstanceData& InstanceData)
+	{
+		return InstanceData.GetMutable<FStateTreeExecutionState>(0);
+	}
+
+	/** @return const StateTree execution state from the instance storage. */
+	static const FStateTreeExecutionState& GetExecState(const FStateTreeInstanceData& InstanceData)
 	{
-		const FStateTreeInstanceStorageOffset& ItemOffset = StateTree->InstanceStorageOffsets[0];
-		check(ItemOffset.Struct == FStateTreeExecutionState::StaticStruct());
-		return *reinterpret_cast<FStateTreeExecutionState*>(Storage.GetMutableMemory() + ItemOffset.Offset);
+		return InstanceData.Get<FStateTreeExecutionState>(0);
 	}
 
-	/** @return Item at specified index. */
+	/** @return StateTree node at specified index. */
 	template <typename T>
-	T& GetItem(const int32 Index) const
+	const T& GetNode(const int32 Index) const
 	{
-		return StateTree->Items[Index].template GetMutable<T>();
+		return StateTree->Nodes[Index].template Get<T>();
 	}
 	
 	/** @return String describing state status for logging and debug. */
@@ -382,7 +384,7 @@ protected:
 	TArray<FStateTreeDataView> DataViews;
 
 	/** Optional Instance of the storage */
-	FInstancedStruct StorageInstance;
+	FStateTreeInstanceData InternalInstanceData;
 
 	/** Storage type of the context */
 	EStateTreeStorage StorageType = EStateTreeStorage::Internal;
diff --git a/Engine/Plugins/Runtime/StateTree/Source/StateTreeModule/Public/StateTreeInstanceData.h b/Engine/Plugins/Runtime/StateTree/Source/StateTreeModule/Public/StateTreeInstanceData.h
new file mode 100644
--- /dev/null
+++ b/Engine/Plugins/Runtime/StateTree/Source/StateTreeModule/Public/StateTreeInstanceData.h
@@ -0,0 +1,166 @@
+﻿// Copyright Epic Games, Inc. All Rights Reserved.
+
+#pragma once
+
+#include "InstancedStruct.h"
+#include "StructView.h"
+#include "StateTreeInstanceData.generated.h"
+
+/**
+ * StateTree instance data layout describes the layout of the instance data.
+ */
+struct STATETREEMODULE_API FStateTreeInstanceDataLayout
+{
+	/** Struct describing each item in the layout. */
+	struct FLayoutItem
+	{
+		const UScriptStruct* ScriptStruct = nullptr;
+		int32 Offset = 0;
+	};
+	
+	/**
+	 * Creates instance data layout from an array of structs.
+	 * @return shared pointer to the layout
+	 */
+	static TSharedPtr<FStateTreeInstanceDataLayout> Create(TConstArrayView<const UScriptStruct*> Structs);
+
+	/** @retrun the memory size required for the instance data. */
+	int32 GetLayoutInstanceSize() const;
+	/** @return the minimum alignment required by the instance data. */
+	int32 GetLayoutInstanceMinAlignment() const;
+
+	/** @return number if items in the layout. */
+	int32 Num() const { return NumItems; }
+
+	/** @return layout item at specified index */
+	 FLayoutItem& GetMutableItem(const int32 Index) const
+	{
+		check(Index >= 0 && Index < NumItems);
+		FLayoutItem* Items = GetItemsPtr();
+		return Items[Index];
+	}
+
+	/** @return const layout item at specified index */
+	const  FLayoutItem& GetItem(const int32 Index) const { return GetMutableItem(Index); }
+
+private:
+	FStateTreeInstanceDataLayout() = default;
+	
+	FLayoutItem* GetItemsPtr() const
+	{
+		static constexpr int32 ItemsOffset = Align(sizeof(FStateTreeInstanceDataLayout), alignof( FLayoutItem));
+		return ( FLayoutItem*)((uint8*)(this) + ItemsOffset);
+	}
+
+	struct FLayoutMemoryDeleter
+	{
+		FORCEINLINE void operator()(FStateTreeInstanceDataLayout* Layout) const
+		{
+			FMemory::Free(Layout);
+		}
+	};
+
+	/** Number of items in the layout. */
+	int32 NumItems = 0;
+};
+
+
+/**
+ * StateTree instance data is used to store the runtime state of a StateTree.
+ * The layout of the data is described in a FStateTreeInstanceDataLayout.
+ */
+USTRUCT()
+struct STATETREEMODULE_API FStateTreeInstanceData
+{
+	GENERATED_BODY()
+
+	FStateTreeInstanceData() = default;
+	~FStateTreeInstanceData() { Reset(); }
+
+	/** Creates new layout from an array if instanced structs, and copies the values into the instance data. */
+	void Initialize(TConstArrayView<FInstancedStruct> InValues);
+
+	/** Shares the layout from another instance data, and copies the data over. */
+	void CopyFrom(const FStateTreeInstanceData& InOther);
+
+	/** Resets the data to empty. */
+	void Reset();
+
+	/** @return Number of items in the instance data. */
+	int32 Num() const { return Layout.IsValid() ? Layout->Num() : 0; }
+
+	/** @return true if the instance is correctly initialized. */
+	bool IsValid() const { return Memory != nullptr && Layout.IsValid(); }
+
+	/** @return the layout of the instance data. */
+	TSharedPtr<FStateTreeInstanceDataLayout> GetLayout() const { return Layout; };
+
+	/** @return the script struct describing an item at specified index. */
+	const UScriptStruct* GetScriptStruct(const int32 Index) const
+	{
+		check(IsValid());
+		const FStateTreeInstanceDataLayout::FLayoutItem& Item = Layout->GetItem(Index);
+		return Item.ScriptStruct;
+	}
+
+	/** @returns mutable reference to the struct at specified index, this getter assumes that all data is valid. */
+	template<typename T>
+	T& GetMutable(const int32 Index) const
+	{
+		check(IsValid());
+		const FStateTreeInstanceDataLayout::FLayoutItem& Item = Layout->GetItem(Index);
+		check(Item.ScriptStruct != nullptr);
+		check(Item.ScriptStruct == T::StaticStruct() || Item.ScriptStruct->IsChildOf(T::StaticStruct()));
+		return *((T*)(Memory + Item.Offset));
+	}
+
+	/** @returns mutable pointer to the struct at specified index, or nullptr if cast is not valid. */
+	template<typename T>
+	T* GetMutablePtr(const int32 Index) const
+	{
+		check(IsValid());
+		const FStateTreeInstanceDataLayout::FLayoutItem& Item = Layout->GetItem(Index);
+		if (Item.ScriptStruct != nullptr && (Item.ScriptStruct == T::StaticStruct() || Item.ScriptStruct->IsChildOf(T::StaticStruct())))
+		{
+			return (T*)(Memory + Item.Offset);
+		}
+		return nullptr;
+	}
+
+	/** @return view to the struct at specified index. */
+	FStructView GetMutable(const int32 Index) const
+	{
+		check(IsValid());
+		const FStateTreeInstanceDataLayout::FLayoutItem& Item = Layout->GetItem(Index);
+		return FStructView(Item.ScriptStruct, Memory + Item.Offset);
+	}
+
+	/** @returns const reference to the struct at specified index, this getter assumes that all data is valid. */
+	template<typename T>
+	const T& Get(const int32 Index) const { return GetMutable<T>(Index); }
+
+	/** @returns const pointer to the struct at specified index, or nullptr if cast is not valid. */
+	template<typename T>
+	const T* GetPtr(const int32 Index) const { return GetMutablePtr<T>(Index); }
+
+	/** @return const view to the struct at specified index. */
+	FConstStructView Get(const int32 Index) const { return GetMutable(Index); }
+
+	void AddStructReferencedObjects(class FReferenceCollector& Collector) const;
+
+private:
+	
+	void Allocate(const TSharedPtr<FStateTreeInstanceDataLayout>& InLayout);
+	
+	uint8* Memory = nullptr;
+	TSharedPtr<FStateTreeInstanceDataLayout> Layout;
+};
+
+template<>
+struct TStructOpsTypeTraits<FStateTreeInstanceData> : public TStructOpsTypeTraitsBase2<FStateTreeInstanceData>
+{
+	enum
+	{
+		WithAddStructReferencedObjects = true,
+	};
+};
