StateTree: Fix StateTree asset getting dirtied after PIE
- Mark the StateTree instance data standalone and transient
- Mark references to StateTree instance data struct transient

diff --git a/Engine/Plugins/Runtime/StateTree/Source/StateTreeModule/Public/StateTree.h b/Engine/Plugins/Runtime/StateTree/Source/StateTreeModule/Public/StateTree.h
--- a/Engine/Plugins/Runtime/StateTree/Source/StateTreeModule/Public/StateTree.h
+++ b/Engine/Plugins/Runtime/StateTree/Source/StateTreeModule/Public/StateTree.h
@@ -86,7 +86,7 @@ private:
 	UPROPERTY()
 	TArray<TObjectPtr<UObject>> InstanceObjects;
 
-	UPROPERTY()
+	UPROPERTY(Transient)
 	FStateTreeInstanceData InstanceDataDefaultValue;
 
 	/** List of external data required by the state tree, creating during linking. */
