StateTree: Fix for instance data when compiling BP classes
- Moved instance object directly on FStateTreeInstanceData to make easier to access
- Internal instance data on FStateTreeExecutionContext is not property to simplify handling
- Link & resolve properties on UStateTree when objects are replaced (to refresh cached property pointers)
- Moved instance data initialization out from FStateTreeExecutionContext

diff --git a/Engine/Plugins/AI/MassAI/Source/MassAIBehavior/Private/MassStateTreeSubsystem.cpp b/Engine/Plugins/AI/MassAI/Source/MassAIBehavior/Private/MassStateTreeSubsystem.cpp
--- a/Engine/Plugins/AI/MassAI/Source/MassAIBehavior/Private/MassStateTreeSubsystem.cpp
+++ b/Engine/Plugins/AI/MassAI/Source/MassAIBehavior/Private/MassStateTreeSubsystem.cpp
@@ -2,7 +2,6 @@
 
 #include "MassStateTreeSubsystem.h"
 #include "StateTree.h"
-#include "StateTreeExecutionContext.h"
 #include "Engine/Engine.h"
 
 
@@ -30,8 +29,7 @@ FMassStateTreeInstanceHandle UMassStateTreeSubsystem::AllocateInstanceData(const
 	}
 
 	FMassStateTreeInstanceDataItem& Item = InstanceDataArray[Index];
-	
-	FStateTreeExecutionContext::InitInstanceData(*this, *StateTree, Item.InstanceData);
+	Item.InstanceData.CopyFrom(*this, StateTree->GetInstanceDataDefaultValue());
 	
 	return FMassStateTreeInstanceHandle::Make(Index, Item.Generation);
 }
diff --git a/Engine/Plugins/Runtime/StateTree/Source/StateTreeEditorModule/Private/StateTreeEditor.cpp b/Engine/Plugins/Runtime/StateTree/Source/StateTreeEditorModule/Private/StateTreeEditor.cpp
--- a/Engine/Plugins/Runtime/StateTree/Source/StateTreeEditorModule/Private/StateTreeEditor.cpp
+++ b/Engine/Plugins/Runtime/StateTree/Source/StateTreeEditorModule/Private/StateTreeEditor.cpp
@@ -324,12 +324,13 @@ FText FStateTreeEditor::GetStatisticsText() const
 		return FText::GetEmpty();
 	}
 
-	const FText SizeText = FText::AsMemory((uint64)StateTree->GetInstanceDataSize());
+	const FStateTreeInstanceData& InstanceDataDefaultValue = StateTree->GetInstanceDataDefaultValue();
+	const uint64 LayoutSize = InstanceDataDefaultValue.IsValid() ? InstanceDataDefaultValue.GetLayout()->GetLayoutInstanceSize() : 0;
+	
+	const FText SizeText = FText::AsMemory(LayoutSize);
 	const FText NumNodesText = FText::AsNumber(StateTree->GetNumInstances());
 
 	return FText::Format(LOCTEXT("RuntimeSize", "Runtime size: {0}, {1} nodes"), SizeText, NumNodesText);
-	
-	return FText::GetEmpty();
 }
 
 void FStateTreeEditor::HandleModelAssetChanged()
diff --git a/Engine/Plugins/Runtime/StateTree/Source/StateTreeModule/Private/StateTree.cpp b/Engine/Plugins/Runtime/StateTree/Source/StateTreeModule/Private/StateTree.cpp
--- a/Engine/Plugins/Runtime/StateTree/Source/StateTreeModule/Private/StateTree.cpp
+++ b/Engine/Plugins/Runtime/StateTree/Source/StateTreeModule/Private/StateTree.cpp
@@ -56,12 +56,22 @@ void UStateTree::GetAssetRegistryTags(TArray<FAssetRegistryTag>& OutTags) const
 void UStateTree::PostLoad()
 {
 	Super::PostLoad();
-	
-	PropertyBindings.ResolvePaths();
-
 	Link();
 }
 
+void UStateTree::Serialize(FStructuredArchiveRecord Record)
+{
+	Super::Serialize(Record);
+
+	// We need to link and rebind property bindings each time a BP is compiled,
+	// because property bindings may get invalid, and instance data potentially needs refreshed.
+	if (Record.GetUnderlyingArchive().IsModifyingWeakAndStrongReferences())
+	{
+		Link();
+	}
+}
+
+
 void UStateTree::Link()
 {
 	FStateTreeLinker Linker;
@@ -87,17 +97,8 @@ void UStateTree::Link()
 	InstanceDataDefaultValue.Reset();
 	if (Instances.Num() > 0)
 		{
-		InstanceDataDefaultValue.Initialize(Instances);
-	}
-}
-
-int32 UStateTree::GetInstanceDataSize() const
-{
-	if (!InstanceDataDefaultValue.IsValid() || !InstanceDataDefaultValue.GetLayout().IsValid())
-	{
-		return 0;
+		InstanceDataDefaultValue.Initialize(*this, Instances, InstanceObjects);
 	}
 
-	// TODO: this should count all allocated data (arrays, objects).
-	return InstanceDataDefaultValue.GetLayout()->GetLayoutInstanceSize();
+	PropertyBindings.ResolvePaths();
 }
diff --git a/Engine/Plugins/Runtime/StateTree/Source/StateTreeModule/Private/StateTreeExecutionContext.cpp b/Engine/Plugins/Runtime/StateTree/Source/StateTreeModule/Private/StateTreeExecutionContext.cpp
--- a/Engine/Plugins/Runtime/StateTree/Source/StateTreeModule/Private/StateTreeExecutionContext.cpp
+++ b/Engine/Plugins/Runtime/StateTree/Source/StateTreeModule/Private/StateTreeExecutionContext.cpp
@@ -47,7 +47,7 @@ bool FStateTreeExecutionContext::Init(UObject& InOwner, const UStateTree& InStat
 	if (StorageType == EStateTreeStorage::Internal)
 	{
 		// Duplicate instance state.
-		InitInstanceData(InOwner, *StateTree, InternalInstanceData);
+		InternalInstanceData.CopyFrom(InOwner, StateTree->GetInstanceDataDefaultValue());
 	}
 
 	// Initialize data views for all possible items.
@@ -69,22 +69,6 @@ void FStateTreeExecutionContext::Reset()
 	Owner = nullptr;
 }
 
-bool FStateTreeExecutionContext::InitInstanceData(UObject& InOwner, const UStateTree& InStateTree, FStateTreeInstanceData& OutInstanceData)
-{
-	OutInstanceData.CopyFrom(InStateTree.InstanceDataDefaultValue);
-
-	check (OutInstanceData.IsValid());
-	FStateTreeExecutionState& Exec = GetExecState(OutInstanceData);
-
-	Exec.InstanceObjects.Reset();
-	for (const UObject* Instance : InStateTree.InstanceObjects)
-	{
-		Exec.InstanceObjects.Add(DuplicateObject(Instance, &InOwner));
-	}
-
-	return true;
-}
-
 EStateTreeRunStatus FStateTreeExecutionContext::Start(FStateTreeInstanceData* ExternalInstanceData)
 {
 	CSV_SCOPED_TIMING_STAT_EXCLUSIVE(StateTree_Start);
@@ -95,7 +79,7 @@ EStateTreeRunStatus FStateTreeExecutionContext::Start(FStateTreeInstanceData* Ex
 	}
 	
 	FStateTreeInstanceData& InstanceData = SelectMutableInstanceData(ExternalInstanceData);
-	checkSlow(InstanceData.GetLayout() == StateTree->InstanceDataDefaultValue.GetLayout());
+	checkSlow(InstanceData.GetLayout() == StateTree->GetInstanceDataDefaultValue().GetLayout());
 	FStateTreeExecutionState& Exec = GetExecState(InstanceData);
 
 	// Stop if still running previous state.
@@ -159,7 +143,7 @@ EStateTreeRunStatus FStateTreeExecutionContext::Stop(FStateTreeInstanceData* Ext
 	}
 	
 	FStateTreeInstanceData& InstanceData = SelectMutableInstanceData(ExternalInstanceData);
-	checkSlow(InstanceData.GetLayout() == StateTree->InstanceDataDefaultValue.GetLayout());
+	checkSlow(InstanceData.GetLayout() == StateTree->GetInstanceDataDefaultValue().GetLayout());
 	FStateTreeExecutionState& Exec = GetExecState(InstanceData);
 
 	// Exit states if still in some valid state.
@@ -186,7 +170,7 @@ EStateTreeRunStatus FStateTreeExecutionContext::Tick(const float DeltaTime, FSta
 		return EStateTreeRunStatus::Failed;
 	}
 	FStateTreeInstanceData& InstanceData = SelectMutableInstanceData(ExternalInstanceData);
-	checkSlow(InstanceData.GetLayout() == StateTree->InstanceDataDefaultValue.GetLayout());
+	checkSlow(InstanceData.GetLayout() == StateTree->GetInstanceDataDefaultValue().GetLayout());
 	FStateTreeExecutionState& Exec = GetExecState(InstanceData);
 
 	// No ticking of the tree is done or stopped.
@@ -908,28 +892,6 @@ EStateTreeRunStatus FStateTreeExecutionContext::GetLastTickStatus(const FStateTr
 	return Exec.LastTickStatus;
 }
 
-void FStateTreeExecutionContext::AddStructReferencedObjects(FReferenceCollector& Collector) const
-{
-	if (StorageType == EStateTreeStorage::Internal)
-	{
-		AddStructReferencedObjects(&InternalInstanceData, Collector);
-	}
-}
-
-void FStateTreeExecutionContext::AddStructReferencedObjects(const FStateTreeInstanceData* InstanceData, FReferenceCollector& Collector) const
-{
-	if (InstanceData && InstanceData->IsValid())
-	{
-		const FStateTreeExecutionState& Exec = GetExecState(*InstanceData);
-		for (const UObject* Object : Exec.InstanceObjects)
-		{
-			Collector.AddReferencedObject(Object);
-		}
-		
-		InstanceData->AddStructReferencedObjects(Collector);
-	}
-}
-
 #if WITH_GAMEPLAY_DEBUGGER
 
 FString FStateTreeExecutionContext::GetDebugInfoString(const FStateTreeInstanceData* ExternalInstanceData) const
@@ -939,7 +901,7 @@ FString FStateTreeExecutionContext::GetDebugInfoString(const FStateTreeInstanceD
 		return FString(TEXT("No StateTree asset."));
 	}
 	const FStateTreeInstanceData& InstanceData = SelectInstanceData(ExternalInstanceData);
-	checkSlow(InstanceData.GetLayout() == StateTree->InstanceDataDefaultValue.GetLayout());
+	checkSlow(InstanceData.GetLayout() == StateTree->GetInstanceDataDefaultValue().GetLayout());
 	const FStateTreeExecutionState& Exec = GetExecState(InstanceData);
 
 	FString DebugString = FString::Printf(TEXT("StateTree (asset: '%s')\n"), *GetNameSafe(StateTree));
@@ -1027,7 +989,7 @@ void FStateTreeExecutionContext::DebugPrintInternalLayout(const FStateTreeInstan
 
 	// Instance InstanceData offsets (e.g. tasks, evaluators, conditions)
 
-	if (const FStateTreeInstanceDataLayout* InstanceLayout = StateTree->InstanceDataDefaultValue.GetLayout().Get())
+	if (const FStateTreeInstanceDataLayout* InstanceLayout = StateTree->GetInstanceDataDefaultValue().GetLayout().Get())
 	{
 		DebugString += FString::Printf(TEXT("\nInstance Data Offsets(%d)\n  [ %-40s | %-6s ]\n"), InstanceLayout->Num(), TEXT("Name"), TEXT("Offset"));
 
@@ -1123,7 +1085,7 @@ FString FStateTreeExecutionContext::GetActiveStateName(const FStateTreeInstanceD
 		return FString(TEXT("<None>"));
 	}
 	const FStateTreeInstanceData& InstanceData = SelectInstanceData(ExternalInstanceData);
-	checkSlow(InstanceData.GetLayout() == StateTree->InstanceDataDefaultValue.GetLayout());
+	checkSlow(InstanceData.GetLayout() == StateTree->GetInstanceDataDefaultValue().GetLayout());
 	const FStateTreeExecutionState& Exec = GetExecState(InstanceData);
 
 	FString FullStateName;
diff --git a/Engine/Plugins/Runtime/StateTree/Source/StateTreeModule/Private/StateTreeInstanceData.cpp b/Engine/Plugins/Runtime/StateTree/Source/StateTreeModule/Private/StateTreeInstanceData.cpp
--- a/Engine/Plugins/Runtime/StateTree/Source/StateTreeModule/Private/StateTreeInstanceData.cpp
+++ b/Engine/Plugins/Runtime/StateTree/Source/StateTreeModule/Private/StateTreeInstanceData.cpp
@@ -1,6 +1,7 @@
 ﻿// Copyright Epic Games, Inc. All Rights Reserved.
 
 #include "StateTreeInstanceData.h"
+#include "UObject/UnrealType.h"
 #include "Algo/Transform.h"
 
 //----------------------------------------------------------------//
@@ -89,6 +90,136 @@ void FStateTreeInstanceData::AddStructReferencedObjects(class FReferenceCollecto
 			Collector.AddReferencedObjects(Item.ScriptStruct, Memory + Item.Offset);
 		}
 	}
+
+	for (const UObject* Object : InstanceObjects)
+	{
+		Collector.AddReferencedObject(Object);
+	}	
+}
+
+bool FStateTreeInstanceData::Identical(const FStateTreeInstanceData* Other, uint32 PortFlags) const
+{
+	if (UNLIKELY(Other == nullptr))
+	{
+		return false;
+	}
+
+	// Identical if both are uninitialized.
+	if (UNLIKELY(!IsValid() && !Other->IsValid()))
+	{
+		return true;
+	}
+
+	// Not identical if one is valid and other is not.
+	if (UNLIKELY(IsValid() != Other->IsValid()))
+	{
+		return false;
+	}
+
+	// Not identical if different amount of instanced objects.
+	if (UNLIKELY(InstanceObjects.Num() != Other->InstanceObjects.Num()))
+	{
+		return false;
+	}
+
+	// Not identical if different layouts.
+	if (Layout != Other->Layout)
+	{
+		return false;
+	}
+
+	bool bResult = true;
+
+	// Check that the struct contents are identical.
+	for (int32 Index = 0; Index < Layout->Num(); Index++)
+	{
+		const FStateTreeInstanceDataLayout::FLayoutItem Item = Layout->GetItem(Index);
+		const FStateTreeInstanceDataLayout::FLayoutItem OtherItem = Other->Layout->GetItem(Index);
+		if (Item.ScriptStruct != nullptr && OtherItem.ScriptStruct != nullptr)
+		{
+			check(Item.ScriptStruct == OtherItem.ScriptStruct);
+			const uint8* ItemMemory = Memory + Item.Offset;
+			const uint8* OtherItemMemory = Other->Memory + OtherItem.Offset;
+			
+			if (Item.ScriptStruct->CompareScriptStruct(ItemMemory, OtherItemMemory, PortFlags) == false)
+			{
+				bResult = false;
+				break;
+			}
+		}
+	}
+
+	// Check that the instance object contents are identical.
+	if (bResult)
+	{
+		// Copied from object property.
+		auto AreObjectsIndentical = [](UObject* A, UObject* B, uint32 PortFlags) -> bool
+		{
+			if ((PortFlags & PPF_DuplicateForPIE) != 0)
+			{
+				return false;
+			}
+
+			if (A == B)
+			{
+				return true;
+			}
+
+			// Resolve the object handles and run the deep comparison logic 
+			if ((PortFlags & (PPF_DeepCompareInstances | PPF_DeepComparison)) != 0)
+			{
+				return FObjectPropertyBase::StaticIdentical(A, B, PortFlags);
+			}
+
+			return true;
+		};
+		
+		for (int32 Index = 0; Index < InstanceObjects.Num(); Index++)
+		{
+			if (InstanceObjects[Index] != nullptr && Other->InstanceObjects[Index] != nullptr)
+			{
+				if (!AreObjectsIndentical(InstanceObjects[Index], Other->InstanceObjects[Index], PortFlags))
+				{
+					bResult = false;
+					break;
+				}
+			}
+			else
+			{
+				bResult = false;
+				break;
+			}
+		}
+	}
+	
+	return bResult;
+}
+
+bool FStateTreeInstanceData::Serialize(FArchive& Ar)
+{
+	if (Ar.IsModifyingWeakAndStrongReferences())
+	{
+		// Just enough support to allow to replace a recompiled Blueprint.
+
+		// Items
+		if (Layout.IsValid())
+		{
+			for (int32 Index = 0; Index < Layout->Num(); ++Index)
+			{
+				const FStateTreeInstanceDataLayout::FLayoutItem Item = Layout->GetItem(Index);
+				UScriptStruct* NonConstScriptStruct = const_cast<UScriptStruct*>(Item.ScriptStruct);
+				if (NonConstScriptStruct != nullptr)
+				{
+					NonConstScriptStruct->SerializeItem(Ar, Memory + Item.Offset, nullptr);
+				}
+			}
+		}
+
+		// Instances
+		Ar << InstanceObjects;
+	}
+
+	return true;
 }
 
 void FStateTreeInstanceData::Allocate(const TSharedPtr<FStateTreeInstanceDataLayout>& InLayout)
@@ -110,7 +241,7 @@ void FStateTreeInstanceData::Allocate(const TSharedPtr<FStateTreeInstanceDataLay
 	}
 }
 
-void FStateTreeInstanceData::CopyFrom(const FStateTreeInstanceData& InOther)
+void FStateTreeInstanceData::CopyFrom(UObject& InOwner, const FStateTreeInstanceData& InOther)
 {
 	if (&InOther == this)
 	{
@@ -119,7 +250,7 @@ void FStateTreeInstanceData::CopyFrom(const FStateTreeInstanceData& InOther)
 	
 	Allocate(InOther.GetLayout());
 
-	// Copy values over
+	// Copy struct values over
 	for (int32 Index = 0; Index < Layout->Num(); Index++)
 	{
 		const FStateTreeInstanceDataLayout::FLayoutItem& Item = Layout->GetItem(Index);
@@ -128,9 +259,20 @@ void FStateTreeInstanceData::CopyFrom(const FStateTreeInstanceData& InOther)
 			Item.ScriptStruct->CopyScriptStruct(Memory + Item.Offset, InOther.Memory + Item.Offset);
 		}
 	}
+
+	// Copy instance objects.
+	InstanceObjects.Reset();
+	for (const UObject* Instance : InOther.InstanceObjects)
+	{
+		if (ensure(Instance != nullptr))
+		{
+			ensure(Instance->GetClass()->HasAnyClassFlags(CLASS_NewerVersionExists) == false);
+			InstanceObjects.Add(DuplicateObject(Instance, &InOwner));
+		}
+	}
 }
 
-void FStateTreeInstanceData::Initialize(TConstArrayView<FInstancedStruct> InValues)
+void FStateTreeInstanceData::Initialize(UObject& InOwner, TConstArrayView<FInstancedStruct> InValues, TConstArrayView<TObjectPtr<UObject>> InObjects)
 {
 	TArray<const UScriptStruct*> Structs;
 	Algo::Transform(InValues, Structs, [](const FInstancedStruct& Value) { return Value.GetScriptStruct(); });
@@ -149,6 +291,17 @@ void FStateTreeInstanceData::Initialize(TConstArrayView<FInstancedStruct> InValu
 			Item.ScriptStruct->CopyScriptStruct(Memory + Item.Offset, Value.GetMutableMemory());
 		}
 	}
+
+	// Copy UObjects
+	InstanceObjects.Reset();
+	for (const UObject* Instance : InObjects)
+	{
+		if (ensure(Instance != nullptr))
+		{
+			ensure(Instance->GetClass()->HasAnyClassFlags(CLASS_NewerVersionExists) == false);
+			InstanceObjects.Add(DuplicateObject(Instance, &InOwner));
+		}
+	}
 }
 
 void FStateTreeInstanceData::Reset()
@@ -170,4 +323,6 @@ void FStateTreeInstanceData::Reset()
 	FMemory::Free(Memory);
 	Memory = nullptr;
 	Layout = nullptr;
+
+	InstanceObjects.Reset();
 }
diff --git a/Engine/Plugins/Runtime/StateTree/Source/StateTreeModule/Public/Components/StateTreeBrainComponent.h b/Engine/Plugins/Runtime/StateTree/Source/StateTreeModule/Public/Components/StateTreeBrainComponent.h
--- a/Engine/Plugins/Runtime/StateTree/Source/StateTreeModule/Public/Components/StateTreeBrainComponent.h
+++ b/Engine/Plugins/Runtime/StateTree/Source/StateTreeModule/Public/Components/StateTreeBrainComponent.h
@@ -68,7 +68,7 @@ protected:
 	UPROPERTY(EditDefaultsOnly, Category = AI, meta=(RequiredAssetDataTags="Schema=BrainComponentStateTreeSchema"))
 	UStateTree* StateTree;
 
-	UPROPERTY()
+	UPROPERTY(Transient)
 	FStateTreeExecutionContext StateTreeContext;
 
 	/** if set, state tree execution is allowed */
diff --git a/Engine/Plugins/Runtime/StateTree/Source/StateTreeModule/Public/StateTree.h b/Engine/Plugins/Runtime/StateTree/Source/StateTreeModule/Public/StateTree.h
--- a/Engine/Plugins/Runtime/StateTree/Source/StateTreeModule/Public/StateTree.h
+++ b/Engine/Plugins/Runtime/StateTree/Source/StateTreeModule/Public/StateTree.h
@@ -21,8 +21,8 @@ class STATETREEMODULE_API UStateTree : public UDataAsset
 
 public:
 
-	/** @return Size required for the instance data in bytes. */
-	int32 GetInstanceDataSize() const;
+	/** @return Default value for the instance data. */
+	const FStateTreeInstanceData& GetInstanceDataDefaultValue() const { return InstanceDataDefaultValue; };
 	
 	/** @return Number of runtime data (Evaluators, Tasks, Conditions) in the runtime storage. */
 	int32 GetNumInstances() const { return Instances.Num(); }
@@ -61,6 +61,7 @@ protected:
 	void Link();
 
 	virtual void PostLoad() override;
+	virtual void Serialize(FStructuredArchiveRecord Record) override;
 	
 #if WITH_EDITOR
 	virtual void PostEditChangeProperty(struct FPropertyChangedEvent& PropertyChangedEvent) override;
diff --git a/Engine/Plugins/Runtime/StateTree/Source/StateTreeModule/Public/StateTreeExecutionContext.h b/Engine/Plugins/Runtime/StateTree/Source/StateTreeModule/Public/StateTreeExecutionContext.h
--- a/Engine/Plugins/Runtime/StateTree/Source/StateTreeModule/Public/StateTreeExecutionContext.h
+++ b/Engine/Plugins/Runtime/StateTree/Source/StateTreeModule/Public/StateTreeExecutionContext.h
@@ -35,9 +35,6 @@ struct STATETREEMODULE_API FStateTreeExecutionState
 
 	/** Running time of the delayed transition */
 	float GatedTransitionTime = 0.0f;
-
-	/** Object instances, ref counting handled manually. */
-	TArray<UObject*> InstanceObjects;
 };
 
 UENUM()
@@ -68,9 +65,6 @@ public:
 	/** Resets the instance to initial empty state. Note: Does not call ExitState(). */
 	void Reset();
 
-	/** Initializes external instance data (no need to call, if using internal storage). */
-	static bool InitInstanceData(UObject& InOwner, const UStateTree& InStateTree, FStateTreeInstanceData& OutInstanceData);
-
 	/** Returns the StateTree asset in use. */
 	const UStateTree* GetStateTree() const { return StateTree; }
 
@@ -241,9 +235,6 @@ public:
 	void DebugPrintInternalLayout(const FStateTreeInstanceData* ExternalInstanceData = nullptr);
 #endif
 
-	void AddStructReferencedObjects(class FReferenceCollector& Collector) const;
-	void AddStructReferencedObjects(const FStateTreeInstanceData* ExternalInstanceData, class FReferenceCollector& Collector) const;
-
 protected:
 
 	/** @return Prefix that will be used by STATETREE_LOG and STATETREE_CLOG, empty by default. */
@@ -335,8 +326,7 @@ protected:
 	{
 		if (UNLIKELY(bIsObject == true))
 		{
-			const FStateTreeExecutionState& Exec = GetExecState(InstanceData);
-			return FStateTreeDataView(Exec.InstanceObjects[Index]);
+			return FStateTreeDataView(InstanceData.GetMutableObject(Index));
 		}
 		
 		return FStateTreeDataView(InstanceData.GetMutable(Index));
@@ -371,30 +361,22 @@ protected:
 	FString DebugGetStatePath(TArrayView<FStateTreeHandle> ActiveStateHandles, int32 ActiveStateIndex) const;
 
 	/** The StateTree asset the context is initialized for */
-	UPROPERTY()
-	const UStateTree* StateTree = nullptr;
+	UPROPERTY(Transient)
+	TObjectPtr<const UStateTree> StateTree = nullptr;
 
-	UPROPERTY()
-	UObject* Owner = nullptr;
+	UPROPERTY(Transient)
+	TObjectPtr<UObject> Owner = nullptr;
 	
+	/** Optional Instance of the storage */
+	UPROPERTY(Transient)
+	FStateTreeInstanceData InternalInstanceData;
+
 	/** States visited during a tick while updating evaluators. Initialized to match the number of states in the asset. */ 
 	TArray<bool> VisitedStates;
 
 	/** Array of data pointers (external data, tasks, evaluators, conditions), used during evaluation. Initialized to match the number of items in the asset. */
 	TArray<FStateTreeDataView> DataViews;
 
-	/** Optional Instance of the storage */
-	FStateTreeInstanceData InternalInstanceData;
-
 	/** Storage type of the context */
 	EStateTreeStorage StorageType = EStateTreeStorage::Internal;
 };
-
-template<>
-struct TStructOpsTypeTraits<FStateTreeExecutionContext> : public TStructOpsTypeTraitsBase2<FStateTreeExecutionContext>
-{
-	enum
-	{
-		WithAddStructReferencedObjects = true,
-	};
-};
diff --git a/Engine/Plugins/Runtime/StateTree/Source/StateTreeModule/Public/StateTreeInstanceData.h b/Engine/Plugins/Runtime/StateTree/Source/StateTreeModule/Public/StateTreeInstanceData.h
--- a/Engine/Plugins/Runtime/StateTree/Source/StateTreeModule/Public/StateTreeInstanceData.h
+++ b/Engine/Plugins/Runtime/StateTree/Source/StateTreeModule/Public/StateTreeInstanceData.h
@@ -64,10 +64,11 @@ private:
 	int32 NumItems = 0;
 };
 
-
 /**
  * StateTree instance data is used to store the runtime state of a StateTree.
  * The layout of the data is described in a FStateTreeInstanceDataLayout.
+ *
+ * Note: Serialization is supported only for FArchive::IsModifyingWeakAndStrongReferences(), that is replacing object references.
  */
 USTRUCT()
 struct STATETREEMODULE_API FStateTreeInstanceData
@@ -78,10 +79,10 @@ struct STATETREEMODULE_API FStateTreeInstanceData
 	~FStateTreeInstanceData() { Reset(); }
 
 	/** Creates new layout from an array if instanced structs, and copies the values into the instance data. */
-	void Initialize(TConstArrayView<FInstancedStruct> InValues);
+	void Initialize(UObject& InOwner, TConstArrayView<FInstancedStruct> InValues, TConstArrayView<TObjectPtr<UObject>> InObjects);
 
 	/** Shares the layout from another instance data, and copies the data over. */
-	void CopyFrom(const FStateTreeInstanceData& InOther);
+	void CopyFrom(UObject& InOwner, const FStateTreeInstanceData& InOther);
 
 	/** Resets the data to empty. */
 	void Reset();
@@ -146,14 +147,31 @@ struct STATETREEMODULE_API FStateTreeInstanceData
 	/** @return const view to the struct at specified index. */
 	FConstStructView Get(const int32 Index) const { return GetMutable(Index); }
 
-	void AddStructReferencedObjects(class FReferenceCollector& Collector) const;
+	/** @return number of instance objects */
+	int32 NumObjects() const { return InstanceObjects.Num(); }
+
+	/** @return pointer to an instance object   */
+	UObject* GetMutableObject(const int32 Index) const { return InstanceObjects[Index]; }
 
+	/** @return const pointer to an instance object   */
+	const UObject* GetObject(const int32 Index) const { return InstanceObjects[Index]; }
+
+	/** Type traits */
+	void AddStructReferencedObjects(class FReferenceCollector& Collector) const;
+	bool Identical(const FStateTreeInstanceData* Other, uint32 PortFlags) const;
+	bool Serialize(FArchive& Ar);
+	
 private:
 	
 	void Allocate(const TSharedPtr<FStateTreeInstanceDataLayout>& InLayout);
-	
+
+	/** Struct instances */
 	uint8* Memory = nullptr;
 	TSharedPtr<FStateTreeInstanceDataLayout> Layout;
+
+	/** Object instances. */
+	UPROPERTY(Transient)
+	TArray<TObjectPtr<UObject>> InstanceObjects;
 };
 
 template<>
@@ -161,6 +179,8 @@ struct TStructOpsTypeTraits<FStateTreeInstanceData> : public TStructOpsTypeTrait
 {
 	enum
 	{
+		WithSerializer = true,
+		WithIdentical = true,
 		WithAddStructReferencedObjects = true,
 	};
 };
