StateTree: Condition expressions
- Added ability to combine conditions into expressions
- Added editor UI for the expressions
- Replaced condition description with a name (makes the expressions faster to skim)

diff --git a/Engine/Plugins/AI/MassAI/Source/MassAIBehavior/Private/Conditions/ZoneGraphTagConditions.cpp b/Engine/Plugins/AI/MassAI/Source/MassAIBehavior/Private/Conditions/ZoneGraphTagConditions.cpp
--- a/Engine/Plugins/AI/MassAI/Source/MassAIBehavior/Private/Conditions/ZoneGraphTagConditions.cpp
+++ b/Engine/Plugins/AI/MassAI/Source/MassAIBehavior/Private/Conditions/ZoneGraphTagConditions.cpp
@@ -97,52 +97,6 @@ bool FZoneGraphTagFilterCondition::TestCondition(FStateTreeExecutionContext& Con
 	return Filter.Pass(Tags) ^ bInvert;
 }
 
-#if WITH_EDITOR
-FText FZoneGraphTagFilterCondition::GetDescription(const FGuid& ID, FStateTreeDataView InstanceData, const IStateTreeBindingLookup& BindingLookup) const
-{
-	const FStateTreeEditorPropertyPath LeftPath(ID, GET_MEMBER_NAME_STRING_CHECKED(FZoneGraphTagFilterConditionInstanceData, Tags));
-
-	FText InvertText;
-	if (bInvert)
-	{
-		InvertText = LOCTEXT("Not", "Not");
-	}
-
-	FText LeftText;
-	if (const FStateTreeEditorPropertyPath* Binding = BindingLookup.GetPropertyBindingSource(LeftPath))
-	{
-		LeftText = BindingLookup.GetPropertyPathDisplayName(*Binding);
-	}
-	else
-	{
-		LeftText = LOCTEXT("NotBound", "Not Bound");
-	}
-
-	FText FilterParts[6] = { FText::GetEmpty(), FText::GetEmpty(), FText::GetEmpty(), FText::GetEmpty(), FText::GetEmpty(), FText::GetEmpty() };
-	int32 PartIndex = 0;
-	if (Filter.AnyTags != FZoneGraphTagMask::None)
-	{
-		FilterParts[PartIndex++] = UE::MassBehavior::ZoneGraph::GetMaskOperatorText(EZoneLaneTagMaskComparison::Any);
-		FilterParts[PartIndex++] = UE::MassBehavior::ZoneGraph::GetTagMaskName(Filter.AnyTags);
-	}
-
-	if (Filter.AllTags != FZoneGraphTagMask::None)
-	{
-		FilterParts[PartIndex++] = UE::MassBehavior::ZoneGraph::GetMaskOperatorText(EZoneLaneTagMaskComparison::All);
-		FilterParts[PartIndex++] = UE::MassBehavior::ZoneGraph::GetTagMaskName(Filter.AllTags);
-	}
-
-	if (Filter.NotTags != FZoneGraphTagMask::None)
-	{
-		FilterParts[PartIndex++] = UE::MassBehavior::ZoneGraph::GetMaskOperatorText(EZoneLaneTagMaskComparison::Not);
-		FilterParts[PartIndex++] = UE::MassBehavior::ZoneGraph::GetTagMaskName(Filter.NotTags);
-	}
-
-	return FText::Format(LOCTEXT("CompareZoneGraphTagFilterDesc", "{0} <Details.Bold>{1}</> contains {2} <Details.Bold>{3}</> {4} <Details.Bold>{5}</> {6} <Details.Bold>{7}</>"),
-		InvertText, LeftText, FilterParts[0], FilterParts[1], FilterParts[2], FilterParts[3], FilterParts[4], FilterParts[5]);
-}
-#endif// WITH_EDITOR
-
 
 //----------------------------------------------------------------------//
 //  FZoneGraphTagMaskCondition
@@ -163,46 +117,6 @@ bool FZoneGraphTagMaskCondition::TestCondition(FStateTreeExecutionContext& Conte
 	return Left.CompareMasks(Right, Operator) ^ bInvert;
 }
 
-#if WITH_EDITOR
-FText FZoneGraphTagMaskCondition::GetDescription(const FGuid& ID, FStateTreeDataView InstanceData, const IStateTreeBindingLookup& BindingLookup) const
-{
-	const FZoneGraphTagMaskConditionInstanceData& Instance = InstanceData.Get<FZoneGraphTagMaskConditionInstanceData>();
-	const FStateTreeEditorPropertyPath LeftPath(ID, GET_MEMBER_NAME_STRING_CHECKED(FZoneGraphTagMaskConditionInstanceData, Left));
-	const FStateTreeEditorPropertyPath RightPath(ID, GET_MEMBER_NAME_STRING_CHECKED(FZoneGraphTagMaskConditionInstanceData, Right));
-
-	FText InvertText;
-	if (bInvert)
-	{
-		InvertText = LOCTEXT("Not", "Not");
-	}
-
-	FText LeftText;
-	if (const FStateTreeEditorPropertyPath* Binding = BindingLookup.GetPropertyBindingSource(LeftPath))
-	{
-		LeftText = BindingLookup.GetPropertyPathDisplayName(*Binding);
-	}
-	else
-	{
-		LeftText = LOCTEXT("NotBound", "Not Bound");
-	}
-
-	FText OperatorText = UE::MassBehavior::ZoneGraph::GetMaskOperatorText(Operator);
-
-	FText RightText;
-	if (const FStateTreeEditorPropertyPath* Binding = BindingLookup.GetPropertyBindingSource(RightPath))
-	{
-		RightText = BindingLookup.GetPropertyPathDisplayName(*Binding);
-	}
-	else
-	{
-		RightText = UE::MassBehavior::ZoneGraph::GetTagMaskName(Instance.Right);
-	}
-
-	return FText::Format(LOCTEXT("CompareZoneGraphTagMaskDesc", "{0} <Details.Bold>{1}</> contains {2} <Details.Bold>{3}</>"),
-		InvertText, LeftText, OperatorText, RightText);
-}
-#endif// WITH_EDITOR
-
 
 //----------------------------------------------------------------------//
 //  FZoneGraphTagCondition
@@ -223,43 +137,6 @@ bool FZoneGraphTagCondition::TestCondition(FStateTreeExecutionContext& Context)
 	return (Left == Right) ^ bInvert;
 }
 
-#if WITH_EDITOR
-FText FZoneGraphTagCondition::GetDescription(const FGuid& ID, FStateTreeDataView InstanceData, const IStateTreeBindingLookup& BindingLookup) const
-{
-	const FZoneGraphTagConditionInstanceData& Instance = InstanceData.Get<FZoneGraphTagConditionInstanceData>();
-	const FStateTreeEditorPropertyPath LeftPath(ID, GET_MEMBER_NAME_STRING_CHECKED(FZoneGraphTagConditionInstanceData, Left));
-	const FStateTreeEditorPropertyPath RightPath(ID, GET_MEMBER_NAME_STRING_CHECKED(FZoneGraphTagConditionInstanceData, Right));
-
-	FText InvertText;
-	if (bInvert)
-	{
-		InvertText = LOCTEXT("Not", "Not");
-	}
-
-	FText LeftText;
-	if (const FStateTreeEditorPropertyPath* Binding = BindingLookup.GetPropertyBindingSource(LeftPath))
-	{
-		LeftText = BindingLookup.GetPropertyPathDisplayName(*Binding);
-	}
-	else
-	{
-		LeftText = LOCTEXT("NotBound", "Not Bound");
-	}
-
-	FText RightText;
-	if (const FStateTreeEditorPropertyPath* Binding = BindingLookup.GetPropertyBindingSource(RightPath))
-	{
-		RightText = BindingLookup.GetPropertyPathDisplayName(*Binding);
-	}
-	else
-	{
-		RightText = UE::MassBehavior::ZoneGraph::GetTagName(Instance.Right);
-	}
-
-	return FText::Format(LOCTEXT("CompareZoneGraphTagDesc", "{0} <Details.Bold>{1}</> is <Details.Bold>{2}</>"),
-		InvertText, LeftText, RightText);
-}
-#endif// WITH_EDITOR
 
 #if WITH_EDITOR
 #undef LOCTEXT_NAMESPACE
diff --git a/Engine/Plugins/AI/MassAI/Source/MassAIBehavior/Public/Conditions/ZoneGraphTagConditions.h b/Engine/Plugins/AI/MassAI/Source/MassAIBehavior/Public/Conditions/ZoneGraphTagConditions.h
--- a/Engine/Plugins/AI/MassAI/Source/MassAIBehavior/Public/Conditions/ZoneGraphTagConditions.h
+++ b/Engine/Plugins/AI/MassAI/Source/MassAIBehavior/Public/Conditions/ZoneGraphTagConditions.h
@@ -31,10 +31,6 @@ struct MASSAIBEHAVIOR_API FZoneGraphTagFilterCondition : public FStateTreeCondit
 	virtual const UStruct* GetInstanceDataType() const override { return FZoneGraphTagFilterConditionInstanceData::StaticStruct(); }
 	virtual bool TestCondition(FStateTreeExecutionContext& Context) const override;
 
-#if WITH_EDITOR
-	virtual FText GetDescription(const FGuid& ID, FStateTreeDataView InstanceData, const IStateTreeBindingLookup& BindingLookup) const override;
-#endif
-
 	TStateTreeInstanceDataPropertyHandle<FZoneGraphTagMask> TagsHandle;
 
 	UPROPERTY(EditAnywhere, Category = Condition)
@@ -71,10 +67,6 @@ struct MASSAIBEHAVIOR_API FZoneGraphTagMaskCondition : public FStateTreeConditio
 	virtual const UStruct* GetInstanceDataType() const override { return FZoneGraphTagMaskConditionInstanceData::StaticStruct(); }
 	virtual bool TestCondition(FStateTreeExecutionContext& Context) const override;
 
-#if WITH_EDITOR
-	virtual FText GetDescription(const FGuid& ID, FStateTreeDataView InstanceData, const IStateTreeBindingLookup& BindingLookup) const override;
-#endif
-
 	TStateTreeInstanceDataPropertyHandle<FZoneGraphTagMask> LeftHandle;
 	TStateTreeInstanceDataPropertyHandle<FZoneGraphTagMask> RightHandle;
 
@@ -112,10 +104,6 @@ struct MASSAIBEHAVIOR_API FZoneGraphTagCondition : public FStateTreeConditionBas
 	virtual const UStruct* GetInstanceDataType() const override { return FZoneGraphTagConditionInstanceData::StaticStruct(); }
 	virtual bool TestCondition(FStateTreeExecutionContext& Context) const override;
 
-#if WITH_EDITOR
-	virtual FText GetDescription(const FGuid& ID, FStateTreeDataView InstanceData, const IStateTreeBindingLookup& BindingLookup) const override;
-#endif
-
 	TStateTreeInstanceDataPropertyHandle<FZoneGraphTag> LeftHandle;
 	TStateTreeInstanceDataPropertyHandle<FZoneGraphTag> RightHandle;
 
diff --git a/Engine/Plugins/AI/MassAI/Source/MassAIDebug/Private/GameplayDebuggerCategory_Mass.cpp b/Engine/Plugins/AI/MassAI/Source/MassAIDebug/Private/GameplayDebuggerCategory_Mass.cpp
--- a/Engine/Plugins/AI/MassAI/Source/MassAIDebug/Private/GameplayDebuggerCategory_Mass.cpp
+++ b/Engine/Plugins/AI/MassAI/Source/MassAIDebug/Private/GameplayDebuggerCategory_Mass.cpp
@@ -564,7 +564,6 @@ void FGameplayDebuggerCategory_Mass::CollectData(APlayerController* OwnerPC, AAc
 							}
 							else
 							{
-								Status += StateTreeContext.GetActiveStateName(InstanceData);
 								Status += TEXT("{red}<No StateTree instance>{white}\n");
 							}
 						}
diff --git a/Engine/Plugins/Runtime/StateTree/Source/StateTreeEditorModule/Private/Customizations/StateTreeEditorNodeDetails.cpp b/Engine/Plugins/Runtime/StateTree/Source/StateTreeEditorModule/Private/Customizations/StateTreeEditorNodeDetails.cpp
--- a/Engine/Plugins/Runtime/StateTree/Source/StateTreeEditorModule/Private/Customizations/StateTreeEditorNodeDetails.cpp
+++ b/Engine/Plugins/Runtime/StateTree/Source/StateTreeEditorModule/Private/Customizations/StateTreeEditorNodeDetails.cpp
@@ -34,6 +34,7 @@
 #include "Blueprint/StateTreeConditionBlueprintBase.h"
 #include "StateTreeEditorModule.h"
 #include "StateTreeNodeClassCache.h"
+#include "Styling/StyleColors.h"
 
 #define LOCTEXT_NAMESPACE "StateTreeEditor"
 
@@ -172,10 +173,15 @@ void FStateTreeEditorNodeDetails::CustomizeHeader(TSharedRef<class IPropertyHand
 	InstanceObjectProperty = StructProperty->GetChildHandle(TEXT("InstanceObject"));
 	IDProperty = StructProperty->GetChildHandle(TEXT("ID"));
 
+	IndentProperty = StructProperty->GetChildHandle(TEXT("ConditionIndent"));
+	OperandProperty = StructProperty->GetChildHandle(TEXT("ConditionOperand"));
+
 	check(NodeProperty.IsValid());
 	check(InstanceProperty.IsValid());
 	check(IDProperty.IsValid());
-	
+	check(IndentProperty.IsValid());
+	check(OperandProperty.IsValid());
+
 	// Find base class and struct from meta data.
 	static const FName BaseStructMetaName(TEXT("BaseStruct")); // TODO: move these names into one central place.
 	static const FName BaseClassMetaName(TEXT("BaseClass")); // TODO: move these names into one central place.
@@ -200,30 +206,95 @@ void FStateTreeEditorNodeDetails::CustomizeHeader(TSharedRef<class IPropertyHand
 		.VAlign(VAlign_Center)
 		[
 			SNew(SHorizontalBox)
-			// Description (Condition)
-			+ SHorizontalBox::Slot()
-			.FillWidth(3.0f)
-			.VAlign(VAlign_Center)
-			[
-				SNew(SRichTextBlock)
-				.DecoratorStyleSet(FStateTreeEditorStyle::Get().Get())
-				.TextStyle(FStateTreeEditorStyle::Get(), "Details.Normal")
-				.Text(this, &FStateTreeEditorNodeDetails::GetDescription)
-				.Visibility(this, &FStateTreeEditorNodeDetails::IsDescriptionVisible)
-			]
-			// Name (Eval/Task)
 			+ SHorizontalBox::Slot()
-			.FillWidth(1.5f)
-			.VAlign(VAlign_Center)
+			.FillWidth(2.0f)
 			[
-				SNew(SEditableTextBox)
-				.IsEnabled(TAttribute<bool>(this, &FStateTreeEditorNodeDetails::IsNameEnabled))
-				.Text(this, &FStateTreeEditorNodeDetails::GetName)
-				.OnTextCommitted(this, &FStateTreeEditorNodeDetails::OnNameCommitted)
-				.SelectAllTextWhenFocused(true)
-				.RevertTextOnEscape(true)
-				.Font(IDetailLayoutBuilder::GetDetailFont())
-				.Visibility(this, &FStateTreeEditorNodeDetails::IsNameVisible)
+				SNew(SHorizontalBox)
+				// Indent
+				+ SHorizontalBox::Slot()
+				.AutoWidth()
+				.VAlign(VAlign_Center)
+				[
+					SNew(SBox)
+					.WidthOverride(this, &FStateTreeEditorNodeDetails::GetIndentSize)
+					.Visibility(this, &FStateTreeEditorNodeDetails::IsConditionVisible)
+					[
+						SNew(SComboButton)
+						.ComboButtonStyle(FStateTreeEditorStyle::Get(), "StateTree.Node.Indent.ComboBox")
+						.ContentPadding(2.0f)
+						.HasDownArrow(false)
+						.VAlign(VAlign_Center)
+						.OnGetMenuContent(this, &FStateTreeEditorNodeDetails::OnGetIndentContent)
+						.ButtonContent()
+						[
+							SNew(STextBlock)
+							.Text(FText::FromString(TEXT(">")))
+							.TextStyle(FStateTreeEditorStyle::Get(), "Details.Normal")
+							.ToolTipText(LOCTEXT("IndentTooltip", "Indent of the expression row, controls parentheses and evaluation order."))
+						]
+					]
+				]
+				// Operand
+				+ SHorizontalBox::Slot()
+				.AutoWidth()
+				.VAlign(VAlign_Center)
+				[
+					SNew(SBox)
+					.WidthOverride(30.0f)
+					.Padding(FMargin(2, 4))
+					.VAlign(VAlign_Center)
+					.Visibility(this, &FStateTreeEditorNodeDetails::IsConditionVisible)
+					[
+						SNew(SComboButton)
+						.IsEnabled(TAttribute<bool>(this, &FStateTreeEditorNodeDetails::IsOperandEnabled))
+						.ComboButtonStyle(FStateTreeEditorStyle::Get(), "StateTree.Node.Operand.ComboBox")
+						.ButtonColorAndOpacity(this, &FStateTreeEditorNodeDetails::GetOperandColor)
+						.HasDownArrow(false)
+						.HAlign(HAlign_Center)
+						.VAlign(VAlign_Center)
+						.OnGetMenuContent(this, &FStateTreeEditorNodeDetails::OnGetOperandContent)
+						.ButtonContent()
+						[
+							SNew(STextBlock)
+							.TextStyle(FStateTreeEditorStyle::Get(), "StateTree.Node.Operand")
+							.Text(this, &FStateTreeEditorNodeDetails::GetOperandText)
+						]
+					]
+				]
+				// Open parens
+				+ SHorizontalBox::Slot()
+				.AutoWidth()
+				.VAlign(VAlign_Center)
+				[
+					SNew(STextBlock)
+					.TextStyle(FStateTreeEditorStyle::Get(), "StateTree.Node.Parens")
+					.Text(this, &FStateTreeEditorNodeDetails::GetOpenParens)
+					.Visibility(this, &FStateTreeEditorNodeDetails::IsConditionVisible)
+				]
+				// Name (Eval/Task)
+				+ SHorizontalBox::Slot()
+				.AutoWidth()
+				.VAlign(VAlign_Center)
+				[
+					SNew(SEditableTextBox)
+					.IsEnabled(TAttribute<bool>(this, &FStateTreeEditorNodeDetails::IsNameEnabled))
+					.Text(this, &FStateTreeEditorNodeDetails::GetName)
+					.OnTextCommitted(this, &FStateTreeEditorNodeDetails::OnNameCommitted)
+					.SelectAllTextWhenFocused(true)
+					.RevertTextOnEscape(true)
+					.Style(FStateTreeEditorStyle::Get(), "StateTree.Node.Name")
+					.Visibility(this, &FStateTreeEditorNodeDetails::IsNameVisible)
+				]
+				// Close parens
+				+ SHorizontalBox::Slot()
+				.AutoWidth()
+				.VAlign(VAlign_Center)
+				[
+					SNew(STextBlock)
+					.TextStyle(FStateTreeEditorStyle::Get(), "StateTree.Node.Parens")
+					.Text(this, &FStateTreeEditorNodeDetails::GetCloseParens)
+					.Visibility(this, &FStateTreeEditorNodeDetails::IsConditionVisible)
+				]
 			]
 			// Class picker
 			+ SHorizontalBox::Slot()
@@ -440,44 +511,248 @@ void FStateTreeEditorNodeDetails::FindOuterObjects()
 	}
 }
 
-FText FStateTreeEditorNodeDetails::GetDescription() const
+FOptionalSize FStateTreeEditorNodeDetails::GetIndentSize() const
+{
+	return FOptionalSize(15.0f + GetIndent() * 30.0f);
+}
+
+TSharedRef<SWidget> FStateTreeEditorNodeDetails::OnGetIndentContent() const
+{
+	FMenuBuilder MenuBuilder(true, NULL);
+
+	for (int32 Indent = 0; Indent < UE::StateTree::MaxConditionIndent; Indent++)
+	{
+		FUIAction ItemAction(
+			FExecuteAction::CreateSP(this, &FStateTreeEditorNodeDetails::SetIndent, Indent),
+			FCanExecuteAction(),
+			FIsActionChecked::CreateSP(this, &FStateTreeEditorNodeDetails::IsIndent, Indent));
+		MenuBuilder.AddMenuEntry(FText::AsNumber(Indent), TAttribute<FText>(), FSlateIcon(), ItemAction, FName(), EUserInterfaceActionType::Check);
+	}
+
+	return MenuBuilder.MakeWidget();
+}
+
+int32 FStateTreeEditorNodeDetails::GetIndent() const
+{
+	check(IndentProperty);
+	
+	uint8 Indent = 0;
+	IndentProperty->GetValue(Indent);
+
+	return Indent;
+}
+
+void FStateTreeEditorNodeDetails::SetIndent(const int32 Indent) const
+{
+	check(IndentProperty);
+	
+	IndentProperty->SetValue((uint8)FMath::Clamp(Indent, 0, UE::StateTree::MaxConditionIndent - 1));
+}
+
+bool FStateTreeEditorNodeDetails::IsIndent(const int32 Indent) const
+{
+	return Indent == GetIndent();
+}
+
+bool FStateTreeEditorNodeDetails::IsFirstItem() const
 {
 	check(StructProperty);
+	return StructProperty->GetIndexInArray() == 0;
+}
 
-	TArray<void*> RawNodeData;
-	StructProperty->AccessRawData(RawNodeData);
+int32 FStateTreeEditorNodeDetails::GetCurrIndent() const
+{
+	// First item needs to be zero indent to make the parentheses counting to work properly.
+	return IsFirstItem() ? 0 : GetIndent();
+}
+
+int32 FStateTreeEditorNodeDetails::GetNextIndent() const
+{
+	// Find the intent of the next item by finding the item in the parent array.
+	check(StructProperty);
+	TSharedPtr<IPropertyHandle> ParentProp = StructProperty->GetParentHandle();
+	if (!ParentProp.IsValid())
+	{
+		return 0;
+	}
+	TSharedPtr<IPropertyHandleArray> ParentArray = ParentProp->AsArray();
+	if (!ParentArray.IsValid())
+	{
+		return 0;
+	}
 
-	// Multiple descriptions do not make sense, just if only one node is selected.
-	if (RawNodeData.Num() != 1)
+	uint32 NumElements = 0;
+	if (ParentArray->GetNumElements(NumElements) != FPropertyAccess::Success)
 	{
-		return LOCTEXT("MultipleSelectedFormatted", "<Details.Subdued>Multiple Selected</>");
+		return 0;
+	}
+	
+	const int32 NextIndex = StructProperty->GetIndexInArray() + 1;
+	if (NextIndex >= (int32)NumElements)
+	{
+		return 0;
 	}
 
-	FStateTreeEditorNode* Node = static_cast<FStateTreeEditorNode*>(RawNodeData[0]);
-	if (Node != nullptr && EditorData != nullptr && Node->Node.IsValid())
+	TSharedPtr<IPropertyHandle> NextStructProperty = ParentArray->GetElement(NextIndex);
+	if (!NextStructProperty.IsValid())
 	{
-		if (const FStateTreeConditionBase* Condition = Node->Node.GetPtr<const FStateTreeConditionBase>())
-		{
-			const FStateTreeBindingLookup BindingLookup(EditorData);
-			if (Node->Instance.IsValid())
-			{
-				return Condition->GetDescription(Node->ID, FStateTreeDataView(Node->Instance), BindingLookup);
-			}
-			else if (Node->InstanceObject != nullptr)
-			{
-				return Condition->GetDescription(Node->ID, FStateTreeDataView(Node->InstanceObject), BindingLookup);
-			}
-			else
-			{
-				return Node->Node.GetScriptStruct()->GetDisplayNameText();
-			}
-		}
+		return 0;
 	}
 	
-	return LOCTEXT("ConditionNotSet", "<Details.Subdued>Condition Not Set</>");
+	TSharedPtr<IPropertyHandle> NextIndentProperty = NextStructProperty->GetChildHandle(TEXT("ConditionIndent"));
+	if (!NextIndentProperty.IsValid())
+	{
+		return 0;
+	}
+	
+	uint8 Indent = 0;
+	NextIndentProperty->GetValue(Indent);
+
+	return Indent;
+}
+
+FText FStateTreeEditorNodeDetails::GetOpenParens() const
+{
+	check(IndentProperty);
+
+	const int32 CurrIndent = GetCurrIndent();
+	const int32 NextIndent = GetNextIndent();
+	const int32 DeltaIndent = NextIndent - CurrIndent;
+	const int32 OpenParens = FMath::Max(0, DeltaIndent);
+
+	static_assert(UE::StateTree::MaxConditionIndent == 4);
+	switch (OpenParens)
+	{
+		case 1: return FText::FromString(TEXT("("));
+		case 2: return FText::FromString(TEXT("(("));
+		case 3: return FText::FromString(TEXT("((("));
+		case 4: return FText::FromString(TEXT("(((("));
+	}
+	return FText::GetEmpty();
+}
+
+FText FStateTreeEditorNodeDetails::GetCloseParens() const
+{
+	check(IndentProperty);
+
+	const int32 CurrIndent = GetCurrIndent();
+	const int32 NextIndent = GetNextIndent();
+	const int32 DeltaIndent = NextIndent - CurrIndent;
+	const int32 CloseParens = FMath::Max(0, -DeltaIndent);
+
+	static_assert(UE::StateTree::MaxConditionIndent == 4);
+	switch (CloseParens)
+	{
+	case 1: return FText::FromString(TEXT(")"));
+	case 2: return FText::FromString(TEXT("))"));
+	case 3: return FText::FromString(TEXT(")))"));
+	case 4: return FText::FromString(TEXT("))))"));
+	}
+	return FText::GetEmpty();
+}
+
+FText FStateTreeEditorNodeDetails::GetOperandText() const
+{
+	check(OperandProperty);
+
+	// First item does not relate to anything existing, it could be empty, we return IF to indicate that we're building condition. 
+	if (IsFirstItem())
+	{
+		return LOCTEXT("IfOperand", "IF");
+	}
+
+	uint8 Value = 0;
+	OperandProperty->GetValue(Value);
+	const EStateTreeConditionOperand Operand = (EStateTreeConditionOperand)Value;
+
+	if (Operand == EStateTreeConditionOperand::And)
+	{
+		return LOCTEXT("AndOperand", "AND");
+	}
+	else if (Operand == EStateTreeConditionOperand::Or)
+	{
+		return LOCTEXT("OrOperand", "OR");
+	}
+	else
+	{
+		ensureMsgf(false, TEXT("Unhandled operand %s"), *UEnum::GetValueAsString(Operand));
+	}
+
+	return FText::GetEmpty();
 }
 
-EVisibility FStateTreeEditorNodeDetails::IsDescriptionVisible() const
+FSlateColor FStateTreeEditorNodeDetails::GetOperandColor() const
+{
+	check(OperandProperty);
+
+	if (IsFirstItem())
+	{
+		return FStyleColors::Transparent;
+	}
+
+	uint8 Value = 0; 
+	OperandProperty->GetValue(Value);
+	const EStateTreeConditionOperand Operand = (EStateTreeConditionOperand)Value;
+
+	if (Operand == EStateTreeConditionOperand::And)
+	{
+		return FStyleColors::AccentPink;
+	}
+	else if (Operand == EStateTreeConditionOperand::Or)
+	{
+		return FStyleColors::AccentBlue;
+	}
+	else
+	{
+		ensureMsgf(false, TEXT("Unhandled operand %s"), *UEnum::GetValueAsString(Operand));
+	}
+
+	return FStyleColors::Transparent;
+}
+
+TSharedRef<SWidget> FStateTreeEditorNodeDetails::OnGetOperandContent() const
+{
+	FMenuBuilder MenuBuilder(true, NULL);
+
+	FUIAction AndAction(
+		FExecuteAction::CreateSP(this, &FStateTreeEditorNodeDetails::SetOperand, EStateTreeConditionOperand::And),
+		FCanExecuteAction(),
+		FIsActionChecked::CreateSP(this, &FStateTreeEditorNodeDetails::IsOperand, EStateTreeConditionOperand::And));
+	MenuBuilder.AddMenuEntry(LOCTEXT("AndOperand", "AND"), TAttribute<FText>(), FSlateIcon(), AndAction, FName(), EUserInterfaceActionType::Check);
+
+	FUIAction OrAction(FExecuteAction::CreateSP(this, &FStateTreeEditorNodeDetails::SetOperand, EStateTreeConditionOperand::Or),
+		FCanExecuteAction(),
+		FIsActionChecked::CreateSP(this, &FStateTreeEditorNodeDetails::IsOperand, EStateTreeConditionOperand::Or));
+	MenuBuilder.AddMenuEntry(LOCTEXT("OrOperand", "OR"), TAttribute<FText>(), FSlateIcon(), OrAction, FName(), EUserInterfaceActionType::Check);
+
+	return MenuBuilder.MakeWidget();
+}
+
+bool FStateTreeEditorNodeDetails::IsOperandEnabled() const
+{
+	return !IsFirstItem();
+}
+
+bool FStateTreeEditorNodeDetails::IsOperand(const EStateTreeConditionOperand Operand) const
+{
+	check(OperandProperty);
+
+	uint8 Value = 0; 
+	OperandProperty->GetValue(Value);
+	const EStateTreeConditionOperand CurrOperand = (EStateTreeConditionOperand)Value;
+
+	return CurrOperand == Operand;
+}
+
+void FStateTreeEditorNodeDetails::SetOperand(const EStateTreeConditionOperand Operand) const
+{
+	check(OperandProperty);
+
+	OperandProperty->SetValue((uint8)Operand);
+}
+
+
+EVisibility FStateTreeEditorNodeDetails::IsConditionVisible() const
 {
 	const UScriptStruct* ScriptStruct = nullptr;
 	if (const FStateTreeEditorNode* Node = GetCommonNode())
@@ -510,6 +785,10 @@ FText FStateTreeEditorNodeDetails::GetName() const
 				{
 					void* Ptr = const_cast<void*>(static_cast<const void*>(Node->Node.GetMemory()));
 					const FName NameValue = *NameProperty->ContainerPtrToValuePtr<FName>(Ptr);
+					if (NameValue.IsNone())
+					{
+						return GetDisplayValueString();
+					}
 					return FText::FromName(NameValue);
 				}
 			}
@@ -522,23 +801,26 @@ FText FStateTreeEditorNodeDetails::GetName() const
 
 EVisibility FStateTreeEditorNodeDetails::IsNameVisible() const
 {
-	const UStateTreeSchema* Schema = StateTree ? StateTree->GetSchema() : nullptr;
-	if (Schema && Schema->AllowMultipleTasks() == false)
-	{
-		// Single task states use the state name as task name.
-		return EVisibility::Collapsed;
-	}
-	
 	const UScriptStruct* ScriptStruct = nullptr;
 	if (const FStateTreeEditorNode* Node = GetCommonNode())
 	{
 		ScriptStruct = Node->Node.GetScriptStruct();
 	}
+
+	if (ScriptStruct != nullptr && ScriptStruct->IsChildOf(FStateTreeTaskBase::StaticStruct()))
+	{
+		const UStateTreeSchema* Schema = StateTree ? StateTree->GetSchema() : nullptr;
+		if (Schema && Schema->AllowMultipleTasks() == false)
+		{
+			// Single task states use the state name as task name.
+			return EVisibility::Collapsed;
+		}
+	}
 	
-	return ScriptStruct != nullptr && (ScriptStruct->IsChildOf(FStateTreeTaskBase::StaticStruct()) || ScriptStruct->IsChildOf(FStateTreeEvaluatorBase::StaticStruct())) ? EVisibility::Visible : EVisibility::Collapsed;
+	return EVisibility::Visible;
 }
 
-void FStateTreeEditorNodeDetails::OnNameCommitted(const FText& NewText, ETextCommit::Type InTextCommit)
+void FStateTreeEditorNodeDetails::OnNameCommitted(const FText& NewText, ETextCommit::Type InTextCommit) const
 {
 	check(StructProperty);
 
@@ -755,7 +1037,7 @@ TSharedRef<SWidget> FStateTreeEditorNodeDetails::GeneratePicker()
 	return MenuBuilder.MakeWidget();
 }
 
-void FStateTreeEditorNodeDetails::OnStructPicked(const UScriptStruct* InStruct)
+void FStateTreeEditorNodeDetails::OnStructPicked(const UScriptStruct* InStruct) const
 {
 	check(StructProperty);
 	check(StateTree);
@@ -813,6 +1095,7 @@ void FStateTreeEditorNodeDetails::OnStructPicked(const UScriptStruct* InStruct)
 				else if (InStruct->IsChildOf(FStateTreeConditionBase::StaticStruct()))
 				{
 					FStateTreeConditionBase& Cond = Node->Node.GetMutable<FStateTreeConditionBase>();
+					Cond.Name = FName(InStruct->GetDisplayNameText().ToString());
 
 					if (const UScriptStruct* InstanceType = Cast<const UScriptStruct>(Cond.GetInstanceDataType()))
 					{
@@ -840,7 +1123,7 @@ void FStateTreeEditorNodeDetails::OnStructPicked(const UScriptStruct* InStruct)
 	}
 }
 
-void FStateTreeEditorNodeDetails::OnClassPicked(UClass* InClass)
+void FStateTreeEditorNodeDetails::OnClassPicked(UClass* InClass) const
 {
 	check(StructProperty);
 	check(StateTree);
@@ -885,7 +1168,8 @@ void FStateTreeEditorNodeDetails::OnClassPicked(UClass* InClass)
 				Node->Node.InitializeAs(FStateTreeBlueprintConditionWrapper::StaticStruct());
 				FStateTreeBlueprintConditionWrapper& Cond = Node->Node.GetMutable<FStateTreeBlueprintConditionWrapper>();
 				Cond.ConditionClass = InClass;
-				
+				Cond.Name = FName(InClass->GetDisplayNameText().ToString());
+
 				Node->InstanceObject = NewObject<UObject>(EditorData, InClass);
 
 				Node->ID = FGuid::NewGuid();
diff --git a/Engine/Plugins/Runtime/StateTree/Source/StateTreeEditorModule/Private/Customizations/StateTreeEditorNodeDetails.h b/Engine/Plugins/Runtime/StateTree/Source/StateTreeEditorModule/Private/Customizations/StateTreeEditorNodeDetails.h
--- a/Engine/Plugins/Runtime/StateTree/Source/StateTreeEditorModule/Private/Customizations/StateTreeEditorNodeDetails.h
+++ b/Engine/Plugins/Runtime/StateTree/Source/StateTreeEditorModule/Private/Customizations/StateTreeEditorNodeDetails.h
@@ -11,6 +11,7 @@ class UStateTree;
 class UStateTreeState;
 class UStateTreeEditorData;
 struct FStateTreeEditorPropertyPath;
+enum class EStateTreeConditionOperand : uint8;
 
 /**
  * Type customization for nodes (Conditions, Evaluators and Tasks) in StateTreeState.
@@ -33,12 +34,33 @@ private:
 	void ResetToDefault(TSharedPtr<IPropertyHandle> PropertyHandle);
 	TSharedPtr<IPropertyHandle> GetInstancedObjectValueHandle(TSharedPtr<IPropertyHandle> PropertyHandle);
 
-	FText GetDescription() const;
-	EVisibility IsDescriptionVisible() const;
+	FOptionalSize GetIndentSize() const;
+	TSharedRef<SWidget> OnGetIndentContent() const;
+
+	int32 GetIndent() const;
+	void SetIndent(const int32 Indent) const;
+	bool IsIndent(const int32 Indent) const;
+	
+	FText GetOperandText() const;
+	FSlateColor GetOperandColor() const;
+	TSharedRef<SWidget> OnGetOperandContent() const;
+	bool IsOperandEnabled() const;
+
+	void SetOperand(const EStateTreeConditionOperand Operand) const;
+	bool IsOperand(const EStateTreeConditionOperand Operand) const;
+
+	bool IsFirstItem() const;
+	int32 GetCurrIndent() const;
+	int32 GetNextIndent() const;
 	
+	FText GetOpenParens() const;
+	FText GetCloseParens() const;
+
+	EVisibility IsConditionVisible() const;
+
 	FText GetName() const;
 	EVisibility IsNameVisible() const;
-	void OnNameCommitted(const FText& NewText, ETextCommit::Type InTextCommit);
+	void OnNameCommitted(const FText& NewText, ETextCommit::Type InTextCommit) const;
 	bool IsNameEnabled() const;
 
 	const struct FStateTreeEditorNode* GetCommonNode() const;
@@ -47,8 +69,8 @@ private:
 
 	TSharedRef<SWidget> GeneratePicker();
 
-	void OnStructPicked(const UScriptStruct* InStruct);
-	void OnClassPicked(UClass* InClass);
+	void OnStructPicked(const UScriptStruct* InStruct) const;
+	void OnClassPicked(UClass* InClass) const;
 
 	void OnIdentifierChanged(const UStateTree& StateTree);
 	void OnBindingChanged(const FStateTreeEditorPropertyPath& SourcePath, const FStateTreeEditorPropertyPath& TargetPath);
@@ -68,5 +90,8 @@ private:
 	TSharedPtr<IPropertyHandle> InstanceObjectProperty;
 	TSharedPtr<IPropertyHandle> IDProperty;
 
+	TSharedPtr<IPropertyHandle> IndentProperty;
+	TSharedPtr<IPropertyHandle> OperandProperty;
+
 	FDelegateHandle OnBindingChangedHandle;
 };
diff --git a/Engine/Plugins/Runtime/StateTree/Source/StateTreeEditorModule/Private/StateTreeBaker.cpp b/Engine/Plugins/Runtime/StateTree/Source/StateTreeEditorModule/Private/StateTreeBaker.cpp
--- a/Engine/Plugins/Runtime/StateTree/Source/StateTreeEditorModule/Private/StateTreeBaker.cpp
+++ b/Engine/Plugins/Runtime/StateTree/Source/StateTreeEditorModule/Private/StateTreeBaker.cpp
@@ -78,6 +78,31 @@ bool FStateTreeBaker::CreateStates()
 	return true;
 }
 
+bool FStateTreeBaker::CreateConditions(TConstArrayView<FStateTreeEditorNode> Conditions)
+{
+	for (int32 Index = 0; Index < Conditions.Num(); Index++)
+	{
+		const bool bIsFirst = Index == 0;
+		const FStateTreeEditorNode& CondNode = Conditions[Index];
+		// First operand should be copy as we dont have a previous item to operate on.
+		const EStateTreeConditionOperand Operand = bIsFirst ? EStateTreeConditionOperand::Copy : CondNode.ConditionOperand;
+		// First indent must be 0 to make the parentheses calculation match.
+		const int32 CurrIndent = bIsFirst ? 0 : CondNode.ConditionIndent;
+		// Next indent, or terminate at zero.
+		const int32 NextIndent = Conditions.IsValidIndex(Index + 1) ? Conditions[Index].ConditionIndent : 0;
+		const int32 DeltaIndent = NextIndent - CurrIndent;
+		
+		check(DeltaIndent >= MIN_int8 && DeltaIndent <= MAX_int8);
+
+		if (!CreateCondition(CondNode, Operand, (int8)DeltaIndent))
+		{
+			return false;
+		}
+	}
+
+	return true;
+}
+
 bool FStateTreeBaker::CreateStateTransitions()
 {
 	for (int32 i = 0; i < StateTree->States.Num(); i++)
@@ -90,14 +115,11 @@ bool FStateTreeBaker::CreateStateTransitions()
 		
 		// Enter conditions.
 		BakedState.EnterConditionsBegin = uint16(StateTree->Nodes.Num());
-		for (FStateTreeEditorNode& CondNode : SourceState->EnterConditions)
+		if (!CreateConditions(SourceState->EnterConditions))
 		{
-			if (!CreateCondition(CondNode))
-			{
-				Log.Reportf(EMessageSeverity::Error,
-					TEXT("Failed to create state enter condition."));
-				return false;
-			}
+			Log.Reportf(EMessageSeverity::Error,
+				TEXT("Failed to create state enter condition."));
+			return false;
 		}
 		BakedState.EnterConditionsNum = uint8(uint16(StateTree->Nodes.Num()) - BakedState.EnterConditionsBegin);
 
@@ -117,15 +139,12 @@ bool FStateTreeBaker::CreateStateTransitions()
 			// Note: Unset transition is allowed here. It can be used to mask a transition at parent.
 
 			BakedTransition.ConditionsBegin = uint16(StateTree->Nodes.Num());
-			for (FStateTreeEditorNode& CondNode : Transition.Conditions)
+			if (!CreateConditions(Transition.Conditions))
 			{
-				if (!CreateCondition(CondNode))
-				{
-					Log.Reportf(EMessageSeverity::Error,
-						TEXT("Failed to create condition for transition to %s."),
-						*Transition.State.Name.ToString());
-					return false;
-				}
+				Log.Reportf(EMessageSeverity::Error,
+					TEXT("Failed to create condition for transition to %s."),
+					*Transition.State.Name.ToString());
+				return false;
 			}
 			BakedTransition.ConditionsNum = uint8(uint16(StateTree->Nodes.Num()) - BakedTransition.ConditionsBegin);
 		}
@@ -173,7 +192,7 @@ bool FStateTreeBaker::ResolveTransitionState(const UStateTreeState& SourceState,
 	return true;
 }
 
-bool FStateTreeBaker::CreateCondition(const FStateTreeEditorNode& CondNode)
+bool FStateTreeBaker::CreateCondition(const FStateTreeEditorNode& CondNode, const EStateTreeConditionOperand Operand, const int8 DeltaIndent)
 {
 	if (!CondNode.Node.IsValid())
 	{
@@ -199,6 +218,9 @@ bool FStateTreeBaker::CreateCondition(const FStateTreeEditorNode& CondNode)
 
 	FStateTreeConditionBase& Cond = Item.GetMutable<FStateTreeConditionBase>();
 
+	Cond.Operand = Operand;
+	Cond.DeltaIndent = DeltaIndent;
+	
 	if (CondNode.Instance.IsValid())
 	{
 		// Struct instance
diff --git a/Engine/Plugins/Runtime/StateTree/Source/StateTreeEditorModule/Private/StateTreeEditorStyle.cpp b/Engine/Plugins/Runtime/StateTree/Source/StateTreeEditorModule/Private/StateTreeEditorStyle.cpp
--- a/Engine/Plugins/Runtime/StateTree/Source/StateTreeEditorModule/Private/StateTreeEditorStyle.cpp
+++ b/Engine/Plugins/Runtime/StateTree/Source/StateTreeEditorModule/Private/StateTreeEditorStyle.cpp
@@ -7,6 +7,9 @@
 #include "EditorStyleSet.h"
 #include "Interfaces/IPluginManager.h"
 
+#include "Styling/StyleColors.h"
+
+
 #define IMAGE_BRUSH( RelativePath, ... ) FSlateImageBrush( StyleSet->RootToContentDir( RelativePath, TEXT( ".png" ) ), __VA_ARGS__ )
 #define IMAGE_PLUGIN_BRUSH( RelativePath, ... ) FSlateImageBrush( FMeshEditorStyle::InContent( RelativePath, ".png" ), __VA_ARGS__ )
 #define BOX_BRUSH( RelativePath, ... ) FSlateBoxBrush( StyleSet->RootToContentDir( RelativePath, TEXT( ".png" ) ), __VA_ARGS__ )
@@ -102,6 +105,48 @@ void FStateTreeEditorStyle::Initialize()
 		.SetSelectorFocusedBrush(IMAGE_BRUSH("Common/Selection", Icon8x8, SelectionColor))
 	);
 
+	const FComboButtonStyle& ComboButtonStyle = FCoreStyle::Get().GetWidgetStyle<FComboButtonStyle>("ComboButton");
+
+	// Condition Operand combo button
+	const FButtonStyle OperandButton = FButtonStyle()
+		.SetNormal(FSlateRoundedBoxBrush(FStyleColors::AccentGreen.GetSpecifiedColor().Desaturate(0.3f), 4.0f))
+		.SetHovered(FSlateRoundedBoxBrush(FStyleColors::AccentGreen.GetSpecifiedColor().Desaturate(0.2f), 4.0f))
+		.SetPressed(FSlateRoundedBoxBrush(FStyleColors::AccentGreen.GetSpecifiedColor().Desaturate(0.1f), 4.0f))
+		.SetNormalForeground(FStyleColors::Foreground)
+		.SetHoveredForeground(FStyleColors::ForegroundHover)
+		.SetPressedForeground(FStyleColors::ForegroundHover)
+		.SetDisabledForeground(FStyleColors::ForegroundHover)
+		.SetNormalPadding(FMargin(2, 2, 2, 2))
+		.SetPressedPadding(FMargin(2, 3, 2, 1));
+
+	StyleSet->Set("StateTree.Node.Operand.ComboBox", FComboButtonStyle(ComboButtonStyle).SetButtonStyle(OperandButton));
+
+	StyleSet->Set("StateTree.Node.Operand", FTextBlockStyle(NormalText)
+		.SetFont(FEditorStyle::GetFontStyle(TEXT("PropertyWindow.BoldFont")))
+		.SetFontSize(7));
+
+	StyleSet->Set("StateTree.Node.Parens", FTextBlockStyle(NormalText)
+		.SetFont(FEditorStyle::GetFontStyle(TEXT("PropertyWindow.NormalFont")))
+		.SetFontSize(12));
+
+	// Condition Indent combo button
+	const FButtonStyle IndentButton = FButtonStyle()
+		.SetNormal(FSlateRoundedBoxBrush(FLinearColor::Transparent, 2.0f))
+		.SetHovered(FSlateRoundedBoxBrush(FStyleColors::Background, 2.0f, FStyleColors::InputOutline, 1.0f))
+		.SetPressed(FSlateRoundedBoxBrush(FStyleColors::Background, 2.0f, FStyleColors::Hover, 1.0f))
+		.SetNormalForeground(FStyleColors::Transparent)
+		.SetHoveredForeground(FStyleColors::Hover)
+		.SetPressedForeground(FStyleColors::Foreground)
+		.SetNormalPadding(FMargin(2, 2, 2, 2))
+		.SetPressedPadding(FMargin(2, 3, 2, 1));
+	
+	StyleSet->Set("StateTree.Node.Indent.ComboBox", FComboButtonStyle(ComboButtonStyle).SetButtonStyle(IndentButton));
+
+	const FEditableTextBoxStyle& NormalEditableTextBox = FCoreStyle::Get().GetWidgetStyle<FEditableTextBoxStyle>("NormalEditableTextBox");
+	FEditableTextBoxStyle NameEditStyle(NormalEditableTextBox);
+	NameEditStyle.Font.Size = 10;
+	StyleSet->Set("StateTree.Node.Name", NameEditStyle);
+	
 	FSlateStyleRegistry::RegisterSlateStyle(*StyleSet.Get());
 }
 
diff --git a/Engine/Plugins/Runtime/StateTree/Source/StateTreeEditorModule/Public/StateTreeBaker.h b/Engine/Plugins/Runtime/StateTree/Source/StateTreeEditorModule/Public/StateTreeBaker.h
--- a/Engine/Plugins/Runtime/StateTree/Source/StateTreeEditorModule/Public/StateTreeBaker.h
+++ b/Engine/Plugins/Runtime/StateTree/Source/StateTreeEditorModule/Public/StateTreeBaker.h
@@ -36,7 +36,8 @@ private:
 	bool CreateStates();
 	bool CreateStateRecursive(UStateTreeState& State, const FStateTreeHandle Parent);
 	bool CreateStateTransitions();
-	bool CreateCondition(const FStateTreeEditorNode& CondNode);
+	bool CreateConditions(TConstArrayView<FStateTreeEditorNode> Conditions);
+	bool CreateCondition(const FStateTreeEditorNode& CondNode, const EStateTreeConditionOperand Operand, const int8 DeltaIndent);
 	bool CreateTask(const FStateTreeEditorNode& TaskNode);
 	bool CreateEvaluator(const FStateTreeEditorNode& EvalNode);
 	bool GetAndValidateBindings(const FStateTreeBindableStructDesc& TargetStruct, TArray<FStateTreeEditorPropertyBinding>& OutBindings) const;
diff --git a/Engine/Plugins/Runtime/StateTree/Source/StateTreeEditorModule/Public/StateTreeState.h b/Engine/Plugins/Runtime/StateTree/Source/StateTreeEditorModule/Public/StateTreeState.h
--- a/Engine/Plugins/Runtime/StateTree/Source/StateTreeEditorModule/Public/StateTreeState.h
+++ b/Engine/Plugins/Runtime/StateTree/Source/StateTreeEditorModule/Public/StateTreeState.h
@@ -72,6 +72,12 @@ struct STATETREEEDITORMODULE_API FStateTreeEditorNode
 	
 	UPROPERTY(EditDefaultsOnly, Category = Node)
 	FGuid ID;
+
+	UPROPERTY(EditDefaultsOnly, Category = Node)
+	uint8 ConditionIndent = 0;
+
+	UPROPERTY(EditDefaultsOnly, Category = Node)
+	EStateTreeConditionOperand ConditionOperand = EStateTreeConditionOperand::And; 
 };
 
 template <typename T>
@@ -82,7 +88,6 @@ struct TStateTreeEditorNode : public FStateTreeEditorNode
 	FORCEINLINE typename T::InstanceDataType& GetInstance() { return Instance.template GetMutable<typename T::InstanceDataType>(); }
 };
 
-
 /**
  * Editor representation of a transition in StateTree
  */
diff --git a/Engine/Plugins/Runtime/StateTree/Source/StateTreeEditorModule/StateTreeEditorModule.Build.cs b/Engine/Plugins/Runtime/StateTree/Source/StateTreeEditorModule/StateTreeEditorModule.Build.cs
--- a/Engine/Plugins/Runtime/StateTree/Source/StateTreeEditorModule/StateTreeEditorModule.Build.cs
+++ b/Engine/Plugins/Runtime/StateTree/Source/StateTreeEditorModule/StateTreeEditorModule.Build.cs
@@ -30,6 +30,7 @@ namespace UnrealBuildTool.Rules
 				"PropertyAccessEditor",
 				"StructUtils",
 				"StructUtilsEditor",
+				"GameplayTags"
 			}
 			);
 
diff --git a/Engine/Plugins/Runtime/StateTree/Source/StateTreeModule/Private/Blueprint/StateTreeConditionBlueprintBase.cpp b/Engine/Plugins/Runtime/StateTree/Source/StateTreeModule/Private/Blueprint/StateTreeConditionBlueprintBase.cpp
--- a/Engine/Plugins/Runtime/StateTree/Source/StateTreeModule/Private/Blueprint/StateTreeConditionBlueprintBase.cpp
+++ b/Engine/Plugins/Runtime/StateTree/Source/StateTreeModule/Private/Blueprint/StateTreeConditionBlueprintBase.cpp
@@ -40,24 +40,3 @@ bool FStateTreeBlueprintConditionWrapper::TestCondition(FStateTreeExecutionConte
 	return false;
 }
 
-#if WITH_EDITOR
-FText FStateTreeBlueprintConditionWrapper::GetDescription(const FGuid& ID, FStateTreeDataView InstanceData, const IStateTreeBindingLookup& BindingLookup) const
-{
-	const UClass* Class = Cast<const UClass>(InstanceData.GetStruct());
-
-	if (Class != nullptr)
-	{
-		if (UBlueprint* ClassBP = Cast<UBlueprint>(Class->ClassGeneratedBy))
-		{
-			if (!ClassBP->BlueprintDescription.IsEmpty())
-			{
-				return FText::FromString(ClassBP->BlueprintDescription);
-			}
-		}
-
-		return Class->GetDisplayNameText();
-	}
-	
-	return FText::GetEmpty();
-}
-#endif
diff --git a/Engine/Plugins/Runtime/StateTree/Source/StateTreeModule/Private/Conditions/GameplayTagConditions.cpp b/Engine/Plugins/Runtime/StateTree/Source/StateTreeModule/Private/Conditions/GameplayTagConditions.cpp
--- a/Engine/Plugins/Runtime/StateTree/Source/StateTreeModule/Private/Conditions/GameplayTagConditions.cpp
+++ b/Engine/Plugins/Runtime/StateTree/Source/StateTreeModule/Private/Conditions/GameplayTagConditions.cpp
@@ -66,44 +66,6 @@ bool FGameplayTagMatchCondition::TestCondition(FStateTreeExecutionContext& Conte
 	return (bExactMatch ? TagContainer.HasTagExact(Tag) : TagContainer.HasTag(Tag)) ^ bInvert;
 }
 
-#if WITH_EDITOR
-FText FGameplayTagMatchCondition::GetDescription(const FGuid& ID, FStateTreeDataView InstanceData, const IStateTreeBindingLookup& BindingLookup) const
-{
-	const FGameplayTagMatchConditionInstanceData& Instance = InstanceData.Get<FGameplayTagMatchConditionInstanceData>();
-	const FStateTreeEditorPropertyPath TagContainerPath(ID, GET_MEMBER_NAME_STRING_CHECKED(FGameplayTagMatchConditionInstanceData, TagContainer));
-	const FStateTreeEditorPropertyPath TagPath(ID, GET_MEMBER_NAME_STRING_CHECKED(FGameplayTagMatchConditionInstanceData, Tag));
-
-	FText InvertText;
-	if (bInvert)
-	{
-		InvertText = LOCTEXT("Not", "Not");
-	}
-
-	FText TagContainerText;
-	if (const FStateTreeEditorPropertyPath* Binding = BindingLookup.GetPropertyBindingSource(TagContainerPath))
-	{
-		TagContainerText = BindingLookup.GetPropertyPathDisplayName(*Binding);
-	}
-	else
-	{
-		TagContainerText = UE::StateTree::Conditions::GetContainerAsText(Instance.TagContainer);
-	}
-
-	FText TagText;
-	if (const FStateTreeEditorPropertyPath* Binding = BindingLookup.GetPropertyBindingSource(TagContainerPath))
-	{
-		TagText = BindingLookup.GetPropertyPathDisplayName(*Binding);
-	}
-	else
-	{
-		TagText = FText::FromString(Instance.Tag.ToString());
-	}
-
-	return FText::Format(LOCTEXT("GameplayTagMatchDesc", "{0} <Details.Bold>{1}</> has <Details.Bold>{2}</>"),
-		InvertText, TagContainerText, TagText);
-}
-#endif// WITH_EDITOR
-
 //----------------------------------------------------------------------//
 //  FGameplayTagContainerMatchCondition
 //----------------------------------------------------------------------//
@@ -137,45 +99,6 @@ bool FGameplayTagContainerMatchCondition::TestCondition(FStateTreeExecutionConte
 	return bResult ^ bInvert;
 }
 
-#if WITH_EDITOR
-FText FGameplayTagContainerMatchCondition::GetDescription(const FGuid& ID, FStateTreeDataView InstanceData, const IStateTreeBindingLookup& BindingLookup) const
-{
-	const FGameplayTagContainerMatchConditionInstanceData& Instance = InstanceData.Get<FGameplayTagContainerMatchConditionInstanceData>();
-	const FStateTreeEditorPropertyPath TagContainerPath(ID, GET_MEMBER_NAME_STRING_CHECKED(FGameplayTagContainerMatchConditionInstanceData, TagContainer));
-	const FStateTreeEditorPropertyPath OtherContainerPath(ID, GET_MEMBER_NAME_STRING_CHECKED(FGameplayTagContainerMatchConditionInstanceData, OtherContainer));
-
-	FText InvertText;
-	if (bInvert)
-	{
-		InvertText = LOCTEXT("Not", "Not");
-	}
-
-	FText TagContainerText;
-	if (const FStateTreeEditorPropertyPath* Binding = BindingLookup.GetPropertyBindingSource(TagContainerPath))
-	{
-		TagContainerText = BindingLookup.GetPropertyPathDisplayName(*Binding);
-	}
-	else
-	{
-		TagContainerText = LOCTEXT("NotBound", "Not Bound");
-	}
-
-	FText MatchTypeText = UEnum::GetDisplayValueAsText(MatchType);
-
-	FText OtherContainerText;
-	if (const FStateTreeEditorPropertyPath* Binding = BindingLookup.GetPropertyBindingSource(TagContainerPath))
-	{
-		OtherContainerText = BindingLookup.GetPropertyPathDisplayName(*Binding);
-	}
-	else
-	{
-		OtherContainerText = UE::StateTree::Conditions::GetContainerAsText(Instance.TagContainer);
-	}
-
-	return FText::Format(LOCTEXT("GameplayTagContainerMatchDesc", "{0} <Details.Bold>{1}</> {2} <Details.Bold>{3}</>"),
-		InvertText, TagContainerText, MatchTypeText, OtherContainerText);
-}
-#endif// WITH_EDITOR
 
 //----------------------------------------------------------------------//
 //  FGameplayTagQueryCondition
@@ -194,33 +117,6 @@ bool FGameplayTagQueryCondition::TestCondition(FStateTreeExecutionContext& Conte
 	return TagQuery.Matches(TagContainer) ^ bInvert;
 }
 
-#if WITH_EDITOR
-FText FGameplayTagQueryCondition::GetDescription(const FGuid& ID, FStateTreeDataView InstanceData, const IStateTreeBindingLookup& BindingLookup) const
-{
-	const FGameplayTagQueryConditionInstanceData& Instance = InstanceData.Get<FGameplayTagQueryConditionInstanceData>();
-	const FStateTreeEditorPropertyPath TagContainerPath(ID, GET_MEMBER_NAME_STRING_CHECKED(FGameplayTagQueryConditionInstanceData, TagContainer));
-
-	FText InvertText;
-	if (bInvert)
-	{
-		InvertText = LOCTEXT("Not", "Not");
-	}
-
-	FText TagContainerText;
-	if (const FStateTreeEditorPropertyPath* Binding = BindingLookup.GetPropertyBindingSource(TagContainerPath))
-	{
-		TagContainerText = BindingLookup.GetPropertyPathDisplayName(*Binding);
-	}
-	else
-	{
-		TagContainerText = LOCTEXT("NotBound", "Not Bound");
-	}
-
-	return FText::Format(LOCTEXT("GameplayTagQueryMatchDesc", "{0} <Details.Bold>{1}</> matches query</>"),
-		InvertText, TagContainerText);
-}
-#endif// WITH_EDITOR
-
 
 #if WITH_EDITOR
 #undef LOCTEXT_NAMESPACE
diff --git a/Engine/Plugins/Runtime/StateTree/Source/StateTreeModule/Private/Conditions/StateTreeCondition_Common.cpp b/Engine/Plugins/Runtime/StateTree/Source/StateTreeModule/Private/Conditions/StateTreeCondition_Common.cpp
--- a/Engine/Plugins/Runtime/StateTree/Source/StateTreeModule/Private/Conditions/StateTreeCondition_Common.cpp
+++ b/Engine/Plugins/Runtime/StateTree/Source/StateTreeModule/Private/Conditions/StateTreeCondition_Common.cpp
@@ -95,45 +95,6 @@ bool FStateTreeCondition_CompareInt::TestCondition(FStateTreeExecutionContext& C
 	return bResult ^ bInvert;
 }
 
-#if WITH_EDITOR
-FText FStateTreeCondition_CompareInt::GetDescription(const FGuid& ID, FStateTreeDataView InstanceData, const IStateTreeBindingLookup& BindingLookup) const
-{
-	const FStateTreeCondition_CompareIntInstanceData& Instance = InstanceData.Get<FStateTreeCondition_CompareIntInstanceData>();
-	const FStateTreeEditorPropertyPath LeftPath(ID, GET_MEMBER_NAME_STRING_CHECKED(FStateTreeCondition_CompareIntInstanceData, Left));
-	const FStateTreeEditorPropertyPath RightPath(ID, GET_MEMBER_NAME_STRING_CHECKED(FStateTreeCondition_CompareIntInstanceData, Right));
-
-	FText InvertText;
-	if (bInvert)
-	{
-		InvertText = LOCTEXT("Not", "Not");
-	}
-
-	FText LeftText;
-	if (const FStateTreeEditorPropertyPath* Binding = BindingLookup.GetPropertyBindingSource(LeftPath))
-	{
-		LeftText = BindingLookup.GetPropertyPathDisplayName(*Binding);
-	}
-	else
-	{
-		LeftText = LOCTEXT("NotBound", "Not Bound");
-	}
-
-	FText OperatorText = UE::StateTree::Conditions::GetOperatorText(Operator);
-
-	FText RightText;
-	if (const FStateTreeEditorPropertyPath* Binding = BindingLookup.GetPropertyBindingSource(RightPath))
-	{
-		RightText = BindingLookup.GetPropertyPathDisplayName(*Binding);
-	}
-	else
-	{
-		RightText = FText::AsNumber(Instance.Right);
-	}
-
-	return FText::Format(LOCTEXT("CompareIntDesc", "{0} <Details.Bold>{1}</> {2} <Details.Bold>{3}</>"), InvertText, LeftText, OperatorText, RightText);
-}
-#endif
-
 
 //----------------------------------------------------------------------//
 //  FStateTreeCondition_CompareFloat
@@ -156,46 +117,6 @@ bool FStateTreeCondition_CompareFloat::TestCondition(FStateTreeExecutionContext&
 }
 
 
-#if WITH_EDITOR
-FText FStateTreeCondition_CompareFloat::GetDescription(const FGuid& ID, FStateTreeDataView InstanceData, const IStateTreeBindingLookup& BindingLookup) const
-{
-	const FStateTreeCondition_CompareFloatInstanceData& Instance = InstanceData.Get<FStateTreeCondition_CompareFloatInstanceData>();
-	const FStateTreeEditorPropertyPath LeftPath(ID, GET_MEMBER_NAME_STRING_CHECKED(FStateTreeCondition_CompareFloatInstanceData, Left));
-	const FStateTreeEditorPropertyPath RightPath(ID, GET_MEMBER_NAME_STRING_CHECKED(FStateTreeCondition_CompareFloatInstanceData, Right));
-
-	FText InvertText;
-	if (bInvert)
-	{
-		InvertText = LOCTEXT("Not", "Not");
-	}
-
-	FText LeftText;
-	if (const FStateTreeEditorPropertyPath* Binding = BindingLookup.GetPropertyBindingSource(LeftPath))
-	{
-		LeftText = BindingLookup.GetPropertyPathDisplayName(*Binding);
-	}
-	else
-	{
-		LeftText = LOCTEXT("NotBound", "Not Bound");
-	}
-
-	FText OperatorText = UE::StateTree::Conditions::GetOperatorText(Operator);
-
-	FText RightText;
-	if (const FStateTreeEditorPropertyPath* Binding = BindingLookup.GetPropertyBindingSource(RightPath))
-	{
-		RightText = BindingLookup.GetPropertyPathDisplayName(*Binding);
-	}
-	else
-	{
-		RightText = FText::AsNumber(Instance.Right);
-	}
-
-	return FText::Format(LOCTEXT("CompareFloatDesc", "{0} <Details.Bold>{1}</> {2} <Details.Bold>{3}</>"), InvertText, LeftText, OperatorText, RightText);
-}
-#endif
-
-
 //----------------------------------------------------------------------//
 //  FStateTreeCondition_CompareBool
 //----------------------------------------------------------------------//
@@ -215,43 +136,6 @@ bool FStateTreeCondition_CompareBool::TestCondition(FStateTreeExecutionContext&
 	return (bLeft == bRight) ^ bInvert;
 }
 
-#if WITH_EDITOR
-FText FStateTreeCondition_CompareBool::GetDescription(const FGuid& ID, FStateTreeDataView InstanceData, const IStateTreeBindingLookup& BindingLookup) const
-{
-	const FStateTreeCondition_CompareBoolInstanceData& Instance = InstanceData.Get<FStateTreeCondition_CompareBoolInstanceData>();
-	const FStateTreeEditorPropertyPath LeftPath(ID, GET_MEMBER_NAME_STRING_CHECKED(FStateTreeCondition_CompareBoolInstanceData, bLeft));
-	const FStateTreeEditorPropertyPath RightPath(ID, GET_MEMBER_NAME_STRING_CHECKED(FStateTreeCondition_CompareBoolInstanceData, bRight));
-
-	FText InvertText;
-	if (bInvert)
-	{
-		InvertText = LOCTEXT("Not", "Not");
-	}
-
-	FText LeftText;
-	if (const FStateTreeEditorPropertyPath* Binding = BindingLookup.GetPropertyBindingSource(LeftPath))
-	{
-		LeftText = BindingLookup.GetPropertyPathDisplayName(*Binding);
-	}
-	else
-	{
-		LeftText = LOCTEXT("NotBound", "Not Bound");
-	}
-
-	FText RightText;
-	if (const FStateTreeEditorPropertyPath* Binding = BindingLookup.GetPropertyBindingSource(RightPath))
-	{
-		RightText = BindingLookup.GetPropertyPathDisplayName(*Binding);
-	}
-	else
-	{
-		RightText = FText::FromString(LexToSanitizedString(Instance.bRight));
-	}
-
-	return FText::Format(LOCTEXT("CompareBoolDesc", "{0} <Details.Bold>{1}</> is <Details.Bold>{2}</>"), InvertText, LeftText, RightText);
-}
-#endif
-
 
 //----------------------------------------------------------------------//
 //  FStateTreeCondition_CompareEnum
@@ -273,48 +157,6 @@ bool FStateTreeCondition_CompareEnum::TestCondition(FStateTreeExecutionContext&
 }
 
 #if WITH_EDITOR
-FText FStateTreeCondition_CompareEnum::GetDescription(const FGuid& ID, FStateTreeDataView InstanceData, const IStateTreeBindingLookup& BindingLookup) const
-{
-	const FStateTreeCondition_CompareEnumInstanceData& Instance = InstanceData.Get<FStateTreeCondition_CompareEnumInstanceData>();
-	const FStateTreeEditorPropertyPath LeftPath(ID, GET_MEMBER_NAME_STRING_CHECKED(FStateTreeCondition_CompareEnumInstanceData, Left));
-	const FStateTreeEditorPropertyPath RightPath(ID, GET_MEMBER_NAME_STRING_CHECKED(FStateTreeCondition_CompareEnumInstanceData, Right));
-
-	FText InvertText;
-	if (bInvert)
-	{
-		InvertText = LOCTEXT("Not", "Not");
-	}
-
-	FText LeftText;
-	if (const FStateTreeEditorPropertyPath* Binding = BindingLookup.GetPropertyBindingSource(LeftPath))
-	{
-		LeftText = BindingLookup.GetPropertyPathDisplayName(*Binding);
-	}
-	else
-	{
-		LeftText = LOCTEXT("NotBound", "Not Bound");
-	}
-
-	FText RightText;
-	if (const FStateTreeEditorPropertyPath* Binding = BindingLookup.GetPropertyBindingSource(RightPath))
-	{
-		RightText = BindingLookup.GetPropertyPathDisplayName(*Binding);
-	}
-	else
-	{
-		if (Instance.Right.Enum)
-		{
-			RightText = Instance.Right.Enum->GetDisplayNameTextByValue(int64(Instance.Right.Value));
-		}
-		else
-		{
-			RightText = LOCTEXT("Invalid", "Invalid");
-		}
-	}
-
-	return FText::Format(LOCTEXT("EnumEqualsDesc", "{0} <Details.Bold>{1}</> is <Details.Bold>{2}</>"), InvertText, LeftText, RightText);
-}
-
 void FStateTreeCondition_CompareEnum::OnBindingChanged(const FGuid& ID, FStateTreeDataView InstanceData, const FStateTreeEditorPropertyPath& SourcePath, const FStateTreeEditorPropertyPath& TargetPath, const IStateTreeBindingLookup& BindingLookup)
 {
 	if (!TargetPath.IsValid())
@@ -385,55 +227,6 @@ bool FStateTreeCondition_CompareDistance::TestCondition(FStateTreeExecutionConte
 	return bResult ^ bInvert;
 }
 
-#if WITH_EDITOR
-FText FStateTreeCondition_CompareDistance::GetDescription(const FGuid& ID, FStateTreeDataView InstanceData, const IStateTreeBindingLookup& BindingLookup) const
-{
-	const FStateTreeCondition_CompareDistanceInstanceData& Instance = InstanceData.Get<FStateTreeCondition_CompareDistanceInstanceData>();
-	const FStateTreeEditorPropertyPath SourcePath(ID, GET_MEMBER_NAME_STRING_CHECKED(FStateTreeCondition_CompareDistanceInstanceData, Source));
-	const FStateTreeEditorPropertyPath TargetPath(ID, GET_MEMBER_NAME_STRING_CHECKED(FStateTreeCondition_CompareDistanceInstanceData, Target));
-	const FStateTreeEditorPropertyPath DistancePath(ID, GET_MEMBER_NAME_STRING_CHECKED(FStateTreeCondition_CompareDistanceInstanceData, Distance));
-
-	FText InvertText;
-	if (bInvert)
-	{
-		InvertText = LOCTEXT("Not", "Not");
-	}
-
-	FText SourceText;
-	if (const FStateTreeEditorPropertyPath* Binding = BindingLookup.GetPropertyBindingSource(SourcePath))
-	{
-		SourceText = BindingLookup.GetPropertyPathDisplayName(*Binding);
-	}
-	else
-	{
-		SourceText = LOCTEXT("NotBound", "Not Bound");
-	}
-
-	FText TargetText;
-	if (const FStateTreeEditorPropertyPath* Binding = BindingLookup.GetPropertyBindingSource(TargetPath))
-	{
-		TargetText = BindingLookup.GetPropertyPathDisplayName(*Binding);
-	}
-	else
-	{
-		TargetText = Instance.Target.ToCompactText();
-	}
-
-	FText OperatorText = UE::StateTree::Conditions::GetOperatorText(Operator);
-
-	FText DistanceText;
-	if (const FStateTreeEditorPropertyPath* DistanceBinding = BindingLookup.GetPropertyBindingSource(DistancePath))
-	{
-		DistanceText = BindingLookup.GetPropertyPathDisplayName(*DistanceBinding);
-	}
-	else
-	{
-		DistanceText = FText::AsNumber(Instance.Distance);
-	}
-
-	return FText::Format(LOCTEXT("CompareDistanceDesc", "{0} Distance <Details.Subdued>from</> <Details.Bold>{1}</> <Details.Subdued>to</> <Details.Bold>{2}</> {3} <Details.Bold>{4}</>"), InvertText, SourceText, TargetText, OperatorText, DistanceText);
-}
-#endif
 
 //----------------------------------------------------------------------//
 //  FStateTreeCondition_Random
@@ -452,24 +245,5 @@ bool FStateTreeCondition_Random::TestCondition(FStateTreeExecutionContext& Conte
 	return FMath::FRandRange(0.0f, 1.0f) < Threshold;
 }
 
-#if WITH_EDITOR
-FText FStateTreeCondition_Random::GetDescription(const FGuid& ID, FStateTreeDataView InstanceData, const IStateTreeBindingLookup& BindingLookup) const
-{
-	const FStateTreeCondition_RandomInstanceData& Instance = InstanceData.Get<FStateTreeCondition_RandomInstanceData>();
-	const FStateTreeEditorPropertyPath ThresholdPath(ID, GET_MEMBER_NAME_STRING_CHECKED(FStateTreeCondition_RandomInstanceData, Threshold));
-
-	FText ThresholdText;
-	if (const FStateTreeEditorPropertyPath* Binding = BindingLookup.GetPropertyBindingSource(ThresholdPath))
-	{
-		ThresholdText = BindingLookup.GetPropertyPathDisplayName(*Binding);
-	}
-	else
-	{
-		ThresholdText = FText::AsPercent(Instance.Threshold);
-	}
-
-	return FText::Format(LOCTEXT("RandomDesc", "<Details.Bold>Random</> {1}"), ThresholdText);
-}
-#endif
 
 #undef LOCTEXT_NAMESPACE
diff --git a/Engine/Plugins/Runtime/StateTree/Source/StateTreeModule/Private/StateTreeExecutionContext.cpp b/Engine/Plugins/Runtime/StateTree/Source/StateTreeModule/Private/StateTreeExecutionContext.cpp
--- a/Engine/Plugins/Runtime/StateTree/Source/StateTreeModule/Private/StateTreeExecutionContext.cpp
+++ b/Engine/Plugins/Runtime/StateTree/Source/StateTreeModule/Private/StateTreeExecutionContext.cpp
@@ -654,9 +654,19 @@ bool FStateTreeExecutionContext::TestAllConditions(FStateTreeInstanceData& Insta
 {
 	CSV_SCOPED_TIMING_STAT_EXCLUSIVE(StateTree_TestConditions);
 
-	for (uint32 i = 0; i < ConditionsNum; i++)
+	if (ConditionsNum == 0)
 	{
-		const FStateTreeConditionBase& Cond = GetNode<FStateTreeConditionBase>(ConditionsOffset + i);
+		return true;
+	}
+	
+	TStaticArray<EStateTreeConditionOperand, UE::StateTree::MaxConditionIndent + 1> Operands(InPlace, EStateTreeConditionOperand::Copy);
+	TStaticArray<bool, UE::StateTree::MaxConditionIndent + 1> Values(InPlace, false);
+
+	int32 Level = 0;
+	
+	for (uint32 Index = 0; Index < ConditionsNum; Index++)
+	{
+		const FStateTreeConditionBase& Cond = GetNode<FStateTreeConditionBase>(ConditionsOffset + Index);
 		const FStateTreeDataView CondInstanceView = GetInstanceData(InstanceData, Cond.bInstanceIsObject, Cond.InstanceIndex);
 		DataViews[Cond.DataViewIndex] = CondInstanceView;
 
@@ -665,13 +675,46 @@ bool FStateTreeExecutionContext::TestAllConditions(FStateTreeInstanceData& Insta
 		{
 			StateTree->PropertyBindings.CopyTo(DataViews, Cond.BindingsBatch.Index, CondInstanceView);
 		}
-		
-		if (!Cond.TestCondition(*this))
+
+		const bool bValue = Cond.TestCondition(*this);
+
+		const int32 DeltaIndent = Cond.DeltaIndent;
+		const int32 OpenParens = FMath::Max(0, DeltaIndent) + 1;	// +1 for the current value that is stored at the empty slot at the top of the value stack.
+		const int32 ClosedParens = FMath::Max(0, -DeltaIndent) + 1;
+
+		// Store the operand to apply when merging higher level down when returning to this level.
+		// @todo: remove this conditions in 5.1, needs resaving existing StateTrees.
+		const EStateTreeConditionOperand Operand = Index == 0 ? EStateTreeConditionOperand::Copy : Cond.Operand;
+		Operands[Level] = Operand;
+
+		// Store current value at the top of the stack.
+		Level += OpenParens;
+		Values[Level] = bValue;
+
+		// Evaluate and merge down values based on closed braces.
+		// The current value is placed in parens (see +1 above), which makes merging down and applying the new value consistent.
+		// The default operand is copy, so if the value is needed immediately, it is just copied down, or if we're on the same level,
+		// the operand storing above gives handles with the right logic.
+		for (int32 Paren = 0; Paren < ClosedParens; Paren++)
 		{
-			return false;
+			Level--;
+			switch (Operands[Level])
+			{
+			case EStateTreeConditionOperand::Copy:
+				Values[Level] = Values[Level + 1];
+				break;
+			case EStateTreeConditionOperand::And:
+				Values[Level] &= Values[Level + 1];
+				break;
+			case EStateTreeConditionOperand::Or:
+				Values[Level] |= Values[Level + 1];
+				break;
+			}
+			Operands[Level] = EStateTreeConditionOperand::Copy;
 		}
 	}
-	return true;
+	
+	return Values[0];
 }
 
 FStateTreeTransitionResult FStateTreeExecutionContext::TriggerTransitions(FStateTreeInstanceData& InstanceData, const FStateTreeStateStatus CurrentStatus, const int32 Depth)
diff --git a/Engine/Plugins/Runtime/StateTree/Source/StateTreeModule/Public/Blueprint/StateTreeConditionBlueprintBase.h b/Engine/Plugins/Runtime/StateTree/Source/StateTreeModule/Public/Blueprint/StateTreeConditionBlueprintBase.h
--- a/Engine/Plugins/Runtime/StateTree/Source/StateTreeModule/Public/Blueprint/StateTreeConditionBlueprintBase.h
+++ b/Engine/Plugins/Runtime/StateTree/Source/StateTreeModule/Public/Blueprint/StateTreeConditionBlueprintBase.h
@@ -39,11 +39,6 @@ struct STATETREEMODULE_API FStateTreeBlueprintConditionWrapper : public FStateTr
 	virtual bool Link(FStateTreeLinker& Linker) override;
 	virtual bool TestCondition(FStateTreeExecutionContext& Context) const override;
 
-#if WITH_EDITOR
-	/** @return Rich text description of the condition. */
-	virtual FText GetDescription(const FGuid& ID, FStateTreeDataView InstanceData, const IStateTreeBindingLookup& BindingLookup) const;
-#endif
-	
 	UPROPERTY()
 	TSubclassOf<UStateTreeConditionBlueprintBase> ConditionClass = nullptr;
 
diff --git a/Engine/Plugins/Runtime/StateTree/Source/StateTreeModule/Public/Conditions/GameplayTagConditions.h b/Engine/Plugins/Runtime/StateTree/Source/StateTreeModule/Public/Conditions/GameplayTagConditions.h
--- a/Engine/Plugins/Runtime/StateTree/Source/StateTreeModule/Public/Conditions/GameplayTagConditions.h
+++ b/Engine/Plugins/Runtime/StateTree/Source/StateTreeModule/Public/Conditions/GameplayTagConditions.h
@@ -34,10 +34,6 @@ struct STATETREEMODULE_API FGameplayTagMatchCondition : public FStateTreeConditi
 	virtual const UStruct* GetInstanceDataType() const override { return FGameplayTagMatchConditionInstanceData::StaticStruct(); }
 	virtual bool TestCondition(FStateTreeExecutionContext& Context) const override;
 
-#if WITH_EDITOR
-	virtual FText GetDescription(const FGuid& ID, FStateTreeDataView InstanceData, const IStateTreeBindingLookup& BindingLookup) const override;
-#endif
-
 	TStateTreeInstanceDataPropertyHandle<FGameplayTagContainer> TagContainerHandle;
 	TStateTreeInstanceDataPropertyHandle<FGameplayTag> TagHandle;
 	
@@ -74,10 +70,6 @@ struct STATETREEMODULE_API FGameplayTagContainerMatchCondition : public FStateTr
 	virtual const UStruct* GetInstanceDataType() const override { return FGameplayTagContainerMatchConditionInstanceData::StaticStruct(); }
 	virtual bool TestCondition(FStateTreeExecutionContext& Context) const override;
 
-#if WITH_EDITOR
-	virtual FText GetDescription(const FGuid& ID, FStateTreeDataView InstanceData, const IStateTreeBindingLookup& BindingLookup) const override;
-#endif
-
 	TStateTreeInstanceDataPropertyHandle<FGameplayTagContainer> TagContainerHandle;
 	TStateTreeInstanceDataPropertyHandle<FGameplayTagContainer> OtherContainerHandle;
 
@@ -114,10 +106,6 @@ struct STATETREEMODULE_API FGameplayTagQueryCondition : public FStateTreeConditi
 	virtual const UStruct* GetInstanceDataType() const override { return FGameplayTagQueryConditionInstanceData::StaticStruct(); }
 	virtual bool TestCondition(FStateTreeExecutionContext& Context) const override;
 
-#if WITH_EDITOR
-	virtual FText GetDescription(const FGuid& ID, FStateTreeDataView InstanceData, const IStateTreeBindingLookup& BindingLookup) const override;
-#endif
-
 	TStateTreeInstanceDataPropertyHandle<FGameplayTagContainer> TagContainerHandle;
 
 	UPROPERTY(EditAnywhere, Category = Condition)
diff --git a/Engine/Plugins/Runtime/StateTree/Source/StateTreeModule/Public/Conditions/StateTreeCondition_Common.h b/Engine/Plugins/Runtime/StateTree/Source/StateTreeModule/Public/Conditions/StateTreeCondition_Common.h
--- a/Engine/Plugins/Runtime/StateTree/Source/StateTreeModule/Public/Conditions/StateTreeCondition_Common.h
+++ b/Engine/Plugins/Runtime/StateTree/Source/StateTreeModule/Public/Conditions/StateTreeCondition_Common.h
@@ -47,10 +47,6 @@ struct STATETREEMODULE_API FStateTreeCondition_CompareInt : public FStateTreeCon
 	virtual const UStruct* GetInstanceDataType() const override { return FStateTreeCondition_CompareIntInstanceData::StaticStruct(); }
 	virtual bool TestCondition(FStateTreeExecutionContext& Context) const override;
 
-#if WITH_EDITOR
-	virtual FText GetDescription(const FGuid& ID, FStateTreeDataView InstanceData, const IStateTreeBindingLookup& BindingLookup) const override;
-#endif
-
 	UPROPERTY(EditAnywhere, Category = Parameter)
 	bool bInvert = false;
 
@@ -93,10 +89,6 @@ struct STATETREEMODULE_API FStateTreeCondition_CompareFloat : public FStateTreeC
 	virtual const UStruct* GetInstanceDataType() const override { return FStateTreeCondition_CompareFloatInstanceData::StaticStruct(); }
 	virtual bool TestCondition(FStateTreeExecutionContext& Context) const override;
 
-#if WITH_EDITOR
-	virtual FText GetDescription(const FGuid& ID, FStateTreeDataView InstanceData, const IStateTreeBindingLookup& BindingLookup) const override;
-#endif
-
 	UPROPERTY(EditAnywhere, Category = Parameter)
 	bool bInvert = false;
 
@@ -142,10 +134,6 @@ struct STATETREEMODULE_API FStateTreeCondition_CompareBool : public FStateTreeCo
 	virtual const UStruct* GetInstanceDataType() const override { return FStateTreeCondition_CompareBoolInstanceData::StaticStruct(); }
 	virtual bool TestCondition(FStateTreeExecutionContext& Context) const override;
 
-#if WITH_EDITOR
-	virtual FText GetDescription(const FGuid& ID, FStateTreeDataView InstanceData, const IStateTreeBindingLookup& BindingLookup) const override;
-#endif
-
 	UPROPERTY(EditAnywhere, Category = Parameter)
 	bool bInvert = false;
 
@@ -189,7 +177,6 @@ struct STATETREEMODULE_API FStateTreeCondition_CompareEnum : public FStateTreeCo
 	virtual bool TestCondition(FStateTreeExecutionContext& Context) const override;
 
 #if WITH_EDITOR
-	virtual FText GetDescription(const FGuid& ID, FStateTreeDataView InstanceData, const IStateTreeBindingLookup& BindingLookup) const override;
 	virtual void OnBindingChanged(const FGuid& ID, FStateTreeDataView InstanceData, const FStateTreeEditorPropertyPath& SourcePath, const FStateTreeEditorPropertyPath& TargetPath, const IStateTreeBindingLookup& BindingLookup) override;
 #endif
 
@@ -235,10 +222,6 @@ struct STATETREEMODULE_API FStateTreeCondition_CompareDistance : public FStateTr
 	virtual const UStruct* GetInstanceDataType() const override { return FStateTreeCondition_CompareDistanceInstanceData::StaticStruct(); }
 	virtual bool TestCondition(FStateTreeExecutionContext& Context) const override;
 
-#if WITH_EDITOR
-	virtual FText GetDescription(const FGuid& ID, FStateTreeDataView InstanceData, const IStateTreeBindingLookup& BindingLookup) const override;
-#endif
-
 	UPROPERTY(EditAnywhere, Category = Condition)
 	bool bInvert = false;
 
@@ -275,9 +258,5 @@ struct STATETREEMODULE_API FStateTreeCondition_Random : public FStateTreeConditi
 	virtual const UStruct* GetInstanceDataType() const override { return FStateTreeCondition_RandomInstanceData::StaticStruct(); }
 	virtual bool TestCondition(FStateTreeExecutionContext& Context) const override;
 
-#if WITH_EDITOR
-	virtual FText GetDescription(const FGuid& ID, FStateTreeDataView InstanceData, const IStateTreeBindingLookup& BindingLookup) const override;
-#endif
-
 	TStateTreeInstanceDataPropertyHandle<float> ThresholdHandle;
 };
diff --git a/Engine/Plugins/Runtime/StateTree/Source/StateTreeModule/Public/StateTreeConditionBase.h b/Engine/Plugins/Runtime/StateTree/Source/StateTreeModule/Public/StateTreeConditionBase.h
--- a/Engine/Plugins/Runtime/StateTree/Source/StateTreeModule/Public/StateTreeConditionBase.h
+++ b/Engine/Plugins/Runtime/StateTree/Source/StateTreeModule/Public/StateTreeConditionBase.h
@@ -24,8 +24,6 @@ struct STATETREEMODULE_API FStateTreeConditionBase : public FStateTreeNodeBase
 	GENERATED_BODY()
 
 #if WITH_EDITOR
-	/** @return Rich text description of the condition. */
-	virtual FText GetDescription(const FGuid& ID, FStateTreeDataView InstanceData, const IStateTreeBindingLookup& BindingLookup) const { return FText::GetEmpty(); }
 	/**
 	 * Called when binding of any of the properties in the condition changes.
 	 * @param ID ID of the item, can be used make property paths to this item.
@@ -39,6 +37,12 @@ struct STATETREEMODULE_API FStateTreeConditionBase : public FStateTreeNodeBase
 	
 	/** @return True if the condition passes. */
 	virtual bool TestCondition(FStateTreeExecutionContext& Context) const { return false; }
+
+	UPROPERTY()
+	EStateTreeConditionOperand Operand = EStateTreeConditionOperand::And;
+
+	UPROPERTY()
+	int8 DeltaIndent = 0;
 };
 
 /**
diff --git a/Engine/Plugins/Runtime/StateTree/Source/StateTreeModule/Public/StateTreeTypes.h b/Engine/Plugins/Runtime/StateTree/Source/StateTreeModule/Public/StateTreeTypes.h
--- a/Engine/Plugins/Runtime/StateTree/Source/StateTreeModule/Public/StateTreeTypes.h
+++ b/Engine/Plugins/Runtime/StateTree/Source/StateTreeModule/Public/StateTreeTypes.h
@@ -11,46 +11,59 @@ STATETREEMODULE_API DECLARE_LOG_CATEGORY_EXTERN(LogStateTree, Warning, All);
 #define WITH_STATETREE_DEBUG (!(UE_BUILD_SHIPPING || UE_BUILD_SHIPPING_WITH_EDITOR || UE_BUILD_TEST) && 1)
 #endif // WITH_STATETREE_DEBUG
 
-/**
- * Status describing current ticking state. 
- */
+/** Status describing current ticking state. */
 UENUM()
 enum class EStateTreeRunStatus : uint8
 {
-	Running,			/** Tree is still running. */
-	Failed,				/** Tree execution has stopped on failure. */
-	Succeeded,			/** Tree execution has stopped on success. */
-	Unset,				/** Status not set. */
+	Running,			/// Tree is still running.
+	Failed,				/// Tree execution has stopped on failure.
+	Succeeded,			/// Tree execution has stopped on success.
+	Unset,				/// Status not set.
 };
 
 /**  Evaluator evaluation type. */
 UENUM()
 enum class EStateTreeEvaluationType : uint8
 {
-	PreSelect,			/** Called during selection process on states that have not been visited yet. */
-    Tick,				/** Called during tick on active states. */
+	PreSelect,			/// Called during selection process on states that have not been visited yet.
+    Tick,				/// Called during tick on active states.
 };
 
 /**  State change type. Passed to EnterState() and ExitState() to indicate how the state change affects the state and Evaluator or Task is on. */
 UENUM()
 enum class EStateTreeStateChangeType : uint8
 {
-	None,				/** Not an activation */
-	Changed,			/** The state became activated or deactivated. */
-    Sustained,			/** The state is parent of new active state and sustained previous active state. */
+	None,				/// Not an activation
+	Changed,			/// The state became activated or deactivated.
+    Sustained,			/// The state is parent of new active state and sustained previous active state.
 };
 
 /** Transitions behavior. */
 UENUM()
 enum class EStateTreeTransitionType : uint8
 {
-	Succeeded,			// Signal StateTree execution succeeded.
-	Failed,				// Signal StateTree execution failed.
-	GotoState,			// Transition to specified state.
-	NotSet,				// No transition.
-	NextState,			// Goto next sibling state.
+	Succeeded,			/// Signal StateTree execution succeeded.
+	Failed,				/// Signal StateTree execution failed.
+	GotoState,			/// Transition to specified state.
+	NotSet,				/// No transition.
+	NextState,			/// Goto next sibling state.
 };
 
+/** Operand between conditions */
+UENUM()
+enum class EStateTreeConditionOperand : uint8
+{
+	Copy UMETA(Hidden),	/// Copy result
+	And,				/// Combine results with AND.
+	Or,					/// Combine results with OR.
+};
+
+namespace UE::StateTree
+{
+	constexpr int32 MaxConditionIndent = 4; 
+}; // UE::StateTree
+
+
 
 /** Transitions event. */
 UENUM()
