StateTree: fix for nonunity build

diff --git a/Engine/Plugins/Runtime/StateTree/Source/StateTreeEditorModule/Private/Customizations/StateTreeEditorNodeDetails.h b/Engine/Plugins/Runtime/StateTree/Source/StateTreeEditorModule/Private/Customizations/StateTreeEditorNodeDetails.h
--- a/Engine/Plugins/Runtime/StateTree/Source/StateTreeEditorModule/Private/Customizations/StateTreeEditorNodeDetails.h
+++ b/Engine/Plugins/Runtime/StateTree/Source/StateTreeEditorModule/Private/Customizations/StateTreeEditorNodeDetails.h
@@ -4,6 +4,7 @@
 
 #include "CoreMinimal.h"
 #include "UnrealClient.h"
+#include "Types/SlateStructs.h"
 #include "IPropertyTypeCustomization.h"
 
 class IPropertyHandle;
