[StateTree] misc trivial changes
- removed dependencies on someplugins (GameplayBehaviors, GameplayAbilities and SmartObjects)
#rnx

diff --git a/Engine/Plugins/AI/MassAI/Source/MassAIBehavior/Public/MassStateTreeExecutionContext.h b/Engine/Plugins/AI/MassAI/Source/MassAIBehavior/Public/MassStateTreeExecutionContext.h
--- a/Engine/Plugins/AI/MassAI/Source/MassAIBehavior/Public/MassStateTreeExecutionContext.h
+++ b/Engine/Plugins/AI/MassAI/Source/MassAIBehavior/Public/MassStateTreeExecutionContext.h
@@ -31,7 +31,7 @@ public:
 protected:
 
 	/** Prefix that will be used by STATETREE_LOG and STATETREE_CLOG, using Entity description. */
-	virtual FString GetInstanceDescription() const override { return FString::Printf(TEXT("Entity [%s]"), *Entity.DebugGetDescription()); }
+	virtual FString GetInstanceDescription() const override { return FString::Printf(TEXT("Entity [%s]: "), *Entity.DebugGetDescription()); }
 
 	virtual void BeginGatedTransition(const FStateTreeExecutionState& Exec) override;
 
diff --git a/Engine/Plugins/Runtime/StateTree/Source/StateTreeModule/Private/StateTreeExecutionContext.cpp b/Engine/Plugins/Runtime/StateTree/Source/StateTreeModule/Private/StateTreeExecutionContext.cpp
--- a/Engine/Plugins/Runtime/StateTree/Source/StateTreeModule/Private/StateTreeExecutionContext.cpp
+++ b/Engine/Plugins/Runtime/StateTree/Source/StateTreeModule/Private/StateTreeExecutionContext.cpp
@@ -13,8 +13,8 @@
 #include "Logging/LogScopedVerbosityOverride.h"
 #include "ProfilingDebugging/CsvProfiler.h"
 
-#define STATETREE_LOG(Verbosity, Format, ...) UE_VLOG_UELOG(GetOwner(), LogStateTree, Verbosity, TEXT("%s: ") Format, *GetInstanceDescription(), ##__VA_ARGS__)
-#define STATETREE_CLOG(Condition, Verbosity, Format, ...) UE_CVLOG_UELOG((Condition), GetOwner(), LogStateTree, Verbosity, TEXT("%s: ") Format, *GetInstanceDescription(), ##__VA_ARGS__)
+#define STATETREE_LOG(Verbosity, Format, ...) UE_VLOG_UELOG(GetOwner(), LogStateTree, Verbosity, TEXT("%s") Format, *GetInstanceDescription(), ##__VA_ARGS__)
+#define STATETREE_CLOG(Condition, Verbosity, Format, ...) UE_CVLOG_UELOG((Condition), GetOwner(), LogStateTree, Verbosity, TEXT("%s") Format, *GetInstanceDescription(), ##__VA_ARGS__)
 
 namespace UE::StateTree
 {
@@ -935,6 +935,11 @@ EStateTreeRunStatus FStateTreeExecutionContext::GetLastTickStatus(const FStateTr
 	return Exec.LastTickStatus;
 }
 
+FString FStateTreeExecutionContext::GetInstanceDescription() const
+{
+	return Owner != nullptr? FString::Printf(TEXT("%s: "), *Owner->GetName()) : TEXT("");
+}
+
 #if WITH_GAMEPLAY_DEBUGGER
 
 FString FStateTreeExecutionContext::GetDebugInfoString(const FStateTreeInstanceData* ExternalInstanceData) const
diff --git a/Engine/Plugins/Runtime/StateTree/Source/StateTreeModule/Private/StateTreeSchema.cpp b/Engine/Plugins/Runtime/StateTree/Source/StateTreeModule/Private/StateTreeSchema.cpp
--- a/Engine/Plugins/Runtime/StateTree/Source/StateTreeModule/Private/StateTreeSchema.cpp
+++ b/Engine/Plugins/Runtime/StateTree/Source/StateTreeModule/Private/StateTreeSchema.cpp
@@ -1,10 +1,7 @@
 ﻿// Copyright Epic Games, Inc. All Rights Reserved.
 
 #include "StateTreeSchema.h"
-#include "CoreMinimal.h"
-#include "Blueprint/StateTreeEvaluatorBlueprintBase.h"
-#include "Blueprint/StateTreeTaskBlueprintBase.h"
-#include "Blueprint/StateTreeConditionBlueprintBase.h"
+#include "Blueprint/StateTreeItemBlueprintBase.h"
 
 bool UStateTreeSchema::IsChildOfBlueprintBase(const UClass* InClass) const
 {
diff --git a/Engine/Plugins/Runtime/StateTree/Source/StateTreeModule/Public/StateTreeExecutionContext.h b/Engine/Plugins/Runtime/StateTree/Source/StateTreeModule/Public/StateTreeExecutionContext.h
--- a/Engine/Plugins/Runtime/StateTree/Source/StateTreeModule/Public/StateTreeExecutionContext.h
+++ b/Engine/Plugins/Runtime/StateTree/Source/StateTreeModule/Public/StateTreeExecutionContext.h
@@ -238,7 +238,7 @@ public:
 protected:
 
 	/** @return Prefix that will be used by STATETREE_LOG and STATETREE_CLOG, empty by default. */
-	virtual FString GetInstanceDescription() const { return TEXT(""); }
+	virtual FString GetInstanceDescription() const;
 
 	/** Callback when gated transition is triggered. Contexts that are event based can use this to trigger a future event. */
 	virtual void BeginGatedTransition(const FStateTreeExecutionState& Exec) {};
diff --git a/Engine/Plugins/Runtime/StateTree/Source/StateTreeModule/Public/StateTreeNodeBase.h b/Engine/Plugins/Runtime/StateTree/Source/StateTreeModule/Public/StateTreeNodeBase.h
--- a/Engine/Plugins/Runtime/StateTree/Source/StateTreeModule/Public/StateTreeNodeBase.h
+++ b/Engine/Plugins/Runtime/StateTree/Source/StateTreeModule/Public/StateTreeNodeBase.h
@@ -17,7 +17,7 @@ struct STATETREEMODULE_API FStateTreeNodeBase
 
 	virtual ~FStateTreeNodeBase() {}
 
-	/** @return Struct that represents the runtime data of the evaluator. */
+	/** @return Struct that represents the runtime data of the node. */
 	virtual const UStruct* GetInstanceDataType() const { return nullptr; };
 
 	/**
diff --git a/Engine/Plugins/Runtime/StateTree/Source/StateTreeModule/Public/StateTreeSchema.h b/Engine/Plugins/Runtime/StateTree/Source/StateTreeModule/Public/StateTreeSchema.h
--- a/Engine/Plugins/Runtime/StateTree/Source/StateTreeModule/Public/StateTreeSchema.h
+++ b/Engine/Plugins/Runtime/StateTree/Source/StateTreeModule/Public/StateTreeSchema.h
@@ -2,7 +2,6 @@
 
 #pragma once
 
-#include "CoreMinimal.h"
 #include "StateTreeSchema.generated.h"
 
 /**
@@ -28,10 +27,10 @@ public:
 	virtual bool IsStructAllowed(const UScriptStruct* InScriptStruct) const { return false; }
 
 	/** @return True if specified class is supported */
-	virtual bool IsClassAllowed(const UClass* InScriptStruct) const { return false; };
+	virtual bool IsClassAllowed(const UClass* InScriptStruct) const { return false; }
 
 	/** @return True if specified struct/class is supported as external data */
-	virtual bool IsExternalItemAllowed(const UStruct& InStruct) const { return false; };
+	virtual bool IsExternalItemAllowed(const UStruct& InStruct) const { return false; }
 
 	/**
 	 * Helper function to check if a class is any of the Blueprint extendable item classes (Eval, Task, Condition).
diff --git a/Engine/Plugins/Runtime/StateTree/StateTree.uplugin b/Engine/Plugins/Runtime/StateTree/StateTree.uplugin
--- a/Engine/Plugins/Runtime/StateTree/StateTree.uplugin
+++ b/Engine/Plugins/Runtime/StateTree/StateTree.uplugin
@@ -36,18 +36,6 @@
 			"Name": "StructUtils",
 			"Enabled": true
 		},
-		{
-			"Name": "GameplayBehaviors",
-			"Enabled": true
-		},
-		{
-			"Name": "GameplayAbilities",
-			"Enabled": true
-		},
-		{
-			"Name": "SmartObjects",
-			"Enabled": true
-		},
 		{
 			"Name": "PropertyAccessEditor",
 			"Enabled": true
