StateTree: Added functionality to call and reuse a subtree
- Removed unnecessary checks for specific transition (handled in enter state logic)
- Added state type to StateTreeState, State, Group, or Linked
- StateTreeState property customization hides properties based on type
- Allow a state to link to another state
- Changed baker binding validation to handle linked states
- Added linked state handling in execution context
- Updated UI deal with linked states
- Updated gameplay debugger to deal with linked states
- Moved automatic Root state adding from Editor data to editor open (tests were broken due to extra root state)
- Fixed tests and added simple test for linked state
- Added meta to tag to state link to allow to select only direct states (no next, etc)
- Added counter to track state changes (mainly for debugging)

diff --git a/Engine/Plugins/AI/MassAI/Source/MassAIBehavior/Private/Tasks/MassStateTreeSmartObjectTask.cpp b/Engine/Plugins/AI/MassAI/Source/MassAIBehavior/Private/Tasks/MassStateTreeSmartObjectTask.cpp
--- a/Engine/Plugins/AI/MassAI/Source/MassAIBehavior/Private/Tasks/MassStateTreeSmartObjectTask.cpp
+++ b/Engine/Plugins/AI/MassAI/Source/MassAIBehavior/Private/Tasks/MassStateTreeSmartObjectTask.cpp
@@ -189,7 +189,7 @@ void FMassUseSmartObjectTask::ExitState(FStateTreeExecutionContext& Context, con
 	}
 }
 
-void FMassUseSmartObjectTask::StateCompleted(FStateTreeExecutionContext& Context, const EStateTreeRunStatus CompletionStatus, const FStateTreeHandle CompletedState) const
+void FMassUseSmartObjectTask::StateCompleted(FStateTreeExecutionContext& Context, const EStateTreeRunStatus CompletionStatus, const FStateTreeActiveStates& CompletedActiveStates) const
 {
 	FMassSmartObjectUserFragment& SOUser = Context.GetExternalData(SmartObjectUserHandle);
 
diff --git a/Engine/Plugins/AI/MassAI/Source/MassAIBehavior/Private/Tasks/MassZoneGraphPathFollowTask.cpp b/Engine/Plugins/AI/MassAI/Source/MassAIBehavior/Private/Tasks/MassZoneGraphPathFollowTask.cpp
--- a/Engine/Plugins/AI/MassAI/Source/MassAIBehavior/Private/Tasks/MassZoneGraphPathFollowTask.cpp
+++ b/Engine/Plugins/AI/MassAI/Source/MassAIBehavior/Private/Tasks/MassZoneGraphPathFollowTask.cpp
@@ -89,12 +89,6 @@ bool FMassZoneGraphPathFollowTask::RequestPath(FMassStateTreeExecutionContext& C
 
 EStateTreeRunStatus FMassZoneGraphPathFollowTask::EnterState(FStateTreeExecutionContext& Context, const EStateTreeStateChangeType ChangeType, const FStateTreeTransitionResult& Transition) const
 {
-	// Do not reset of the state if current state is still active after transition, unless transitioned specifically to this state.
-	if (ChangeType == EStateTreeStateChangeType::Sustained && Transition.Current != Transition.Next)
-	{
-		return EStateTreeRunStatus::Running;
-	}
-	
 	FMassStateTreeExecutionContext& MassContext = static_cast<FMassStateTreeExecutionContext&>(Context);
 
 	bool bDisplayDebug = false;
diff --git a/Engine/Plugins/AI/MassAI/Source/MassAIBehavior/Private/Tasks/MassZoneGraphStandTask.cpp b/Engine/Plugins/AI/MassAI/Source/MassAIBehavior/Private/Tasks/MassZoneGraphStandTask.cpp
--- a/Engine/Plugins/AI/MassAI/Source/MassAIBehavior/Private/Tasks/MassZoneGraphStandTask.cpp
+++ b/Engine/Plugins/AI/MassAI/Source/MassAIBehavior/Private/Tasks/MassZoneGraphStandTask.cpp
@@ -30,12 +30,6 @@ bool FMassZoneGraphStandTask::Link(FStateTreeLinker& Linker)
 
 EStateTreeRunStatus FMassZoneGraphStandTask::EnterState(FStateTreeExecutionContext& Context, const EStateTreeStateChangeType ChangeType, const FStateTreeTransitionResult& Transition) const
 {
-	// Do not reset of the state if current state is still active after transition, unless transitioned specifically to this state.
-	if (ChangeType == EStateTreeStateChangeType::Sustained && Transition.Current != Transition.Next)
-	{
-		return EStateTreeRunStatus::Running;
-	}
-
 	const FMassStateTreeExecutionContext& MassContext = static_cast<FMassStateTreeExecutionContext&>(Context);
 
 	const FMassZoneGraphLaneLocationFragment& LaneLocation = Context.GetExternalData(LocationHandle);
diff --git a/Engine/Plugins/AI/MassAI/Source/MassAIBehavior/Public/Tasks/MassStateTreeSmartObjectTask.h b/Engine/Plugins/AI/MassAI/Source/MassAIBehavior/Public/Tasks/MassStateTreeSmartObjectTask.h
--- a/Engine/Plugins/AI/MassAI/Source/MassAIBehavior/Public/Tasks/MassStateTreeSmartObjectTask.h
+++ b/Engine/Plugins/AI/MassAI/Source/MassAIBehavior/Public/Tasks/MassStateTreeSmartObjectTask.h
@@ -70,7 +70,7 @@ protected:
 	virtual const UStruct* GetInstanceDataType() const override { return FMassUseSmartObjectTaskInstanceData::StaticStruct(); }
 	virtual EStateTreeRunStatus EnterState(FStateTreeExecutionContext& Context, const EStateTreeStateChangeType ChangeType, const FStateTreeTransitionResult& Transition) const override;
 	virtual void ExitState(FStateTreeExecutionContext& Context, const EStateTreeStateChangeType ChangeType, const FStateTreeTransitionResult& Transition) const override;
-	virtual void StateCompleted(FStateTreeExecutionContext& Context, const EStateTreeRunStatus CompletionStatus, const FStateTreeHandle CompletedState) const override;
+	virtual void StateCompleted(FStateTreeExecutionContext& Context, const EStateTreeRunStatus CompletionStatus, const FStateTreeActiveStates& CompletedActiveStates) const override;
 	virtual EStateTreeRunStatus Tick(FStateTreeExecutionContext& Context, const float DeltaTime) const override;
 
 	TStateTreeExternalDataHandle<USmartObjectSubsystem> SmartObjectSubsystemHandle;
diff --git a/Engine/Plugins/AI/MassAI/Source/MassAIDebug/Private/GameplayDebuggerCategory_Mass.cpp b/Engine/Plugins/AI/MassAI/Source/MassAIDebug/Private/GameplayDebuggerCategory_Mass.cpp
--- a/Engine/Plugins/AI/MassAI/Source/MassAIDebug/Private/GameplayDebuggerCategory_Mass.cpp
+++ b/Engine/Plugins/AI/MassAI/Source/MassAIDebug/Private/GameplayDebuggerCategory_Mass.cpp
@@ -560,7 +560,7 @@ void FGameplayDebuggerCategory_Mass::CollectData(APlayerController* OwnerPC, AAc
 							if (FStateTreeInstanceData* InstanceData = MassStateTreeSubsystem->GetInstanceData(StateTreeInstance.InstanceHandle))
 							{
 								Status += StateTreeContext.GetActiveStateName(InstanceData);
-								Status += TEXT("\n");
+								Status += FString::Printf(TEXT("  {yellow}%d{white}\n"), StateTreeContext.GetStateChangeCount(InstanceData));
 							}
 							else
 							{
diff --git a/Engine/Plugins/Runtime/StateTree/Config/BaseStateTree.ini b/Engine/Plugins/Runtime/StateTree/Config/BaseStateTree.ini
--- a/Engine/Plugins/Runtime/StateTree/Config/BaseStateTree.ini
+++ b/Engine/Plugins/Runtime/StateTree/Config/BaseStateTree.ini
@@ -1,8 +1,5 @@
 [CoreRedirects]
 
-+PropertyRedirects=(OldName="StateTreeState.DefaultTransition",NewName="StateDoneTransition")
-+PropertyRedirects=(OldName="BakedStateTreeState.DefaultTransitionState",NewName="StateDoneTransitionState")
-+PropertyRedirects=(OldName="BakedStateTreeState.DefaultTransitionType",NewName="StateDoneTransitionType")
 +PropertyRedirects=(OldName="StateTreeBindableStructDesc.ScriptStruct",NewName="StateTreeBindableStructDesc.Struct")
 +PropertyRedirects=(OldName="StateTreeEditorData.Routines",NewName="StateTreeEditorData.SubTrees")
 
diff --git a/Engine/Plugins/Runtime/StateTree/Source/StateTreeEditorModule/Private/Customizations/StateTreeStateDetails.cpp b/Engine/Plugins/Runtime/StateTree/Source/StateTreeEditorModule/Private/Customizations/StateTreeStateDetails.cpp
--- a/Engine/Plugins/Runtime/StateTree/Source/StateTreeEditorModule/Private/Customizations/StateTreeStateDetails.cpp
+++ b/Engine/Plugins/Runtime/StateTree/Source/StateTreeEditorModule/Private/Customizations/StateTreeStateDetails.cpp
@@ -29,6 +29,7 @@
 #include "Modules/ModuleManager.h"
 #include "ISettingsModule.h"
 #include "Editor.h"
+#include "IPropertyUtilities.h"
 #include "StateTreeEditorData.h"
 #include "StateTreeViewModel.h"
 #include "StateTree.h"
@@ -66,12 +67,23 @@ void FStateTreeStateDetails::CustomizeDetails(IDetailLayoutBuilder& DetailBuilde
 	TSharedPtr<IPropertyHandle> EnterConditionsProperty = DetailBuilder.GetProperty(TEXT("EnterConditions"));
 	TSharedPtr<IPropertyHandle> EvaluatorsProperty = DetailBuilder.GetProperty(TEXT("Evaluators"));
 	TSharedPtr<IPropertyHandle> TransitionsProperty = DetailBuilder.GetProperty(TEXT("Transitions"));
+	TSharedPtr<IPropertyHandle> TypeProperty = DetailBuilder.GetProperty(TEXT("Type"));
+	TSharedPtr<IPropertyHandle> LinkedStateProperty = DetailBuilder.GetProperty(TEXT("LinkedState"));
 
+	uint8 StateTypeValue = 0;
+	TypeProperty->GetValue(StateTypeValue);
+	const EStateTreeStateType StateType = (EStateTreeStateType)StateTypeValue;
+	
 	IDetailCategoryBuilder& StateCategory = DetailBuilder.EditCategory(TEXT("State"), LOCTEXT("StateDetailsState", "State"));
 	StateCategory.SetSortOrder(0);
 
+	if (StateType != EStateTreeStateType::Linked)
+	{
+		LinkedStateProperty->MarkHiddenByCustomization();
+	}
+	
 	const FName EvalCategoryName(TEXT("Evaluators"));
-	if (Schema && Schema->AllowEvaluators())
+	if (StateType == EStateTreeStateType::State && Schema && Schema->AllowEvaluators())
 	{
 		MakeArrayCategory(DetailBuilder, EvalCategoryName, LOCTEXT("StateDetailsEvaluators", "Evaluators"), 1, EvaluatorsProperty);
 	}
@@ -90,25 +102,43 @@ void FStateTreeStateDetails::CustomizeDetails(IDetailLayoutBuilder& DetailBuilde
 		DetailBuilder.EditCategory(EnterConditionsCategoryName).SetCategoryVisibility(false);
 	}
 
-	if (Schema && Schema->AllowMultipleTasks())
+	if (StateType == EStateTreeStateType::State)
 	{
-		const FName TasksCategoryName(TEXT("Tasks"));
-		MakeArrayCategory(DetailBuilder, TasksCategoryName, LOCTEXT("StateDetailsTasks", "Tasks"), 3, TasksProperty);
-		SingleTaskProperty->MarkHiddenByCustomization();
+		if (Schema && Schema->AllowMultipleTasks())
+		{
+			const FName TasksCategoryName(TEXT("Tasks"));
+			MakeArrayCategory(DetailBuilder, TasksCategoryName, LOCTEXT("StateDetailsTasks", "Tasks"), 3, TasksProperty);
+			SingleTaskProperty->MarkHiddenByCustomization();
+		}
+		else
+		{
+			const FName TaskCategoryName(TEXT("Task"));
+			IDetailCategoryBuilder& Category = DetailBuilder.EditCategory(TaskCategoryName);
+			Category.SetSortOrder(3);
+			
+			IDetailPropertyRow& Row = Category.AddProperty(SingleTaskProperty);
+			Row.ShouldAutoExpand(true);
+			
+			TasksProperty->MarkHiddenByCustomization();
+		}
 	}
 	else
 	{
-		const FName TaskCategoryName(TEXT("Task"));
-		IDetailCategoryBuilder& Category = DetailBuilder.EditCategory(TaskCategoryName);
-		Category.SetSortOrder(3);
-		
-		IDetailPropertyRow& Row = Category.AddProperty(SingleTaskProperty);
-		Row.ShouldAutoExpand(true);
-		
+		SingleTaskProperty->MarkHiddenByCustomization();
 		TasksProperty->MarkHiddenByCustomization();
 	}
 
 	MakeArrayCategory(DetailBuilder, "Transitions", LOCTEXT("StateDetailsTransitions", "Transitions"), 4, TransitionsProperty);
+
+	// Refresh the UI when the type changes.	
+	TSharedPtr<IPropertyUtilities> PropUtils = DetailBuilder.GetPropertyUtilities();
+	TypeProperty->SetOnPropertyValueChanged(FSimpleDelegate::CreateLambda([PropUtils] ()
+	{
+		if (PropUtils.IsValid())
+		{
+			PropUtils->ForceRefresh();
+		}
+	}));
 }
 
 void FStateTreeStateDetails::MakeArrayCategory(IDetailLayoutBuilder& DetailBuilder, FName CategoryName, const FText& DisplayName, int32 SortOrder, TSharedPtr<IPropertyHandle> PropertyHandle)
diff --git a/Engine/Plugins/Runtime/StateTree/Source/StateTreeEditorModule/Private/Customizations/StateTreeStateLinkDetails.cpp b/Engine/Plugins/Runtime/StateTree/Source/StateTreeEditorModule/Private/Customizations/StateTreeStateLinkDetails.cpp
--- a/Engine/Plugins/Runtime/StateTree/Source/StateTreeEditorModule/Private/Customizations/StateTreeStateLinkDetails.cpp
+++ b/Engine/Plugins/Runtime/StateTree/Source/StateTreeEditorModule/Private/Customizations/StateTreeStateLinkDetails.cpp
@@ -29,6 +29,12 @@ void FStateTreeStateLinkDetails::CustomizeHeader(TSharedRef<class IPropertyHandl
 	IDProperty = StructProperty->GetChildHandle(TEXT("ID"));
 	TypeProperty = StructProperty->GetChildHandle(TEXT("Type"));
 
+	static const FName NAME_DirectStatesOnly = "DirectStatesOnly";
+	if (const FProperty* MetaDataProperty = StructProperty->GetMetaDataProperty())
+	{
+		bAllowMetaStates = !MetaDataProperty->HasMetaData(NAME_DirectStatesOnly);
+	}
+	
 	CacheStates();
 
 	HeaderRow
@@ -143,11 +149,14 @@ TSharedRef<SWidget> FStateTreeStateLinkDetails::OnGetStateContent() const
 {
 	FMenuBuilder MenuBuilder(true, NULL);
 
-	FUIAction NextItemAction(FExecuteAction::CreateSP(const_cast<FStateTreeStateLinkDetails*>(this), &FStateTreeStateLinkDetails::OnStateComboChange, ComboNextState));
-	MenuBuilder.AddMenuEntry(LOCTEXT("TransitionNextState", "Next State"), LOCTEXT("TransitionNextTooltip", "Goto next sibling State."), FSlateIcon(), NextItemAction);
+	if (bAllowMetaStates)
+	{
+		FUIAction NextItemAction(FExecuteAction::CreateSP(const_cast<FStateTreeStateLinkDetails*>(this), &FStateTreeStateLinkDetails::OnStateComboChange, ComboNextState));
+		MenuBuilder.AddMenuEntry(LOCTEXT("TransitionNextState", "Next State"), LOCTEXT("TransitionNextTooltip", "Goto next sibling State."), FSlateIcon(), NextItemAction);
 
-	FUIAction NotSetItemAction(FExecuteAction::CreateSP(const_cast<FStateTreeStateLinkDetails*>(this), &FStateTreeStateLinkDetails::OnStateComboChange, ComboNotSet));
-	MenuBuilder.AddMenuEntry(LOCTEXT("TransitionBlock", "Block Transition"), LOCTEXT("TransitionBlockTooltip", "Will not transition to any state, but will block other transitions to trigger if the condition passes."), FSlateIcon(), NotSetItemAction);
+		FUIAction NotSetItemAction(FExecuteAction::CreateSP(const_cast<FStateTreeStateLinkDetails*>(this), &FStateTreeStateLinkDetails::OnStateComboChange, ComboNotSet));
+		MenuBuilder.AddMenuEntry(LOCTEXT("TransitionBlock", "Block Transition"), LOCTEXT("TransitionBlockTooltip", "Will not transition to any state, but will block other transitions to trigger if the condition passes."), FSlateIcon(), NotSetItemAction);
+	}
 
 	if (CachedNames.Num() > 0)
 	{
diff --git a/Engine/Plugins/Runtime/StateTree/Source/StateTreeEditorModule/Private/Customizations/StateTreeStateLinkDetails.h b/Engine/Plugins/Runtime/StateTree/Source/StateTreeEditorModule/Private/Customizations/StateTreeStateLinkDetails.h
--- a/Engine/Plugins/Runtime/StateTree/Source/StateTreeEditorModule/Private/Customizations/StateTreeStateLinkDetails.h
+++ b/Engine/Plugins/Runtime/StateTree/Source/StateTreeEditorModule/Private/Customizations/StateTreeStateLinkDetails.h
@@ -49,6 +49,9 @@ private:
 	TSharedPtr<IPropertyHandle> IDProperty;
 	TSharedPtr<IPropertyHandle> TypeProperty;
 
+	// If set, allow to select meta states like Next or (tree) Succeeded.
+	bool bAllowMetaStates = true;
+	
 	class IPropertyUtilities* PropUtils;
 	TSharedPtr<IPropertyHandle> StructProperty;
 };
diff --git a/Engine/Plugins/Runtime/StateTree/Source/StateTreeEditorModule/Private/SStateTreeViewRow.cpp b/Engine/Plugins/Runtime/StateTree/Source/StateTreeEditorModule/Private/SStateTreeViewRow.cpp
--- a/Engine/Plugins/Runtime/StateTree/Source/StateTreeEditorModule/Private/SStateTreeViewRow.cpp
+++ b/Engine/Plugins/Runtime/StateTree/Source/StateTreeEditorModule/Private/SStateTreeViewRow.cpp
@@ -58,6 +58,7 @@ void SStateTreeViewRow::Construct(const FArguments& InArgs, const TSharedRef<STa
 
 	static const FLinearColor TasksBackground = FLinearColor(FColor(17, 117, 131));
 	static const FLinearColor EvaluatorsBackground = FLinearColor(FColor(48, 48, 48));
+	static const FLinearColor LinkBackground = FLinearColor(FColor(84, 84, 84));
 	
 	this->ChildSlot
     .HAlign(HAlign_Fill)
@@ -217,7 +218,48 @@ void SStateTreeViewRow::Construct(const FArguments& InArgs, const TSharedRef<STa
 				]
 			]
 		]
+		
+		// Linked State box
+		+ SHorizontalBox::Slot()
+		.VAlign(VAlign_Center)
+		.Padding(FMargin(0.0f, 4.0f))
+		.AutoWidth()
+		[
+			SNew(SBox)
+			.HeightOverride(28.0f)
+			.VAlign(VAlign_Fill)
+			.Visibility(this, &SStateTreeViewRow::GetLinkedStateVisibility)
+			[
+				SNew(SBorder)
+				.BorderImage(FEditorStyle::GetBrush("WhiteBrush"))
+				.BorderBackgroundColor(LinkBackground)
+				.Padding(FMargin(12.0f, 0.0f, 16.0f, 0.0f))
+				[
+					// Link icon
+					SNew(SHorizontalBox)
+					+ SHorizontalBox::Slot()
+					.VAlign(VAlign_Center)
+					.Padding(FMargin(0.0f, 0.0f, 4.0f, 0.0f))
+					.AutoWidth()
+					[
+						SNew(STextBlock)
+						.Text(FEditorFontGlyphs::Link)
+						.TextStyle(FStateTreeEditorStyle::Get(), "StateTree.DetailsIcon")
+					]
 
+					// Linked State
+					+ SHorizontalBox::Slot()
+					.VAlign(VAlign_Center)
+					.AutoWidth()
+					[
+						SNew(STextBlock)
+						.Text(this, &SStateTreeViewRow::GetLinkedStateDesc)
+						.TextStyle(FStateTreeEditorStyle::Get(), "StateTree.Details")
+					]
+				]
+			]
+		]
+		
 		// Completed transitions
 		+ SHorizontalBox::Slot()
 		.VAlign(VAlign_Center)
@@ -490,6 +532,31 @@ FText SStateTreeViewRow::GetTasksDesc() const
 	return FText::Join((FText::FromString(TEXT(" & "))), Names);
 }
 
+EVisibility SStateTreeViewRow::GetLinkedStateVisibility() const
+{
+	if (const UStateTreeState* State = WeakState.Get())
+	{
+		return State->Type == EStateTreeStateType::Linked ? EVisibility::Visible : EVisibility::Collapsed;
+	}
+	return EVisibility::Collapsed;
+}
+
+FText SStateTreeViewRow::GetLinkedStateDesc() const
+{
+	const UStateTreeState* State = WeakState.Get();
+	if (!State)
+	{
+		return FText::GetEmpty();
+	}
+
+	if (State->Type == EStateTreeStateType::Linked)
+	{
+		return FText::FromName(State->LinkedState.Name);
+	}
+	
+	return FText::GetEmpty();
+}
+
 bool SStateTreeViewRow::HasParentTransitionForEvent(const UStateTreeState& State, const EStateTreeTransitionEvent Event) const
 {
 	EStateTreeTransitionEvent CombinedEvents = EStateTreeTransitionEvent::None;
@@ -535,6 +602,7 @@ FText SStateTreeViewRow::GetTransitionsDesc(const UStateTreeState& State, const
 	}
 
 	if (State.Children.Num() == 0
+		&& State.Type == EStateTreeStateType::State
 		&& DescItems.Num() == 0
 		&& Event != EStateTreeTransitionEvent::OnCondition)
 	{
@@ -599,6 +667,7 @@ FText SStateTreeViewRow::GetTransitionsIcon(const UStateTreeState& State, const
 	}
 	
 	if (State.Children.Num() == 0
+		&& State.Type == EStateTreeStateType::State
         && IconType == IconNone
         && Event != EStateTreeTransitionEvent::OnCondition)
 	{
diff --git a/Engine/Plugins/Runtime/StateTree/Source/StateTreeEditorModule/Private/SStateTreeViewRow.h b/Engine/Plugins/Runtime/StateTree/Source/StateTreeEditorModule/Private/SStateTreeViewRow.h
--- a/Engine/Plugins/Runtime/StateTree/Source/StateTreeEditorModule/Private/SStateTreeViewRow.h
+++ b/Engine/Plugins/Runtime/StateTree/Source/StateTreeEditorModule/Private/SStateTreeViewRow.h
@@ -37,6 +37,9 @@ private:
 	EVisibility GetTasksVisibility() const;
 	FText GetTasksDesc() const;
 
+	EVisibility GetLinkedStateVisibility() const;
+	FText GetLinkedStateDesc() const;
+
 	EVisibility GetCompletedTransitionVisibility() const;
 	FText GetCompletedTransitionsDesc() const;
 	FText GetCompletedTransitionsIcon() const;
diff --git a/Engine/Plugins/Runtime/StateTree/Source/StateTreeEditorModule/Private/StateTreeBaker.cpp b/Engine/Plugins/Runtime/StateTree/Source/StateTreeEditorModule/Private/StateTreeBaker.cpp
--- a/Engine/Plugins/Runtime/StateTree/Source/StateTreeEditorModule/Private/StateTreeBaker.cpp
+++ b/Engine/Plugins/Runtime/StateTree/Source/StateTreeEditorModule/Private/StateTreeBaker.cpp
@@ -34,6 +34,24 @@ bool FStateTreeBaker::Bake(UStateTree& InStateTree)
 		return false;
 	}
 
+	if (!CreateExecutionInfos())
+	{
+		StateTree->ResetBaked();
+		return false;
+	}
+	
+	if (!CreateStateEvaluators())
+	{
+		StateTree->ResetBaked();
+		return false;
+	}
+
+	if (!CreateStateTasks())
+	{
+		StateTree->ResetBaked();
+		return false;
+	}
+
 	if (!CreateStateTransitions())
 	{
 		StateTree->ResetBaked();
@@ -59,6 +77,17 @@ FStateTreeHandle FStateTreeBaker::GetStateHandle(const FGuid& StateID) const
 	return FStateTreeHandle(uint16(*Idx));
 }
 
+UStateTreeState* FStateTreeBaker::GetState(const FGuid& StateID)
+{
+	const int32* Idx = IDToState.Find(StateID);
+	if (Idx == nullptr)
+	{
+		return nullptr;
+	}
+
+	return SourceStates[*Idx];
+}
+
 bool FStateTreeBaker::CreateStates()
 {
 	// Create item for the runtime execution state
@@ -78,7 +107,144 @@ bool FStateTreeBaker::CreateStates()
 	return true;
 }
 
-bool FStateTreeBaker::CreateConditions(TConstArrayView<FStateTreeEditorNode> Conditions)
+bool FStateTreeBaker::CreateStateRecursive(UStateTreeState& State, const FStateTreeHandle Parent)
+{
+	FStateTreeCompilerLogStateScope LogStateScope(&State, Log);
+
+	const int32 StateIdx = StateTree->States.AddDefaulted();
+	FBakedStateTreeState& BakedState = StateTree->States[StateIdx];
+	BakedState.Name = State.Name;
+	BakedState.Parent = Parent;
+
+	SourceStates.Add(&State);
+	IDToState.Add(State.ID, StateIdx);
+
+	check(StateTree->Instances.Num() <= int32(MAX_uint16));
+
+	// Child states
+	check(StateTree->States.Num() <= int32(MAX_uint16));
+	BakedState.ChildrenBegin = uint16(StateTree->States.Num());
+	for (UStateTreeState* Child : State.Children)
+	{
+		if (Child)
+		{
+			if (!CreateStateRecursive(*Child, FStateTreeHandle((uint16)StateIdx)))
+			{
+				return false;
+			}
+		}
+	}
+	check(StateTree->States.Num() <= int32(MAX_uint16));
+	StateTree->States[StateIdx].ChildrenEnd = uint16(StateTree->States.Num()); // Cannot use BakedState here, it may be invalid due to array resize.
+
+	return true;
+}
+
+FString FStateTreeBaker::GetExecutionPathString(const TConstArrayView<const UStateTreeState*> Path)
+{
+	FString PathStr;
+	const UStateTreeState* PrevState = nullptr;
+	for (int32 Index = 0; Index < Path.Num(); Index++)
+	{
+		const UStateTreeState* State = Path[Index];
+		check(State);
+		if (PrevState != nullptr)
+		{
+			if (PrevState != State->Parent)
+			{
+				// Linked
+				PathStr += TEXT(">");
+			}
+			else
+			{
+				PathStr += TEXT("/");
+			}
+		}
+		PathStr += State->Name.ToString();
+		PrevState = State;
+	}
+
+	return PathStr;
+}
+
+bool FStateTreeBaker::IsPathLinked(const TConstArrayView<const UStateTreeState*> Path)
+{
+	bool bIsLinked = false; 
+	const UStateTreeState* PrevState = nullptr;
+	for (int32 Index = 0; Index < Path.Num(); Index++)
+	{
+		const UStateTreeState* State = Path[Index];
+		check(State);
+		// If the previous state is not the parent state in tree, then the connection must have come from a linked state (a state linking to an arbitrary state in the tree).
+		if (PrevState && PrevState != State->Parent)
+		{
+			bIsLinked = true;
+			break;
+		}
+		PrevState = State;
+	}
+	return bIsLinked;
+}
+
+bool FStateTreeBaker::CreateExecutionInfos()
+{
+	for (UStateTreeState* SubTree : TreeData->SubTrees)
+	{
+		if (SubTree != nullptr)
+		{
+			TArray<const UStateTreeState*> Path;
+			if (!CreateExecutionInfosRecursive(*SubTree, Path))
+			{
+				return false;
+			}
+		}
+	}
+
+	return true;
+}
+
+bool FStateTreeBaker::CreateExecutionInfosRecursive(UStateTreeState& State, TArray<const UStateTreeState*>& Path)
+{
+	Path.Add(&State);
+	
+	if (Path.Num() > FStateTreeActiveStates::MaxStates)
+	{
+		Log.Reportf(EMessageSeverity::Error,
+		TEXT("Reached maximum execution depth %d at: '%s'."), FStateTreeActiveStates::MaxStates, *GetExecutionPathString(Path));
+		return false;
+	}
+
+	FStateExecutionInfo& ExecInfo = ExecutionInfos.FindOrAdd(&State);
+	ExecInfo.ExecutionPaths.Emplace(Path);
+
+	if (State.Type == EStateTreeStateType::Linked)
+	{
+		if (UStateTreeState* LinkedState = GetState(State.LinkedState.ID))
+		{
+			if (!CreateExecutionInfosRecursive(*LinkedState, Path))
+			{
+				return false;
+			}
+		}
+	}
+	
+	for (UStateTreeState* Child : State.Children)
+	{
+		if (Child)
+		{
+			if (!CreateExecutionInfosRecursive(*Child, Path))
+			{
+				return false;
+			}
+		}
+	}
+
+	Path.Pop();
+
+	return true;
+}
+
+bool FStateTreeBaker::CreateConditions(UStateTreeState& State, TConstArrayView<FStateTreeEditorNode> Conditions)
 {
 	for (int32 Index = 0; Index < Conditions.Num(); Index++)
 	{
@@ -94,7 +260,7 @@ bool FStateTreeBaker::CreateConditions(TConstArrayView<FStateTreeEditorNode> Con
 		
 		check(DeltaIndent >= MIN_int8 && DeltaIndent <= MAX_int8);
 
-		if (!CreateCondition(CondNode, Operand, (int8)DeltaIndent))
+		if (!CreateCondition(State, CondNode, Operand, (int8)DeltaIndent))
 		{
 			return false;
 		}
@@ -103,6 +269,71 @@ bool FStateTreeBaker::CreateConditions(TConstArrayView<FStateTreeEditorNode> Con
 	return true;
 }
 
+bool FStateTreeBaker::CreateStateTasks()
+{
+	for (int32 i = 0; i < StateTree->States.Num(); i++)
+	{
+		FBakedStateTreeState& BakedState = StateTree->States[i];
+		UStateTreeState* SourceState = SourceStates[i];
+		check(SourceState != nullptr);
+
+		FStateTreeCompilerLogStateScope LogStateScope(SourceState, Log);
+		
+		// Create tasks
+		check(StateTree->Nodes.Num() <= int32(MAX_uint16));
+		BakedState.TasksBegin = uint16(StateTree->Nodes.Num());
+
+		for (FStateTreeEditorNode& TaskNode : SourceState->Tasks)
+		{
+			if (!CreateTask(*SourceState, TaskNode))
+			{
+				return false;
+			}
+		}
+
+		if (!CreateTask(*SourceState, SourceState->SingleTask))
+		{
+			return false;
+		}
+	
+		const int32 TasksNum = StateTree->Nodes.Num() - int32(BakedState.TasksBegin);
+		check(TasksNum <= int32(MAX_uint8));
+		BakedState.TasksNum = uint8(TasksNum);
+	}
+	
+	return true;
+}
+
+bool FStateTreeBaker::CreateStateEvaluators()
+{
+	for (int32 i = 0; i < StateTree->States.Num(); i++)
+	{
+		FBakedStateTreeState& BakedState = StateTree->States[i];
+		UStateTreeState* SourceState = SourceStates[i];
+		check(SourceState != nullptr);
+
+		FStateTreeCompilerLogStateScope LogStateScope(SourceState, Log);
+		
+		// Collect evaluators
+		check(StateTree->Nodes.Num() <= int32(MAX_uint16));
+		BakedState.EvaluatorsBegin = uint16(StateTree->Nodes.Num());
+
+		for (FStateTreeEditorNode& EvalNode : SourceState->Evaluators)
+		{
+			if (!CreateEvaluator(*SourceState, EvalNode))
+			{
+				return false;
+			}
+		}
+		
+		const int32 EvaluatorsNum = StateTree->Nodes.Num() - int32(BakedState.EvaluatorsBegin);
+		check(EvaluatorsNum <= int32(MAX_uint8));
+		BakedState.EvaluatorsNum = uint8(EvaluatorsNum);
+	}
+
+	return true;
+}
+
 bool FStateTreeBaker::CreateStateTransitions()
 {
 	for (int32 i = 0; i < StateTree->States.Num(); i++)
@@ -115,7 +346,7 @@ bool FStateTreeBaker::CreateStateTransitions()
 		
 		// Enter conditions.
 		BakedState.EnterConditionsBegin = uint16(StateTree->Nodes.Num());
-		if (!CreateConditions(SourceState->EnterConditions))
+		if (!CreateConditions(*SourceState, SourceState->EnterConditions))
 		{
 			Log.Reportf(EMessageSeverity::Error,
 				TEXT("Failed to create state enter condition."));
@@ -123,6 +354,39 @@ bool FStateTreeBaker::CreateStateTransitions()
 		}
 		BakedState.EnterConditionsNum = uint8(uint16(StateTree->Nodes.Num()) - BakedState.EnterConditionsBegin);
 
+		// Linked state
+		if (SourceState->Type == EStateTreeStateType::Linked)
+		{
+			// Make sure the linked state is not self or parent to this state.
+			const UStateTreeState* LinkedParentState = nullptr;
+			for (const UStateTreeState* State = SourceState; State != nullptr; State = State->Parent)
+			{
+				if (State->ID == SourceState->LinkedState.ID)
+				{
+					LinkedParentState = State;
+					break;
+				}
+			}
+			
+			if (LinkedParentState != nullptr)
+			{
+				Log.Reportf(EMessageSeverity::Error,
+					TEXT("State is linked to it's parent state '%s', which will create infinite loop."),
+					*LinkedParentState->Name.ToString());
+				return false;
+			}
+			
+			BakedState.LinkedState = GetStateHandle(SourceState->LinkedState.ID);
+			
+			if (!BakedState.LinkedState.IsValid())
+			{
+				Log.Reportf(EMessageSeverity::Error,
+					TEXT("Failed to resolve linked state '%s'."),
+					*SourceState->LinkedState.Name.ToString());
+				return false;
+			}
+		}
+		
 		// Transitions
 		BakedState.TransitionsBegin = uint16(StateTree->Transitions.Num());
 		for (FStateTreeTransition& Transition : SourceState->Transitions)
@@ -139,10 +403,10 @@ bool FStateTreeBaker::CreateStateTransitions()
 			// Note: Unset transition is allowed here. It can be used to mask a transition at parent.
 
 			BakedTransition.ConditionsBegin = uint16(StateTree->Nodes.Num());
-			if (!CreateConditions(Transition.Conditions))
+			if (!CreateConditions(*SourceState, Transition.Conditions))
 			{
 				Log.Reportf(EMessageSeverity::Error,
-					TEXT("Failed to create condition for transition to %s."),
+					TEXT("Failed to create condition for transition to '%s'."),
 					*Transition.State.Name.ToString());
 				return false;
 			}
@@ -164,7 +428,7 @@ bool FStateTreeBaker::ResolveTransitionState(const UStateTreeState& SourceState,
 		if (!OutTransitionHandle.IsValid())
 		{
 			Log.Reportf(EMessageSeverity::Error,
-				TEXT("Failed to resolve transition to state %s."),
+				TEXT("Failed to resolve transition to state '%s'."),
 				*Link.Name.ToString());
 			return false;
 		}
@@ -183,7 +447,7 @@ bool FStateTreeBaker::ResolveTransitionState(const UStateTreeState& SourceState,
 		if (!OutTransitionHandle.IsValid())
 		{
 			Log.Reportf(EMessageSeverity::Error,
-				TEXT("Failed to resolve transition next state, no handle found for %s."),
+				TEXT("Failed to resolve transition next state, no handle found for '%s'."),
 				*NextState->Name.ToString());
 			return false;
 		}
@@ -192,7 +456,7 @@ bool FStateTreeBaker::ResolveTransitionState(const UStateTreeState& SourceState,
 	return true;
 }
 
-bool FStateTreeBaker::CreateCondition(const FStateTreeEditorNode& CondNode, const EStateTreeConditionOperand Operand, const int8 DeltaIndent)
+bool FStateTreeBaker::CreateCondition(UStateTreeState& State, const FStateTreeEditorNode& CondNode, const EStateTreeConditionOperand Operand, const int8 DeltaIndent)
 {
 	if (!CondNode.Node.IsValid())
 	{
@@ -231,7 +495,7 @@ bool FStateTreeBaker::CreateCondition(const FStateTreeEditorNode& CondNode, cons
 	
 		// Create binding source struct descriptor.
 		StructDesc.Struct = Instance.GetScriptStruct();
-		StructDesc.Name = StructDesc.Struct->GetFName();
+		StructDesc.Name = Cond.Name;
 
 		check(InstanceIndex <= int32(MAX_uint16));
 		Cond.InstanceIndex = uint16(InstanceIndex);
@@ -249,7 +513,7 @@ bool FStateTreeBaker::CreateCondition(const FStateTreeEditorNode& CondNode, cons
 		
 		// Create binding source struct descriptor.
 		StructDesc.Struct = Instance->GetClass();
-		StructDesc.Name = StructDesc.Struct->GetFName();
+		StructDesc.Name = Cond.Name;
 
 		check(InstanceIndex <= int32(MAX_uint16));
 		Cond.InstanceIndex = uint16(InstanceIndex);
@@ -261,7 +525,7 @@ bool FStateTreeBaker::CreateCondition(const FStateTreeEditorNode& CondNode, cons
 
 	// Check that the bindings for this struct are still all valid.
 	TArray<FStateTreeEditorPropertyBinding> Bindings;
-	if (!GetAndValidateBindings(StructDesc, Bindings))
+	if (!GetAndValidateBindings(State, StructDesc, Bindings))
 	{
 		return false;
 	}
@@ -281,7 +545,7 @@ bool FStateTreeBaker::CreateCondition(const FStateTreeEditorNode& CondNode, cons
 	return true;
 }
 
-bool FStateTreeBaker::CreateTask(const FStateTreeEditorNode& TaskNode)
+bool FStateTreeBaker::CreateTask(UStateTreeState& State, const FStateTreeEditorNode& TaskNode)
 {
 	// Silently ignore empty items.
 	if (!TaskNode.Node.IsValid())
@@ -298,7 +562,7 @@ bool FStateTreeBaker::CreateTask(const FStateTreeEditorNode& TaskNode)
 	if (!TaskNode.Instance.IsValid() && TaskNode.InstanceObject == nullptr)
 	{
 		Log.Reportf(EMessageSeverity::Error, StructDesc,
-			TEXT("Malformed evaluator, missing instance value."));
+			TEXT("Malformed task, missing instance value."));
 		return false;
 	}
 
@@ -348,7 +612,7 @@ bool FStateTreeBaker::CreateTask(const FStateTreeEditorNode& TaskNode)
 
 	// Check that the bindings for this struct are still all valid.
 	TArray<FStateTreeEditorPropertyBinding> Bindings;
-	if (!GetAndValidateBindings(StructDesc, Bindings))
+	if (!GetAndValidateBindings(State, StructDesc, Bindings))
 	{
 		return false;
 	}
@@ -369,7 +633,7 @@ bool FStateTreeBaker::CreateTask(const FStateTreeEditorNode& TaskNode)
 	return true;
 }
 
-bool FStateTreeBaker::CreateEvaluator(const FStateTreeEditorNode& EvalNode)
+bool FStateTreeBaker::CreateEvaluator(UStateTreeState& State, const FStateTreeEditorNode& EvalNode)
 {
 	// Silently ignore empty items.
 	if (!EvalNode.Node.IsValid())
@@ -436,7 +700,7 @@ bool FStateTreeBaker::CreateEvaluator(const FStateTreeEditorNode& EvalNode)
 
 	// Check that the bindings for this struct are still all valid.
 	TArray<FStateTreeEditorPropertyBinding> Bindings;
-	if (!GetAndValidateBindings(StructDesc, Bindings))
+	if (!GetAndValidateBindings(State, StructDesc, Bindings))
 	{
 		return false;
 	}
@@ -477,7 +741,7 @@ bool FStateTreeBaker::IsPropertyAnyEnum(const FStateTreeBindableStructDesc& Stru
 	return bIsAnyEnum;
 }
 
-bool FStateTreeBaker::GetAndValidateBindings(const FStateTreeBindableStructDesc& TargetStruct, TArray<FStateTreeEditorPropertyBinding>& OutBindings) const
+bool FStateTreeBaker::GetAndValidateBindings(UStateTreeState& State, const FStateTreeBindableStructDesc& TargetStruct, TArray<FStateTreeEditorPropertyBinding>& OutBindings) const
 {
 	OutBindings.Reset();
 	
@@ -494,27 +758,44 @@ bool FStateTreeBaker::GetAndValidateBindings(const FStateTreeBindableStructDesc&
 		if (SourceStructIdx == INDEX_NONE)
 		{
 			Log.Reportf(EMessageSeverity::Error, TargetStruct,
-						TEXT("Failed to find binding source %s:%s."),
-						*TargetStruct.Name.ToString(), *Binding.SourcePath.ToString());
+						TEXT("Failed to find binding source '%s:%s'."),
+						*TargetStruct.Name.ToString(), *Binding.TargetPath.ToString());
 			return false;
 		}
 		const FStateTreeBindableStructDesc& SourceStruct = BindingsCompiler.GetSourceStructDesc(SourceStructIdx);
 
-		// Source must be accessible by the target struct.
+		// Source must be accessible by the target struct via all execution paths.
 		TArray<FStateTreeBindableStructDesc> AccessibleStructs;
-		TreeData->GetAccessibleStructs(Binding.TargetPath.StructID, AccessibleStructs);
-		const bool SourceAccessible = AccessibleStructs.ContainsByPredicate([SourceStructID](const FStateTreeBindableStructDesc& Structs)
-			{
-				return (Structs.ID == SourceStructID);
-			});
-		if (!SourceAccessible)
+		const FStateExecutionInfo& ExecInfo = ExecutionInfos.FindChecked(&State);
+		for (const FExecutionPath& ExecPath : ExecInfo.ExecutionPaths)
 		{
-			Log.Reportf(EMessageSeverity::Error, TargetStruct,
-        				TEXT("%s:%s is not accessible to %s:%s."),
-        				*SourceStruct.Name.ToString(), *Binding.SourcePath.ToString(), *TargetStruct.Name.ToString(), *Binding.TargetPath.ToString());
-			return false;
-		}
+			AccessibleStructs.Reset();
+			TreeData->GetAccessibleStructs(ExecPath.Path, Binding.TargetPath.StructID, AccessibleStructs);
 
+			const bool SourceAccessible = AccessibleStructs.ContainsByPredicate([SourceStructID](const FStateTreeBindableStructDesc& Structs)
+				{
+					return (Structs.ID == SourceStructID);
+				});
+
+			if (!SourceAccessible)
+			{
+				Log.Reportf(EMessageSeverity::Error, TargetStruct,
+							TEXT("Property '%s:%s' cannot be bound to '%s:%s', because the binding source '%s' is not updated before '%s' in the tree."),
+							*SourceStruct.Name.ToString(), *Binding.SourcePath.ToString(),
+							*TargetStruct.Name.ToString(), *Binding.TargetPath.ToString(),
+							*SourceStruct.Name.ToString(), *TargetStruct.Name.ToString());
+
+				if (IsPathLinked(ExecPath.Path))
+				{
+					Log.Reportf(EMessageSeverity::Error, TargetStruct,
+								TEXT("The binding source is not updated when executing via linked state: %s."),
+								*GetExecutionPathString(ExecPath.Path));
+				}
+				
+				return false;
+			}
+		}
+		
 		// Special case fo AnyEnum. StateTreeBindingExtension allows AnyEnums to bind to other enum types.
 		// The actual copy will be done via potential type promotion copy, into the value property inside the AnyEnum.
 		// We amend the paths here to point to the 'Value' property.
@@ -541,73 +822,3 @@ bool FStateTreeBaker::GetAndValidateBindings(const FStateTreeBindableStructDesc&
 
 	return true;
 }
-
-bool FStateTreeBaker::CreateStateRecursive(UStateTreeState& State, const FStateTreeHandle Parent)
-{
-	FStateTreeCompilerLogStateScope LogStateScope(&State, Log);
-
-	const int32 StateIdx = StateTree->States.AddDefaulted();
-	FBakedStateTreeState& BakedState = StateTree->States[StateIdx];
-	BakedState.Name = State.Name;
-	BakedState.Parent = Parent;
-
-	SourceStates.Add(&State);
-	IDToState.Add(State.ID, StateIdx);
-
-	check(StateTree->Instances.Num() <= int32(MAX_uint16));
-
-	// Collect evaluators
-	check(StateTree->Nodes.Num() <= int32(MAX_uint16));
-	BakedState.EvaluatorsBegin = uint16(StateTree->Nodes.Num());
-
-	for (FStateTreeEditorNode& EvalNode : State.Evaluators)
-	{
-		if (!CreateEvaluator(EvalNode))
-		{
-			return false;
-		}
-	}
-	
-	const int32 EvaluatorsNum = StateTree->Nodes.Num() - int32(BakedState.EvaluatorsBegin);
-	check(EvaluatorsNum <= int32(MAX_uint8));
-	BakedState.EvaluatorsNum = uint8(EvaluatorsNum);
-
-	// Collect tasks
-	check(StateTree->Nodes.Num() <= int32(MAX_uint16));
-	BakedState.TasksBegin = uint16(StateTree->Nodes.Num());
-
-	for (FStateTreeEditorNode& TaskNode : State.Tasks)
-	{
-		if (!CreateTask(TaskNode))
-		{
-			return false;
-		}
-	}
-
-	if (!CreateTask(State.SingleTask))
-	{
-		return false;
-	}
-	
-	const int32 TasksNum = StateTree->Nodes.Num() - int32(BakedState.TasksBegin);
-	check(TasksNum <= int32(MAX_uint8));
-	BakedState.TasksNum = uint8(TasksNum);
-
-	// Child states
-	check(StateTree->States.Num() <= int32(MAX_uint16));
-	BakedState.ChildrenBegin = uint16(StateTree->States.Num());
-	for (UStateTreeState* Child : State.Children)
-	{
-		if (Child)
-		{
-			if (!CreateStateRecursive(*Child, FStateTreeHandle((uint16)StateIdx)))
-			{
-				return false;
-			}
-		}
-	}
-	check(StateTree->States.Num() <= int32(MAX_uint16));
-	StateTree->States[StateIdx].ChildrenEnd = uint16(StateTree->States.Num()); // Cannot use BakedState here, it may be invalid due to array resize.
-
-	return true;
-}
diff --git a/Engine/Plugins/Runtime/StateTree/Source/StateTreeEditorModule/Private/StateTreeEditor.cpp b/Engine/Plugins/Runtime/StateTree/Source/StateTreeEditorModule/Private/StateTreeEditor.cpp
--- a/Engine/Plugins/Runtime/StateTree/Source/StateTreeEditorModule/Private/StateTreeEditor.cpp
+++ b/Engine/Plugins/Runtime/StateTree/Source/StateTreeEditorModule/Private/StateTreeEditor.cpp
@@ -106,6 +106,7 @@ void FStateTreeEditor::InitEditor( const EToolkitMode::Type Mode, const TSharedP
 	if (EditorData == NULL)
 	{
 		EditorData = NewObject<UStateTreeEditorData>(StateTree, FName(), RF_Transactional);
+		EditorData->AddRootState();
 		StateTree->EditorData = EditorData;
 	}
 
@@ -493,6 +494,11 @@ namespace UE::StateTree::Editor::Internal
 				{
 					UStateTreeState* CurState = Stack.Pop();
 
+					if (CurState->Type == EStateTreeStateType::Linked)
+					{
+						FixChangedStateLinkName(CurState->LinkedState, IDToName);
+					}
+					
 					for (FStateTreeTransition& Transition : CurState->Transitions)
 					{
 						FixChangedStateLinkName(Transition.State, IDToName);
diff --git a/Engine/Plugins/Runtime/StateTree/Source/StateTreeEditorModule/Private/StateTreeEditorData.cpp b/Engine/Plugins/Runtime/StateTree/Source/StateTreeEditorModule/Private/StateTreeEditorData.cpp
--- a/Engine/Plugins/Runtime/StateTree/Source/StateTreeEditorModule/Private/StateTreeEditorData.cpp
+++ b/Engine/Plugins/Runtime/StateTree/Source/StateTreeEditorModule/Private/StateTreeEditorData.cpp
@@ -10,77 +10,78 @@
 void UStateTreeEditorData::PostInitProperties()
 {
 	Super::PostInitProperties();
-	
-	if (!HasAnyFlags(RF_ClassDefaultObject) && !(GetOuter() && GetOuter()->HasAnyFlags(RF_ClassDefaultObject|RF_ArchetypeObject)))
-	{
-		// Add root node by default.
-		AddRootState();
-	}
 }
 
 void UStateTreeEditorData::GetAccessibleStructs(const FGuid TargetStructID, TArray<FStateTreeBindableStructDesc>& OutStructDescs) const
 {
 	// Find the states that are updated before the current state.
-	TSet<const UStateTreeState*> ValidStates;
+	TArray<const UStateTreeState*> Path;
 	const UStateTreeState* CurrentState = GetStateByStructID(TargetStructID);
 	for (const UStateTreeState* State = CurrentState; State != nullptr; State = State->Parent)
 	{
-		ValidStates.Add(State);
+		Path.Insert(State, 0);
 	}
+	
+	GetAccessibleStructs(Path, TargetStructID, OutStructDescs);
+}
 
+void UStateTreeEditorData::GetAccessibleStructs(const TConstArrayView<const UStateTreeState*> Path, const FGuid TargetStructID, TArray<FStateTreeBindableStructDesc>& OutStructDescs) const
+{
 	TArray<FStateTreeBindableStructDesc> EvalDescs;
 	TArray<FStateTreeBindableStructDesc> TaskDescs;
 
-	if (ValidStates.Num() > 0)
+	for (const UStateTreeState* State : Path)
 	{
-		VisitHierarchy([&ValidStates, &OutStructDescs, &EvalDescs, &TaskDescs, TargetStructID, CurrentState]
+		if (State == nullptr)
+		{
+			continue;
+		}
+		
+		VisitState(*State, [&OutStructDescs, &EvalDescs, &TaskDescs, TargetStructID]
 			(const UStateTreeState& State, const FGuid& ID, const FName& Name, const EStateTreeNodeType NodeType, const UScriptStruct* NodeStruct, const UStruct* InstanceStruct) -> bool
 			{
-				if (ValidStates.Contains(&State))
+				if (ID == TargetStructID)
 				{
-					if (ID == TargetStructID)
+					if (NodeType == EStateTreeNodeType::EnterCondition)
 					{
-						if (NodeType == EStateTreeNodeType::EnterCondition)
-						{
-							// Enter conditions can only see evaluators
-							OutStructDescs.Append(EvalDescs);
-						}
-						else if (NodeType == EStateTreeNodeType::Evaluator)
-						{
-							// Evaluators can only see other evaluators
-							OutStructDescs.Append(EvalDescs);
-						}
-						else if (NodeType == EStateTreeNodeType::Task)
-						{
-							// Tasks can see evals and tasks.
-							OutStructDescs.Append(EvalDescs);
-							OutStructDescs.Append(TaskDescs);
-						}
-						else if (NodeType == EStateTreeNodeType::TransitionCondition)
-						{
-							// Transitions can see evals and tasks.
-							OutStructDescs.Append(EvalDescs);
-							OutStructDescs.Append(TaskDescs);
-						}
-						
-						return false; // Stop visit
+						// Enter conditions can only see evaluators
+						OutStructDescs.Append(EvalDescs);
 					}
-
-					// Not at target yet, collect all evaluators and tasks visible so far.
-					if (NodeStruct->IsChildOf(FStateTreeEvaluatorBase::StaticStruct()))
+					else if (NodeType == EStateTreeNodeType::Evaluator)
 					{
-						FStateTreeBindableStructDesc& Desc = EvalDescs.AddDefaulted_GetRef();
-						Desc.Struct = InstanceStruct;
-						Desc.Name = Name;
-						Desc.ID = ID;
+						// Evaluators can only see other evaluators
+						OutStructDescs.Append(EvalDescs);
 					}
-					else if (NodeStruct->IsChildOf(FStateTreeTaskBase::StaticStruct()))
+					else if (NodeType == EStateTreeNodeType::Task)
 					{
-						FStateTreeBindableStructDesc& Desc = TaskDescs.AddDefaulted_GetRef();
-						Desc.Struct = InstanceStruct;
-						Desc.Name = Name;
-						Desc.ID = ID;
+						// Tasks can see evals and tasks.
+						OutStructDescs.Append(EvalDescs);
+						OutStructDescs.Append(TaskDescs);
 					}
+					else if (NodeType == EStateTreeNodeType::TransitionCondition)
+					{
+						// Transitions can see evals and tasks.
+						OutStructDescs.Append(EvalDescs);
+						OutStructDescs.Append(TaskDescs);
+					}
+					
+					return false; // Stop visit
+				}
+
+				// Not at target yet, collect all evaluators and tasks visible so far.
+				if (NodeStruct->IsChildOf(FStateTreeEvaluatorBase::StaticStruct()))
+				{
+					FStateTreeBindableStructDesc& Desc = EvalDescs.AddDefaulted_GetRef();
+					Desc.Struct = InstanceStruct;
+					Desc.Name = Name;
+					Desc.ID = ID;
+				}
+				else if (NodeStruct->IsChildOf(FStateTreeTaskBase::StaticStruct()))
+				{
+					FStateTreeBindableStructDesc& Desc = TaskDescs.AddDefaulted_GetRef();
+					Desc.Struct = InstanceStruct;
+					Desc.Name = Name;
+					Desc.ID = ID;
 				}
 				return true; // Continue
 			});
@@ -135,99 +136,108 @@ void UStateTreeEditorData::GetAllStructIDs(TMap<FGuid, const UStruct*>& AllStruc
 		});
 }
 
-void UStateTreeEditorData::VisitHierarchy(TFunctionRef<bool(const UStateTreeState& State, const FGuid& ID, const FName& Name, const EStateTreeNodeType NodeType, const UScriptStruct* NodeStruct, const UStruct* InstanceStruct)> InFunc) const
+bool UStateTreeEditorData::VisitState(const UStateTreeState& State, TFunctionRef<bool(const UStateTreeState& State, const FGuid& ID, const FName& Name, const EStateTreeNodeType NodeType, const UScriptStruct* NodeStruct, const UStruct* InstanceStruct)> InFunc) const
 {
-	TArray<const UStateTreeState*> Stack;
 	bool bContinue = true;
-
-	for (const UStateTreeState* SubTree : SubTrees)
+	
+	// Evaluators
+	for (const FStateTreeEditorNode& Node : State.Evaluators)
 	{
-		if (!SubTree)
+		if (const FStateTreeEvaluatorBase* Evaluator = Node.Node.GetPtr<FStateTreeEvaluatorBase>())
 		{
-			continue;
+			if (!InFunc(State, Node.ID, Evaluator->Name, EStateTreeNodeType::Evaluator, Node.Node.GetScriptStruct(), Evaluator->GetInstanceDataType()))
+			{
+				bContinue = false;
+				break;
+			}
 		}
-
-		Stack.Add(SubTree);
-
-		while (!Stack.IsEmpty() && bContinue)
+	}
+	if (bContinue)
+	{
+		// Enter conditions
+		for (const FStateTreeEditorNode& Node : State.EnterConditions)
 		{
-			const UStateTreeState* State = Stack[0];
-			check(State);
-
-			Stack.RemoveAt(0);
-
-			// Evaluators
-			for (const FStateTreeEditorNode& Node : State->Evaluators)
+			if (const FStateTreeConditionBase* Cond = Node.Node.GetPtr<FStateTreeConditionBase>())
 			{
-				if (const FStateTreeEvaluatorBase* Evaluator = Node.Node.GetPtr<FStateTreeEvaluatorBase>())
+				if (!InFunc(State, Node.ID, Node.Node.GetScriptStruct()->GetFName(), EStateTreeNodeType::EnterCondition, Node.Node.GetScriptStruct(), Cond->GetInstanceDataType()))
 				{
-					if (!InFunc(*State, Node.ID, Evaluator->Name, EStateTreeNodeType::Evaluator, Node.Node.GetScriptStruct(), Evaluator->GetInstanceDataType()))
-					{
-						bContinue = false;
-						break;
-					}
+					bContinue = false;
+					break;
 				}
 			}
-			if (bContinue)
+		}
+	}
+	if (bContinue)
+	{
+		// Tasks
+		for (const FStateTreeEditorNode& Node : State.Tasks)
+		{
+			if (const FStateTreeTaskBase* Task = Node.Node.GetPtr<FStateTreeTaskBase>())
 			{
-				// Enter conditions
-				for (const FStateTreeEditorNode& Node : State->EnterConditions)
+				if (!InFunc(State, Node.ID, Task->Name, EStateTreeNodeType::Task, Node.Node.GetScriptStruct(), Task->GetInstanceDataType()))
 				{
-					if (const FStateTreeConditionBase* Cond = Node.Node.GetPtr<FStateTreeConditionBase>())
-					{
-						if (!InFunc(*State, Node.ID, Node.Node.GetScriptStruct()->GetFName(), EStateTreeNodeType::EnterCondition, Node.Node.GetScriptStruct(), Cond->GetInstanceDataType()))
-						{
-							bContinue = false;
-							break;
-						}
-					}
+					bContinue = false;
+					break;
 				}
 			}
-			if (bContinue)
+		}
+	}
+	if (bContinue)
+	{
+		if (const FStateTreeTaskBase* Task = State.SingleTask.Node.GetPtr<FStateTreeTaskBase>())
+		{
+			if (!InFunc(State, State.SingleTask.ID, Task->Name, EStateTreeNodeType::Task, State.SingleTask.Node.GetScriptStruct(), Task->GetInstanceDataType()))
 			{
-				// Tasks
-				for (const FStateTreeEditorNode& Node : State->Tasks)
-				{
-					if (const FStateTreeTaskBase* Task = Node.Node.GetPtr<FStateTreeTaskBase>())
-					{
-						if (!InFunc(*State, Node.ID, Task->Name, EStateTreeNodeType::Task, Node.Node.GetScriptStruct(), Task->GetInstanceDataType()))
-						{
-							bContinue = false;
-							break;
-						}
-					}
-				}
+				bContinue = false;
 			}
-			if (bContinue)
+		}
+
+	}
+	if (bContinue)
+	{
+		// Transitions
+		for (const FStateTreeTransition& Transition : State.Transitions)
+		{
+			for (const FStateTreeEditorNode& Node : Transition.Conditions)
 			{
-				if (const FStateTreeTaskBase* Task = State->SingleTask.Node.GetPtr<FStateTreeTaskBase>())
+				if (const FStateTreeConditionBase* Cond = Node.Node.GetPtr<FStateTreeConditionBase>())
 				{
-					if (!InFunc(*State, State->SingleTask.ID, Task->Name, EStateTreeNodeType::Task, State->SingleTask.Node.GetScriptStruct(), Task->GetInstanceDataType()))
+					if (!InFunc(State, Node.ID, Node.Node.GetScriptStruct()->GetFName(), EStateTreeNodeType::TransitionCondition, Node.Node.GetScriptStruct(), Cond->GetInstanceDataType()))
 					{
 						bContinue = false;
 						break;
 					}
 				}
-
-			}
-			if (bContinue)
-			{
-				// Transitions
-				for (const FStateTreeTransition& Transition : State->Transitions)
-				{
-					for (const FStateTreeEditorNode& Node : Transition.Conditions)
-					{
-						if (const FStateTreeConditionBase* Cond = Node.Node.GetPtr<FStateTreeConditionBase>())
-						{
-							if (!InFunc(*State, Node.ID, Node.Node.GetScriptStruct()->GetFName(), EStateTreeNodeType::TransitionCondition, Node.Node.GetScriptStruct(), Cond->GetInstanceDataType()))
-							{
-								bContinue = false;
-								break;
-							}
-						}
-					}
-				}
 			}
+		}
+	}
+
+	return bContinue;
+}
+
+void UStateTreeEditorData::VisitHierarchy(TFunctionRef<bool(const UStateTreeState& State, const FGuid& ID, const FName& Name, const EStateTreeNodeType NodeType, const UScriptStruct* NodeStruct, const UStruct* InstanceStruct)> InFunc) const
+{
+	TArray<const UStateTreeState*> Stack;
+	bool bContinue = true;
+
+	for (const UStateTreeState* SubTree : SubTrees)
+	{
+		if (!SubTree)
+		{
+			continue;
+		}
+
+		Stack.Add(SubTree);
+
+		while (!Stack.IsEmpty() && bContinue)
+		{
+			const UStateTreeState* State = Stack[0];
+			check(State);
+
+			Stack.RemoveAt(0);
+
+			bContinue = VisitState(*State, InFunc);
+			
 			if (bContinue)
 			{
 				// Children
@@ -237,6 +247,7 @@ void UStateTreeEditorData::VisitHierarchy(TFunctionRef<bool(const UStateTreeStat
 				}
 			}
 		}
+		
 		if (!bContinue)
 		{
 			break;
diff --git a/Engine/Plugins/Runtime/StateTree/Source/StateTreeEditorModule/Private/StateTreeState.cpp b/Engine/Plugins/Runtime/StateTree/Source/StateTreeEditorModule/Private/StateTreeState.cpp
--- a/Engine/Plugins/Runtime/StateTree/Source/StateTreeEditorModule/Private/StateTreeState.cpp
+++ b/Engine/Plugins/Runtime/StateTree/Source/StateTreeEditorModule/Private/StateTreeState.cpp
@@ -58,8 +58,8 @@ void UStateTreeState::PostEditChangeChainProperty(FPropertyChangedChainEvent& Pr
 
 	if (Property)
 	{
-		if (Property->GetOwnerClass() == UStateTreeState::StaticClass() &&
-			Property->GetFName() == GET_MEMBER_NAME_CHECKED(UStateTreeState, Name))
+		if (Property->GetOwnerClass() == UStateTreeState::StaticClass()
+			&& Property->GetFName() == GET_MEMBER_NAME_CHECKED(UStateTreeState, Name))
 		{
 			UStateTree* StateTree = GetTypedOuter<UStateTree>();
 			if (ensure(StateTree))
@@ -68,6 +68,21 @@ void UStateTreeState::PostEditChangeChainProperty(FPropertyChangedChainEvent& Pr
 			}
 		}
 
+		if (Property->GetOwnerClass() == UStateTreeState::StaticClass()
+			&& Property->GetFName() == GET_MEMBER_NAME_CHECKED(UStateTreeState, Type))
+		{
+			if (Type == EStateTreeStateType::Group || Type == EStateTreeStateType::Linked)
+			{
+				// Remove any tasks and evaluators.
+				Tasks.Reset();
+				Evaluators.Reset();
+			}
+			if (Type != EStateTreeStateType::Linked)
+			{
+				LinkedState = FStateTreeStateLink();
+			}
+		}
+
 		if (MemberProperty)
 		{
 			// Ensure unique ID on duplicated items.
diff --git a/Engine/Plugins/Runtime/StateTree/Source/StateTreeEditorModule/Public/StateTreeBaker.h b/Engine/Plugins/Runtime/StateTree/Source/StateTreeEditorModule/Public/StateTreeBaker.h
--- a/Engine/Plugins/Runtime/StateTree/Source/StateTreeEditorModule/Public/StateTreeBaker.h
+++ b/Engine/Plugins/Runtime/StateTree/Source/StateTreeEditorModule/Public/StateTreeBaker.h
@@ -32,21 +32,44 @@ private:
 
 	bool ResolveTransitionState(const UStateTreeState& SourceState, const FStateTreeStateLink& Link, FStateTreeHandle& OutTransitionHandle) const;
 	FStateTreeHandle GetStateHandle(const FGuid& StateID) const;
+	UStateTreeState* GetState(const FGuid& StateID);
+	static FString GetExecutionPathString(const TConstArrayView<const UStateTreeState*> Path);
+	static bool IsPathLinked(const TConstArrayView<const UStateTreeState*> Path);
 
+	bool CreateExecutionInfos();
+	bool CreateExecutionInfosRecursive(UStateTreeState& State, TArray<const UStateTreeState*>& Path);
 	bool CreateStates();
 	bool CreateStateRecursive(UStateTreeState& State, const FStateTreeHandle Parent);
+	
+	bool CreateStateTasks();
+	bool CreateStateEvaluators();
 	bool CreateStateTransitions();
-	bool CreateConditions(TConstArrayView<FStateTreeEditorNode> Conditions);
-	bool CreateCondition(const FStateTreeEditorNode& CondNode, const EStateTreeConditionOperand Operand, const int8 DeltaIndent);
-	bool CreateTask(const FStateTreeEditorNode& TaskNode);
-	bool CreateEvaluator(const FStateTreeEditorNode& EvalNode);
-	bool GetAndValidateBindings(const FStateTreeBindableStructDesc& TargetStruct, TArray<FStateTreeEditorPropertyBinding>& OutBindings) const;
-	bool IsPropertyAnyEnum(const FStateTreeBindableStructDesc& Struct, FStateTreeEditorPropertyPath Path) const;
 	
+	bool CreateConditions(UStateTreeState& State, TConstArrayView<FStateTreeEditorNode> Conditions);
+	bool CreateCondition(UStateTreeState& State, const FStateTreeEditorNode& CondNode, const EStateTreeConditionOperand Operand, const int8 DeltaIndent);
+	bool CreateTask(UStateTreeState& State, const FStateTreeEditorNode& TaskNode);
+	bool CreateEvaluator(UStateTreeState& State, const FStateTreeEditorNode& EvalNode);
+	bool GetAndValidateBindings(UStateTreeState& State, const FStateTreeBindableStructDesc& TargetStruct, TArray<FStateTreeEditorPropertyBinding>& OutBindings) const;
+	bool IsPropertyAnyEnum(const FStateTreeBindableStructDesc& Struct, FStateTreeEditorPropertyPath Path) const;
+
+	struct FExecutionPath
+	{
+		FExecutionPath() = default;
+		FExecutionPath(const TConstArrayView<const UStateTreeState*> InPath) : Path(InPath) {}
+		
+		TArray<const UStateTreeState*> Path;
+	};
+
+	struct FStateExecutionInfo
+	{
+		TArray<FExecutionPath> ExecutionPaths;
+	};
+
 	FStateTreeCompilerLog& Log;
 	UStateTree* StateTree = nullptr;
 	UStateTreeEditorData* TreeData = nullptr;
 	TMap<FGuid, int32> IDToState;
 	TArray<UStateTreeState*> SourceStates;
+	TMap<UStateTreeState*, FStateExecutionInfo> ExecutionInfos;
 	FStateTreePropertyBindingCompiler BindingsCompiler;
 };
diff --git a/Engine/Plugins/Runtime/StateTree/Source/StateTreeEditorModule/Public/StateTreeEditorData.h b/Engine/Plugins/Runtime/StateTree/Source/StateTreeEditorModule/Public/StateTreeEditorData.h
--- a/Engine/Plugins/Runtime/StateTree/Source/StateTreeEditorModule/Public/StateTreeEditorData.h
+++ b/Engine/Plugins/Runtime/StateTree/Source/StateTreeEditorModule/Public/StateTreeEditorData.h
@@ -44,6 +44,20 @@ public:
 	 */
 	void VisitHierarchy(TFunctionRef<bool(const UStateTreeState& State, const FGuid& ID, const FName& Name, const EStateTreeNodeType NodeType, const UScriptStruct* NodeStruct, const UStruct* InstanceStruct)> InFunc) const;
 
+	/**
+	 * Iterates over all nodes in a given state.
+	 * @param InFunc function called at each node, should return true if visiting is continued or false to stop.
+	 */
+	bool VisitState(const UStateTreeState& State, TFunctionRef<bool(const UStateTreeState& State, const FGuid& ID, const FName& Name, const EStateTreeNodeType NodeType, const UScriptStruct* NodeStruct, const UStruct* InstanceStruct)> InFunc) const;
+
+	/**
+	 * Returns array of nodes along the execution path, up to the TargetStruct.
+	 * @param Path The states to visit during the check
+	 * @param TargetStructID The ID of the node where to stop.
+	 * @param OutStructDescs Array of nodes accessible on the given path.  
+	 */
+	void GetAccessibleStructs(const TConstArrayView<const UStateTreeState*> Path, const FGuid TargetStructID, TArray<FStateTreeBindableStructDesc>& OutStructDescs) const;
+
 	// StateTree Builder API
 
 	/**
diff --git a/Engine/Plugins/Runtime/StateTree/Source/StateTreeEditorModule/Public/StateTreeState.h b/Engine/Plugins/Runtime/StateTree/Source/StateTreeEditorModule/Public/StateTreeState.h
--- a/Engine/Plugins/Runtime/StateTree/Source/StateTreeEditorModule/Public/StateTreeState.h
+++ b/Engine/Plugins/Runtime/StateTree/Source/StateTreeEditorModule/Public/StateTreeState.h
@@ -22,6 +22,7 @@ struct STATETREEEDITORMODULE_API FStateTreeStateLink
 	FStateTreeStateLink(EStateTreeTransitionType InType) : Type(InType) {}
 
 	void Set(const EStateTreeTransitionType InType, const UStateTreeState* InState = nullptr);
+	void Set(const UStateTreeState* InState) { Set(EStateTreeTransitionType::GotoState, InState); }
 	
 	bool IsValid() const { return ID.IsValid(); }
 
@@ -127,6 +128,14 @@ struct STATETREEEDITORMODULE_API FStateTreeTransition
 	TArray<FStateTreeEditorNode> Conditions;
 };
 
+UENUM()
+enum class EStateTreeStateType : uint8
+{
+	State,	// A State containing tasks and evaluators.
+	Group,	// A State containing just sub states. 
+	Linked,	// A State that is linked to another state in the tree (the execution continues on the linked state).
+};
+
 
 /**
  * Editor representation of a state in StateTree
@@ -148,12 +157,13 @@ public:
 	// StateTree Builder API
 	
 	/** Adds child state with specified name. */
-	UStateTreeState& AddChildState(const FName ChildName)
+	UStateTreeState& AddChildState(const FName ChildName, const EStateTreeStateType StateType = EStateTreeStateType::State)
 	{
 		UStateTreeState* ChildState = NewObject<UStateTreeState>(this);
 		check(ChildState);
 		ChildState->Name = ChildName;
 		ChildState->Parent = this;
+		ChildState->Type = StateType;
 		Children.Add(ChildState);
 		return *ChildState;
 	}
@@ -226,6 +236,12 @@ public:
 	UPROPERTY(EditDefaultsOnly, Category = "State")
 	FName Name;
 
+	UPROPERTY(EditDefaultsOnly, Category = "State")
+	EStateTreeStateType Type;
+
+	UPROPERTY(EditDefaultsOnly, Category = "State", Meta=(DirectStatesOnly))
+	FStateTreeStateLink LinkedState;
+
 	UPROPERTY()
 	FGuid ID;
 
diff --git a/Engine/Plugins/Runtime/StateTree/Source/StateTreeModule/Private/Blueprint/StateTreeEvaluatorBlueprintBase.cpp b/Engine/Plugins/Runtime/StateTree/Source/StateTreeModule/Private/Blueprint/StateTreeEvaluatorBlueprintBase.cpp
--- a/Engine/Plugins/Runtime/StateTree/Source/StateTreeModule/Private/Blueprint/StateTreeEvaluatorBlueprintBase.cpp
+++ b/Engine/Plugins/Runtime/StateTree/Source/StateTreeModule/Private/Blueprint/StateTreeEvaluatorBlueprintBase.cpp
@@ -20,10 +20,10 @@ void UStateTreeEvaluatorBlueprintBase::ExitState(FStateTreeExecutionContext& Con
 	ReceiveExitState(OwnerActor, ChangeType, Transition);
 }
 
-void UStateTreeEvaluatorBlueprintBase::StateCompleted(FStateTreeExecutionContext& Context, const EStateTreeRunStatus CompletionStatus, const FStateTreeHandle CompletedState)
+void UStateTreeEvaluatorBlueprintBase::StateCompleted(FStateTreeExecutionContext& Context, const EStateTreeRunStatus CompletionStatus, const FStateTreeActiveStates& CompletedActiveStates)
 {
 	AActor* OwnerActor = GetOwnerActor(Context);
-	ReceiveStateCompleted(OwnerActor, CompletionStatus, CompletedState);
+	ReceiveStateCompleted(OwnerActor, CompletionStatus, CompletedActiveStates);
 }
 
 void UStateTreeEvaluatorBlueprintBase::Evaluate(FStateTreeExecutionContext& Context, const EStateTreeEvaluationType EvalType, const float DeltaTime)
@@ -65,13 +65,13 @@ void FStateTreeBlueprintEvaluatorWrapper::ExitState(FStateTreeExecutionContext&
 	Instance->ExitState(Context, ChangeType, Transition);
 }
 
-void FStateTreeBlueprintEvaluatorWrapper::StateCompleted(FStateTreeExecutionContext& Context, const EStateTreeRunStatus CompletionStatus, const FStateTreeHandle CompletedState) const
+void FStateTreeBlueprintEvaluatorWrapper::StateCompleted(FStateTreeExecutionContext& Context, const EStateTreeRunStatus CompletionStatus, const FStateTreeActiveStates& CompletedActiveStates) const
 {
 	UStateTreeEvaluatorBlueprintBase* Instance = Context.GetInstanceObjectInternal<UStateTreeEvaluatorBlueprintBase>(DataViewIndex);
 	check(Instance);
 	
 	Instance->CopyExternalData(Context, ExternalDataHandles);
-	Instance->StateCompleted(Context, CompletionStatus, CompletedState);
+	Instance->StateCompleted(Context, CompletionStatus, CompletedActiveStates);
 }
 
 void FStateTreeBlueprintEvaluatorWrapper::Evaluate(FStateTreeExecutionContext& Context, const EStateTreeEvaluationType EvalType, const float DeltaTime) const
diff --git a/Engine/Plugins/Runtime/StateTree/Source/StateTreeModule/Private/Blueprint/StateTreeTaskBlueprintBase.cpp b/Engine/Plugins/Runtime/StateTree/Source/StateTreeModule/Private/Blueprint/StateTreeTaskBlueprintBase.cpp
--- a/Engine/Plugins/Runtime/StateTree/Source/StateTreeModule/Private/Blueprint/StateTreeTaskBlueprintBase.cpp
+++ b/Engine/Plugins/Runtime/StateTree/Source/StateTreeModule/Private/Blueprint/StateTreeTaskBlueprintBase.cpp
@@ -25,10 +25,10 @@ void UStateTreeTaskBlueprintBase::ExitState(FStateTreeExecutionContext& Context,
 	ReceiveExitState(OwnerActor, ChangeType, Transition);
 }
 
-void UStateTreeTaskBlueprintBase::StateCompleted(FStateTreeExecutionContext& Context, const EStateTreeRunStatus CompletionStatus, const FStateTreeHandle CompletedState)
+void UStateTreeTaskBlueprintBase::StateCompleted(FStateTreeExecutionContext& Context, const EStateTreeRunStatus CompletionStatus, const FStateTreeActiveStates& CompletedActiveStates)
 {
 	AActor* OwnerActor = GetOwnerActor(Context);
-	ReceiveStateCompleted(OwnerActor, CompletionStatus, CompletedState);
+	ReceiveStateCompleted(OwnerActor, CompletionStatus, CompletedActiveStates);
 }
 
 EStateTreeRunStatus UStateTreeTaskBlueprintBase::Tick(FStateTreeExecutionContext& Context, const float DeltaTime)
@@ -70,13 +70,13 @@ void FStateTreeBlueprintTaskWrapper::ExitState(FStateTreeExecutionContext& Conte
 	Instance->ExitState(Context, ChangeType, Transition);
 }
 
-void FStateTreeBlueprintTaskWrapper::StateCompleted(FStateTreeExecutionContext& Context, const EStateTreeRunStatus CompletionStatus, const FStateTreeHandle CompletedState) const
+void FStateTreeBlueprintTaskWrapper::StateCompleted(FStateTreeExecutionContext& Context, const EStateTreeRunStatus CompletionStatus, const FStateTreeActiveStates& CompletedActiveStates) const
 {
 	UStateTreeTaskBlueprintBase* Instance = Context.GetInstanceObjectInternal<UStateTreeTaskBlueprintBase>(DataViewIndex);
 	check(Instance);
 	
 	Instance->CopyExternalData(Context, ExternalDataHandles);
-	Instance->StateCompleted(Context, CompletionStatus, CompletedState);
+	Instance->StateCompleted(Context, CompletionStatus, CompletedActiveStates);
 }
 
 EStateTreeRunStatus FStateTreeBlueprintTaskWrapper::Tick(FStateTreeExecutionContext& Context, const float DeltaTime) const
diff --git a/Engine/Plugins/Runtime/StateTree/Source/StateTreeModule/Private/Evaluators/CooldownGuardEvaluator.cpp b/Engine/Plugins/Runtime/StateTree/Source/StateTreeModule/Private/Evaluators/CooldownGuardEvaluator.cpp
--- a/Engine/Plugins/Runtime/StateTree/Source/StateTreeModule/Private/Evaluators/CooldownGuardEvaluator.cpp
+++ b/Engine/Plugins/Runtime/StateTree/Source/StateTreeModule/Private/Evaluators/CooldownGuardEvaluator.cpp
@@ -34,7 +34,7 @@ void FCooldownGuardEvaluator::EnterState(FStateTreeExecutionContext& Context, co
 	}
 }
  
-void FCooldownGuardEvaluator::StateCompleted(FStateTreeExecutionContext& Context, const EStateTreeRunStatus CompletionStatus, const FStateTreeHandle CompletedState) const
+void FCooldownGuardEvaluator::StateCompleted(FStateTreeExecutionContext& Context, const EStateTreeRunStatus CompletionStatus, const FStateTreeActiveStates& CompletedActiveStates) const
 {
 	if (CompletionStatus == EStateTreeRunStatus::Failed
 		&& (Activation == ECooldownGuardActivation::OnCompleted || Activation == ECooldownGuardActivation::OnFailed))
diff --git a/Engine/Plugins/Runtime/StateTree/Source/StateTreeModule/Private/Evaluators/SelectionGuardEvaluator.cpp b/Engine/Plugins/Runtime/StateTree/Source/StateTreeModule/Private/Evaluators/SelectionGuardEvaluator.cpp
--- a/Engine/Plugins/Runtime/StateTree/Source/StateTreeModule/Private/Evaluators/SelectionGuardEvaluator.cpp
+++ b/Engine/Plugins/Runtime/StateTree/Source/StateTreeModule/Private/Evaluators/SelectionGuardEvaluator.cpp
@@ -30,7 +30,7 @@ void FSelectionGuardEvaluator::ExitState(FStateTreeExecutionContext& Context, co
 	bActive = false;
 }
 
-void FSelectionGuardEvaluator::StateCompleted(FStateTreeExecutionContext& Context, const EStateTreeRunStatus CompletionStatus, const FStateTreeHandle CompletedState) const
+void FSelectionGuardEvaluator::StateCompleted(FStateTreeExecutionContext& Context, const EStateTreeRunStatus CompletionStatus, const FStateTreeActiveStates& CompletedActiveStates) const
 {
 	bool& bActive = Context.GetInstanceData(ActiveHandle);
 
diff --git a/Engine/Plugins/Runtime/StateTree/Source/StateTreeModule/Private/StateTreeExecutionContext.cpp b/Engine/Plugins/Runtime/StateTree/Source/StateTreeModule/Private/StateTreeExecutionContext.cpp
--- a/Engine/Plugins/Runtime/StateTree/Source/StateTreeModule/Private/StateTreeExecutionContext.cpp
+++ b/Engine/Plugins/Runtime/StateTree/Source/StateTreeModule/Private/StateTreeExecutionContext.cpp
@@ -13,8 +13,8 @@
 #include "Logging/LogScopedVerbosityOverride.h"
 #include "ProfilingDebugging/CsvProfiler.h"
 
-#define STATETREE_LOG(Verbosity, Format, ...) UE_VLOG_UELOG(GetOwner(), LogStateTree, Verbosity, TEXT("%s") Format, *GetInstanceDescription(), ##__VA_ARGS__)
-#define STATETREE_CLOG(Condition, Verbosity, Format, ...) UE_CVLOG_UELOG((Condition), GetOwner(), LogStateTree, Verbosity, TEXT("%s") Format, *GetInstanceDescription(), ##__VA_ARGS__)
+#define STATETREE_LOG(Verbosity, Format, ...) UE_VLOG_UELOG(GetOwner(), LogStateTree, Verbosity, TEXT("%s: ") Format, *GetInstanceDescription(), ##__VA_ARGS__)
+#define STATETREE_CLOG(Condition, Verbosity, Format, ...) UE_CVLOG_UELOG((Condition), GetOwner(), LogStateTree, Verbosity, TEXT("%s: ") Format, *GetInstanceDescription(), ##__VA_ARGS__)
 
 namespace UE::StateTree
 {
@@ -85,44 +85,54 @@ EStateTreeRunStatus FStateTreeExecutionContext::Start(FStateTreeInstanceData* Ex
 	// Stop if still running previous state.
 	if (Exec.TreeRunStatus == EStateTreeRunStatus::Running)
 	{
-		const FStateTreeTransitionResult Transition(FStateTreeStateStatus(Exec.CurrentState, Exec.LastTickStatus), FStateTreeHandle::Succeeded);
+		FStateTreeTransitionResult Transition;
+		Transition.TargetState = FStateTreeHandle::Succeeded;
+		Transition.CurrentActiveStates = Exec.ActiveStates;
+		Transition.CurrentRunStatus = Exec.LastTickStatus;
+		Transition.NextActiveStates = FStateTreeActiveStates(FStateTreeHandle::Succeeded);
+
 		ExitState(InstanceData, Transition);
 	}
 	
 	// Initialize to unset running state.
 	Exec.TreeRunStatus = EStateTreeRunStatus::Running;
-	Exec.CurrentState = FStateTreeHandle::Invalid;
+	Exec.ActiveStates.Reset();
 	Exec.LastTickStatus = EStateTreeRunStatus::Unset;
 
 	// Select new Starting from root.
 	VisitedStates.Init(false, StateTree->States.Num());
+	
 	static const FStateTreeHandle RootState = FStateTreeHandle(0);
-	const FStateTreeTransitionResult Transition = SelectState(InstanceData, FStateTreeStateStatus(FStateTreeHandle::Invalid), RootState, RootState, 0);
 
-	// Handle potential transition.
-	if (Transition.Next.IsValid())
+	FStateTreeActiveStates NextActiveStates;
+	if (SelectState(InstanceData, RootState, NextActiveStates))
 	{
-		if (Transition.Next == FStateTreeHandle::Succeeded || Transition.Next == FStateTreeHandle::Failed)
+		if (NextActiveStates.Last() == FStateTreeHandle::Succeeded || NextActiveStates.Last() == FStateTreeHandle::Failed)
 		{
 			// Transition to a terminal state (succeeded/failed), or default transition failed.
 			STATETREE_LOG(Warning, TEXT("%s: Tree %s at StateTree start on '%s' using StateTree '%s'."),
-				ANSI_TO_TCHAR(__FUNCTION__), Transition.Next == FStateTreeHandle::Succeeded ? TEXT("succeeded") : TEXT("failed"), *GetNameSafe(Owner), *GetNameSafe(StateTree));
-			Exec.TreeRunStatus = Transition.Next == FStateTreeHandle::Succeeded ? EStateTreeRunStatus::Succeeded : EStateTreeRunStatus::Failed;
+				ANSI_TO_TCHAR(__FUNCTION__), NextActiveStates.Last() == FStateTreeHandle::Succeeded ? TEXT("succeeded") : TEXT("failed"), *GetNameSafe(Owner), *GetNameSafe(StateTree));
+			Exec.TreeRunStatus = NextActiveStates.Last() == FStateTreeHandle::Succeeded ? EStateTreeRunStatus::Succeeded : EStateTreeRunStatus::Failed;
 		}
 		else
 		{
 			// Enter state tasks can fail/succeed, treat it same as tick.
+			FStateTreeTransitionResult Transition;
+			Transition.TargetState = RootState;
+			Transition.CurrentActiveStates = Exec.ActiveStates;
+			Transition.CurrentRunStatus = Exec.LastTickStatus;
+			Transition.NextActiveStates = NextActiveStates; // Enter state will update Exec.ActiveStates.
 			Exec.LastTickStatus = EnterState(InstanceData, Transition);
 			
 			// Report state completed immediately.
 			if (Exec.LastTickStatus != EStateTreeRunStatus::Running)
 			{
-				StateCompleted(InstanceData, Exec.CurrentState, Exec.LastTickStatus);
+				StateCompleted(InstanceData);
 			}
 		}
 	}
 
-	if (!Exec.CurrentState.IsValid())
+	if (Exec.ActiveStates.IsEmpty())
 	{
 		// Should not happen. This may happen if initial state could not be selected.
 		STATETREE_LOG(Error, TEXT("%s: Failed to select initial state on '%s' using StateTree '%s'. This should not happen, check that the StateTree logic can always select a state at start."),
@@ -147,15 +157,26 @@ EStateTreeRunStatus FStateTreeExecutionContext::Stop(FStateTreeInstanceData* Ext
 	FStateTreeExecutionState& Exec = GetExecState(InstanceData);
 
 	// Exit states if still in some valid state.
-	if (Exec.CurrentState.IsValid() && (Exec.CurrentState != FStateTreeHandle::Succeeded || Exec.CurrentState != FStateTreeHandle::Failed))
+	if (!Exec.ActiveStates.IsEmpty() && (Exec.ActiveStates.Last() != FStateTreeHandle::Succeeded || Exec.ActiveStates.Last() != FStateTreeHandle::Failed))
 	{
 		// Transition to Succeeded state.
-		const FStateTreeTransitionResult Transition(FStateTreeStateStatus(Exec.CurrentState, Exec.LastTickStatus), FStateTreeHandle::Succeeded);
+		FStateTreeTransitionResult Transition;
+		Transition.TargetState = FStateTreeHandle::Succeeded;
+		Transition.CurrentActiveStates = Exec.ActiveStates;
+		Transition.CurrentRunStatus = Exec.LastTickStatus;
+		Transition.NextActiveStates = FStateTreeActiveStates(FStateTreeHandle::Succeeded);
+
 		ExitState(InstanceData, Transition);
-		Exec.CurrentState = FStateTreeHandle::Succeeded;
+
+		Exec.TreeRunStatus = EStateTreeRunStatus::Succeeded;
 	}
+	else
+	{
+		Exec.TreeRunStatus = Exec.ActiveStates.Last() == FStateTreeHandle::Succeeded ? EStateTreeRunStatus::Succeeded : EStateTreeRunStatus::Failed;
+	}
+
+	Exec.ActiveStates.Reset();
 
-	Exec.TreeRunStatus = Exec.CurrentState == FStateTreeHandle::Succeeded ? EStateTreeRunStatus::Succeeded : EStateTreeRunStatus::Failed;
 	Exec.LastTickStatus = EStateTreeRunStatus::Unset;
 
 	return Exec.TreeRunStatus;
@@ -192,15 +213,15 @@ EStateTreeRunStatus FStateTreeExecutionContext::Tick(const float DeltaTime, FSta
 	{
 		// Tick evaluators on active states. Further evaluator may be ticked during selection as non-active states are visited.
 		// Ticked states are kept track of and evaluators are ticked just once per frame.
-		TickEvaluators(InstanceData, Exec.CurrentState, EStateTreeEvaluationType::Tick, DeltaTime);
+		TickEvaluators(InstanceData, Exec.ActiveStates, EStateTreeEvaluationType::Tick, DeltaTime);
 
 		// Tick tasks on active states.
-		Exec.LastTickStatus = TickTasks(InstanceData, Exec.CurrentState, DeltaTime);
+		Exec.LastTickStatus = TickTasks(InstanceData, Exec, DeltaTime);
 		
 		// Report state completed immediately.
 		if (Exec.LastTickStatus != EStateTreeRunStatus::Running)
 		{
-			StateCompleted(InstanceData, Exec.CurrentState, Exec.LastTickStatus);
+			StateCompleted(InstanceData);
 		}
 	}
 
@@ -210,17 +231,15 @@ EStateTreeRunStatus FStateTreeExecutionContext::Tick(const float DeltaTime, FSta
 	for (int32 Iter = 0; Iter < MaxIterations; Iter++)
 	{
 		// Trigger conditional transitions or state succeed/failed transitions. First tick transition is handled here too.
-		FStateTreeTransitionResult Transition = TriggerTransitions(InstanceData, FStateTreeStateStatus(Exec.CurrentState, Exec.LastTickStatus), 0);
-
-		// Handle potential transition.
-		if (Transition.Next.IsValid())
+		FStateTreeTransitionResult Transition;
+		if (TriggerTransitions(InstanceData, Transition))
 		{
 			ExitState(InstanceData, Transition);
 			
-			if (Transition.Next == FStateTreeHandle::Succeeded || Transition.Next == FStateTreeHandle::Failed)
+			if (Transition.NextActiveStates.Last() == FStateTreeHandle::Succeeded || Transition.NextActiveStates.Last() == FStateTreeHandle::Failed)
 			{
 				// Transition to a terminal state (succeeded/failed), or default transition failed.
-				Exec.TreeRunStatus = Transition.Next == FStateTreeHandle::Succeeded ? EStateTreeRunStatus::Succeeded : EStateTreeRunStatus::Failed;
+				Exec.TreeRunStatus = Transition.NextActiveStates.Last() == FStateTreeHandle::Succeeded ? EStateTreeRunStatus::Succeeded : EStateTreeRunStatus::Failed;
 				return Exec.TreeRunStatus;
 			}
 			
@@ -230,7 +249,7 @@ EStateTreeRunStatus FStateTreeExecutionContext::Tick(const float DeltaTime, FSta
 			// Report state completed immediately.
 			if (Exec.LastTickStatus != EStateTreeRunStatus::Running)
 			{
-				StateCompleted(InstanceData, Exec.CurrentState, Exec.LastTickStatus);
+				StateCompleted(InstanceData);
 			}
 		}
 		
@@ -244,7 +263,7 @@ EStateTreeRunStatus FStateTreeExecutionContext::Tick(const float DeltaTime, FSta
 		VisitedStates.Init(false, StateTree->States.Num());
 	}
 	
-	if (!Exec.CurrentState.IsValid())
+	if (Exec.ActiveStates.IsEmpty())
 	{
 		// Should not happen. This may happen if a state completion transition could not be selected. 
 		STATETREE_LOG(Error, TEXT("%s: Failed to select state on '%s' using StateTree '%s'. This should not happen, state completion transition is likely missing."),
@@ -260,23 +279,13 @@ EStateTreeRunStatus FStateTreeExecutionContext::EnterState(FStateTreeInstanceDat
 {
 	CSV_SCOPED_TIMING_STAT_EXCLUSIVE(StateTree_EnterState);
 
-	if (!Transition.Next.IsValid())
+	if (Transition.NextActiveStates.IsEmpty())
 	{
 		return EStateTreeRunStatus::Failed;
 	}
 
 	FStateTreeExecutionState& Exec = GetExecState(InstanceData);
-
-	// Activate evaluators along the active branch.
-	TStaticArray<FStateTreeHandle, 32> States;
-	const int32 NumStates = GetActiveStates(Transition.Next, States);
-	TStaticArray<FStateTreeHandle, 32> PrevStates;
-	const int32 NumPrevStates = GetActiveStates(Transition.Source.State, PrevStates);
-	// Pad prev states with invalid states if it's shorter.
-	for (int32 Index = NumPrevStates; Index < NumStates; Index++)
-	{
-		PrevStates[Index] = FStateTreeHandle::Invalid;
-	}
+	Exec.StateChangeCount++;
 
 	// On target branch means that the state is the target of current transition or child of it.
 	// States which were active before and will remain active, but are not on target branch will not get
@@ -287,15 +296,24 @@ EStateTreeRunStatus FStateTreeExecutionContext::EnterState(FStateTreeInstanceDat
 	
 	EStateTreeRunStatus Result = EStateTreeRunStatus::Running;
 
-	Exec.EnterStateFailedTaskIndex = INDEX_NONE;
-	Exec.CurrentState = FStateTreeHandle::Failed;
+	Exec.EnterStateFailedTaskIndex = MAX_uint16; // This will make all tasks to be accepted.
+	Exec.ActiveStates.Reset();
 	
-	for (int32 Index = 0; Index < NumStates && Result != EStateTreeRunStatus::Failed; Index++)
+	for (int32 Index = 0; Index < Transition.NextActiveStates.Num() && Result != EStateTreeRunStatus::Failed; Index++)
 	{
-		const FStateTreeHandle CurrentHandle = States[Index];
+		const FStateTreeHandle CurrentHandle = Transition.NextActiveStates[Index];
+		const FStateTreeHandle PreviousHandle = Transition.CurrentActiveStates.GetStateSafe(Index);
 		const FBakedStateTreeState& State = StateTree->States[CurrentHandle.Index];
-		bOnTargetBranch = bOnTargetBranch || CurrentHandle == Transition.Target;
-		const bool bWasActive = PrevStates[Index] == CurrentHandle;
+
+		if (!Exec.ActiveStates.Push(CurrentHandle))
+		{
+			STATETREE_LOG(Error, TEXT("%s: Reached max execution depth when trying to enter state '%s'.  '%s' using StateTree '%s'."),
+				ANSI_TO_TCHAR(__FUNCTION__), *GetStateStatusString(Exec), *GetNameSafe(Owner), *GetNameSafe(StateTree));
+			break;
+		}
+		
+		bOnTargetBranch = bOnTargetBranch || CurrentHandle == Transition.TargetState;
+		const bool bWasActive = PreviousHandle == CurrentHandle;
 		if (bWasActive && !bOnTargetBranch)
 		{
 			// States which will keep on being active and were not part of the transition will not get enter/exit state.
@@ -313,11 +331,11 @@ EStateTreeRunStatus FStateTreeExecutionContext::EnterState(FStateTreeInstanceDat
 			continue;
 		}
 
-		CurrentTransition.Current = CurrentHandle;
+		CurrentTransition.CurrentState = CurrentHandle;
 		
 		const EStateTreeStateChangeType ChangeType = bWasActive ? EStateTreeStateChangeType::Sustained : EStateTreeStateChangeType::Changed;
 
-		STATETREE_LOG(Log, TEXT("%*sEnter state '%s' %s"), Index*UE::StateTree::DebugIndentSize, TEXT(""), *DebugGetStatePath(States, Index), *StaticEnum<EStateTreeStateChangeType>()->GetValueAsString(ChangeType));
+		STATETREE_LOG(Log, TEXT("%*sEnter state '%s' %s"), Index*UE::StateTree::DebugIndentSize, TEXT(""), *DebugGetStatePath(Transition.NextActiveStates, Index), *UEnum::GetValueAsString(ChangeType));
 		
 		for (int32 EvalIndex = State.EvaluatorsBegin; EvalIndex < (int32)(State.EvaluatorsBegin + State.EvaluatorsNum); EvalIndex++)
 		{
@@ -362,7 +380,6 @@ EStateTreeRunStatus FStateTreeExecutionContext::EnterState(FStateTreeInstanceDat
 				{
 					// Store how far in the enter state we got. This will be used to match the ExitState() calls.
 					Exec.EnterStateFailedTaskIndex = TaskIndex;
-					Exec.CurrentState = CurrentHandle;
 					Result = Status;
 					break;
 				}
@@ -370,11 +387,6 @@ EStateTreeRunStatus FStateTreeExecutionContext::EnterState(FStateTreeInstanceDat
 		}
 	}
 
-	if (Result != EStateTreeRunStatus::Failed)
-	{
-		Exec.CurrentState = Transition.Next;
-	}
-
 	return Result;
 }
 
@@ -382,7 +394,7 @@ void FStateTreeExecutionContext::ExitState(FStateTreeInstanceData& InstanceData,
 {
 	CSV_SCOPED_TIMING_STAT_EXCLUSIVE(StateTree_ExitState);
 
-	if (!Transition.Source.State.IsValid())
+	if (Transition.CurrentActiveStates.IsEmpty())
 	{
 		return;
 	}
@@ -392,17 +404,6 @@ void FStateTreeExecutionContext::ExitState(FStateTreeInstanceData& InstanceData,
 	Exec.GatedTransitionIndex = INDEX_NONE;
 	Exec.GatedTransitionTime = 0.0f;
 
-	// Deactivate evaluators along the active branch.
-	TStaticArray<FStateTreeHandle, 32> States;
-	const int32 NumStates = GetActiveStates(Transition.Source.State, States);
-	TStaticArray<FStateTreeHandle, 32> NextStates;
-	const int32 NumNextStates = GetActiveStates(Transition.Next, NextStates);
-	// Pad next states with invalid states if it's shorter.
-	for (int32 Index = NumNextStates; Index < NumStates; Index++)
-	{
-		NextStates[Index] = FStateTreeHandle::Invalid;
-	}
-
 	// On target branch means that the state is the target of current transition or child of it.
 	// States which were active before and will remain active, but are not on target branch will not get
 	// EnterState called. That is, a transition is handled as "replan from this state".
@@ -412,13 +413,14 @@ void FStateTreeExecutionContext::ExitState(FStateTreeInstanceData& InstanceData,
 
 	// It would be more symmetrical, if the evals, tasks and states were executed in reverse order, but we need
 	// to do it like this because of the property copies.
-	for (int32 Index = 0; Index < NumStates; Index++)
+	for (int32 Index = 0; Index < Transition.CurrentActiveStates.Num(); Index++)
 	{
-		const FStateTreeHandle CurrentHandle = States[Index];
-		const FStateTreeHandle NextHandle = NextStates[Index];
+		const FStateTreeHandle CurrentHandle = Transition.CurrentActiveStates[Index];
+		const FStateTreeHandle NextHandle = Transition.NextActiveStates.GetStateSafe(Index);
 		const FBakedStateTreeState& State = StateTree->States[CurrentHandle.Index];
-		const bool bRemainsActive = NextHandle == States[Index];
-		bOnTargetBranch = bOnTargetBranch || NextHandle == Transition.Target;
+
+		const bool bRemainsActive = NextHandle == CurrentHandle;
+		bOnTargetBranch = bOnTargetBranch || NextHandle == Transition.TargetState;
 		if (bRemainsActive && !bOnTargetBranch)
 		{
 			// States which will keep on being active and were not part of the transition will not get enter/exit state.
@@ -436,11 +438,11 @@ void FStateTreeExecutionContext::ExitState(FStateTreeInstanceData& InstanceData,
 			continue;
 		}
 
-		CurrentTransition.Current = CurrentHandle;
+		CurrentTransition.CurrentState = CurrentHandle;
 		
 		const EStateTreeStateChangeType ChangeType = bRemainsActive ? EStateTreeStateChangeType::Sustained : EStateTreeStateChangeType::Changed;
 
-		STATETREE_LOG(Log, TEXT("%*sExit state '%s' %s"), Index*UE::StateTree::DebugIndentSize, TEXT(""), *DebugGetStatePath(States, Index), *StaticEnum<EStateTreeStateChangeType>()->GetValueAsString(ChangeType));
+		STATETREE_LOG(Log, TEXT("%*sExit state '%s' %s"), Index*UE::StateTree::DebugIndentSize, TEXT(""), *DebugGetStatePath(Transition.CurrentActiveStates, Index), *UEnum::GetValueAsString(ChangeType));
 
 		for (int32 EvalIndex = State.EvaluatorsBegin; EvalIndex < (int32)(State.EvaluatorsBegin + State.EvaluatorsNum); EvalIndex++)
 		{
@@ -489,28 +491,25 @@ void FStateTreeExecutionContext::ExitState(FStateTreeInstanceData& InstanceData,
 	}
 }
 
-void FStateTreeExecutionContext::StateCompleted(FStateTreeInstanceData& InstanceData, const FStateTreeHandle CurrentState, const EStateTreeRunStatus CompletionStatus)
+void FStateTreeExecutionContext::StateCompleted(FStateTreeInstanceData& InstanceData)
 {
 	CSV_SCOPED_TIMING_STAT_EXCLUSIVE(StateTree_StateCompleted);
 
-	if (!CurrentState.IsValid())
+	const FStateTreeExecutionState& Exec = GetExecState(InstanceData);
+
+	if (Exec.ActiveStates.IsEmpty())
 	{
 		return;
 	}
 
-	TStaticArray<FStateTreeHandle, 32> States;
-	const int32 NumStates = GetActiveStates(CurrentState, States);
-
-	FStateTreeExecutionState& Exec = GetExecState(InstanceData);
-
 	// Call from child towards root to allow to pass results back.
 	// Note: Completed is assumed to be called immediately after tick or enter state, so there's no property copying.
-	for (int32 Index = NumStates - 1; Index >= 0; Index--)
+	for (int32 Index = Exec.ActiveStates.Num() - 1; Index >= 0; Index--)
 	{
-		const FStateTreeHandle CurrentHandle = States[Index];
+		const FStateTreeHandle CurrentHandle = Exec.ActiveStates[Index];
 		const FBakedStateTreeState& State = StateTree->States[CurrentHandle.Index];
 
-		STATETREE_LOG(Verbose, TEXT("%*sState Completed '%s' %s"), Index*UE::StateTree::DebugIndentSize, TEXT(""), *DebugGetStatePath(States, Index), *StaticEnum<EStateTreeRunStatus>()->GetValueAsString(CompletionStatus));
+		STATETREE_LOG(Verbose, TEXT("%*sState Completed '%s' %s"), Index*UE::StateTree::DebugIndentSize, TEXT(""), *DebugGetStatePath(Exec.ActiveStates, Index), *UEnum::GetValueAsString(Exec.LastTickStatus));
 
 		// Notify Tasks
 		for (int32 TaskIndex = (int32)(State.TasksBegin + State.TasksNum) - 1; TaskIndex >= State.TasksBegin; TaskIndex--)
@@ -522,7 +521,7 @@ void FStateTreeExecutionContext::StateCompleted(FStateTreeInstanceData& Instance
 				const FStateTreeTaskBase& Task = GetNode<FStateTreeTaskBase>(TaskIndex);
 
 				STATETREE_LOG(Verbose, TEXT("%*s  Notify Task '%s'"), Index*UE::StateTree::DebugIndentSize, TEXT(""), *Task.Name.ToString());
-				Task.StateCompleted(*this, CompletionStatus, CurrentState);
+				Task.StateCompleted(*this, Exec.LastTickStatus, Exec.ActiveStates);
 			}
 		}
 
@@ -532,35 +531,49 @@ void FStateTreeExecutionContext::StateCompleted(FStateTreeInstanceData& Instance
 			const FStateTreeEvaluatorBase& Eval = GetNode<FStateTreeEvaluatorBase>(EvalIndex);
 
 			STATETREE_LOG(Verbose, TEXT("%*s  Notify Evaluator '%s'"), Index*UE::StateTree::DebugIndentSize, TEXT(""), *Eval.Name.ToString());
-			Eval.StateCompleted(*this, CompletionStatus, CurrentState);
+			Eval.StateCompleted(*this, Exec.LastTickStatus, Exec.ActiveStates);
 		}
 	}
 }
 
-void FStateTreeExecutionContext::TickEvaluators(FStateTreeInstanceData& InstanceData, const FStateTreeHandle CurrentState, const EStateTreeEvaluationType EvalType, const float DeltaTime)
+void FStateTreeExecutionContext::TickEvaluatorsForSelect(FStateTreeInstanceData& InstanceData, const FStateTreeHandle CurrentState, const EStateTreeEvaluationType EvalType, const float DeltaTime)
 {
-	CSV_SCOPED_TIMING_STAT_EXCLUSIVE(StateTree_TickEvaluators);
-
-	if (!CurrentState.IsValid())
+	if (CurrentState == FStateTreeHandle::Succeeded || CurrentState == FStateTreeHandle::Failed || CurrentState == FStateTreeHandle::Invalid)
 	{
 		return;
 	}
 
-	TStaticArray<FStateTreeHandle, 32> States;
-	const int32 NumStates = GetActiveStates(CurrentState, States);
+	FStateTreeActiveStates Branch;
+	FStateTreeHandle StateHandle = CurrentState;
+	while (StateHandle.IsValid())
+	{
+		if (!Branch.PushFront(StateHandle))
+		{
+			break;
+		}
+		StateHandle = StateTree->States[StateHandle.Index].Parent;
+	}
 
-	for (int32 Index = 0; Index < NumStates; Index++)
+	TickEvaluators(InstanceData, Branch, EvalType, DeltaTime);
+}
+
+void FStateTreeExecutionContext::TickEvaluators(FStateTreeInstanceData& InstanceData, const FStateTreeActiveStates& ActiveStates, const EStateTreeEvaluationType EvalType, const float DeltaTime)
+{
+	CSV_SCOPED_TIMING_STAT_EXCLUSIVE(StateTree_TickEvaluators);
+
+	for (int32 Index = 0; Index < ActiveStates.Num(); Index++)
 	{
-		const FStateTreeHandle CurrentHandle = States[Index];
-		if (VisitedStates[CurrentHandle.Index])
+		const FStateTreeHandle StateHandle = ActiveStates[Index];
+
+		if (VisitedStates[StateHandle.Index])
 		{
 			// Already ticked this frame.
 			continue;
 		}
 
-		const FBakedStateTreeState& State = StateTree->States[CurrentHandle.Index];
+		const FBakedStateTreeState& State = StateTree->States[StateHandle.Index];
 
-		STATETREE_CLOG(State.EvaluatorsNum > 0, Verbose, TEXT("%*sTicking Evaluators of state '%s' %s"), Index*UE::StateTree::DebugIndentSize, TEXT(""), *DebugGetStatePath(States, Index), *StaticEnum<EStateTreeEvaluationType>()->GetValueAsString(EvalType));
+		STATETREE_CLOG(State.EvaluatorsNum > 0, Verbose, TEXT("%*sTicking Evaluators of state '%s' %s"), Index*UE::StateTree::DebugIndentSize, TEXT(""), *DebugGetStatePath(ActiveStates, Index), *UEnum::GetValueAsString(EvalType));
 
 		// Tick evaluators
 		for (int32 EvalIndex = 0; EvalIndex < int32(State.EvaluatorsNum); EvalIndex++)
@@ -581,31 +594,28 @@ void FStateTreeExecutionContext::TickEvaluators(FStateTreeInstanceData& Instance
 			}
 		}
 
-		VisitedStates[CurrentHandle.Index] = true;
+		VisitedStates[StateHandle.Index] = true;
 	}
 }
 
-EStateTreeRunStatus FStateTreeExecutionContext::TickTasks(FStateTreeInstanceData& InstanceData, const FStateTreeHandle CurrentState, const float DeltaTime)
+EStateTreeRunStatus FStateTreeExecutionContext::TickTasks(FStateTreeInstanceData& InstanceData, const FStateTreeExecutionState& Exec, const float DeltaTime)
 {
 	CSV_SCOPED_TIMING_STAT_EXCLUSIVE(StateTree_TickTasks);
 
-	if (!CurrentState.IsValid())
+	if (Exec.ActiveStates.IsEmpty())
 	{
 		return EStateTreeRunStatus::Failed;
 	}
 
-	TStaticArray<FStateTreeHandle, 32> States;
-	const int32 NumStates = GetActiveStates(CurrentState, States);
-
 	EStateTreeRunStatus Result = EStateTreeRunStatus::Running;
 	int32 NumTotalTasks = 0;
 
-	for (int32 Index = 0; Index < NumStates && Result != EStateTreeRunStatus::Failed; Index++)
+	for (int32 Index = 0; Index < Exec.ActiveStates.Num() && Result != EStateTreeRunStatus::Failed; Index++)
 	{
-		const FStateTreeHandle CurrentHandle = States[Index];
+		const FStateTreeHandle CurrentHandle = Exec.ActiveStates[Index];
 		const FBakedStateTreeState& State = StateTree->States[CurrentHandle.Index];
 
-		STATETREE_CLOG(State.TasksNum > 0, Verbose, TEXT("%*sTicking Tasks of state '%s'"), Index*UE::StateTree::DebugIndentSize, TEXT(""), *DebugGetStatePath(States, Index));
+		STATETREE_CLOG(State.TasksNum > 0, Verbose, TEXT("%*sTicking Tasks of state '%s'"), Index*UE::StateTree::DebugIndentSize, TEXT(""), *DebugGetStatePath(Exec.ActiveStates, Index));
 
 		// Tick Tasks
 		for (int32 TaskIndex = 0; TaskIndex < int32(State.TasksNum); TaskIndex++)
@@ -717,25 +727,26 @@ bool FStateTreeExecutionContext::TestAllConditions(FStateTreeInstanceData& Insta
 	return Values[0];
 }
 
-FStateTreeTransitionResult FStateTreeExecutionContext::TriggerTransitions(FStateTreeInstanceData& InstanceData, const FStateTreeStateStatus CurrentStatus, const int32 Depth)
+bool FStateTreeExecutionContext::TriggerTransitions(FStateTreeInstanceData& InstanceData, FStateTreeTransitionResult& OutTransition)
 {
 	CSV_SCOPED_TIMING_STAT_EXCLUSIVE(StateTree_TriggerTransition);
 
 	FStateTreeExecutionState& Exec = GetExecState(InstanceData);
 	EStateTreeTransitionEvent Event = EStateTreeTransitionEvent::OnCondition;
-	if (CurrentStatus.RunStatus == EStateTreeRunStatus::Succeeded)
+	if (Exec.LastTickStatus == EStateTreeRunStatus::Succeeded)
 	{
 		Event = EStateTreeTransitionEvent::OnSucceeded;
 	}
-	else if (CurrentStatus.RunStatus == EStateTreeRunStatus::Failed)
+	else if (Exec.LastTickStatus == EStateTreeRunStatus::Failed)
 	{
 		Event = EStateTreeTransitionEvent::OnFailed;
 	}
 
 	// Walk towards root and check all transitions along the way.
-	for (FStateTreeHandle Handle = CurrentStatus.State; Handle.IsValid(); Handle = StateTree->States[Handle.Index].Parent)
+	for (int32 StateIndex = Exec.ActiveStates.Num() - 1; StateIndex >= 0; StateIndex--)
 	{
-		const FBakedStateTreeState& State = StateTree->States[Handle.Index];
+		const FBakedStateTreeState& State = StateTree->States[Exec.ActiveStates[StateIndex].Index];
+		
 		for (uint8 i = 0; i < State.TransitionsNum; i++)
 		{
 			// All transition conditions must pass
@@ -751,41 +762,50 @@ FStateTreeTransitionResult FStateTreeExecutionContext::TriggerTransitions(FState
 						Exec.GatedTransitionIndex = TransitionIndex;
 						Exec.GatedTransitionTime = FMath::RandRange(0.0f, Transition.GateDelay * 0.1f); // TODO: we need variance too.
 						BeginGatedTransition(Exec);
-						STATETREE_LOG(Verbose, TEXT("Gated transition triggered from '%s' (%s) -> '%s' %.1fs"), *GetSafeStateName(CurrentStatus.State), *State.Name.ToString(), *GetSafeStateName(Transition.State), Exec.GatedTransitionTime);
+						STATETREE_LOG(Verbose, TEXT("Gated transition triggered from '%s' (%s) -> '%s' %.1fs"), *GetSafeStateName(Exec.ActiveStates.Last()), *State.Name.ToString(), *GetSafeStateName(Transition.State), Exec.GatedTransitionTime);
 					}
 
 					// Keep on updating current state, until we have tried to trigger
 					if (Exec.GatedTransitionTime > 0.0f)
 					{
-						return FStateTreeTransitionResult(CurrentStatus, FStateTreeHandle::Invalid);
+						return false;
 					}
 
-					STATETREE_LOG(Verbose, TEXT("Passed gated transition from '%s' (%s) -> '%s'"), *GetSafeStateName(CurrentStatus.State), *State.Name.ToString(), *GetSafeStateName(Transition.State));
+					STATETREE_LOG(Verbose, TEXT("Passed gated transition from '%s' (%s) -> '%s'"), *GetSafeStateName(Exec.ActiveStates.Last()), *State.Name.ToString(), *GetSafeStateName(Transition.State));
 				}
 				
 				if (Transition.Type == EStateTreeTransitionType::GotoState || Transition.Type == EStateTreeTransitionType::NextState)
 				{
-					FStateTreeTransitionResult Result = SelectState(InstanceData, CurrentStatus, Transition.State, Transition.State, Depth + 1);
-					if (Result.Next.IsValid())
+					OutTransition.CurrentActiveStates = Exec.ActiveStates;
+					OutTransition.TargetState = Transition.State;
+					OutTransition.NextActiveStates.Reset();
+
+					if (SelectState(InstanceData, Transition.State, OutTransition.NextActiveStates))
 					{
-						STATETREE_LOG(Verbose, TEXT("Transition on state '%s' (%s) -[%s]-> state '%s'"), *GetSafeStateName(CurrentStatus.State), *State.Name.ToString(), *GetSafeStateName(Result.Target), *GetSafeStateName(Result.Next));
-						return Result;
+						STATETREE_LOG(Verbose, TEXT("Transition on state '%s' (%s) -[%s]-> state '%s'"), *GetSafeStateName(Exec.ActiveStates.Last()), *State.Name.ToString(), *GetSafeStateName(Transition.State), *GetSafeStateName(OutTransition.NextActiveStates.Last()));
+						return true;
 					}
 				}
 				else if (Transition.Type == EStateTreeTransitionType::NotSet)
 				{
 					// NotSet is no-operation, but can be used to mask a transition at parent state. Returning unset keeps updating current state.
-					return FStateTreeTransitionResult(CurrentStatus, FStateTreeHandle::Invalid);
+					return false;
 				}
 				else if (Transition.Type == EStateTreeTransitionType::Succeeded)
 				{
-					STATETREE_LOG(Verbose, TEXT("Stop tree execution from state '%s' (%s): Succeeded"), *GetSafeStateName(CurrentStatus.State), *State.Name.ToString());
-					return FStateTreeTransitionResult(CurrentStatus, FStateTreeHandle::Succeeded);
+					STATETREE_LOG(Verbose, TEXT("Stop tree execution from state '%s' (%s): Succeeded"), *GetSafeStateName(Exec.ActiveStates.Last()), *State.Name.ToString());
+					OutTransition.CurrentActiveStates = Exec.ActiveStates;
+					OutTransition.TargetState = FStateTreeHandle::Succeeded;
+					OutTransition.NextActiveStates = FStateTreeActiveStates(FStateTreeHandle::Succeeded);
+					return true;
 				}
 				else
 				{
-					STATETREE_LOG(Verbose, TEXT("Stop tree execution from state '%s' (%s): Failed"), *GetSafeStateName(CurrentStatus.State), *State.Name.ToString());
-					return FStateTreeTransitionResult(CurrentStatus, FStateTreeHandle::Failed);
+					STATETREE_LOG(Verbose, TEXT("Stop tree execution from state '%s' (%s): Failed"), *GetSafeStateName(Exec.ActiveStates.Last()), *State.Name.ToString());
+					OutTransition.CurrentActiveStates = Exec.ActiveStates;
+					OutTransition.TargetState = FStateTreeHandle::Failed;
+					OutTransition.NextActiveStates = FStateTreeActiveStates(FStateTreeHandle::Failed);
+					return true;
 				}
 			}
 			else if (Exec.GatedTransitionIndex == TransitionIndex)
@@ -797,79 +817,138 @@ FStateTreeTransitionResult FStateTreeExecutionContext::TriggerTransitions(FState
 		}
 	}
 
-	if (CurrentStatus.RunStatus != EStateTreeRunStatus::Running)
+	if (Exec.LastTickStatus != EStateTreeRunStatus::Running)
 	{
-		STATETREE_LOG(Error, TEXT("%s: Default transitions on state '%s' could not be triggered. '%s' using StateTree '%s'."),
-            ANSI_TO_TCHAR(__FUNCTION__), *GetSafeStateName(CurrentStatus.State), *GetNameSafe(Owner), *GetNameSafe(StateTree));
-		return FStateTreeTransitionResult(CurrentStatus, FStateTreeHandle::Failed);
+		// Could not trigger completion transition, jump back to start.
+		static const FStateTreeHandle RootState = FStateTreeHandle(0);
+		OutTransition.TargetState = RootState;
+		return SelectState(InstanceData, RootState, OutTransition.NextActiveStates);
 	}
 
 	// No transition triggered, keep on updating current state.
-	return FStateTreeTransitionResult(CurrentStatus, FStateTreeHandle::Invalid);
+	return false;
 }
 
-FStateTreeTransitionResult FStateTreeExecutionContext::SelectState(FStateTreeInstanceData& InstanceData, const FStateTreeStateStatus InitialStateStatus, const FStateTreeHandle InitialTargetState, const FStateTreeHandle NextState, const int Depth)
+bool FStateTreeExecutionContext::SelectState(FStateTreeInstanceData& InstanceData, const FStateTreeHandle NextState, FStateTreeActiveStates& OutNewActiveState)
+{
+	const FStateTreeExecutionState& Exec = GetExecState(InstanceData);
+
+	if (!NextState.IsValid())
+	{
+		return false;
+	}
+
+	// Find common ancestor of `NextState` in the current active states and connect.
+	// This allows transitions within a subtree.
+	OutNewActiveState = Exec.ActiveStates;
+	
+	TStaticArray<FStateTreeHandle, FStateTreeActiveStates::MaxStates> InBetweenStates;
+	int32 NumInBetweenStates = 0;
+	int32 CommonActiveAncestorIndex = INDEX_NONE;
+
+	// Walk towards the root from current state.
+	FStateTreeHandle CurrState = NextState;
+	while (CurrState.IsValid())
+	{
+		// Store the states that are in between the 'NextState' and common ancestor. 
+		InBetweenStates[NumInBetweenStates++] = CurrState;
+		// Check if the state can be found in the active states.
+		CommonActiveAncestorIndex = OutNewActiveState.IndexOfReverse(CurrState); 
+		if (CommonActiveAncestorIndex != INDEX_NONE)
+		{
+			break;
+		}
+		if (NumInBetweenStates == InBetweenStates.Num())
+		{
+			STATETREE_LOG(Error, TEXT("%s: Too many parent states when selecting state '%s' from '%s'.  '%s' using StateTree '%s'."),
+				ANSI_TO_TCHAR(__FUNCTION__), *GetSafeStateName(NextState), *GetStateStatusString(Exec), *GetNameSafe(Owner), *GetNameSafe(StateTree));
+			return false;
+		}
+
+		CurrState = StateTree->States[CurrState.Index].Parent;
+	}
+
+	// Max takes care of INDEX_NONE, by setting the num to 0.
+	OutNewActiveState.SetNum(FMath::Max(0, CommonActiveAncestorIndex));
+	
+	// Append in between state in reverse order, they were collected from leaf towards the root.
+	bool bActiveStatesOverflow = false;
+	for (int32 Index = NumInBetweenStates - 1; Index > 0; Index--)
+	{
+		bActiveStatesOverflow |= !OutNewActiveState.Push(InBetweenStates[Index]);
+	}
+
+	if (bActiveStatesOverflow)
+	{
+		STATETREE_LOG(Error, TEXT("%s: Reached max execution depth when trying to select state %s from '%s'.  '%s' using StateTree '%s'."),
+			ANSI_TO_TCHAR(__FUNCTION__), *GetSafeStateName(NextState), *GetStateStatusString(Exec), *GetNameSafe(Owner), *GetNameSafe(StateTree));
+		return false;
+	}
+
+	return SelectStateInternal(InstanceData, NextState, OutNewActiveState);
+}
+
+bool FStateTreeExecutionContext::SelectStateInternal(FStateTreeInstanceData& InstanceData, const FStateTreeHandle NextState, FStateTreeActiveStates& OutNewActiveState)
 {
 	CSV_SCOPED_TIMING_STAT_EXCLUSIVE(StateTree_SelectState);
 
+	const FStateTreeExecutionState& Exec = GetExecState(InstanceData);
+
 	if (!NextState.IsValid())
 	{
 		// Trying to select non-existing state.
-		STATETREE_LOG(Error, TEXT("%s: Trying to select invalid state from '%s' via '%s.  '%s' using StateTree '%s'."),
-            ANSI_TO_TCHAR(__FUNCTION__), *GetStateStatusString(InitialStateStatus), *GetSafeStateName(InitialTargetState), *GetNameSafe(Owner), *GetNameSafe(StateTree));
-		return FStateTreeTransitionResult(InitialStateStatus, InitialTargetState, FStateTreeHandle::Invalid);
+		STATETREE_LOG(Error, TEXT("%s: Trying to select invalid state from '%s'.  '%s' using StateTree '%s'."),
+            ANSI_TO_TCHAR(__FUNCTION__), *GetStateStatusString(Exec), *GetNameSafe(Owner), *GetNameSafe(StateTree));
+		return false;
 	}
 
 	const FBakedStateTreeState& State = StateTree->States[NextState.Index];
 
 	// Make sure all the evaluators for the target state are up to date.
-	TickEvaluators(InstanceData, NextState, EStateTreeEvaluationType::PreSelect, 0.0f);
+	TickEvaluatorsForSelect(InstanceData, NextState, EStateTreeEvaluationType::PreSelect, 0.0f);
 	
 	// Check that the state can be entered
 	if (TestAllConditions(InstanceData, State.EnterConditionsBegin, State.EnterConditionsNum))
 	{
-		// If the state has children, proceed to select children.
-		if (State.HasChildren())
+		if (!OutNewActiveState.Push(NextState))
+		{
+			STATETREE_LOG(Error, TEXT("%s: Reached max execution depth when trying to select state %s from '%s'.  '%s' using StateTree '%s'."),
+				ANSI_TO_TCHAR(__FUNCTION__), *GetSafeStateName(NextState), *GetStateStatusString(Exec), *GetNameSafe(Owner), *GetNameSafe(StateTree));
+			return false;
+		}
+		
+		if (State.LinkedState.IsValid())
 		{
+			// If State is linked, proceed to the linked state.
+			if (SelectStateInternal(InstanceData, State.LinkedState, OutNewActiveState))
+			{
+				// Selection succeeded
+				return true;
+			}
+		}
+		else if (State.HasChildren())
+		{
+			// If the state has children, proceed to select children.
 			for (uint16 ChildState = State.ChildrenBegin; ChildState < State.ChildrenEnd; ChildState = StateTree->States[ChildState].GetNextSibling())
 			{
-				const FStateTreeTransitionResult ChildResult = SelectState(InstanceData, InitialStateStatus, InitialTargetState, FStateTreeHandle(ChildState), Depth + 1);
-				if (ChildResult.Next.IsValid())
+				if (SelectStateInternal(InstanceData, FStateTreeHandle(ChildState), OutNewActiveState))
 				{
 					// Selection succeeded
-					return ChildResult;
+					return true;
 				}
 			}
 		}
 		else
 		{
 			// Select this state.
-			return FStateTreeTransitionResult(InitialStateStatus, InitialTargetState, NextState);
+			return true;
 		}
-	}
-	
-	// Nothing got selected.
-	return FStateTreeTransitionResult(InitialStateStatus, InitialTargetState, FStateTreeHandle::Invalid);
-}
-
-int32 FStateTreeExecutionContext::GetActiveStates(const FStateTreeHandle StateHandle, TStaticArray<FStateTreeHandle, 32>& OutStateHandles) const
-{
-	if (StateHandle == FStateTreeHandle::Succeeded || StateHandle == FStateTreeHandle::Failed || StateHandle == FStateTreeHandle::Invalid)
-	{
-		return 0;
-	}
-	
-	int32 NumStates = 0;
-	FStateTreeHandle CurrentHandle = StateHandle;
-	while (CurrentHandle.IsValid() && NumStates < OutStateHandles.Num())
-	{
-		OutStateHandles[NumStates++] = CurrentHandle;
-		CurrentHandle = StateTree->States[CurrentHandle.Index].Parent;
+		
+		OutNewActiveState.Pop();
 	}
 
-	Algo::Reverse(OutStateHandles.GetData(), NumStates);
-	
-	return NumStates;
+	// Nothing got selected.
+	return false;
 }
 
 FString FStateTreeExecutionContext::GetSafeStateName(const FStateTreeHandle State) const
@@ -894,39 +973,25 @@ FString FStateTreeExecutionContext::GetSafeStateName(const FStateTreeHandle Stat
 	return TEXT("(Unknown)");
 }
 
-FString FStateTreeExecutionContext::DebugGetStatePath(const TArrayView<FStateTreeHandle> ActiveStateHandles, const int32 ActiveStateIndex) const
+FString FStateTreeExecutionContext::DebugGetStatePath(const FStateTreeActiveStates& ActiveStates, int32 ActiveStateIndex) const
 {
 	FString StatePath;
-	if (!ensureMsgf(ActiveStateHandles.IsValidIndex(ActiveStateIndex), TEXT("Provided index must be valid")))
+	if (!ensureMsgf(ActiveStates.IsValidIndex(ActiveStateIndex), TEXT("Provided index must be valid")))
 	{
 		return StatePath;
 	}
 
 	for (int32 i = 0; i <= ActiveStateIndex; i++)
 	{
-        const FBakedStateTreeState& State = StateTree->States[ActiveStateHandles[i].Index];
+		const FBakedStateTreeState& State = StateTree->States[ActiveStates[i].Index];
 		StatePath.Appendf(TEXT("%s%s"), i == 0 ? TEXT("") : TEXT("."), *State.Name.ToString());
 	}
 	return StatePath;
 }
 
-FString FStateTreeExecutionContext::GetStateStatusString(const FStateTreeStateStatus StateStatus) const
+FString FStateTreeExecutionContext::GetStateStatusString(const FStateTreeExecutionState& ExecState) const
 {
-	// Invalid (null) state in status refers to the whole tree.
-	const FString StateName = GetSafeStateName(StateStatus.State);
-	switch(StateStatus.RunStatus)
-	{
-	case EStateTreeRunStatus::Unset:
-		return StateName + TEXT("Unset");
-	case EStateTreeRunStatus::Running:
-		return StateName + TEXT("Running");
-	case EStateTreeRunStatus::Succeeded:
-		return StateName + TEXT("Succeeded");
-	case EStateTreeRunStatus::Failed:
-		return StateName + TEXT("Failed");
-	default:
-		return StateName + TEXT("(Unknown)");
-	}
+	return GetSafeStateName(ExecState.ActiveStates.Last()) + TEXT(":") + UEnum::GetDisplayValueAsText(ExecState.LastTickStatus).ToString();
 }
 
 EStateTreeRunStatus FStateTreeExecutionContext::GetLastTickStatus(const FStateTreeInstanceData* ExternalInstanceData) const
@@ -940,6 +1005,13 @@ FString FStateTreeExecutionContext::GetInstanceDescription() const
 	return Owner != nullptr? FString::Printf(TEXT("%s: "), *Owner->GetName()) : TEXT("");
 }
 
+const FStateTreeActiveStates& FStateTreeExecutionContext::GetActiveStates(const FStateTreeInstanceData* ExternalInstanceData) const
+{
+	const FStateTreeExecutionState& Exec = GetExecState(SelectInstanceData(ExternalInstanceData));
+	return Exec.ActiveStates;
+}
+
+
 #if WITH_GAMEPLAY_DEBUGGER
 
 FString FStateTreeExecutionContext::GetDebugInfoString(const FStateTreeInstanceData* ExternalInstanceData) const
@@ -971,13 +1043,10 @@ FString FStateTreeExecutionContext::GetDebugInfoString(const FStateTreeInstanceD
 	}
 
 	// Active States
-	TStaticArray<FStateTreeHandle, 32> ActiveStates;
-	const int32 NumStates = GetActiveStates(Exec.CurrentState, ActiveStates);
-
 	DebugString += TEXT("Current State:\n");
-	for (int32 Index = 0; Index < NumStates; Index++)
+	for (int32 Index = 0; Index < Exec.ActiveStates.Num(); Index++)
 	{
-		FStateTreeHandle Handle = ActiveStates[Index];
+		FStateTreeHandle Handle = Exec.ActiveStates[Index];
 		if (Handle.IsValid())
 		{
 			const FBakedStateTreeState& State = StateTree->States[Handle.Index];
@@ -1087,13 +1156,11 @@ void FStateTreeExecutionContext::DebugPrintInternalLayout(const FStateTreeInstan
 		);
 	for (const FBakedStateTreeState& State : StateTree->States)
 	{
-		DebugString += FString::Printf(TEXT("  | %-30s | %15s | %5s [%3d:%-3d[ | %9s   %4d %4d %4d %4d | %3s   %4d %4d %4d %4d | %11s   %-16s %-40s %-16s %-40s |\n"),
+		DebugString += FString::Printf(TEXT("  | %-30s | %15s | %5s [%3d:%-3d[ | %9s   %4d %4d %4d %4d | %3s   %4d %4d %4d %4d\n"),
 									*State.Name.ToString(), *State.Parent.Describe(),
 									TEXT(""), State.ChildrenBegin, State.ChildrenEnd,
 									TEXT(""), State.EnterConditionsBegin, State.TransitionsBegin, State.TasksBegin, State.EvaluatorsBegin,
-									TEXT(""), State.EnterConditionsNum, State.TransitionsNum, State.TasksNum, State.EvaluatorsNum,
-									TEXT(""), *State.StateDoneTransitionState.Describe(), *StaticEnum<EStateTreeTransitionType>()->GetValueAsString(State.StateDoneTransitionType),
-									*State.StateFailedTransitionState.Describe(), *StaticEnum<EStateTreeTransitionType>()->GetValueAsString(State.StateFailedTransitionType));
+									TEXT(""), State.EnterConditionsNum, State.TransitionsNum, State.TasksNum, State.EvaluatorsNum);
 	}
 
 	DebugString += FString::Printf(TEXT("\nStates\n  [ %-30s | %-12s | %-30s | %15s | %10s ]\n"),
@@ -1132,6 +1199,7 @@ FString FStateTreeExecutionContext::GetActiveStateName(const FStateTreeInstanceD
 	{
 		return FString(TEXT("<None>"));
 	}
+	
 	const FStateTreeInstanceData& InstanceData = SelectInstanceData(ExternalInstanceData);
 	checkSlow(InstanceData.GetLayout() == StateTree->GetInstanceDataDefaultValue().GetLayout());
 	const FStateTreeExecutionState& Exec = GetExecState(InstanceData);
@@ -1139,21 +1207,24 @@ FString FStateTreeExecutionContext::GetActiveStateName(const FStateTreeInstanceD
 	FString FullStateName;
 	
 	// Active States
-	TStaticArray<FStateTreeHandle, 32> ActiveStates;
-	const int32 NumStates = GetActiveStates(Exec.CurrentState, ActiveStates);
-
-	for (int32 Index = 0; Index < NumStates; Index++)
+	for (int32 Index = 0; Index < Exec.ActiveStates.Num(); Index++)
 	{
-		FStateTreeHandle Handle = ActiveStates[Index];
+		const FStateTreeHandle Handle = Exec.ActiveStates[Index];
 		if (Handle.IsValid())
 		{
 			const FBakedStateTreeState& State = StateTree->States[Handle.Index];
+			bool bIsLinked = false;
 			if (Index > 0)
 			{
 				FullStateName += TEXT("\n");
+				bIsLinked = Exec.ActiveStates[Index - 1] != State.Parent;
 			}
 			FullStateName += FString::Printf(TEXT("%*s-"), Index * 3, TEXT("")); // Indent
 			FullStateName += *State.Name.ToString();
+			if (bIsLinked)
+			{
+				FullStateName += TEXT(" >");
+			}
 		}
 	}
 
@@ -1174,6 +1245,42 @@ FString FStateTreeExecutionContext::GetActiveStateName(const FStateTreeInstanceD
 
 	return FullStateName;
 }
+
+int32 FStateTreeExecutionContext::GetStateChangeCount(const FStateTreeInstanceData* ExternalInstanceData) const
+{
+	const FStateTreeInstanceData& InstanceData = SelectInstanceData(ExternalInstanceData);
+	checkSlow(InstanceData.GetLayout() == StateTree->GetInstanceDataDefaultValue().GetLayout());
+	const FStateTreeExecutionState& Exec = GetExecState(InstanceData);
+	return Exec.StateChangeCount;
+}
+
+TArray<FName> FStateTreeExecutionContext::GetActiveStateNames(const FStateTreeInstanceData* ExternalInstanceData) const
+{
+	TArray<FName> Result;
+
+	if (!StateTree)
+	{
+		return Result;
+	}
+	
+	const FStateTreeInstanceData& InstanceData = SelectInstanceData(ExternalInstanceData);
+	checkSlow(InstanceData.GetLayout() == StateTree->GetInstanceDataDefaultValue().GetLayout());
+	const FStateTreeExecutionState& Exec = GetExecState(InstanceData);
+
+	// Active States
+	for (int32 Index = 0; Index < Exec.ActiveStates.Num(); Index++)
+	{
+		const FStateTreeHandle Handle = Exec.ActiveStates[Index];
+		if (Handle.IsValid())
+		{
+			const FBakedStateTreeState& State = StateTree->States[Handle.Index];
+			Result.Add(State.Name);
+		}
+	}
+
+	return Result;
+}
+
 #endif // WITH_STATETREE_DEBUG
 
 #undef STATETREE_LOG
diff --git a/Engine/Plugins/Runtime/StateTree/Source/StateTreeModule/Public/Blueprint/StateTreeEvaluatorBlueprintBase.h b/Engine/Plugins/Runtime/StateTree/Source/StateTreeModule/Public/Blueprint/StateTreeEvaluatorBlueprintBase.h
--- a/Engine/Plugins/Runtime/StateTree/Source/StateTreeModule/Public/Blueprint/StateTreeEvaluatorBlueprintBase.h
+++ b/Engine/Plugins/Runtime/StateTree/Source/StateTreeModule/Public/Blueprint/StateTreeEvaluatorBlueprintBase.h
@@ -28,7 +28,7 @@ public:
 	void ReceiveExitState(AActor* OwnerActor, const EStateTreeStateChangeType ChangeType, const FStateTreeTransitionResult& Transition);
 
 	UFUNCTION(BlueprintImplementableEvent, meta = (DisplayName = "StateCompleted"))
-	void ReceiveStateCompleted(AActor* OwnerActor, const EStateTreeRunStatus CompletionStatus, const FStateTreeHandle CompletedState);
+	void ReceiveStateCompleted(AActor* OwnerActor, const EStateTreeRunStatus CompletionStatus, const FStateTreeActiveStates CompletedState);
 
 	UFUNCTION(BlueprintImplementableEvent, meta = (DisplayName = "Evaluate"))
 	void ReceiveEvaluate(AActor* OwnerActor, const EStateTreeEvaluationType EvalType, const float DeltaTime);
@@ -38,7 +38,7 @@ public:
 
 	virtual void ExitState(FStateTreeExecutionContext& Context, const EStateTreeStateChangeType ChangeType, const FStateTreeTransitionResult& Transition);
 
-	virtual void StateCompleted(FStateTreeExecutionContext& Context, const EStateTreeRunStatus CompletionStatus, const FStateTreeHandle CompletedState);
+	virtual void StateCompleted(FStateTreeExecutionContext& Context, const EStateTreeRunStatus CompletionStatus, const FStateTreeActiveStates& CompletedActiveStates);
 
 	virtual void Evaluate(FStateTreeExecutionContext& Context, const EStateTreeEvaluationType EvalType, const float DeltaTime);
 };
@@ -55,7 +55,7 @@ struct STATETREEMODULE_API FStateTreeBlueprintEvaluatorWrapper : public FStateTr
 	virtual bool Link(FStateTreeLinker& Linker) override;
 	virtual void EnterState(FStateTreeExecutionContext& Context, const EStateTreeStateChangeType ChangeType, const FStateTreeTransitionResult& Transition) const override;
 	virtual void ExitState(FStateTreeExecutionContext& Context, const EStateTreeStateChangeType ChangeType, const FStateTreeTransitionResult& Transition) const override;
-	virtual void StateCompleted(FStateTreeExecutionContext& Context, const EStateTreeRunStatus CompletionStatus, const FStateTreeHandle CompletedState) const override;
+	virtual void StateCompleted(FStateTreeExecutionContext& Context, const EStateTreeRunStatus CompletionStatus, const FStateTreeActiveStates& CompletedActiveStates) const override;
 	virtual void Evaluate(FStateTreeExecutionContext& Context, const EStateTreeEvaluationType EvalType, const float DeltaTime) const override;
 
 	UPROPERTY()
diff --git a/Engine/Plugins/Runtime/StateTree/Source/StateTreeModule/Public/Blueprint/StateTreeTaskBlueprintBase.h b/Engine/Plugins/Runtime/StateTree/Source/StateTreeModule/Public/Blueprint/StateTreeTaskBlueprintBase.h
--- a/Engine/Plugins/Runtime/StateTree/Source/StateTreeModule/Public/Blueprint/StateTreeTaskBlueprintBase.h
+++ b/Engine/Plugins/Runtime/StateTree/Source/StateTreeModule/Public/Blueprint/StateTreeTaskBlueprintBase.h
@@ -28,7 +28,7 @@ public:
 	void ReceiveExitState(AActor* OwnerActor, const EStateTreeStateChangeType ChangeType, const FStateTreeTransitionResult& Transition);
 
 	UFUNCTION(BlueprintImplementableEvent, meta = (DisplayName = "StateCompleted"))
-	void ReceiveStateCompleted(AActor* OwnerActor, const EStateTreeRunStatus CompletionStatus, const FStateTreeHandle CompletedState);
+	void ReceiveStateCompleted(AActor* OwnerActor, const EStateTreeRunStatus CompletionStatus, const FStateTreeActiveStates CompletedActiveStates);
 
 	UFUNCTION(BlueprintImplementableEvent, meta = (DisplayName = "Tick"))
 	EStateTreeRunStatus ReceiveTick(AActor* OwnerActor, const float DeltaTime);
@@ -38,7 +38,7 @@ public:
 
 	virtual void ExitState(FStateTreeExecutionContext& Context, const EStateTreeStateChangeType ChangeType, const FStateTreeTransitionResult& Transition);
 
-	virtual void StateCompleted(FStateTreeExecutionContext& Context, const EStateTreeRunStatus CompletionStatus, const FStateTreeHandle CompletedState);
+	virtual void StateCompleted(FStateTreeExecutionContext& Context, const EStateTreeRunStatus CompletionStatus, const FStateTreeActiveStates& CompletedActiveStates);
 
 	virtual EStateTreeRunStatus Tick(FStateTreeExecutionContext& Context, const float DeltaTime);
 };
@@ -55,7 +55,7 @@ struct STATETREEMODULE_API FStateTreeBlueprintTaskWrapper : public FStateTreeTas
 	virtual bool Link(FStateTreeLinker& Linker) override;
 	virtual EStateTreeRunStatus EnterState(FStateTreeExecutionContext& Context, const EStateTreeStateChangeType ChangeType, const FStateTreeTransitionResult& Transition) const override;
 	virtual void ExitState(FStateTreeExecutionContext& Context, const EStateTreeStateChangeType ChangeType, const FStateTreeTransitionResult& Transition) const override;
-	virtual void StateCompleted(FStateTreeExecutionContext& Context, const EStateTreeRunStatus CompletionStatus, const FStateTreeHandle CompletedState) const override;
+	virtual void StateCompleted(FStateTreeExecutionContext& Context, const EStateTreeRunStatus CompletionStatus, const FStateTreeActiveStates& CompletedActiveStates) const override;
 	virtual EStateTreeRunStatus Tick(FStateTreeExecutionContext& Context, const float DeltaTime) const override;
 
 	UPROPERTY()
diff --git a/Engine/Plugins/Runtime/StateTree/Source/StateTreeModule/Public/Evaluators/CooldownGuardEvaluator.h b/Engine/Plugins/Runtime/StateTree/Source/StateTreeModule/Public/Evaluators/CooldownGuardEvaluator.h
--- a/Engine/Plugins/Runtime/StateTree/Source/StateTreeModule/Public/Evaluators/CooldownGuardEvaluator.h
+++ b/Engine/Plugins/Runtime/StateTree/Source/StateTreeModule/Public/Evaluators/CooldownGuardEvaluator.h
@@ -49,7 +49,7 @@ struct STATETREEMODULE_API FCooldownGuardEvaluator : public FStateTreeEvaluatorC
 	virtual const UStruct* GetInstanceDataType() const override { return FCooldownGuardEvaluatorInstanceData::StaticStruct(); };
 	virtual bool Link(FStateTreeLinker& Linker) override;
 	virtual void EnterState(FStateTreeExecutionContext& Context, const EStateTreeStateChangeType ChangeType, const FStateTreeTransitionResult& Transition) const override;
-	virtual void StateCompleted(FStateTreeExecutionContext& Context, const EStateTreeRunStatus CompletionStatus, const FStateTreeHandle CompletedState) const override;
+	virtual void StateCompleted(FStateTreeExecutionContext& Context, const EStateTreeRunStatus CompletionStatus, const FStateTreeActiveStates& CompletedActiveStates) const override;
 	virtual void Evaluate(FStateTreeExecutionContext& Context, const EStateTreeEvaluationType EvalType, const float DeltaTime) const override;
 
 	void SetCoolDown(FStateTreeExecutionContext& Context) const;
diff --git a/Engine/Plugins/Runtime/StateTree/Source/StateTreeModule/Public/Evaluators/SelectionGuardEvaluator.h b/Engine/Plugins/Runtime/StateTree/Source/StateTreeModule/Public/Evaluators/SelectionGuardEvaluator.h
--- a/Engine/Plugins/Runtime/StateTree/Source/StateTreeModule/Public/Evaluators/SelectionGuardEvaluator.h
+++ b/Engine/Plugins/Runtime/StateTree/Source/StateTreeModule/Public/Evaluators/SelectionGuardEvaluator.h
@@ -43,7 +43,7 @@ struct STATETREEMODULE_API FSelectionGuardEvaluator : public FStateTreeEvaluator
 	virtual const UStruct* GetInstanceDataType() const override { return FSelectionGuardEvaluatorInstanceData::StaticStruct(); };
 	virtual bool Link(FStateTreeLinker& Linker) override;
 	virtual void EnterState(FStateTreeExecutionContext& Context, const EStateTreeStateChangeType ChangeType, const FStateTreeTransitionResult& Transition) const override;
-	virtual void StateCompleted(FStateTreeExecutionContext& Context, const EStateTreeRunStatus CompletionStatus, const FStateTreeHandle CompletedState) const override;
+	virtual void StateCompleted(FStateTreeExecutionContext& Context, const EStateTreeRunStatus CompletionStatus, const FStateTreeActiveStates& CompletedActiveStates) const override;
 	virtual void ExitState(FStateTreeExecutionContext& Context, const EStateTreeStateChangeType ChangeType, const FStateTreeTransitionResult& Transition) const override;
 	
 	TStateTreeInstanceDataPropertyHandle<bool> ActiveHandle;
diff --git a/Engine/Plugins/Runtime/StateTree/Source/StateTreeModule/Public/StateTreeEvaluatorBase.h b/Engine/Plugins/Runtime/StateTree/Source/StateTreeModule/Public/StateTreeEvaluatorBase.h
--- a/Engine/Plugins/Runtime/StateTree/Source/StateTreeModule/Public/StateTreeEvaluatorBase.h
+++ b/Engine/Plugins/Runtime/StateTree/Source/StateTreeModule/Public/StateTreeEvaluatorBase.h
@@ -40,9 +40,9 @@ struct STATETREEMODULE_API FStateTreeEvaluatorBase : public FStateTreeNodeBase
 	 * are executed earlier in the tree. Note that StateCompleted is not called if conditional transition changes the state.
 	 * @param Context Reference to current execution context.
 	 * @param CompletionStatus Describes the running status of the completed state (Succeeded/Failed).
-	 * @param CompletedState Handle of the state that was completed.
+	 * @param CompletedActiveStates Active states at the time of completion.
 	 */
-	virtual void StateCompleted(FStateTreeExecutionContext& Context, const EStateTreeRunStatus CompletionStatus, const FStateTreeHandle CompletedState) const {}
+	virtual void StateCompleted(FStateTreeExecutionContext& Context, const EStateTreeRunStatus CompletionStatus, const FStateTreeActiveStates& CompletedActiveStates) const {}
 	
 	/**
 	 * Called when evaluator needs to be updated. EvalType describes if the tick happens during state tree tick when the evaluator is on active state (Tick),
diff --git a/Engine/Plugins/Runtime/StateTree/Source/StateTreeModule/Public/StateTreeExecutionContext.h b/Engine/Plugins/Runtime/StateTree/Source/StateTreeModule/Public/StateTreeExecutionContext.h
--- a/Engine/Plugins/Runtime/StateTree/Source/StateTreeModule/Public/StateTreeExecutionContext.h
+++ b/Engine/Plugins/Runtime/StateTree/Source/StateTreeModule/Public/StateTreeExecutionContext.h
@@ -18,11 +18,11 @@ struct STATETREEMODULE_API FStateTreeExecutionState
 {
 	GENERATED_BODY()
 
-	/** Currently active state */
-	FStateTreeHandle CurrentState = FStateTreeHandle::Invalid;
+	/** Currently active states */
+	FStateTreeActiveStates ActiveStates;
 
 	/** The index of the task that failed during enter state. Exit state uses it to call ExitState() symmetrically. */
-	uint16 EnterStateFailedTaskIndex = INDEX_NONE;
+	uint16 EnterStateFailedTaskIndex = MAX_uint16;
 
 	/** Result of last tick */
 	EStateTreeRunStatus LastTickStatus = EStateTreeRunStatus::Failed;
@@ -33,6 +33,9 @@ struct STATETREEMODULE_API FStateTreeExecutionState
 	/** Delayed transition handle, if exists */
 	int16 GatedTransitionIndex = INDEX_NONE;
 
+	/** Number of times a new state has been changed. */
+	uint16 StateChangeCount = 0;
+
 	/** Running time of the delayed transition */
 	float GatedTransitionTime = 0.0f;
 };
@@ -224,6 +227,8 @@ public:
 	
 	EStateTreeRunStatus GetLastTickStatus(const FStateTreeInstanceData* ExternalInstanceData = nullptr) const;
 
+	const FStateTreeActiveStates& GetActiveStates(const FStateTreeInstanceData* ExternalInstanceData = nullptr) const;
+
 #if WITH_GAMEPLAY_DEBUGGER
 	/** @return Debug string describing the current state of the execution */
 	FString GetDebugInfoString(const FStateTreeInstanceData* ExternalInstanceData = nullptr) const;
@@ -231,7 +236,10 @@ public:
 
 #if WITH_STATETREE_DEBUG
 	FString GetActiveStateName(const FStateTreeInstanceData* ExternalInstanceData = nullptr) const;
-	
+	int32 GetStateChangeCount(const FStateTreeInstanceData* ExternalInstanceData = nullptr) const;
+
+	TArray<FName> GetActiveStateNames(const FStateTreeInstanceData* ExternalInstanceData = nullptr) const;
+
 	void DebugPrintInternalLayout(const FStateTreeInstanceData* ExternalInstanceData = nullptr);
 #endif
 
@@ -259,21 +267,27 @@ protected:
 	void ExitState(FStateTreeInstanceData& InstanceData, const FStateTreeTransitionResult& Transition);
 
 	/**
-	 * Handles logic for exiting State. ExitState is called on current active Evaluators and Tasks in reverse order (from leaf to root).
+	 * Handles logic for signalling State completed. StateCompleted is called on current active Evaluators and Tasks in reverse order (from leaf to root).
 	 */
-	void StateCompleted(FStateTreeInstanceData& InstanceData, const FStateTreeHandle CurrentState, const EStateTreeRunStatus CompletionStatus);
+	void StateCompleted(FStateTreeInstanceData& InstanceData);
 
 	/**
-	 * Ticks evaluators of all active states starting from current state by delta time.
+	 * Ticks evaluators of CurrentState and all of it's parent states.
+	 * If TickEvaluatorsForSelect() is called multiple times per frame (i.e. during selection when visiting new states), each state and evaluator is ticked only once.
+	 */
+	void TickEvaluatorsForSelect(FStateTreeInstanceData& InstanceData, const FStateTreeHandle CurrentState, const EStateTreeEvaluationType EvalType, const float DeltaTime);
+
+	/**
+	 * Ticks evaluators of all active states by delta time.
 	 * If TickEvaluators() is called multiple times per frame (i.e. during selection when visiting new states), each state and evaluator is ticked only once.
 	 */
-	void TickEvaluators(FStateTreeInstanceData& InstanceData, const FStateTreeHandle CurrentState, const EStateTreeEvaluationType EvalType, const float DeltaTime);
+	void TickEvaluators(FStateTreeInstanceData& InstanceData, const FStateTreeActiveStates& ActiveStates, const EStateTreeEvaluationType EvalType, const float DeltaTime);
 
 	/**
 	 * Ticks tasks of all active states starting from current state by delta time.
 	 * @return Run status returned by the tasks.
 	 */
-	EStateTreeRunStatus TickTasks(FStateTreeInstanceData& InstanceData, const FStateTreeHandle CurrentState, const float DeltaTime);
+	EStateTreeRunStatus TickTasks(FStateTreeInstanceData& InstanceData, const FStateTreeExecutionState& Exec, const float DeltaTime);
 
 	/**
 	 * Checks all conditions at given range
@@ -288,24 +302,24 @@ protected:
 	 * the actual next state returned by the selector.
 	 * @return Transition result describing the source state, state transitioned to, and next selected state.
 	 */
-	FStateTreeTransitionResult TriggerTransitions(FStateTreeInstanceData& InstanceData, const FStateTreeStateStatus CurrentStatus, const int Depth);
+	bool TriggerTransitions(FStateTreeInstanceData& InstanceData, FStateTreeTransitionResult& OutTransition);
 
 	/**
 	 * Runs state selection logic starting at the specified state, walking towards the leaf states.
-	 * If the preconditions of NextState are not met, "Invalid" is returned. 
-	 * If NextState is a selector state, SelectState is called recursively (depth-first) to all child states (where NextState will be one of child states).
-	 * If NextState is a leaf state, the NextState is returned.
-	 * @param Storage View representing all instance data used by tasks and evaluators
-	 * @param InitialStateStatus Describes the current state and running status (will be passed intact to next selector)
-	 * @param InitialTargetState The state the initial transition target state (will be passed intact to next selector) 
+	 * If a state cannot be selected, false is returned. 
+	 * If NextState is a selector state, SelectStateInternal is called recursively (depth-first) to all child states (where NextState will be one of child states).
+	 * If NextState is a leaf state, the active states leading from root to the leaf are returned.
+	 * @param InstanceData Reference to the instance data
 	 * @param NextState The state which we try to select next.
-	 * @param Depth Depth of recursion.
-	 * @return Transition result describing the source state, transition target state, and next selected state.
+	 * @param OutNewActiveStates Active states that got selected.
+	 * @return True if succeeded to select new active states.
 	 */
-	FStateTreeTransitionResult SelectState(FStateTreeInstanceData& InstanceData, const FStateTreeStateStatus InitialStateStatus, const FStateTreeHandle InitialTargetState, const FStateTreeHandle NextState, const int Depth);
+	bool SelectState(FStateTreeInstanceData& InstanceData, const FStateTreeHandle NextState, FStateTreeActiveStates& OutNewActiveStates);
 
-	/** @return State handles from specified state handle back to the root, specified handle included. */
-	int32 GetActiveStates(const FStateTreeHandle StateHandle, TStaticArray<FStateTreeHandle, 32>& OutStateHandles) const;
+	/**
+	 * Used internally to do the recursive part of the SelectState().
+	 */
+	bool SelectStateInternal(FStateTreeInstanceData& InstanceData, const FStateTreeHandle NextState, FStateTreeActiveStates& OutNewActiveStates);
 
 	/** @return Mutable storage based on storage settings. */
 	FStateTreeInstanceData& SelectMutableInstanceData(FStateTreeInstanceData* ExternalInstanceData)
@@ -352,13 +366,13 @@ protected:
 	}
 	
 	/** @return String describing state status for logging and debug. */
-	FString GetStateStatusString(const FStateTreeStateStatus StateStatus) const;
+	FString GetStateStatusString(const FStateTreeExecutionState& ExecState) const;
 
 	/** @return String describing state name for logging and debug. */
 	FString GetSafeStateName(const FStateTreeHandle State) const;
 
 	/** @return String describing full path of an activate state for logging and debug. */
-	FString DebugGetStatePath(TArrayView<FStateTreeHandle> ActiveStateHandles, int32 ActiveStateIndex) const;
+	FString DebugGetStatePath(const FStateTreeActiveStates& ActiveStates, int32 ActiveStateIndex) const;
 
 	/** The StateTree asset the context is initialized for */
 	UPROPERTY(Transient)
diff --git a/Engine/Plugins/Runtime/StateTree/Source/StateTreeModule/Public/StateTreeTaskBase.h b/Engine/Plugins/Runtime/StateTree/Source/StateTreeModule/Public/StateTreeTaskBase.h
--- a/Engine/Plugins/Runtime/StateTree/Source/StateTreeModule/Public/StateTreeTaskBase.h
+++ b/Engine/Plugins/Runtime/StateTree/Source/StateTreeModule/Public/StateTreeTaskBase.h
@@ -41,9 +41,9 @@ struct STATETREEMODULE_API FStateTreeTaskBase : public FStateTreeNodeBase
 	 * are executed earlier in the tree. Note that StateCompleted is not called if conditional transition changes the state.
 	 * @param Context Reference to current execution context.
 	 * @param CompletionStatus Describes the running status of the completed state (Succeeded/Failed).
-	 * @param CompletedState Handle of the state that was completed.
+	 * @param CompletedActiveStates Active states at the time of completion.
 	 */
-	virtual void StateCompleted(FStateTreeExecutionContext& Context, const EStateTreeRunStatus CompletionStatus, const FStateTreeHandle CompletedState) const {}
+	virtual void StateCompleted(FStateTreeExecutionContext& Context, const EStateTreeRunStatus CompletionStatus, const FStateTreeActiveStates& CompletedActiveStates) const {}
 
 	/**
 	 * Called during state tree tick when the task is on active state.
diff --git a/Engine/Plugins/Runtime/StateTree/Source/StateTreeModule/Public/StateTreeTypes.h b/Engine/Plugins/Runtime/StateTree/Source/StateTreeModule/Public/StateTreeTypes.h
--- a/Engine/Plugins/Runtime/StateTree/Source/StateTreeModule/Public/StateTreeTypes.h
+++ b/Engine/Plugins/Runtime/StateTree/Source/StateTreeModule/Public/StateTreeTypes.h
@@ -118,25 +118,157 @@ struct STATETREEMODULE_API FStateTreeHandle
 
 
 /**
- * Describes current status of a running state or desired state.
+ * Describes an array of active states in a State Tree.
  */
 USTRUCT(BlueprintType)
-struct STATETREEMODULE_API FStateTreeStateStatus
+struct STATETREEMODULE_API FStateTreeActiveStates
 {
 	GENERATED_BODY()
 
-	FStateTreeStateStatus() = default;
-	FStateTreeStateStatus(const FStateTreeHandle InState, const EStateTreeRunStatus InStatus) : State(InState), RunStatus(InStatus) {}
-	FStateTreeStateStatus(const FStateTreeHandle InState) : State(InState), RunStatus(EStateTreeRunStatus::Running) {}
-	FStateTreeStateStatus(const EStateTreeRunStatus InStatus) : State(), RunStatus(InStatus) {}
+	static constexpr uint8 MaxStates = 8;	// Max number of active states
 
-	bool IsSet() const { return RunStatus != EStateTreeRunStatus::Unset; }
+	FStateTreeActiveStates() = default;
 	
-	UPROPERTY(EditDefaultsOnly, Category = Default, BlueprintReadOnly)
-	FStateTreeHandle State = FStateTreeHandle::Invalid;
+	explicit FStateTreeActiveStates(const FStateTreeHandle StateHandle)
+	{
+		Push(StateHandle);
+	}
+	
+	/** Resets the active state array to empty. */
+	void Reset()
+	{
+		NumStates = 0;
+	}
 
-	UPROPERTY(EditDefaultsOnly, Category = Default, BlueprintReadOnly)
-	EStateTreeRunStatus RunStatus = EStateTreeRunStatus::Unset;
+	/** Pushes new state at the back of the array and returns true if there was enough space. */
+	bool Push(const FStateTreeHandle StateHandle)
+	{
+		if ((NumStates + 1) > MaxStates)
+		{
+			return false;
+		}
+		
+		States[NumStates++] = StateHandle;
+		
+		return true;
+	}
+
+	/** Pushes new state at the front of the array and returns true if there was enough space. */
+	bool PushFront(const FStateTreeHandle StateHandle)
+	{
+		if ((NumStates + 1) > MaxStates)
+		{
+			return false;
+		}
+
+		NumStates++;
+		for (int32 Index = (int32)NumStates - 1; Index > 0; Index--)
+		{
+			States[Index] = States[Index - 1];
+		}
+		States[0] = StateHandle;
+		
+		return true;
+	}
+
+	/** Pops a state from the back of the array and returns the popped value, or invalid handle if the array was empty. */
+	FStateTreeHandle Pop()
+	{
+		if (NumStates == 0)
+		{
+			return FStateTreeHandle::Invalid;			
+		}
+		
+		FStateTreeHandle Ret = States[NumStates - 1];
+		NumStates--;
+		return Ret;
+	}
+
+	/** Sets the number of states, new states are set to invalid state. */
+	void SetNum(const int32 NewNum)
+	{
+		check(NewNum >= 0 && NewNum <= MaxStates);
+		if (NewNum > (int32)NumStates)
+		{
+			for (int32 Index = NumStates; Index < NewNum; Index++)
+			{
+				States[Index] = FStateTreeHandle::Invalid;
+			}
+		}
+		NumStates = NewNum;
+	}
+
+	/** Returns true of the array contains specified state. */
+	bool Contains(const FStateTreeHandle StateHandle) const
+	{
+		for (const FStateTreeHandle& Handle : *this)
+		{
+			if (Handle == StateHandle)
+			{
+				return true;
+			}
+		}
+		return false;
+	}
+
+	/** Returns index of a state, searching in reverse order. */
+	int32 IndexOfReverse(const FStateTreeHandle StateHandle) const
+	{
+		for (int32 Index = (int32)NumStates - 1; Index >= 0; Index--)
+		{
+			if (States[Index] == StateHandle)
+				return Index;
+		}
+		return INDEX_NONE;
+	}
+	
+	/** Returns last state in the array, or invalid state if the array is empty. */
+	FStateTreeHandle Last() const { return NumStates > 0 ? States[NumStates - 1] : FStateTreeHandle::Invalid; }
+	
+	/** Returns number of states in the array. */
+	int32 Num() const { return NumStates; }
+
+	/** Returns true if the index is within array bounds. */
+	bool IsValidIndex(const int32 Index) const { return Index >= 0 && Index < (int32)NumStates; }
+	
+	/** Returns true if the array is empty. */
+	bool IsEmpty() const { return NumStates == 0; } 
+
+	/** Returns a specified state in the array. */
+	FORCEINLINE FStateTreeHandle operator[](const int32 Index) const
+	{
+		check(Index >= 0 && Index < (int32)NumStates);
+		return States[Index];
+	}
+
+	/** Returns mutable reference to a specified state in the array. */
+	FORCEINLINE FStateTreeHandle& operator[](const int32 Index)
+	{
+		check(Index >= 0 && Index < (int32)NumStates);
+		return States[Index];
+	}
+
+	/** Returns a specified state in the array, or FStateTreeHandle::Invalid if Index is out of array bounds. */
+	FStateTreeHandle GetStateSafe(const int32 Index) const
+	{
+		return (Index >= 0 && Index < (int32)NumStates) ? States[Index] : FStateTreeHandle::Invalid;
+	}
+
+	/**
+	 * DO NOT USE DIRECTLY
+	 * STL-like iterators to enable range-based for loop support.
+	 */
+	FORCEINLINE FStateTreeHandle* begin() { return &States[0]; }
+	FORCEINLINE FStateTreeHandle* end  () { return &States[0] + Num(); }
+	FORCEINLINE const FStateTreeHandle* begin() const { return &States[0]; }
+	FORCEINLINE const FStateTreeHandle* end  () const { return &States[0] + Num(); }
+
+
+	UPROPERTY(EditDefaultsOnly, Category = Default)
+	FStateTreeHandle States[MaxStates];
+
+	UPROPERTY(EditDefaultsOnly, Category = Default)
+	uint8 NumStates = 0;
 };
 
 /**
@@ -150,24 +282,26 @@ struct STATETREEMODULE_API FStateTreeTransitionResult
 	GENERATED_BODY()
 
 	FStateTreeTransitionResult() = default;
-	FStateTreeTransitionResult(const FStateTreeStateStatus InSource, const FStateTreeHandle InTransitionAndNext) : Source(InSource), Target(InTransitionAndNext), Next(InTransitionAndNext) {}
-	FStateTreeTransitionResult(const FStateTreeStateStatus InSource, const FStateTreeHandle InTransition, const FStateTreeHandle InNext) : Source(InSource), Target(InTransition), Next(InNext) {}
 
-	/** State where the transition started. */
+	/** Current active states, where the transition started. */
+	UPROPERTY(EditDefaultsOnly, Category = Default, BlueprintReadOnly)
+	FStateTreeActiveStates CurrentActiveStates;
+
+	/** Current Run status. */
 	UPROPERTY(EditDefaultsOnly, Category = Default, BlueprintReadOnly)
-	FStateTreeStateStatus Source = FStateTreeStateStatus();
+	EStateTreeRunStatus CurrentRunStatus = EStateTreeRunStatus::Unset;
 
 	/** Transition target state */
 	UPROPERTY(EditDefaultsOnly, Category = Default, BlueprintReadOnly)
-	FStateTreeHandle Target = FStateTreeHandle::Invalid;
+	FStateTreeHandle TargetState = FStateTreeHandle::Invalid;
 
-	/** Selected state, can be different from Transition, if Transition is a selector state. */
+	/** States selected as result of the transition. */
 	UPROPERTY(EditDefaultsOnly, Category = Default, BlueprintReadOnly)
-	FStateTreeHandle Next = FStateTreeHandle::Invalid;
+	FStateTreeActiveStates NextActiveStates;
 
-	/** Current state, update as we execute the tree. */
+	/** The current state being executed. On enter/exit callbacks this is the state of the task or evaluator. */
 	UPROPERTY(EditDefaultsOnly, Category = Default, BlueprintReadOnly)
-	FStateTreeHandle Current = FStateTreeHandle::Invalid;
+	FStateTreeHandle CurrentState = FStateTreeHandle::Invalid;
 };
 
 
@@ -211,17 +345,15 @@ struct STATETREEMODULE_API FBakedStateTreeState
 	FName Name;											// Name of the State
 
 	UPROPERTY()
-	FStateTreeHandle Parent = FStateTreeHandle::Invalid;	// Parent state
+	FStateTreeHandle LinkedState = FStateTreeHandle::Invalid;	// Linked state 
+
+	UPROPERTY()
+	FStateTreeHandle Parent = FStateTreeHandle::Invalid;		// Parent state
 	UPROPERTY()
 	uint16 ChildrenBegin = 0;							// Index to first child state
 	UPROPERTY()
 	uint16 ChildrenEnd = 0;								// Index one past the last child state
 
-	UPROPERTY()
-	FStateTreeHandle StateDoneTransitionState = FStateTreeHandle::Invalid;		// State to transition to when the state execution is done. See also StateDoneTransitionType.
-	UPROPERTY()
-	FStateTreeHandle StateFailedTransitionState = FStateTreeHandle::Invalid;	// State to transition to if the state execution fails. See also StateFailedTransitionType.
-	
 	UPROPERTY()
 	uint16 EnterConditionsBegin = 0;					// Index to first state enter condition
 	UPROPERTY()
@@ -239,11 +371,6 @@ struct STATETREEMODULE_API FBakedStateTreeState
 	uint8 TasksNum = 0;									// Number of tasks
 	UPROPERTY()
 	uint8 EvaluatorsNum = 0;							// Number of evaluators
-
-	UPROPERTY()
-	EStateTreeTransitionType StateDoneTransitionType = EStateTreeTransitionType::NotSet;		// Type of the State Done transition. See also StateDoneTransitionState.
-	UPROPERTY()
-	EStateTreeTransitionType StateFailedTransitionType = EStateTreeTransitionType::NotSet;		// Type of the State Failed transition. See also StateFailedTransitionState.
 };
 
 /** An offset into the StateTree runtime storage type to get a struct view to a specific Task, Evaluator, or Condition. */
diff --git a/Engine/Plugins/Runtime/StateTree/Source/StateTreeTestSuite/Private/StateTreeTest.cpp b/Engine/Plugins/Runtime/StateTree/Source/StateTreeTestSuite/Private/StateTreeTest.cpp
--- a/Engine/Plugins/Runtime/StateTree/Source/StateTreeTestSuite/Private/StateTreeTest.cpp
+++ b/Engine/Plugins/Runtime/StateTree/Source/StateTreeTestSuite/Private/StateTreeTest.cpp
@@ -31,6 +31,7 @@ namespace UE::StateTree::Tests
 		StateTree->SetSchema(Schema);
 		return *StateTree;
 	}
+
 }
 
 struct FStateTreeTest_MakeAndBakeStateTree : FAITestBase
@@ -406,6 +407,64 @@ struct FStateTreeTest_FailEnterState : FAITestBase
 };
 IMPLEMENT_AI_INSTANT_TEST(FStateTreeTest_FailEnterState, "System.StateTree.FailEnterState");
 
+struct FStateTreeTest_SubTree : FAITestBase
+{
+	virtual bool InstantTest() override
+	{
+		UStateTree& StateTree = UE::StateTree::Tests::NewStateTree(&GetWorld());
+		UStateTreeEditorData& EditorData = *Cast<UStateTreeEditorData>(StateTree.EditorData);
+		
+		UStateTreeState& Root = EditorData.AddSubTree(FName(TEXT("Root")));
+		UStateTreeState& State1 = Root.AddChildState(FName(TEXT("State1")), EStateTreeStateType::Linked);
+		UStateTreeState& State2 = Root.AddChildState(FName(TEXT("State2")));
+		UStateTreeState& State3 = Root.AddChildState(FName(TEXT("State3")));
+		UStateTreeState& State3A = State3.AddChildState(FName(TEXT("State3A")));
+		UStateTreeState& State3B = State3.AddChildState(FName(TEXT("State3B")));
+
+		State1.LinkedState.Set(&State3);
+
+		auto& Task2 = State2.AddTask<FTestTask_Stand>(FName(TEXT("Task2")));
+
+		auto& Task3A = State3A.AddTask<FTestTask_Stand>(FName(TEXT("Task3A")));
+		State3A.AddTransition(EStateTreeTransitionEvent::OnCompleted, EStateTreeTransitionType::GotoState, &State3B);
+
+		auto& Task3B = State3B.AddTask<FTestTask_Stand>(FName(TEXT("Task3B")));
+
+
+		FStateTreeCompilerLog Log;
+		FStateTreeBaker Baker(Log);
+		const bool bResult = Baker.Bake(StateTree);
+		AITEST_TRUE("StateTree should get baked", bResult);
+
+		EStateTreeRunStatus Status = EStateTreeRunStatus::Unset;
+		FTestStateTreeExecutionContext Exec;
+		const bool bInitSucceeded = Exec.Init(StateTree, StateTree, EStateTreeStorage::Internal);
+		AITEST_TRUE("StateTree should init", bInitSucceeded);
+
+		const FString TickStr(TEXT("Tick"));
+		const FString EnterStateStr(TEXT("EnterState"));
+		const FString ExitStateStr(TEXT("ExitState"));
+		const FString StateCompletedStr(TEXT("StateCompleted"));
+
+		// Start and enter state
+		Status = Exec.Start();
+
+		AITEST_TRUE("StateTree Active States should be in Root/State1/State3/State3A", Exec.ExpectInActiveStates(Root.Name, State1.Name, State3.Name, State3A.Name));
+		AITEST_TRUE("StateTree Task2 should enter state", !Exec.Expect(Task2.GetName(), EnterStateStr));
+		AITEST_TRUE("StateTree Task3A should enter state", Exec.Expect(Task3A.GetName(), EnterStateStr));
+		AITEST_TRUE("StateTree should be running", Status == EStateTreeRunStatus::Running);
+		Exec.LogClear();
+
+		Status = Exec.Tick(0.1f);
+		AITEST_TRUE("StateTree Active States should be in Root/State1/State3/State3B", Exec.ExpectInActiveStates(Root.Name, State1.Name, State3.Name, State3B.Name));
+		AITEST_TRUE("StateTree Task3B should enter state", Exec.Expect(Task3B.GetName(), EnterStateStr));
+		AITEST_TRUE("StateTree should be running", Status == EStateTreeRunStatus::Running);
+		
+		return true;
+	}
+};
+IMPLEMENT_AI_INSTANT_TEST(FStateTreeTest_SubTree, "System.StateTree.SubTree");
+
 PRAGMA_ENABLE_OPTIMIZATION
 
 #undef LOCTEXT_NAMESPACE
diff --git a/Engine/Plugins/Runtime/StateTree/Source/StateTreeTestSuite/Private/StateTreeTestTypes.h b/Engine/Plugins/Runtime/StateTree/Source/StateTreeTestSuite/Private/StateTreeTestTypes.h
--- a/Engine/Plugins/Runtime/StateTree/Source/StateTreeTestSuite/Private/StateTreeTestTypes.h
+++ b/Engine/Plugins/Runtime/StateTree/Source/StateTreeTestSuite/Private/StateTreeTestTypes.h
@@ -62,6 +62,30 @@ struct FTestStateTreeExecutionContext : public FStateTreeExecutionContext
 	{
 		return FLogOrder(*this, 0).Then(Name, Message);
 	}
+
+	template <class ...Args>
+	bool ExpectInActiveStates(const Args&... States)
+	{
+		FName ExpectedStateNames[] = { States... };
+		const int32 NumExpectedStateNames = sizeof...(States);
+
+		TArray<FName> ActiveStateNames = GetActiveStateNames();
+
+		if (ActiveStateNames.Num() != NumExpectedStateNames)
+		{
+			return false;
+		}
+
+		for (int32 Index = 0; Index != NumExpectedStateNames; Index++)
+		{
+			if (ExpectedStateNames[Index] != ActiveStateNames[Index])
+			{
+				return false;
+			}
+		}
+
+		return true;
+	}
 	
 };
 
@@ -499,7 +523,7 @@ struct FTestTask_ReserveSmartObject : public FStateTreeTaskBase
 		}
 	}
 
-	virtual void StateCompleted(FStateTreeExecutionContext& Context, const EStateTreeRunStatus CompletionStatus, const FStateTreeHandle StateCompleted) const override
+	virtual void StateCompleted(FStateTreeExecutionContext& Context, const EStateTreeRunStatus CompletionStatus, const FStateTreeActiveStates& CompletedActiveStates) const override
 	{
 		const FTestPotentialSmartObjectsResult& PotentialSmartObjects = Context.GetInstanceData(PotentialSmartObjectsHandle);
 		FTestMoveLocationResult& ReservedSmartObjectLocation = Context.GetInstanceData(ReservedSmartObjectLocationHandle);
@@ -680,7 +704,7 @@ struct FTestTask_Stand : public FStateTreeTaskBase
 		TestContext.Log(Name, TEXT("ExitState"));
 	}
 
-	virtual void StateCompleted(FStateTreeExecutionContext& Context, const EStateTreeRunStatus CompletionStatus, const FStateTreeHandle CompletedState) const override
+	virtual void StateCompleted(FStateTreeExecutionContext& Context, const EStateTreeRunStatus CompletionStatus, const FStateTreeActiveStates& CompletedActiveStates) const override
 	{
 		FTestStateTreeExecutionContext& TestContext = static_cast<FTestStateTreeExecutionContext&>(Context);
 		TestContext.Log(Name, TEXT("StateCompleted"));
