StateTree: Renamed baker to compiler

diff --git a/Engine/Plugins/Runtime/StateTree/Config/BaseStateTree.ini b/Engine/Plugins/Runtime/StateTree/Config/BaseStateTree.ini
--- a/Engine/Plugins/Runtime/StateTree/Config/BaseStateTree.ini
+++ b/Engine/Plugins/Runtime/StateTree/Config/BaseStateTree.ini
@@ -12,6 +12,9 @@
 +StructRedirects=(OldName="StateTreeEvaluatorItem",NewName="StateTreeEditorNode")
 +StructRedirects=(OldName="StateTreeTaskItem",NewName="StateTreeEditorNode")
 
++StructRedirects=(OldName="BakedStateTransition",NewName="CompactStateTransition")
++StructRedirects=(OldName="BakedStateTreeState",NewName="CompactStateTreeState")
+
 +PropertyRedirects=(OldName="StateTreeState.Transitions2",NewName="StateTreeState.Transitions")
 +PropertyRedirects=(OldName="StateTreeState.EnterConditions2",NewName="StateTreeState.EnterConditions")
 +PropertyRedirects=(OldName="StateTreeState.Evaluators2",NewName="StateTreeState.Evaluators")
diff --git a/Engine/Plugins/Runtime/StateTree/Source/StateTreeEditorModule/Private/StateTreeBaker.cpp b/Engine/Plugins/Runtime/StateTree/Source/StateTreeEditorModule/Private/StateTreeBaker.cpp
deleted file mode 100644
diff --git a/Engine/Plugins/Runtime/StateTree/Source/StateTreeEditorModule/Private/StateTreeCompiler.cpp b/Engine/Plugins/Runtime/StateTree/Source/StateTreeEditorModule/Private/StateTreeCompiler.cpp
new file mode 100644
--- /dev/null
+++ b/Engine/Plugins/Runtime/StateTree/Source/StateTreeEditorModule/Private/StateTreeCompiler.cpp
@@ -0,0 +1,824 @@
+// Copyright Epic Games, Inc. All Rights Reserved.
+
+#include "StateTreeCompiler.h"
+#include "StateTree.h"
+#include "StateTreeEditorData.h"
+#include "StateTreeTypes.h"
+#include "Conditions/StateTreeCondition_Common.h"
+#include "StateTreeEvaluatorBase.h"
+#include "StateTreeTaskBase.h"
+#include "StateTreeConditionBase.h"
+#include "StateTreeState.h"
+#include "StateTreeExecutionContext.h"
+#include "CoreMinimal.h"
+#include "StateTreePropertyBindingCompiler.h"
+
+
+bool FStateTreeCompiler::Compile(UStateTree& InStateTree)
+{
+	StateTree = &InStateTree;
+	TreeData = Cast<UStateTreeEditorData>(StateTree->EditorData);
+	if (!TreeData)
+	{
+		return false;
+	}
+
+	// Cleanup existing state
+	StateTree->ResetCompiled();
+
+	BindingsCompiler.Init(StateTree->PropertyBindings, Log);
+
+	if (!CreateStates())
+	{
+		StateTree->ResetCompiled();
+		return false;
+	}
+
+	if (!CreateExecutionInfos())
+	{
+		StateTree->ResetCompiled();
+		return false;
+	}
+	
+	if (!CreateStateEvaluators())
+	{
+		StateTree->ResetCompiled();
+		return false;
+	}
+
+	if (!CreateStateTasks())
+	{
+		StateTree->ResetCompiled();
+		return false;
+	}
+
+	if (!CreateStateTransitions())
+	{
+		StateTree->ResetCompiled();
+		return false;
+	}
+
+	BindingsCompiler.Finalize();
+
+	StateTree->PropertyBindings.ResolvePaths();
+	StateTree->Link();
+
+	return true;
+}
+
+FStateTreeHandle FStateTreeCompiler::GetStateHandle(const FGuid& StateID) const
+{
+	const int32* Idx = IDToState.Find(StateID);
+	if (Idx == nullptr)
+	{
+		return FStateTreeHandle::Invalid;
+	}
+
+	return FStateTreeHandle(uint16(*Idx));
+}
+
+UStateTreeState* FStateTreeCompiler::GetState(const FGuid& StateID)
+{
+	const int32* Idx = IDToState.Find(StateID);
+	if (Idx == nullptr)
+	{
+		return nullptr;
+	}
+
+	return SourceStates[*Idx];
+}
+
+bool FStateTreeCompiler::CreateStates()
+{
+	// Create item for the runtime execution state
+	StateTree->Instances.Add(FInstancedStruct::Make<FStateTreeExecutionState>());
+
+	for (UStateTreeState* SubTree : TreeData->SubTrees)
+	{
+		if (SubTree != nullptr)
+		{
+			if (!CreateStateRecursive(*SubTree, FStateTreeHandle::Invalid))
+			{
+				return false;
+			}
+		}
+	}
+	
+	return true;
+}
+
+bool FStateTreeCompiler::CreateStateRecursive(UStateTreeState& State, const FStateTreeHandle Parent)
+{
+	FStateTreeCompilerLogStateScope LogStateScope(&State, Log);
+
+	const int32 StateIdx = StateTree->States.AddDefaulted();
+	FCompactStateTreeState& BakedState = StateTree->States[StateIdx];
+	BakedState.Name = State.Name;
+	BakedState.Parent = Parent;
+
+	SourceStates.Add(&State);
+	IDToState.Add(State.ID, StateIdx);
+
+	check(StateTree->Instances.Num() <= int32(MAX_uint16));
+
+	// Child states
+	check(StateTree->States.Num() <= int32(MAX_uint16));
+	BakedState.ChildrenBegin = uint16(StateTree->States.Num());
+	for (UStateTreeState* Child : State.Children)
+	{
+		if (Child)
+		{
+			if (!CreateStateRecursive(*Child, FStateTreeHandle((uint16)StateIdx)))
+			{
+				return false;
+			}
+		}
+	}
+	check(StateTree->States.Num() <= int32(MAX_uint16));
+	StateTree->States[StateIdx].ChildrenEnd = uint16(StateTree->States.Num()); // Cannot use BakedState here, it may be invalid due to array resize.
+
+	return true;
+}
+
+FString FStateTreeCompiler::GetExecutionPathString(const TConstArrayView<const UStateTreeState*> Path)
+{
+	FString PathStr;
+	const UStateTreeState* PrevState = nullptr;
+	for (int32 Index = 0; Index < Path.Num(); Index++)
+	{
+		const UStateTreeState* State = Path[Index];
+		check(State);
+		if (PrevState != nullptr)
+		{
+			if (PrevState != State->Parent)
+			{
+				// Linked
+				PathStr += TEXT(">");
+			}
+			else
+			{
+				PathStr += TEXT("/");
+			}
+		}
+		PathStr += State->Name.ToString();
+		PrevState = State;
+	}
+
+	return PathStr;
+}
+
+bool FStateTreeCompiler::IsPathLinked(const TConstArrayView<const UStateTreeState*> Path)
+{
+	bool bIsLinked = false; 
+	const UStateTreeState* PrevState = nullptr;
+	for (int32 Index = 0; Index < Path.Num(); Index++)
+	{
+		const UStateTreeState* State = Path[Index];
+		check(State);
+		// If the previous state is not the parent state in tree, then the connection must have come from a linked state (a state linking to an arbitrary state in the tree).
+		if (PrevState && PrevState != State->Parent)
+		{
+			bIsLinked = true;
+			break;
+		}
+		PrevState = State;
+	}
+	return bIsLinked;
+}
+
+bool FStateTreeCompiler::CreateExecutionInfos()
+{
+	for (UStateTreeState* SubTree : TreeData->SubTrees)
+	{
+		if (SubTree != nullptr)
+		{
+			TArray<const UStateTreeState*> Path;
+			if (!CreateExecutionInfosRecursive(*SubTree, Path))
+			{
+				return false;
+			}
+		}
+	}
+
+	return true;
+}
+
+bool FStateTreeCompiler::CreateExecutionInfosRecursive(UStateTreeState& State, TArray<const UStateTreeState*>& Path)
+{
+	Path.Add(&State);
+	
+	if (Path.Num() > FStateTreeActiveStates::MaxStates)
+	{
+		Log.Reportf(EMessageSeverity::Error,
+		TEXT("Reached maximum execution depth %d at: '%s'."), FStateTreeActiveStates::MaxStates, *GetExecutionPathString(Path));
+		return false;
+	}
+
+	FStateExecutionInfo& ExecInfo = ExecutionInfos.FindOrAdd(&State);
+	ExecInfo.ExecutionPaths.Emplace(Path);
+
+	if (State.Type == EStateTreeStateType::Linked)
+	{
+		if (UStateTreeState* LinkedState = GetState(State.LinkedState.ID))
+		{
+			if (!CreateExecutionInfosRecursive(*LinkedState, Path))
+			{
+				return false;
+			}
+		}
+	}
+	
+	for (UStateTreeState* Child : State.Children)
+	{
+		if (Child)
+		{
+			if (!CreateExecutionInfosRecursive(*Child, Path))
+			{
+				return false;
+			}
+		}
+	}
+
+	Path.Pop();
+
+	return true;
+}
+
+bool FStateTreeCompiler::CreateConditions(UStateTreeState& State, TConstArrayView<FStateTreeEditorNode> Conditions)
+{
+	for (int32 Index = 0; Index < Conditions.Num(); Index++)
+	{
+		const bool bIsFirst = Index == 0;
+		const FStateTreeEditorNode& CondNode = Conditions[Index];
+		// First operand should be copy as we dont have a previous item to operate on.
+		const EStateTreeConditionOperand Operand = bIsFirst ? EStateTreeConditionOperand::Copy : CondNode.ConditionOperand;
+		// First indent must be 0 to make the parentheses calculation match.
+		const int32 CurrIndent = bIsFirst ? 0 : CondNode.ConditionIndent;
+		// Next indent, or terminate at zero.
+		const int32 NextIndent = Conditions.IsValidIndex(Index + 1) ? Conditions[Index].ConditionIndent : 0;
+		const int32 DeltaIndent = NextIndent - CurrIndent;
+		
+		check(DeltaIndent >= MIN_int8 && DeltaIndent <= MAX_int8);
+
+		if (!CreateCondition(State, CondNode, Operand, (int8)DeltaIndent))
+		{
+			return false;
+		}
+	}
+
+	return true;
+}
+
+bool FStateTreeCompiler::CreateStateTasks()
+{
+	for (int32 i = 0; i < StateTree->States.Num(); i++)
+	{
+		FCompactStateTreeState& BakedState = StateTree->States[i];
+		UStateTreeState* SourceState = SourceStates[i];
+		check(SourceState != nullptr);
+
+		FStateTreeCompilerLogStateScope LogStateScope(SourceState, Log);
+		
+		// Create tasks
+		check(StateTree->Nodes.Num() <= int32(MAX_uint16));
+		BakedState.TasksBegin = uint16(StateTree->Nodes.Num());
+
+		for (FStateTreeEditorNode& TaskNode : SourceState->Tasks)
+		{
+			if (!CreateTask(*SourceState, TaskNode))
+			{
+				return false;
+			}
+		}
+
+		if (!CreateTask(*SourceState, SourceState->SingleTask))
+		{
+			return false;
+		}
+	
+		const int32 TasksNum = StateTree->Nodes.Num() - int32(BakedState.TasksBegin);
+		check(TasksNum <= int32(MAX_uint8));
+		BakedState.TasksNum = uint8(TasksNum);
+	}
+	
+	return true;
+}
+
+bool FStateTreeCompiler::CreateStateEvaluators()
+{
+	for (int32 i = 0; i < StateTree->States.Num(); i++)
+	{
+		FCompactStateTreeState& BakedState = StateTree->States[i];
+		UStateTreeState* SourceState = SourceStates[i];
+		check(SourceState != nullptr);
+
+		FStateTreeCompilerLogStateScope LogStateScope(SourceState, Log);
+		
+		// Collect evaluators
+		check(StateTree->Nodes.Num() <= int32(MAX_uint16));
+		BakedState.EvaluatorsBegin = uint16(StateTree->Nodes.Num());
+
+		for (FStateTreeEditorNode& EvalNode : SourceState->Evaluators)
+		{
+			if (!CreateEvaluator(*SourceState, EvalNode))
+			{
+				return false;
+			}
+		}
+		
+		const int32 EvaluatorsNum = StateTree->Nodes.Num() - int32(BakedState.EvaluatorsBegin);
+		check(EvaluatorsNum <= int32(MAX_uint8));
+		BakedState.EvaluatorsNum = uint8(EvaluatorsNum);
+	}
+
+	return true;
+}
+
+bool FStateTreeCompiler::CreateStateTransitions()
+{
+	for (int32 i = 0; i < StateTree->States.Num(); i++)
+	{
+		FCompactStateTreeState& BakedState = StateTree->States[i];
+		UStateTreeState* SourceState = SourceStates[i];
+		check(SourceState != nullptr);
+
+		FStateTreeCompilerLogStateScope LogStateScope(SourceState, Log);
+		
+		// Enter conditions.
+		BakedState.EnterConditionsBegin = uint16(StateTree->Nodes.Num());
+		if (!CreateConditions(*SourceState, SourceState->EnterConditions))
+		{
+			Log.Reportf(EMessageSeverity::Error,
+				TEXT("Failed to create state enter condition."));
+			return false;
+		}
+		BakedState.EnterConditionsNum = uint8(uint16(StateTree->Nodes.Num()) - BakedState.EnterConditionsBegin);
+
+		// Linked state
+		if (SourceState->Type == EStateTreeStateType::Linked)
+		{
+			// Make sure the linked state is not self or parent to this state.
+			const UStateTreeState* LinkedParentState = nullptr;
+			for (const UStateTreeState* State = SourceState; State != nullptr; State = State->Parent)
+			{
+				if (State->ID == SourceState->LinkedState.ID)
+				{
+					LinkedParentState = State;
+					break;
+				}
+			}
+			
+			if (LinkedParentState != nullptr)
+			{
+				Log.Reportf(EMessageSeverity::Error,
+					TEXT("State is linked to it's parent state '%s', which will create infinite loop."),
+					*LinkedParentState->Name.ToString());
+				return false;
+			}
+			
+			BakedState.LinkedState = GetStateHandle(SourceState->LinkedState.ID);
+			
+			if (!BakedState.LinkedState.IsValid())
+			{
+				Log.Reportf(EMessageSeverity::Error,
+					TEXT("Failed to resolve linked state '%s'."),
+					*SourceState->LinkedState.Name.ToString());
+				return false;
+			}
+		}
+		
+		// Transitions
+		BakedState.TransitionsBegin = uint16(StateTree->Transitions.Num());
+		for (FStateTreeTransition& Transition : SourceState->Transitions)
+		{
+			FCompactStateTransition& BakedTransition = StateTree->Transitions.AddDefaulted_GetRef();
+			BakedTransition.Event = Transition.Event;
+			BakedTransition.Type = Transition.State.Type;
+			BakedTransition.GateDelay = (uint8)FMath::Clamp(FMath::CeilToInt(Transition.GateDelay * 10.0f), 0, 255);
+			BakedTransition.State = FStateTreeHandle::Invalid;
+			if (!ResolveTransitionState(*SourceState, Transition.State, BakedTransition.State))
+			{
+				return false;
+			}
+			// Note: Unset transition is allowed here. It can be used to mask a transition at parent.
+
+			BakedTransition.ConditionsBegin = uint16(StateTree->Nodes.Num());
+			if (!CreateConditions(*SourceState, Transition.Conditions))
+			{
+				Log.Reportf(EMessageSeverity::Error,
+					TEXT("Failed to create condition for transition to '%s'."),
+					*Transition.State.Name.ToString());
+				return false;
+			}
+			BakedTransition.ConditionsNum = uint8(uint16(StateTree->Nodes.Num()) - BakedTransition.ConditionsBegin);
+		}
+		BakedState.TransitionsNum = uint8(uint16(StateTree->Transitions.Num()) - BakedState.TransitionsBegin);
+	}
+
+	// @todo: Add test to check that all success/failure transition is possible (see editor).
+	
+	return true;
+}
+
+bool FStateTreeCompiler::ResolveTransitionState(const UStateTreeState& SourceState, const FStateTreeStateLink& Link, FStateTreeHandle& OutTransitionHandle) const 
+{
+	if (Link.Type == EStateTreeTransitionType::GotoState)
+	{
+		OutTransitionHandle = GetStateHandle(Link.ID);
+		if (!OutTransitionHandle.IsValid())
+		{
+			Log.Reportf(EMessageSeverity::Error,
+				TEXT("Failed to resolve transition to state '%s'."),
+				*Link.Name.ToString());
+			return false;
+		}
+	}
+	else if (Link.Type == EStateTreeTransitionType::NextState)
+	{
+		// Find next state.
+		const UStateTreeState* NextState = SourceState.GetNextSiblingState();
+		if (NextState == nullptr)
+		{
+			Log.Reportf(EMessageSeverity::Error,
+				TEXT("Failed to resolve transition, there's no next state."));
+			return false;
+		}
+		OutTransitionHandle = GetStateHandle(NextState->ID);
+		if (!OutTransitionHandle.IsValid())
+		{
+			Log.Reportf(EMessageSeverity::Error,
+				TEXT("Failed to resolve transition next state, no handle found for '%s'."),
+				*NextState->Name.ToString());
+			return false;
+		}
+	}
+	
+	return true;
+}
+
+bool FStateTreeCompiler::CreateCondition(UStateTreeState& State, const FStateTreeEditorNode& CondNode, const EStateTreeConditionOperand Operand, const int8 DeltaIndent)
+{
+	if (!CondNode.Node.IsValid())
+	{
+		// Empty line in conditions array, just silently ignore.
+		return true;
+	}
+
+	FStateTreeBindableStructDesc StructDesc;
+	StructDesc.ID = CondNode.ID;
+	StructDesc.Name = CondNode.Node.GetScriptStruct()->GetFName();
+
+	// Check that item has valid instance initialized.
+	if (!CondNode.Instance.IsValid() && CondNode.InstanceObject == nullptr)
+	{
+		Log.Reportf(EMessageSeverity::Error, StructDesc,
+			TEXT("Malformed condition, missing instance value."));
+		return false;
+	}
+
+	// Copy the condition
+	FInstancedStruct& Item = StateTree->Nodes.AddDefaulted_GetRef();
+	Item = CondNode.Node;
+
+	FStateTreeConditionBase& Cond = Item.GetMutable<FStateTreeConditionBase>();
+
+	Cond.Operand = Operand;
+	Cond.DeltaIndent = DeltaIndent;
+	
+	if (CondNode.Instance.IsValid())
+	{
+		// Struct instance
+		FInstancedStruct& Instance = StateTree->Instances.AddDefaulted_GetRef();
+		const int32 InstanceIndex = StateTree->Instances.Num() - 1;
+		
+		Instance = CondNode.Instance;
+	
+		// Create binding source struct descriptor.
+		StructDesc.Struct = Instance.GetScriptStruct();
+		StructDesc.Name = Cond.Name;
+
+		check(InstanceIndex <= int32(MAX_uint16));
+		Cond.InstanceIndex = uint16(InstanceIndex);
+		Cond.bInstanceIsObject = false;
+	}
+	else
+	{
+		// Object Instance
+		check(CondNode.InstanceObject != nullptr);
+		
+		UObject* Instance = DuplicateObject(CondNode.InstanceObject, StateTree);
+		
+		StateTree->InstanceObjects.Add(Instance);
+		const int32 InstanceIndex = StateTree->InstanceObjects.Num() - 1;
+		
+		// Create binding source struct descriptor.
+		StructDesc.Struct = Instance->GetClass();
+		StructDesc.Name = Cond.Name;
+
+		check(InstanceIndex <= int32(MAX_uint16));
+		Cond.InstanceIndex = uint16(InstanceIndex);
+		Cond.bInstanceIsObject = true;
+	}
+
+	// Mark the struct as binding source.
+	const int32 SourceStructIndex = BindingsCompiler.AddSourceStruct(StructDesc);
+
+	// Check that the bindings for this struct are still all valid.
+	TArray<FStateTreeEditorPropertyBinding> Bindings;
+	if (!GetAndValidateBindings(State, StructDesc, Bindings))
+	{
+		return false;
+	}
+
+	// Compile batch copy for this struct, we pass in all the bindings, the compiler will pick up the ones for the target structs.
+	int32 BatchIndex = INDEX_NONE;
+	if (!BindingsCompiler.CompileBatch(StructDesc, Bindings, BatchIndex))
+	{
+		return false;
+	}
+	check(BatchIndex < int32(MAX_uint16));
+	Cond.BindingsBatch = BatchIndex == INDEX_NONE ? FStateTreeHandle::Invalid : FStateTreeHandle(uint16(BatchIndex));
+
+	check(SourceStructIndex <= int32(MAX_uint16));
+	Cond.DataViewIndex = uint16(SourceStructIndex);
+	
+	return true;
+}
+
+bool FStateTreeCompiler::CreateTask(UStateTreeState& State, const FStateTreeEditorNode& TaskNode)
+{
+	// Silently ignore empty items.
+	if (!TaskNode.Node.IsValid())
+	{
+		return true;
+	}
+	
+	// Create binding source struct descriptor.
+	FStateTreeBindableStructDesc StructDesc;
+	StructDesc.ID = TaskNode.ID;
+	StructDesc.Name = TaskNode.Node.GetScriptStruct()->GetFName();
+
+	// Check that item has valid instance initialized.
+	if (!TaskNode.Instance.IsValid() && TaskNode.InstanceObject == nullptr)
+	{
+		Log.Reportf(EMessageSeverity::Error, StructDesc,
+			TEXT("Malformed task, missing instance value."));
+		return false;
+	}
+
+	// Copy the task
+	FInstancedStruct& Item = StateTree->Nodes.AddDefaulted_GetRef();
+	Item = TaskNode.Node;
+	
+	FStateTreeTaskBase& Task = Item.GetMutable<FStateTreeTaskBase>();
+
+	if (TaskNode.Instance.IsValid())
+	{
+		// Struct Instance
+		FInstancedStruct& Instance = StateTree->Instances.AddDefaulted_GetRef();
+		const int32 InstanceIndex = StateTree->Instances.Num() - 1;
+
+		Instance = TaskNode.Instance;
+
+		// Create binding source struct descriptor.
+		StructDesc.Struct = Instance.GetScriptStruct();
+		StructDesc.Name = Task.Name;
+
+		check(InstanceIndex <= int32(MAX_uint16));
+		Task.InstanceIndex = uint16(InstanceIndex);
+		Task.bInstanceIsObject = false;
+	}
+	else
+	{
+		// Object Instance
+		check(TaskNode.InstanceObject != nullptr);
+
+		UObject* Instance = DuplicateObject(TaskNode.InstanceObject, StateTree);
+
+		StateTree->InstanceObjects.Add(Instance);
+		const int32 InstanceIndex = StateTree->InstanceObjects.Num() - 1;
+
+		// Create binding source struct descriptor.
+		StructDesc.Struct = Instance->GetClass();
+		StructDesc.Name = Task.Name;
+
+		check(InstanceIndex <= int32(MAX_uint16));
+		Task.InstanceIndex = uint16(InstanceIndex);
+		Task.bInstanceIsObject = true;
+	}
+
+	// Mark the instance as binding source.
+	const int32 SourceStructIndex = BindingsCompiler.AddSourceStruct(StructDesc);
+
+	// Check that the bindings for this struct are still all valid.
+	TArray<FStateTreeEditorPropertyBinding> Bindings;
+	if (!GetAndValidateBindings(State, StructDesc, Bindings))
+	{
+		return false;
+	}
+
+	// Compile batch copy for this struct, we pass in all the bindings, the compiler will pick up the ones for the target structs.
+	int32 BatchIndex = INDEX_NONE;
+	if (!BindingsCompiler.CompileBatch(StructDesc, Bindings, BatchIndex))
+	{
+		return false;
+	}
+	
+	check(BatchIndex < int32(MAX_uint16));
+	Task.BindingsBatch = BatchIndex == INDEX_NONE ? FStateTreeHandle::Invalid : FStateTreeHandle(uint16(BatchIndex));
+	
+	check(SourceStructIndex <= int32(MAX_uint16));
+	Task.DataViewIndex = uint16(SourceStructIndex);
+
+	return true;
+}
+
+bool FStateTreeCompiler::CreateEvaluator(UStateTreeState& State, const FStateTreeEditorNode& EvalNode)
+{
+	// Silently ignore empty items.
+	if (!EvalNode.Node.IsValid())
+	{
+		return true;
+	}
+	
+	// Create binding source struct descriptor.
+	FStateTreeBindableStructDesc StructDesc;
+    StructDesc.ID = EvalNode.ID;
+    StructDesc.Name = EvalNode.Node.GetScriptStruct()->GetFName();
+
+    // Check that item has valid instance initialized.
+    if (!EvalNode.Instance.IsValid() && EvalNode.InstanceObject == nullptr)
+    {
+        Log.Reportf(EMessageSeverity::Error, StructDesc,
+        	TEXT("Malformed evaluator, missing instance value."));
+        return false;
+    }
+
+	// Copy the evaluator
+	FInstancedStruct& Item = StateTree->Nodes.AddDefaulted_GetRef();
+	Item = EvalNode.Node;
+
+	FStateTreeEvaluatorBase& Eval = Item.GetMutable<FStateTreeEvaluatorBase>();
+
+	if (EvalNode.Instance.IsValid())
+	{
+		// Struct Instance
+		FInstancedStruct& Instance = StateTree->Instances.AddDefaulted_GetRef();
+		const int32 InstanceIndex = StateTree->Instances.Num() - 1;
+
+		Instance = EvalNode.Instance;
+
+		// Create binding source struct descriptor.
+		StructDesc.Struct = Instance.GetScriptStruct();
+		StructDesc.Name = Eval.Name;
+
+		check(InstanceIndex <= int32(MAX_uint16));
+		Eval.InstanceIndex = uint16(InstanceIndex);
+		Eval.bInstanceIsObject = false;
+	}
+	else
+	{
+		// Object Instance
+		check(EvalNode.InstanceObject != nullptr);
+
+		UObject* Instance = DuplicateObject(EvalNode.InstanceObject, StateTree);
+
+		StateTree->InstanceObjects.Add(Instance);
+		const int32 InstanceIndex = StateTree->InstanceObjects.Num() - 1;
+
+		// Create binding source struct descriptor.
+		StructDesc.Struct = Instance->GetClass();
+		StructDesc.Name = Eval.Name;
+
+		check(InstanceIndex <= int32(MAX_uint16));
+		Eval.InstanceIndex = uint16(InstanceIndex);
+		Eval.bInstanceIsObject = true;
+	}
+		
+	// Mark the instance as binding source.
+	const int32 SourceStructIndex = BindingsCompiler.AddSourceStruct(StructDesc);
+
+	// Check that the bindings for this struct are still all valid.
+	TArray<FStateTreeEditorPropertyBinding> Bindings;
+	if (!GetAndValidateBindings(State, StructDesc, Bindings))
+	{
+		return false;
+	}
+
+	// Compile batch copy for this struct, we pass in all the bindings, the compiler will pick up the ones for the target structs.
+	int32 BatchIndex = INDEX_NONE;
+	if (!BindingsCompiler.CompileBatch(StructDesc, Bindings, BatchIndex))
+	{
+		return false;
+	}
+	
+	check(BatchIndex < int32(MAX_uint16));
+	Eval.BindingsBatch = BatchIndex == INDEX_NONE ? FStateTreeHandle::Invalid : FStateTreeHandle(uint16(BatchIndex));
+	
+	check(SourceStructIndex <= int32(MAX_uint16));
+	Eval.DataViewIndex = uint16(SourceStructIndex);
+
+	return true;
+}
+
+bool FStateTreeCompiler::IsPropertyAnyEnum(const FStateTreeBindableStructDesc& Struct, FStateTreeEditorPropertyPath Path) const
+{
+	bool bIsAnyEnum = false;
+	TArray<FStateTreePropertySegment> Segments;
+	const FProperty* LeafProperty = nullptr;
+	int32 LeafArrayIndex = INDEX_NONE;
+	const bool bResolved = FStateTreePropertyBindingCompiler::ResolvePropertyPath(Struct, Path, Segments, LeafProperty, LeafArrayIndex);
+	if (bResolved && LeafProperty)
+	{
+		if (const FProperty* OwnerProperty = LeafProperty->GetOwnerProperty())
+		{
+			if (const FStructProperty* OwnerStructProperty = CastField<FStructProperty>(OwnerProperty))
+			{
+				bIsAnyEnum = OwnerStructProperty->Struct == FStateTreeAnyEnum::StaticStruct();
+			}
+		}
+	}
+	return bIsAnyEnum;
+}
+
+bool FStateTreeCompiler::GetAndValidateBindings(UStateTreeState& State, const FStateTreeBindableStructDesc& TargetStruct, TArray<FStateTreeEditorPropertyBinding>& OutBindings) const
+{
+	OutBindings.Reset();
+	
+	for (const FStateTreeEditorPropertyBinding& Binding : TreeData->EditorBindings.GetBindings())
+	{
+		if (Binding.TargetPath.StructID != TargetStruct.ID)
+		{
+			continue;
+		}
+
+		// Source must be one of the source structs we have discovered in the tree.
+		const FGuid SourceStructID = Binding.SourcePath.StructID;
+		const int32 SourceStructIdx = BindingsCompiler.GetSourceStructIndexByID(SourceStructID);
+		if (SourceStructIdx == INDEX_NONE)
+		{
+			Log.Reportf(EMessageSeverity::Error, TargetStruct,
+						TEXT("Failed to find binding source '%s:%s'."),
+						*TargetStruct.Name.ToString(), *Binding.TargetPath.ToString());
+			return false;
+		}
+		const FStateTreeBindableStructDesc& SourceStruct = BindingsCompiler.GetSourceStructDesc(SourceStructIdx);
+
+		// Source must be accessible by the target struct via all execution paths.
+		TArray<FStateTreeBindableStructDesc> AccessibleStructs;
+		const FStateExecutionInfo& ExecInfo = ExecutionInfos.FindChecked(&State);
+		for (const FExecutionPath& ExecPath : ExecInfo.ExecutionPaths)
+		{
+			AccessibleStructs.Reset();
+			TreeData->GetAccessibleStructs(ExecPath.Path, Binding.TargetPath.StructID, AccessibleStructs);
+
+			const bool SourceAccessible = AccessibleStructs.ContainsByPredicate([SourceStructID](const FStateTreeBindableStructDesc& Structs)
+				{
+					return (Structs.ID == SourceStructID);
+				});
+
+			if (!SourceAccessible)
+			{
+				Log.Reportf(EMessageSeverity::Error, TargetStruct,
+							TEXT("Property '%s:%s' cannot be bound to '%s:%s', because the binding source '%s' is not updated before '%s' in the tree."),
+							*SourceStruct.Name.ToString(), *Binding.SourcePath.ToString(),
+							*TargetStruct.Name.ToString(), *Binding.TargetPath.ToString(),
+							*SourceStruct.Name.ToString(), *TargetStruct.Name.ToString());
+
+				if (IsPathLinked(ExecPath.Path))
+				{
+					Log.Reportf(EMessageSeverity::Error, TargetStruct,
+								TEXT("The binding source is not updated when executing via linked state: %s."),
+								*GetExecutionPathString(ExecPath.Path));
+				}
+				
+				return false;
+			}
+		}
+		
+		// Special case fo AnyEnum. StateTreeBindingExtension allows AnyEnums to bind to other enum types.
+		// The actual copy will be done via potential type promotion copy, into the value property inside the AnyEnum.
+		// We amend the paths here to point to the 'Value' property.
+		const bool bSourceIsAnyEnum = IsPropertyAnyEnum(SourceStruct, Binding.SourcePath);
+		const bool bTargetIsAnyEnum = IsPropertyAnyEnum(TargetStruct, Binding.TargetPath);
+		if (bSourceIsAnyEnum || bTargetIsAnyEnum)
+		{
+			FStateTreeEditorPropertyBinding ModifiedBinding(Binding);
+			if (bSourceIsAnyEnum)
+			{
+				ModifiedBinding.SourcePath.Path.Add(GET_MEMBER_NAME_STRING_CHECKED(FStateTreeAnyEnum, Value));
+			}
+			if (bTargetIsAnyEnum)
+			{
+				ModifiedBinding.TargetPath.Path.Add(GET_MEMBER_NAME_STRING_CHECKED(FStateTreeAnyEnum, Value));
+			}
+			OutBindings.Add(ModifiedBinding);
+		}
+		else
+		{
+			OutBindings.Add(Binding);
+		}
+	}
+
+	return true;
+}
diff --git a/Engine/Plugins/Runtime/StateTree/Source/StateTreeEditorModule/Private/StateTreeEditor.cpp b/Engine/Plugins/Runtime/StateTree/Source/StateTreeEditorModule/Private/StateTreeEditor.cpp
--- a/Engine/Plugins/Runtime/StateTree/Source/StateTreeEditorModule/Private/StateTreeEditor.cpp
+++ b/Engine/Plugins/Runtime/StateTree/Source/StateTreeEditorModule/Private/StateTreeEditor.cpp
@@ -11,7 +11,7 @@
 #include "StateTreeDelegates.h"
 #include "StateTreeState.h"
 #include "StateTreeTaskBase.h"
-#include "StateTreeBaker.h"
+#include "StateTreeCompiler.h"
 #include "StateTreeTypes.h"
 #include "IDetailsView.h"
 #include "Widgets/Docking/SDockTab.h"
@@ -693,9 +693,9 @@ void FStateTreeEditor::Compile()
 	}
 
 	FStateTreeCompilerLog Log;
-	FStateTreeBaker Baker(Log);
+	FStateTreeCompiler Compiler(Log);
 
-	const bool bSuccess = Baker.Bake(*StateTree);
+	const bool bSuccess = Compiler.Compile(*StateTree);
 
 	if (CompilerResultsListing.IsValid())
 	{
@@ -709,8 +709,8 @@ void FStateTreeEditor::Compile()
 	}
 	else
 	{
-		// Make sure not to leave stale data on failed bake.
-		StateTree->ResetBaked();
+		// Make sure not to leave stale data on failed compile.
+		StateTree->ResetCompiled();
 		StateTree->LastCompiledEditorDataHash = 0;
 
 		// Show log
diff --git a/Engine/Plugins/Runtime/StateTree/Source/StateTreeEditorModule/Public/StateTreeBaker.h b/Engine/Plugins/Runtime/StateTree/Source/StateTreeEditorModule/Public/StateTreeBaker.h
deleted file mode 100644
diff --git a/Engine/Plugins/Runtime/StateTree/Source/StateTreeEditorModule/Public/StateTreeCompiler.h b/Engine/Plugins/Runtime/StateTree/Source/StateTreeEditorModule/Public/StateTreeCompiler.h
new file mode 100644
--- /dev/null
+++ b/Engine/Plugins/Runtime/StateTree/Source/StateTreeEditorModule/Public/StateTreeCompiler.h
@@ -0,0 +1,75 @@
+// Copyright Epic Games, Inc. All Rights Reserved.
+
+#pragma once
+
+#include "CoreMinimal.h"
+#include "StateTreeTypes.h"
+#include "StateTreePropertyBindingCompiler.h"
+#include "StateTreeCompilerLog.h"
+
+class UStateTree;
+class UStateTreeState;
+class UStateTreeEditorData;
+struct FStateTreeEditorNode;
+struct FStateTreeStateLink;
+
+/**
+ * Helper class to convert StateTree editor representation into a compact data.
+ * Holds data needed during compiling.
+ */
+struct STATETREEEDITORMODULE_API FStateTreeCompiler
+{
+public:
+
+	FStateTreeCompiler(FStateTreeCompilerLog& InLog)
+		: Log(InLog)
+	{
+	}
+	
+	bool Compile(UStateTree& InStateTree);
+
+private:
+
+	bool ResolveTransitionState(const UStateTreeState& SourceState, const FStateTreeStateLink& Link, FStateTreeHandle& OutTransitionHandle) const;
+	FStateTreeHandle GetStateHandle(const FGuid& StateID) const;
+	UStateTreeState* GetState(const FGuid& StateID);
+	static FString GetExecutionPathString(const TConstArrayView<const UStateTreeState*> Path);
+	static bool IsPathLinked(const TConstArrayView<const UStateTreeState*> Path);
+
+	bool CreateExecutionInfos();
+	bool CreateExecutionInfosRecursive(UStateTreeState& State, TArray<const UStateTreeState*>& Path);
+	bool CreateStates();
+	bool CreateStateRecursive(UStateTreeState& State, const FStateTreeHandle Parent);
+	
+	bool CreateStateTasks();
+	bool CreateStateEvaluators();
+	bool CreateStateTransitions();
+	
+	bool CreateConditions(UStateTreeState& State, TConstArrayView<FStateTreeEditorNode> Conditions);
+	bool CreateCondition(UStateTreeState& State, const FStateTreeEditorNode& CondNode, const EStateTreeConditionOperand Operand, const int8 DeltaIndent);
+	bool CreateTask(UStateTreeState& State, const FStateTreeEditorNode& TaskNode);
+	bool CreateEvaluator(UStateTreeState& State, const FStateTreeEditorNode& EvalNode);
+	bool GetAndValidateBindings(UStateTreeState& State, const FStateTreeBindableStructDesc& TargetStruct, TArray<FStateTreeEditorPropertyBinding>& OutBindings) const;
+	bool IsPropertyAnyEnum(const FStateTreeBindableStructDesc& Struct, FStateTreeEditorPropertyPath Path) const;
+
+	struct FExecutionPath
+	{
+		FExecutionPath() = default;
+		FExecutionPath(const TConstArrayView<const UStateTreeState*> InPath) : Path(InPath) {}
+		
+		TArray<const UStateTreeState*> Path;
+	};
+
+	struct FStateExecutionInfo
+	{
+		TArray<FExecutionPath> ExecutionPaths;
+	};
+
+	FStateTreeCompilerLog& Log;
+	UStateTree* StateTree = nullptr;
+	UStateTreeEditorData* TreeData = nullptr;
+	TMap<FGuid, int32> IDToState;
+	TArray<UStateTreeState*> SourceStates;
+	TMap<UStateTreeState*, FStateExecutionInfo> ExecutionInfos;
+	FStateTreePropertyBindingCompiler BindingsCompiler;
+};
diff --git a/Engine/Plugins/Runtime/StateTree/Source/StateTreeModule/Private/StateTree.cpp b/Engine/Plugins/Runtime/StateTree/Source/StateTreeModule/Private/StateTree.cpp
--- a/Engine/Plugins/Runtime/StateTree/Source/StateTreeModule/Private/StateTree.cpp
+++ b/Engine/Plugins/Runtime/StateTree/Source/StateTreeModule/Private/StateTree.cpp
@@ -14,7 +14,7 @@ bool UStateTree::IsReadyToRun() const
 }
 
 #if WITH_EDITOR
-void UStateTree::ResetBaked()
+void UStateTree::ResetCompiled()
 {
 	States.Reset();
 	Transitions.Reset();
diff --git a/Engine/Plugins/Runtime/StateTree/Source/StateTreeModule/Private/StateTreeExecutionContext.cpp b/Engine/Plugins/Runtime/StateTree/Source/StateTreeModule/Private/StateTreeExecutionContext.cpp
--- a/Engine/Plugins/Runtime/StateTree/Source/StateTreeModule/Private/StateTreeExecutionContext.cpp
+++ b/Engine/Plugins/Runtime/StateTree/Source/StateTreeModule/Private/StateTreeExecutionContext.cpp
@@ -303,7 +303,7 @@ EStateTreeRunStatus FStateTreeExecutionContext::EnterState(FStateTreeInstanceDat
 	{
 		const FStateTreeHandle CurrentHandle = Transition.NextActiveStates[Index];
 		const FStateTreeHandle PreviousHandle = Transition.CurrentActiveStates.GetStateSafe(Index);
-		const FBakedStateTreeState& State = StateTree->States[CurrentHandle.Index];
+		const FCompactStateTreeState& State = StateTree->States[CurrentHandle.Index];
 
 		if (!Exec.ActiveStates.Push(CurrentHandle))
 		{
@@ -417,7 +417,7 @@ void FStateTreeExecutionContext::ExitState(FStateTreeInstanceData& InstanceData,
 	{
 		const FStateTreeHandle CurrentHandle = Transition.CurrentActiveStates[Index];
 		const FStateTreeHandle NextHandle = Transition.NextActiveStates.GetStateSafe(Index);
-		const FBakedStateTreeState& State = StateTree->States[CurrentHandle.Index];
+		const FCompactStateTreeState& State = StateTree->States[CurrentHandle.Index];
 
 		const bool bRemainsActive = NextHandle == CurrentHandle;
 		bOnTargetBranch = bOnTargetBranch || NextHandle == Transition.TargetState;
@@ -507,7 +507,7 @@ void FStateTreeExecutionContext::StateCompleted(FStateTreeInstanceData& Instance
 	for (int32 Index = Exec.ActiveStates.Num() - 1; Index >= 0; Index--)
 	{
 		const FStateTreeHandle CurrentHandle = Exec.ActiveStates[Index];
-		const FBakedStateTreeState& State = StateTree->States[CurrentHandle.Index];
+		const FCompactStateTreeState& State = StateTree->States[CurrentHandle.Index];
 
 		STATETREE_LOG(Verbose, TEXT("%*sState Completed '%s' %s"), Index*UE::StateTree::DebugIndentSize, TEXT(""), *DebugGetStatePath(Exec.ActiveStates, Index), *UEnum::GetValueAsString(Exec.LastTickStatus));
 
@@ -571,7 +571,7 @@ void FStateTreeExecutionContext::TickEvaluators(FStateTreeInstanceData& Instance
 			continue;
 		}
 
-		const FBakedStateTreeState& State = StateTree->States[StateHandle.Index];
+		const FCompactStateTreeState& State = StateTree->States[StateHandle.Index];
 
 		STATETREE_CLOG(State.EvaluatorsNum > 0, Verbose, TEXT("%*sTicking Evaluators of state '%s' %s"), Index*UE::StateTree::DebugIndentSize, TEXT(""), *DebugGetStatePath(ActiveStates, Index), *UEnum::GetValueAsString(EvalType));
 
@@ -613,7 +613,7 @@ EStateTreeRunStatus FStateTreeExecutionContext::TickTasks(FStateTreeInstanceData
 	for (int32 Index = 0; Index < Exec.ActiveStates.Num() && Result != EStateTreeRunStatus::Failed; Index++)
 	{
 		const FStateTreeHandle CurrentHandle = Exec.ActiveStates[Index];
-		const FBakedStateTreeState& State = StateTree->States[CurrentHandle.Index];
+		const FCompactStateTreeState& State = StateTree->States[CurrentHandle.Index];
 
 		STATETREE_CLOG(State.TasksNum > 0, Verbose, TEXT("%*sTicking Tasks of state '%s'"), Index*UE::StateTree::DebugIndentSize, TEXT(""), *DebugGetStatePath(Exec.ActiveStates, Index));
 
@@ -745,13 +745,13 @@ bool FStateTreeExecutionContext::TriggerTransitions(FStateTreeInstanceData& Inst
 	// Walk towards root and check all transitions along the way.
 	for (int32 StateIndex = Exec.ActiveStates.Num() - 1; StateIndex >= 0; StateIndex--)
 	{
-		const FBakedStateTreeState& State = StateTree->States[Exec.ActiveStates[StateIndex].Index];
+		const FCompactStateTreeState& State = StateTree->States[Exec.ActiveStates[StateIndex].Index];
 		
 		for (uint8 i = 0; i < State.TransitionsNum; i++)
 		{
 			// All transition conditions must pass
 			const int16 TransitionIndex = State.TransitionsBegin + i;
-			const FBakedStateTransition& Transition = StateTree->Transitions[TransitionIndex];
+			const FCompactStateTransition& Transition = StateTree->Transitions[TransitionIndex];
 			if (EnumHasAllFlags(Transition.Event, Event) && TestAllConditions(InstanceData, Transition.ConditionsBegin, Transition.ConditionsNum))
 			{
 				// If a transition has delay, we stop testing other transitions, but the transition will not pass the condition until the delay time passes.
@@ -902,7 +902,7 @@ bool FStateTreeExecutionContext::SelectStateInternal(FStateTreeInstanceData& Ins
 		return false;
 	}
 
-	const FBakedStateTreeState& State = StateTree->States[NextState.Index];
+	const FCompactStateTreeState& State = StateTree->States[NextState.Index];
 
 	// Make sure all the evaluators for the target state are up to date.
 	TickEvaluatorsForSelect(InstanceData, NextState, EStateTreeEvaluationType::PreSelect, 0.0f);
@@ -983,7 +983,7 @@ FString FStateTreeExecutionContext::DebugGetStatePath(const FStateTreeActiveStat
 
 	for (int32 i = 0; i <= ActiveStateIndex; i++)
 	{
-		const FBakedStateTreeState& State = StateTree->States[ActiveStates[i].Index];
+		const FCompactStateTreeState& State = StateTree->States[ActiveStates[i].Index];
 		StatePath.Appendf(TEXT("%s%s"), i == 0 ? TEXT("") : TEXT("."), *State.Name.ToString());
 	}
 	return StatePath;
@@ -1049,7 +1049,7 @@ FString FStateTreeExecutionContext::GetDebugInfoString(const FStateTreeInstanceD
 		FStateTreeHandle Handle = Exec.ActiveStates[Index];
 		if (Handle.IsValid())
 		{
-			const FBakedStateTreeState& State = StateTree->States[Handle.Index];
+			const FCompactStateTreeState& State = StateTree->States[Handle.Index];
 			DebugString += FString::Printf(TEXT("[%s]\n"), *State.Name.ToString());
 
 			if (State.TasksNum > 0)
@@ -1130,7 +1130,7 @@ void FStateTreeExecutionContext::DebugPrintInternalLayout(const FStateTreeInstan
 	// Transitions
 	DebugString += FString::Printf(TEXT("\nTransitions(%d)\n  [ %-3s | %15s | %-40s | %-40s | %-8s ]\n"), StateTree->Transitions.Num()
 		, TEXT("Idx"), TEXT("State"), TEXT("Transition Type"), TEXT("Transition Event"), TEXT("Num Cond"));
-	for (const FBakedStateTransition& Transition : StateTree->Transitions)
+	for (const FCompactStateTransition& Transition : StateTree->Transitions)
 	{
 		DebugString += FString::Printf(TEXT("  | %3d | %15s | %-40s | %-40s | %8d |\n"),
 									Transition.ConditionsBegin, *Transition.State.Describe(),
@@ -1154,7 +1154,7 @@ void FStateTreeExecutionContext::DebugPrintInternalLayout(const FStateTreeInstan
 		TEXT("Cond"), TEXT("Tr"), TEXT("Tsk"), TEXT("Evt"), TEXT("Cond"), TEXT("Tr"), TEXT("Tsk"), TEXT("Evt"),
 		TEXT("Done State"), TEXT("Done Type"), TEXT("Failed State"), TEXT("Failed Type")
 		);
-	for (const FBakedStateTreeState& State : StateTree->States)
+	for (const FCompactStateTreeState& State : StateTree->States)
 	{
 		DebugString += FString::Printf(TEXT("  | %-30s | %15s | %5s [%3d:%-3d[ | %9s   %4d %4d %4d %4d | %3s   %4d %4d %4d %4d\n"),
 									*State.Name.ToString(), *State.Parent.Describe(),
@@ -1165,7 +1165,7 @@ void FStateTreeExecutionContext::DebugPrintInternalLayout(const FStateTreeInstan
 
 	DebugString += FString::Printf(TEXT("\nStates\n  [ %-30s | %-12s | %-30s | %15s | %10s ]\n"),
 		TEXT("State"), TEXT("Type"), TEXT("Name"), TEXT("Bindings"), TEXT("Struct Idx"));
-	for (const FBakedStateTreeState& State : StateTree->States)
+	for (const FCompactStateTreeState& State : StateTree->States)
 	{
 		// Evaluators
 		if (State.EvaluatorsNum)
@@ -1212,7 +1212,7 @@ FString FStateTreeExecutionContext::GetActiveStateName(const FStateTreeInstanceD
 		const FStateTreeHandle Handle = Exec.ActiveStates[Index];
 		if (Handle.IsValid())
 		{
-			const FBakedStateTreeState& State = StateTree->States[Handle.Index];
+			const FCompactStateTreeState& State = StateTree->States[Handle.Index];
 			bool bIsLinked = false;
 			if (Index > 0)
 			{
@@ -1273,7 +1273,7 @@ TArray<FName> FStateTreeExecutionContext::GetActiveStateNames(const FStateTreeIn
 		const FStateTreeHandle Handle = Exec.ActiveStates[Index];
 		if (Handle.IsValid())
 		{
-			const FBakedStateTreeState& State = StateTree->States[Handle.Index];
+			const FCompactStateTreeState& State = StateTree->States[Handle.Index];
 			Result.Add(State.Name);
 		}
 	}
diff --git a/Engine/Plugins/Runtime/StateTree/Source/StateTreeModule/Public/StateTree.h b/Engine/Plugins/Runtime/StateTree/Source/StateTreeModule/Public/StateTree.h
--- a/Engine/Plugins/Runtime/StateTree/Source/StateTreeModule/Public/StateTree.h
+++ b/Engine/Plugins/Runtime/StateTree/Source/StateTreeModule/Public/StateTree.h
@@ -41,8 +41,8 @@ public:
 	bool IsReadyToRun() const;
 
 #if WITH_EDITOR
-	/** Resets the baked data to empty. */
-	void ResetBaked();
+	/** Resets the compiled data to empty. */
+	void ResetCompiled();
 #endif
 
 #if WITH_EDITORONLY_DATA
@@ -104,15 +104,15 @@ private:
 	FStateTreePropertyBindings PropertyBindings;
 
 	UPROPERTY()
-	TArray<FBakedStateTreeState> States;
+	TArray<FCompactStateTreeState> States;
 
 	UPROPERTY()
-	TArray<FBakedStateTransition> Transitions;
+	TArray<FCompactStateTransition> Transitions;
 
 	friend struct FStateTreeInstance;
 	friend struct FStateTreeExecutionContext;
 #if WITH_EDITORONLY_DATA
-	friend struct FStateTreeBaker;
+	friend struct FStateTreeCompiler;
 #endif
 };
 
diff --git a/Engine/Plugins/Runtime/StateTree/Source/StateTreeModule/Public/StateTreeExecutionContext.h b/Engine/Plugins/Runtime/StateTree/Source/StateTreeModule/Public/StateTreeExecutionContext.h
--- a/Engine/Plugins/Runtime/StateTree/Source/StateTreeModule/Public/StateTreeExecutionContext.h
+++ b/Engine/Plugins/Runtime/StateTree/Source/StateTreeModule/Public/StateTreeExecutionContext.h
@@ -88,7 +88,7 @@ public:
 	EStateTreeRunStatus Tick(const float DeltaTime, FStateTreeInstanceData* ExternalInstanceData = nullptr);
 
 	/** @return Pointer to a State or null if state not found */ 
-	const FBakedStateTreeState* GetStateFromHandle(const FStateTreeHandle StateHandle) const
+	const FCompactStateTreeState* GetStateFromHandle(const FStateTreeHandle StateHandle) const
 	{
 		return (StateTree && StateTree->States.IsValidIndex(StateHandle.Index)) ? &StateTree->States[StateHandle.Index] : nullptr;
 	}
diff --git a/Engine/Plugins/Runtime/StateTree/Source/StateTreeModule/Public/StateTreeTypes.h b/Engine/Plugins/Runtime/StateTree/Source/StateTreeModule/Public/StateTreeTypes.h
--- a/Engine/Plugins/Runtime/StateTree/Source/StateTreeModule/Public/StateTreeTypes.h
+++ b/Engine/Plugins/Runtime/StateTree/Source/StateTreeModule/Public/StateTreeTypes.h
@@ -78,7 +78,7 @@ enum class EStateTreeTransitionEvent : uint8
 ENUM_CLASS_FLAGS(EStateTreeTransitionEvent)
 
 /**
- * Handle that is used to refer baked state tree data.
+ * Handle that is used to refer compact state tree data.
  */
 USTRUCT(BlueprintType)
 struct STATETREEMODULE_API FStateTreeHandle
@@ -309,7 +309,7 @@ struct STATETREEMODULE_API FStateTreeTransitionResult
  *  Runtime representation of a StateTree transition.
  */
 USTRUCT()
-struct STATETREEMODULE_API FBakedStateTransition
+struct STATETREEMODULE_API FCompactStateTransition
 {
 	GENERATED_BODY()
 
@@ -331,7 +331,7 @@ struct STATETREEMODULE_API FBakedStateTransition
  *  Runtime representation of a StateTree state.
  */
 USTRUCT()
-struct STATETREEMODULE_API FBakedStateTreeState
+struct STATETREEMODULE_API FCompactStateTreeState
 {
 	GENERATED_BODY()
 
diff --git a/Engine/Plugins/Runtime/StateTree/Source/StateTreeTestSuite/Private/StateTreeTest.cpp b/Engine/Plugins/Runtime/StateTree/Source/StateTreeTestSuite/Private/StateTreeTest.cpp
--- a/Engine/Plugins/Runtime/StateTree/Source/StateTreeTestSuite/Private/StateTreeTest.cpp
+++ b/Engine/Plugins/Runtime/StateTree/Source/StateTreeTestSuite/Private/StateTreeTest.cpp
@@ -5,7 +5,7 @@
 #include "AITestsCommon.h"
 #include "StateTreeEditorData.h"
 #include "StateTreeState.h"
-#include "StateTreeBaker.h"
+#include "StateTreeCompiler.h"
 #include "Conditions/StateTreeCondition_Common.h"
 #include "StateTreeEvaluatorBase.h"
 #include "StateTreeTaskBase.h"
@@ -71,10 +71,10 @@ struct FStateTreeTest_MakeAndBakeStateTree : FAITestBase
 		StateB.AddTransition(EStateTreeTransitionEvent::OnCompleted, EStateTreeTransitionType::Succeeded);
 
 		FStateTreeCompilerLog Log;
-		FStateTreeBaker Baker(Log);
-		const bool bResult = Baker.Bake(StateTree);
+		FStateTreeCompiler Compiler(Log);
+		const bool bResult = Compiler.Compile(StateTree);
 
-		AITEST_TRUE("StateTree should get baked", bResult);
+		AITEST_TRUE("StateTree should get compiled", bResult);
 
 		return true;
 	}
@@ -141,10 +141,10 @@ struct FStateTreeTest_WanderLoop : FAITestBase
 		StandOnLane.AddTransition(EStateTreeTransitionEvent::OnCompleted, EStateTreeTransitionType::GotoState, &Wander);
 
 		FStateTreeCompilerLog Log;
-		FStateTreeBaker Baker(Log);
-		bool bResult = Baker.Bake(StateTree);
+		FStateTreeCompiler Compiler(Log);
+		bool bResult = Compiler.Compile(StateTree);
 
-		AITEST_TRUE("StateTree should get baked", bResult);
+		AITEST_TRUE("StateTree should get compiled", bResult);
 
 		EStateTreeRunStatus Status = EStateTreeRunStatus::Unset;
 		FTestStateTreeExecutionContext Exec;
@@ -234,9 +234,9 @@ struct FStateTreeTest_Sequence : FAITestBase
 		State2.AddTransition(EStateTreeTransitionEvent::OnCompleted, EStateTreeTransitionType::Succeeded);
 
 		FStateTreeCompilerLog Log;
-		FStateTreeBaker Baker(Log);
-		const bool bResult = Baker.Bake(StateTree);
-		AITEST_TRUE("StateTree should get baked", bResult);
+		FStateTreeCompiler Compiler(Log);
+		const bool bResult = Compiler.Compile(StateTree);
+		AITEST_TRUE("StateTree should get compiled", bResult);
 
 		EStateTreeRunStatus Status = EStateTreeRunStatus::Unset;
 		FTestStateTreeExecutionContext Exec;
@@ -297,9 +297,9 @@ struct FStateTreeTest_Select : FAITestBase
 		State1A.AddTransition(EStateTreeTransitionEvent::OnCompleted, EStateTreeTransitionType::GotoState, &State1);
 
 		FStateTreeCompilerLog Log;
-		FStateTreeBaker Baker(Log);
-		const bool bResult = Baker.Bake(StateTree);
-		AITEST_TRUE("StateTree should get baked", bResult);
+		FStateTreeCompiler Compiler(Log);
+		const bool bResult = Compiler.Compile(StateTree);
+		AITEST_TRUE("StateTree should get compiled", bResult);
 
 		EStateTreeRunStatus Status = EStateTreeRunStatus::Unset;
 		FTestStateTreeExecutionContext Exec;
@@ -369,9 +369,9 @@ struct FStateTreeTest_FailEnterState : FAITestBase
 		State1A.AddTransition(EStateTreeTransitionEvent::OnCompleted, EStateTreeTransitionType::GotoState, &State1);
 
 		FStateTreeCompilerLog Log;
-		FStateTreeBaker Baker(Log);
-		const bool bResult = Baker.Bake(StateTree);
-		AITEST_TRUE("StateTree should get baked", bResult);
+		FStateTreeCompiler Compiler(Log);
+		const bool bResult = Compiler.Compile(StateTree);
+		AITEST_TRUE("StateTree should get compiled", bResult);
 
 		EStateTreeRunStatus Status = EStateTreeRunStatus::Unset;
 		FTestStateTreeExecutionContext Exec;
@@ -432,9 +432,9 @@ struct FStateTreeTest_SubTree : FAITestBase
 
 
 		FStateTreeCompilerLog Log;
-		FStateTreeBaker Baker(Log);
-		const bool bResult = Baker.Bake(StateTree);
-		AITEST_TRUE("StateTree should get baked", bResult);
+		FStateTreeCompiler Compiler(Log);
+		const bool bResult = Compiler.Compile(StateTree);
+		AITEST_TRUE("StateTree should get compiled", bResult);
 
 		EStateTreeRunStatus Status = EStateTreeRunStatus::Unset;
 		FTestStateTreeExecutionContext Exec;
