[StateTree] external data validation from schema supported types
- fixed external requirements set by StateTreeBrainComponent to find more specific actor classes first before defaulting on AActor
- added missing supported type FMassSharedFragment for MassStateTreeSchema
- moved StateTreeLinker to it own file and handle its potential failures
#rnx

diff --git a/Engine/Plugins/AI/MassAI/Source/MassAIBehavior/Private/Conditions/ZoneGraphTagConditions.cpp b/Engine/Plugins/AI/MassAI/Source/MassAIBehavior/Private/Conditions/ZoneGraphTagConditions.cpp
--- a/Engine/Plugins/AI/MassAI/Source/MassAIBehavior/Private/Conditions/ZoneGraphTagConditions.cpp
+++ b/Engine/Plugins/AI/MassAI/Source/MassAIBehavior/Private/Conditions/ZoneGraphTagConditions.cpp
@@ -2,9 +2,9 @@
 
 #include "Conditions/ZoneGraphTagConditions.h"
 #include "StateTreeExecutionContext.h"
+#include "StateTreeLinker.h"
 #if WITH_EDITOR
 #include "ZoneGraphSettings.h"
-#include "StateTreePropertyBindings.h"
 #endif// WITH_EDITOR
 
 #if WITH_EDITOR
diff --git a/Engine/Plugins/AI/MassAI/Source/MassAIBehavior/Private/Evaluators/MassComponentHitEvaluator.cpp b/Engine/Plugins/AI/MassAI/Source/MassAIBehavior/Private/Evaluators/MassComponentHitEvaluator.cpp
--- a/Engine/Plugins/AI/MassAI/Source/MassAIBehavior/Private/Evaluators/MassComponentHitEvaluator.cpp
+++ b/Engine/Plugins/AI/MassAI/Source/MassAIBehavior/Private/Evaluators/MassComponentHitEvaluator.cpp
@@ -5,6 +5,7 @@
 #include "MassComponentHitSubsystem.h"
 #include "MassStateTreeExecutionContext.h"
 #include "Engine/World.h"
+#include "StateTreeLinker.h"
 
 
 bool FMassComponentHitEvaluator::Link(FStateTreeLinker& Linker)
diff --git a/Engine/Plugins/AI/MassAI/Source/MassAIBehavior/Private/Evaluators/MassStateTreeSmartObjectEvaluator.cpp b/Engine/Plugins/AI/MassAI/Source/MassAIBehavior/Private/Evaluators/MassStateTreeSmartObjectEvaluator.cpp
--- a/Engine/Plugins/AI/MassAI/Source/MassAIBehavior/Private/Evaluators/MassStateTreeSmartObjectEvaluator.cpp
+++ b/Engine/Plugins/AI/MassAI/Source/MassAIBehavior/Private/Evaluators/MassStateTreeSmartObjectEvaluator.cpp
@@ -9,7 +9,7 @@
 #include "MassSmartObjectHandler.h"
 #include "MassStateTreeExecutionContext.h"
 #include "MassZoneGraphNavigationFragments.h"
-#include "SmartObjectSubsystem.h"
+#include "StateTreeLinker.h"
 #include "Engine/World.h"
 
 //----------------------------------------------------------------------//
diff --git a/Engine/Plugins/AI/MassAI/Source/MassAIBehavior/Private/Evaluators/MassZoneGraphAnnotationEvaluator.cpp b/Engine/Plugins/AI/MassAI/Source/MassAIBehavior/Private/Evaluators/MassZoneGraphAnnotationEvaluator.cpp
--- a/Engine/Plugins/AI/MassAI/Source/MassAIBehavior/Private/Evaluators/MassZoneGraphAnnotationEvaluator.cpp
+++ b/Engine/Plugins/AI/MassAI/Source/MassAIBehavior/Private/Evaluators/MassZoneGraphAnnotationEvaluator.cpp
@@ -3,6 +3,7 @@
 #include "Evaluators/MassZoneGraphAnnotationEvaluator.h"
 #include "MassStateTreeExecutionContext.h"
 #include "MassZoneGraphAnnotationFragments.h"
+#include "StateTreeLinker.h"
 
 
 FMassZoneGraphAnnotationEvaluator::FMassZoneGraphAnnotationEvaluator()
diff --git a/Engine/Plugins/AI/MassAI/Source/MassAIBehavior/Private/MassStateTreeSchema.cpp b/Engine/Plugins/AI/MassAI/Source/MassAIBehavior/Private/MassStateTreeSchema.cpp
--- a/Engine/Plugins/AI/MassAI/Source/MassAIBehavior/Private/MassStateTreeSchema.cpp
+++ b/Engine/Plugins/AI/MassAI/Source/MassAIBehavior/Private/MassStateTreeSchema.cpp
@@ -24,5 +24,7 @@ bool UMassStateTreeSchema::IsStructAllowed(const UScriptStruct* InScriptStruct)
 bool UMassStateTreeSchema::IsExternalItemAllowed(const UStruct& InStruct) const
 {
 	// Allow only WorldSubsystems and fragments as external data.
-	return InStruct.IsChildOf(UWorldSubsystem::StaticClass()) || InStruct.IsChildOf(FMassFragment::StaticStruct());
+	return InStruct.IsChildOf(UWorldSubsystem::StaticClass())
+			|| InStruct.IsChildOf(FMassFragment::StaticStruct())
+			|| InStruct.IsChildOf(FMassSharedFragment::StaticStruct());
 }
diff --git a/Engine/Plugins/AI/MassAI/Source/MassAIBehavior/Private/Tasks/MassLookAtTask.cpp b/Engine/Plugins/AI/MassAI/Source/MassAIBehavior/Private/Tasks/MassLookAtTask.cpp
--- a/Engine/Plugins/AI/MassAI/Source/MassAIBehavior/Private/Tasks/MassLookAtTask.cpp
+++ b/Engine/Plugins/AI/MassAI/Source/MassAIBehavior/Private/Tasks/MassLookAtTask.cpp
@@ -5,6 +5,7 @@
 #include "MassLookAtFragments.h"
 #include "MassSignalSubsystem.h"
 #include "MassStateTreeExecutionContext.h"
+#include "StateTreeLinker.h"
 
 bool FMassLookAtTask::Link(FStateTreeLinker& Linker)
 {
diff --git a/Engine/Plugins/AI/MassAI/Source/MassAIBehavior/Private/Tasks/MassStateTreeSmartObjectTask.cpp b/Engine/Plugins/AI/MassAI/Source/MassAIBehavior/Private/Tasks/MassStateTreeSmartObjectTask.cpp
--- a/Engine/Plugins/AI/MassAI/Source/MassAIBehavior/Private/Tasks/MassStateTreeSmartObjectTask.cpp
+++ b/Engine/Plugins/AI/MassAI/Source/MassAIBehavior/Private/Tasks/MassStateTreeSmartObjectTask.cpp
@@ -8,10 +8,10 @@
 #include "MassSmartObjectFragments.h"
 #include "MassStateTreeExecutionContext.h"
 #include "MassNavigationFragments.h"
-#include "SmartObjectSubsystem.h"
 #include "MassZoneGraphNavigationFragments.h"
 #include "MassZoneGraphNavigationUtils.h"
 #include "Engine/World.h"
+#include "StateTreeLinker.h"
 
 //----------------------------------------------------------------------//
 // FMassClaimSmartObjectTask
diff --git a/Engine/Plugins/AI/MassAI/Source/MassAIBehavior/Private/Tasks/MassZoneGraphFindEscapeTarget.cpp b/Engine/Plugins/AI/MassAI/Source/MassAIBehavior/Private/Tasks/MassZoneGraphFindEscapeTarget.cpp
--- a/Engine/Plugins/AI/MassAI/Source/MassAIBehavior/Private/Tasks/MassZoneGraphFindEscapeTarget.cpp
+++ b/Engine/Plugins/AI/MassAI/Source/MassAIBehavior/Private/Tasks/MassZoneGraphFindEscapeTarget.cpp
@@ -10,6 +10,7 @@
 #include "MassAIBehaviorTypes.h"
 #include "MassStateTreeExecutionContext.h"
 #include "ZoneGraphSettings.h"
+#include "StateTreeLinker.h"
 
 FMassZoneGraphFindEscapeTarget::FMassZoneGraphFindEscapeTarget()
 {
diff --git a/Engine/Plugins/AI/MassAI/Source/MassAIBehavior/Private/Tasks/MassZoneGraphFindSmartObjectTarget.cpp b/Engine/Plugins/AI/MassAI/Source/MassAIBehavior/Private/Tasks/MassZoneGraphFindSmartObjectTarget.cpp
--- a/Engine/Plugins/AI/MassAI/Source/MassAIBehavior/Private/Tasks/MassZoneGraphFindSmartObjectTarget.cpp
+++ b/Engine/Plugins/AI/MassAI/Source/MassAIBehavior/Private/Tasks/MassZoneGraphFindSmartObjectTarget.cpp
@@ -8,6 +8,7 @@
 #include "MassStateTreeExecutionContext.h"
 #include "MassZoneGraphNavigationFragments.h"
 #include "SmartObjectZoneAnnotations.h"
+#include "StateTreeLinker.h"
 
 bool FMassZoneGraphFindSmartObjectTarget::Link(FStateTreeLinker& Linker)
 {
diff --git a/Engine/Plugins/AI/MassAI/Source/MassAIBehavior/Private/Tasks/MassZoneGraphPathFollowTask.cpp b/Engine/Plugins/AI/MassAI/Source/MassAIBehavior/Private/Tasks/MassZoneGraphPathFollowTask.cpp
--- a/Engine/Plugins/AI/MassAI/Source/MassAIBehavior/Private/Tasks/MassZoneGraphPathFollowTask.cpp
+++ b/Engine/Plugins/AI/MassAI/Source/MassAIBehavior/Private/Tasks/MassZoneGraphPathFollowTask.cpp
@@ -10,6 +10,7 @@
 #include "MassZoneGraphNavigationFragments.h"
 #include "MassZoneGraphNavigationUtils.h"
 #include "ZoneGraphSubsystem.h"
+#include "StateTreeLinker.h"
 
 bool FMassZoneGraphPathFollowTask::Link(FStateTreeLinker& Linker)
 {
diff --git a/Engine/Plugins/AI/MassAI/Source/MassAIBehavior/Private/Tasks/MassZoneGraphStandTask.cpp b/Engine/Plugins/AI/MassAI/Source/MassAIBehavior/Private/Tasks/MassZoneGraphStandTask.cpp
--- a/Engine/Plugins/AI/MassAI/Source/MassAIBehavior/Private/Tasks/MassZoneGraphStandTask.cpp
+++ b/Engine/Plugins/AI/MassAI/Source/MassAIBehavior/Private/Tasks/MassZoneGraphStandTask.cpp
@@ -11,6 +11,7 @@
 #include "MassStateTreeExecutionContext.h"
 #include "MassSignalSubsystem.h"
 #include "MassSimulationLOD.h"
+#include "StateTreeLinker.h"
 
 bool FMassZoneGraphStandTask::Link(FStateTreeLinker& Linker)
 {
diff --git a/Engine/Plugins/AI/MassCrowd/Source/MassCrowd/Private/Tasks/MassCrowdClaimWaitSlotTask.cpp b/Engine/Plugins/AI/MassCrowd/Source/MassCrowd/Private/Tasks/MassCrowdClaimWaitSlotTask.cpp
--- a/Engine/Plugins/AI/MassCrowd/Source/MassCrowd/Private/Tasks/MassCrowdClaimWaitSlotTask.cpp
+++ b/Engine/Plugins/AI/MassCrowd/Source/MassCrowd/Private/Tasks/MassCrowdClaimWaitSlotTask.cpp
@@ -6,6 +6,7 @@
 #include "MassZoneGraphNavigationFragments.h"
 #include "MassNavigationFragments.h"
 #include "MassStateTreeExecutionContext.h"
+#include "StateTreeLinker.h"
 
 bool FMassCrowdClaimWaitSlotTask::Link(FStateTreeLinker& Linker)
 {
diff --git a/Engine/Plugins/AI/MassCrowd/Source/MassCrowd/Private/Tasks/MassZoneGraphFindWanderTarget.cpp b/Engine/Plugins/AI/MassCrowd/Source/MassCrowd/Private/Tasks/MassZoneGraphFindWanderTarget.cpp
--- a/Engine/Plugins/AI/MassCrowd/Source/MassCrowd/Private/Tasks/MassZoneGraphFindWanderTarget.cpp
+++ b/Engine/Plugins/AI/MassCrowd/Source/MassCrowd/Private/Tasks/MassZoneGraphFindWanderTarget.cpp
@@ -11,6 +11,7 @@
 #include "MassCrowdSettings.h"
 #include "MassCrowdSubsystem.h"
 #include "MassStateTreeExecutionContext.h"
+#include "StateTreeLinker.h"
 
 FMassZoneGraphFindWanderTarget::FMassZoneGraphFindWanderTarget()
 {
diff --git a/Engine/Plugins/Runtime/StateTree/Source/StateTreeEditorModule/Private/StateTreeCompiler.cpp b/Engine/Plugins/Runtime/StateTree/Source/StateTreeEditorModule/Private/StateTreeCompiler.cpp
--- a/Engine/Plugins/Runtime/StateTree/Source/StateTreeEditorModule/Private/StateTreeCompiler.cpp
+++ b/Engine/Plugins/Runtime/StateTree/Source/StateTreeEditorModule/Private/StateTreeCompiler.cpp
@@ -26,7 +26,11 @@ bool FStateTreeCompiler::Compile(UStateTree& InStateTree)
 	// Cleanup existing state
 	StateTree->ResetCompiled();
 
-	BindingsCompiler.Init(StateTree->PropertyBindings, Log);
+	if (!BindingsCompiler.Init(StateTree->PropertyBindings, Log))
+	{
+		StateTree->ResetCompiled();
+		return false;
+	}
 
 	if (!CreateStates())
 	{
@@ -60,9 +64,12 @@ bool FStateTreeCompiler::Compile(UStateTree& InStateTree)
 
 	BindingsCompiler.Finalize();
 
-	StateTree->PropertyBindings.ResolvePaths();
-	StateTree->Link();
-
+	if (!StateTree->Link())
+	{
+		StateTree->ResetCompiled();
+		return false;
+	}
+	
 	return true;
 }
 
diff --git a/Engine/Plugins/Runtime/StateTree/Source/StateTreeEditorModule/Public/StateTreePropertyBindingCompiler.h b/Engine/Plugins/Runtime/StateTree/Source/StateTreeEditorModule/Public/StateTreePropertyBindingCompiler.h
--- a/Engine/Plugins/Runtime/StateTree/Source/StateTreeEditorModule/Public/StateTreePropertyBindingCompiler.h
+++ b/Engine/Plugins/Runtime/StateTree/Source/StateTreeEditorModule/Public/StateTreePropertyBindingCompiler.h
@@ -26,7 +26,7 @@ struct STATETREEEDITORMODULE_API FStateTreePropertyBindingCompiler
 	  * @param PropertyBindings - Reference to the Property Bindings where all the batches will be stored.
 	  * @return true on success.
 	  */
-	bool Init(FStateTreePropertyBindings& InPropertyBindings, FStateTreeCompilerLog& InLog);
+	[[nodiscard]] bool Init(FStateTreePropertyBindings& InPropertyBindings, FStateTreeCompilerLog& InLog);
 
 	/**
 	  * Compiles a batch of property copies.
@@ -35,7 +35,7 @@ struct STATETREEEDITORMODULE_API FStateTreePropertyBindingCompiler
 	  * @param OutBatchIndex - Resulting batch index, if index is INDEX_NONE, no bindings were found and no batch was generated.
 	  * @return True on success, false on failure.
 	 */
-	bool CompileBatch(const FStateTreeBindableStructDesc& TargetStruct, TConstArrayView<FStateTreeEditorPropertyBinding> EditorPropertyBindings, int32& OutBatchIndex);
+	[[nodiscard]] bool CompileBatch(const FStateTreeBindableStructDesc& TargetStruct, TConstArrayView<FStateTreeEditorPropertyBinding> EditorPropertyBindings, int32& OutBatchIndex);
 
 	/** Finalizes compilation, should be called once all batches are compiled. */
 	void Finalize();
@@ -69,9 +69,9 @@ struct STATETREEEDITORMODULE_API FStateTreePropertyBindingCompiler
 	 * @param LogContextStruct Pointer to bindable struct desc where the property path belongs to.
 	 * @return True of the property was solved successfully.
 	 */
-	static bool ResolvePropertyPath(const FStateTreeBindableStructDesc& InStructDesc, const FStateTreeEditorPropertyPath& InPath,
-									TArray<FStateTreePropertySegment>& OutSegments, const FProperty*& OutLeafProperty, int32& OutLeafArrayIndex,
-									FStateTreeCompilerLog* Log = nullptr, const FStateTreeBindableStructDesc* LogContextStruct = nullptr);
+	[[nodiscard]] static bool ResolvePropertyPath(const FStateTreeBindableStructDesc& InStructDesc, const FStateTreeEditorPropertyPath& InPath,
+												  TArray<FStateTreePropertySegment>& OutSegments, const FProperty*& OutLeafProperty, int32& OutLeafArrayIndex,
+												  FStateTreeCompilerLog* Log = nullptr, const FStateTreeBindableStructDesc* LogContextStruct = nullptr);
 
 	/**
 	 * Checks if two property types can are compatible for copying.
@@ -100,5 +100,5 @@ protected:
 	};
 
 	EStateTreePropertyCopyType GetCopyType(const FProperty* SourceProperty, const int32 SourceArrayIndex, const FProperty* TargetProperty, const int32 TargetArrayIndex);
-	bool ResolvePropertyPath(const FStateTreeBindableStructDesc& InOwnerStructDesc, const FStateTreeBindableStructDesc& InStructDesc, const FStateTreeEditorPropertyPath& InPath, FResolvedPathResult& OutResult);
+	[[nodiscard]] bool ResolvePropertyPath(const FStateTreeBindableStructDesc& InOwnerStructDesc, const FStateTreeBindableStructDesc& InStructDesc, const FStateTreeEditorPropertyPath& InPath, FResolvedPathResult& OutResult);
 };
diff --git a/Engine/Plugins/Runtime/StateTree/Source/StateTreeModule/Private/Blueprint/StateTreeItemBlueprintBase.cpp b/Engine/Plugins/Runtime/StateTree/Source/StateTreeModule/Private/Blueprint/StateTreeItemBlueprintBase.cpp
--- a/Engine/Plugins/Runtime/StateTree/Source/StateTreeModule/Private/Blueprint/StateTreeItemBlueprintBase.cpp
+++ b/Engine/Plugins/Runtime/StateTree/Source/StateTreeModule/Private/Blueprint/StateTreeItemBlueprintBase.cpp
@@ -1,9 +1,9 @@
 // Copyright Epic Games, Inc. All Rights Reserved.
 
 #include "Blueprint/StateTreeItemBlueprintBase.h"
-#include "CoreMinimal.h"
 #include "AIController.h"
 #include "StateTreeExecutionContext.h"
+#include "StateTreeLinker.h"
 
 UWorld* UStateTreeItemBlueprintBase::GetWorld() const
 {
diff --git a/Engine/Plugins/Runtime/StateTree/Source/StateTreeModule/Private/Components/StateTreeBrainComponent.cpp b/Engine/Plugins/Runtime/StateTree/Source/StateTreeModule/Private/Components/StateTreeBrainComponent.cpp
--- a/Engine/Plugins/Runtime/StateTree/Source/StateTreeModule/Private/Components/StateTreeBrainComponent.cpp
+++ b/Engine/Plugins/Runtime/StateTree/Source/StateTreeModule/Private/Components/StateTreeBrainComponent.cpp
@@ -3,7 +3,6 @@
 #include "Components/StateTreeBrainComponent.h"
 #include "GameFramework/Actor.h"
 #include "GameplayTasksComponent.h"
-#include "VisualLogger/VisualLoggerTypes.h"
 #include "VisualLogger/VisualLogger.h"
 #include "StateTree.h"
 #include "StateTreeEvaluatorBase.h"
@@ -33,8 +32,6 @@ bool UBrainComponentStateTreeSchema::IsExternalItemAllowed(const UStruct& InStru
 {
 	// Actors and components.
 	return InStruct.IsChildOf(AActor::StaticClass())
-			|| InStruct.IsChildOf(APawn::StaticClass())
-			|| InStruct.IsChildOf(AAIController::StaticClass())
 			|| InStruct.IsChildOf(UActorComponent::StaticClass())
 			|| InStruct.IsChildOf(UWorldSubsystem::StaticClass());
 }
@@ -74,8 +71,8 @@ bool UStateTreeBrainComponent::SetContextRequirements()
 	{
 		return false;
 	}
-	
-	UWorld* World = GetWorld();
+
+	const UWorld* World = GetWorld();
 	if (World == nullptr)
 	{
 		return false;
@@ -95,38 +92,20 @@ bool UStateTreeBrainComponent::SetContextRequirements()
 				UActorComponent* Component = GetOwner()->FindComponentByClass(Cast<UClass>(const_cast<UStruct*>(ItemDesc.Struct)));
 				StateTreeContext.SetExternalData(ItemDesc.Handle, FStateTreeDataView(Component));
 			}
-			else if (ItemDesc.Struct->IsChildOf(AActor::StaticClass()))
-			{
-				if (AIOwner != nullptr)
-				{
-					AActor* OwnerActor = Cast<AActor>(AIOwner->GetPawn());
-					StateTreeContext.SetExternalData(ItemDesc.Handle, FStateTreeDataView(OwnerActor));
-				}
-				else
-				{
-					AActor* OwnerActor = Cast<AActor>(GetOwner());
-					StateTreeContext.SetExternalData(ItemDesc.Handle, FStateTreeDataView(OwnerActor));
-				}
-			}
 			else if (ItemDesc.Struct->IsChildOf(APawn::StaticClass()))
 			{
-				if (AIOwner != nullptr)
-				{
-					APawn* OwnerActor = Cast<APawn>(AIOwner->GetPawn());
-					StateTreeContext.SetExternalData(ItemDesc.Handle, FStateTreeDataView(OwnerActor));
-				}
-				else
-				{
-					APawn* OwnerActor = Cast<APawn>(GetOwner());
-					StateTreeContext.SetExternalData(ItemDesc.Handle, FStateTreeDataView(OwnerActor));
-				}
+				APawn* OwnerPawn = (AIOwner != nullptr) ? AIOwner->GetPawn() : Cast<APawn>(GetOwner());
+				StateTreeContext.SetExternalData(ItemDesc.Handle, FStateTreeDataView(OwnerPawn));
 			}
 			else if (ItemDesc.Struct->IsChildOf(AAIController::StaticClass()))
 			{
-				if (AIOwner != nullptr)
-				{
-					StateTreeContext.SetExternalData(ItemDesc.Handle, FStateTreeDataView(AIOwner));
-				}
+				AAIController* OwnerController = (AIOwner != nullptr) ? AIOwner.Get() : Cast<AAIController>(GetOwner());
+				StateTreeContext.SetExternalData(ItemDesc.Handle, FStateTreeDataView(OwnerController));
+			}
+			else if (ItemDesc.Struct->IsChildOf(AActor::StaticClass()))
+			{
+				AActor* OwnerActor = (AIOwner != nullptr) ? AIOwner->GetPawn() : GetOwner();
+				StateTreeContext.SetExternalData(ItemDesc.Handle, FStateTreeDataView(OwnerActor));
 			}
 		}
 	}
diff --git a/Engine/Plugins/Runtime/StateTree/Source/StateTreeModule/Private/Conditions/GameplayTagConditions.cpp b/Engine/Plugins/Runtime/StateTree/Source/StateTreeModule/Private/Conditions/GameplayTagConditions.cpp
--- a/Engine/Plugins/Runtime/StateTree/Source/StateTreeModule/Private/Conditions/GameplayTagConditions.cpp
+++ b/Engine/Plugins/Runtime/StateTree/Source/StateTreeModule/Private/Conditions/GameplayTagConditions.cpp
@@ -3,9 +3,7 @@
 #include "Conditions/GameplayTagConditions.h"
 #include "GameplayTagContainer.h"
 #include "StateTreeExecutionContext.h"
-#if WITH_EDITOR
-#include "StateTreePropertyBindings.h"
-#endif// WITH_EDITOR
+#include "StateTreeLinker.h"
 
 #if WITH_EDITOR
 #define LOCTEXT_NAMESPACE "StateTreeEditor"
diff --git a/Engine/Plugins/Runtime/StateTree/Source/StateTreeModule/Private/Conditions/StateTreeCondition_Common.cpp b/Engine/Plugins/Runtime/StateTree/Source/StateTreeModule/Private/Conditions/StateTreeCondition_Common.cpp
--- a/Engine/Plugins/Runtime/StateTree/Source/StateTreeModule/Private/Conditions/StateTreeCondition_Common.cpp
+++ b/Engine/Plugins/Runtime/StateTree/Source/StateTreeModule/Private/Conditions/StateTreeCondition_Common.cpp
@@ -4,6 +4,7 @@
 
 #include "StateTreeExecutionContext.h"
 #include "StateTreePropertyBindings.h"
+#include "StateTreeLinker.h"
 
 #define LOCTEXT_NAMESPACE "StateTreeEditor"
 
diff --git a/Engine/Plugins/Runtime/StateTree/Source/StateTreeModule/Private/Evaluators/CooldownGuardEvaluator.cpp b/Engine/Plugins/Runtime/StateTree/Source/StateTreeModule/Private/Evaluators/CooldownGuardEvaluator.cpp
--- a/Engine/Plugins/Runtime/StateTree/Source/StateTreeModule/Private/Evaluators/CooldownGuardEvaluator.cpp
+++ b/Engine/Plugins/Runtime/StateTree/Source/StateTreeModule/Private/Evaluators/CooldownGuardEvaluator.cpp
@@ -2,6 +2,7 @@
 
 #include "Evaluators/CooldownGuardEvaluator.h"
 #include "StateTreeExecutionContext.h"
+#include "StateTreeLinker.h"
 #include "Engine/World.h"
 
 
diff --git a/Engine/Plugins/Runtime/StateTree/Source/StateTreeModule/Private/Evaluators/SelectionGuardEvaluator.cpp b/Engine/Plugins/Runtime/StateTree/Source/StateTreeModule/Private/Evaluators/SelectionGuardEvaluator.cpp
--- a/Engine/Plugins/Runtime/StateTree/Source/StateTreeModule/Private/Evaluators/SelectionGuardEvaluator.cpp
+++ b/Engine/Plugins/Runtime/StateTree/Source/StateTreeModule/Private/Evaluators/SelectionGuardEvaluator.cpp
@@ -2,6 +2,7 @@
 
 #include "Evaluators/SelectionGuardEvaluator.h"
 #include "StateTreeExecutionContext.h"
+#include "StateTreeLinker.h"
 
 
 bool FSelectionGuardEvaluator::Link(FStateTreeLinker& Linker)
diff --git a/Engine/Plugins/Runtime/StateTree/Source/StateTreeModule/Private/StateTree.cpp b/Engine/Plugins/Runtime/StateTree/Source/StateTreeModule/Private/StateTree.cpp
--- a/Engine/Plugins/Runtime/StateTree/Source/StateTreeModule/Private/StateTree.cpp
+++ b/Engine/Plugins/Runtime/StateTree/Source/StateTreeModule/Private/StateTree.cpp
@@ -1,11 +1,9 @@
 // Copyright Epic Games, Inc. All Rights Reserved.
 
 #include "StateTree.h"
-#include "StateTreeEvaluatorBase.h"
-#include "StateTreeTaskBase.h"
-#include "CoreMinimal.h"
-#include "StateTreeConditionBase.h"
 #include "StateTreeDelegates.h"
+#include "StateTreeLinker.h"
+#include "StateTreeNodeBase.h"
 
 bool UStateTree::IsReadyToRun() const
 {
@@ -57,7 +55,10 @@ void UStateTree::PostLoad()
 {
 	Super::PostLoad();
 
-	Link();
+	if (!Link())
+	{
+		UE_LOG(LogStateTree, Error, TEXT("%s failed to link. Asset will not be usable at runtime."), *GetName());	
+	}
 }
 
 void UStateTree::Serialize(FStructuredArchiveRecord Record)
@@ -68,38 +69,56 @@ void UStateTree::Serialize(FStructuredArchiveRecord Record)
 	// because property bindings may get invalid, and instance data potentially needs refreshed.
 	if (Record.GetUnderlyingArchive().IsModifyingWeakAndStrongReferences())
 	{
-		Link();
+		if (!Link())
+		{
+			UE_LOG(LogStateTree, Error, TEXT("%s failed to link. Asset will not be usable at runtime."), *GetName());	
+		}
 	}
 }
 
-
-void UStateTree::Link()
+bool UStateTree::Link()
 {
-	FStateTreeLinker Linker;
+	// Initialize the instance data default value.
+	// This data will be used to allocate runtime instance on all StateTree users.
+	InstanceDataDefaultValue.Reset();
+	
+	ExternalDataDescs.Reset();
+	NumDataViews = 0;
+	ExternalDataBaseIndex = 0;
 
-	ExternalDataBaseIndex = PropertyBindings.GetSourceStructNum();
-	Linker.SetExternalDataBaseIndex(ExternalDataBaseIndex);
+	// Resolves property paths used by bindings a store property pointers
+	if (!PropertyBindings.ResolvePaths())
+	{
+		return false;
+	}
+
+	// Resolves nodes references to other StateTree data
+	FStateTreeLinker Linker(Schema);
+	Linker.SetExternalDataBaseIndex(PropertyBindings.GetSourceStructNum());
 	
 	for (FInstancedStruct& Node : Nodes)
 	{
 		if (FStateTreeNodeBase* NodePtr = Node.GetMutablePtr<FStateTreeNodeBase>())
 		{
 			Linker.SetCurrentInstanceDataType(NodePtr->GetInstanceDataType(), NodePtr->DataViewIndex);
-			NodePtr->Link(Linker);
+			const bool bLinkSucceeded = NodePtr->Link(Linker);
+			if (!bLinkSucceeded || Linker.GetStatus() == EStateTreeLinkerStatus::Failed)
+			{
+				UE_LOG(LogStateTree, Error, TEXT("%s: node '%s' failed to resolve its references."), *GetName(), *NodePtr->StaticStruct()->GetName());
+				return false;
+			}
 		}
 	}
 
+	// Link succeeded, setup tree to be ready to run
+	ExternalDataBaseIndex = PropertyBindings.GetSourceStructNum();
 	ExternalDataDescs = Linker.GetExternalDataDescs();
-	
 	NumDataViews = ExternalDataBaseIndex + ExternalDataDescs.Num();
 
-	// Initialize the instance data default value.
-	// This data will be used to allocate runtime instance on all StateTree users.
-	InstanceDataDefaultValue.Reset();
 	if (Instances.Num() > 0)
 	{
 		InstanceDataDefaultValue.Initialize(*this, Instances, InstanceObjects);
 	}
 
-	PropertyBindings.ResolvePaths();
+	return true;
 }
diff --git a/Engine/Plugins/Runtime/StateTree/Source/StateTreeModule/Public/Blueprint/StateTreeItemBlueprintBase.h b/Engine/Plugins/Runtime/StateTree/Source/StateTreeModule/Public/Blueprint/StateTreeItemBlueprintBase.h
--- a/Engine/Plugins/Runtime/StateTree/Source/StateTreeModule/Public/Blueprint/StateTreeItemBlueprintBase.h
+++ b/Engine/Plugins/Runtime/StateTree/Source/StateTreeModule/Public/Blueprint/StateTreeItemBlueprintBase.h
@@ -2,10 +2,10 @@
 
 #pragma once
 
-#include "CoreMinimal.h"
 #include "StateTreeTypes.h"
 #include "StateTreeItemBlueprintBase.generated.h"
 
+struct FStateTreeLinker;
 struct FStateTreeExecutionContext;
 
 UENUM()
diff --git a/Engine/Plugins/Runtime/StateTree/Source/StateTreeModule/Public/StateTree.h b/Engine/Plugins/Runtime/StateTree/Source/StateTreeModule/Public/StateTree.h
--- a/Engine/Plugins/Runtime/StateTree/Source/StateTreeModule/Public/StateTree.h
+++ b/Engine/Plugins/Runtime/StateTree/Source/StateTreeModule/Public/StateTree.h
@@ -2,7 +2,6 @@
 
 #pragma once
 
-#include "CoreMinimal.h"
 #include "Engine/DataAsset.h"
 #include "StateTreeTypes.h"
 #include "StateTreeSchema.h"
@@ -22,7 +21,7 @@ class STATETREEMODULE_API UStateTree : public UDataAsset
 public:
 
 	/** @return Default value for the instance data. */
-	const FStateTreeInstanceData& GetInstanceDataDefaultValue() const { return InstanceDataDefaultValue; };
+	const FStateTreeInstanceData& GetInstanceDataDefaultValue() const { return InstanceDataDefaultValue; }
 	
 	/** @return Number of runtime data (Evaluators, Tasks, Conditions) in the runtime storage. */
 	int32 GetNumInstances() const { return Instances.Num(); }
@@ -37,7 +36,7 @@ public:
 	const UStateTreeSchema* GetSchema() const { return Schema; }
 	void SetSchema(UStateTreeSchema* InSchema) { Schema = InSchema; }
 
-	/** @return true is the tree asset is to be used at runtime. */
+	/** @return true if the tree asset can be used at runtime. */
 	bool IsReadyToRun() const;
 
 #if WITH_EDITOR
@@ -46,19 +45,22 @@ public:
 #endif
 
 #if WITH_EDITORONLY_DATA
-	// Edit time data for the StateTree, instance of UStateTreeEditorData
+	/** Edit time data for the StateTree, instance of UStateTreeEditorData */
 	UPROPERTY()
 	TObjectPtr<UObject> EditorData;
 
-	// Hash of the editor data from last compile.
+	/** Hash of the editor data from last compile. */
 	UPROPERTY()
 	uint32 LastCompiledEditorDataHash = 0;
 #endif
 
 protected:
 	
-	/** Resolved references between data in the StateTree. */
-	void Link();
+	/**
+	 * Resolves references between data in the StateTree.
+	 * @return true if all references to internal and external data are resolved properly, false otherwise.
+	 */
+	[[nodiscard]] bool Link();
 
 	virtual void PostLoad() override;
 	virtual void Serialize(FStructuredArchiveRecord Record) override;
diff --git a/Engine/Plugins/Runtime/StateTree/Source/StateTreeModule/Public/StateTreeLinker.h b/Engine/Plugins/Runtime/StateTree/Source/StateTreeModule/Public/StateTreeLinker.h
new file mode 100644
--- /dev/null
+++ b/Engine/Plugins/Runtime/StateTree/Source/StateTreeModule/Public/StateTreeLinker.h
@@ -0,0 +1,134 @@
+﻿// Copyright Epic Games, Inc. All Rights Reserved.
+
+#pragma once
+
+#include "StateTreeSchema.h"
+#include "StateTreeTypes.h"
+#include "StateTreeLinker.generated.h"
+
+UENUM()
+enum class EStateTreeLinkerStatus : uint8
+{
+	Succeeded,
+	Failed,
+};
+
+/**
+ * The StateTree linker is used to resolved references to various StateTree data at load time.
+ * @see TStateTreeExternalDataHandle<> for example usage.
+ */
+struct FStateTreeLinker
+{
+	explicit FStateTreeLinker(const UStateTreeSchema* InSchema) : Schema(InSchema) {}
+	
+	/** Sets base index for all external data handles. */
+	void SetExternalDataBaseIndex(const int32 InExternalDataBaseIndex) { ExternalDataBaseIndex = InExternalDataBaseIndex; }
+
+	/** Sets currently linked item's instance data type and index. */ 
+	void SetCurrentInstanceDataType(const UStruct* Struct, const int32 Index)
+	{
+		CurrentInstanceStruct = Struct;
+		CurrentInstanceIndex = Index;
+	}
+
+	EStateTreeLinkerStatus GetStatus() const { return Status; }
+	
+	/**
+	 * Links reference to an external UObject.
+	 * @param Handle Reference to TStateTreeExternalDataHandle<> with UOBJECT type to link to.
+	 */
+	template <typename T>
+	typename TEnableIf<TIsDerivedFrom<typename T::DataType, UObject>::IsDerived, void>::Type LinkExternalData(T& Handle)
+	{
+		LinkExternalData(Handle, T::DataType::StaticClass(), T::DataRequirement);
+	}
+
+	/**
+	 * Links reference to an external UObject.
+	 * @param Handle Reference to TStateTreeExternalDataHandle<> with USTRUCT type to link to.
+	 */
+	template <typename T>
+	typename TEnableIf<!TIsDerivedFrom<typename T::DataType, UObject>::IsDerived, void>::Type LinkExternalData(T& Handle)
+	{
+		LinkExternalData(Handle, T::DataType::StaticStruct(), T::DataRequirement);
+	}
+
+	/**
+	 * Links reference to an external Object or Struct.
+	 * This function should only be used when TStateTreeExternalDataHandle<> cannot be used, i.e. the Struct is based on some data.
+	 * @param Handle Reference to link to.
+	 * @param Struct Expected type of the Object or Struct to link to.
+	 * @param Requirement Describes if the external data is expected to be required or optional.
+	 */
+	void LinkExternalData(FStateTreeExternalDataHandle& Handle, const UStruct* Struct, const EStateTreeExternalDataRequirement Requirement)
+	{
+		if (!ensureMsgf(Schema == nullptr || Schema->IsExternalItemAllowed(*Struct),
+				TEXT("External data of type '%s' used by current node is not allowed by schema '%s' (i.e. rejected by IsExternalItemAllowed)"),
+				*Struct->GetName(),
+				*Schema->GetClass()->GetName()))
+		{
+			Handle = FStateTreeExternalDataHandle();
+			Status = EStateTreeLinkerStatus::Failed;
+			return;
+		}
+		
+		const FStateTreeExternalDataDesc Desc(Struct, Requirement);
+		int32 Index = ExternalDataDescs.Find(Desc);
+		if (Index == INDEX_NONE)
+		{
+			Index = ExternalDataDescs.Add(Desc);
+			check(FStateTreeExternalDataHandle::IsValidIndex(Index + ExternalDataBaseIndex));
+			ExternalDataDescs[Index].Handle.DataViewIndex = (uint8)(Index + ExternalDataBaseIndex);
+		}
+		Handle.DataViewIndex = (uint8)(Index + ExternalDataBaseIndex);
+	}
+	/**
+	 * Links reference to a property in instance data.
+	 * Usage:
+	 * 	  Linker.LinkRuntimeDataProperty(HitPointsHandle, STATETREE_INSTANCEDATA_PROPERTY(FHitPointLayout, HitPoints));
+	 *
+	 * @param Handle Reference to TStateTreeExternalDataHandle<> with USTRUCT type to link to.
+	 * @param DummyProperty Do not use directly.
+	 * @param ScriptStruct Do not use directly.
+	 * @param PropertyName Do not use directly.
+	 */
+	template <typename T, typename S>
+	void LinkInstanceDataProperty(T& Handle, const S& DummyProperty, const UScriptStruct* ScriptStruct, const TCHAR* PropertyName)
+	{
+		static_assert(TIsSame<typename T::DataType, S>::Value, "Expecting linked handle to have same type as the instance data struct member.");
+		LinkInstanceDataPropertyInternal(Handle, ScriptStruct, PropertyName);
+	}
+
+	/** @return linked external data descriptors. */
+	TConstArrayView<FStateTreeExternalDataDesc> GetExternalDataDescs() const { return ExternalDataDescs; }
+
+protected:
+
+	void LinkInstanceDataPropertyInternal(FStateTreeInstanceDataPropertyHandle& Handle, const UScriptStruct* ScriptStruct, const TCHAR* PropertyName)
+	{
+		check(CurrentInstanceStruct != nullptr);
+		check(CurrentInstanceIndex != INDEX_NONE);
+
+		const FProperty* Property = ScriptStruct->FindPropertyByName(FName(PropertyName));
+		if (Property == nullptr)
+		{
+			Handle = FStateTreeInstanceDataPropertyHandle();
+			Status = EStateTreeLinkerStatus::Failed;
+			return;
+		}
+
+		check(CurrentInstanceIndex < MAX_uint8);
+		check(Property->GetOffset_ForInternal() < MAX_uint16);
+		
+		Handle.DataViewIndex = (uint8)CurrentInstanceIndex;
+		Handle.Type = EStateTreePropertyIndirection::Offset;
+		Handle.PropertyOffset = (uint16)Property->GetOffset_ForInternal();
+	}
+
+	const UStateTreeSchema* Schema = nullptr;
+	EStateTreeLinkerStatus Status = EStateTreeLinkerStatus::Succeeded;
+	const UStruct* CurrentInstanceStruct = nullptr;
+	int32 CurrentInstanceIndex = INDEX_NONE;
+	int32 ExternalDataBaseIndex = 0;
+	TArray<FStateTreeExternalDataDesc> ExternalDataDescs;
+};
diff --git a/Engine/Plugins/Runtime/StateTree/Source/StateTreeModule/Public/StateTreeNodeBase.h b/Engine/Plugins/Runtime/StateTree/Source/StateTreeModule/Public/StateTreeNodeBase.h
--- a/Engine/Plugins/Runtime/StateTree/Source/StateTreeModule/Public/StateTreeNodeBase.h
+++ b/Engine/Plugins/Runtime/StateTree/Source/StateTreeModule/Public/StateTreeNodeBase.h
@@ -5,6 +5,8 @@
 #include "StateTreeTypes.h"
 #include "StateTreeNodeBase.generated.h"
 
+struct FStateTreeLinker;
+
 /**
  * Base struct of StateTree Conditions, Evaluators, and Tasks.
  */
@@ -27,7 +29,7 @@ struct STATETREEMODULE_API FStateTreeNodeBase
 	 * @param Linker Reference to the linker
 	 * @return true if linking succeeded. 
 	 */
-	virtual bool Link(FStateTreeLinker& Linker) { return true; }
+	[[nodiscard]] virtual bool Link(FStateTreeLinker& Linker) { return true; }
 
 	/** Name of the node. */
 	UPROPERTY(EditDefaultsOnly, Category = "", meta=(EditCondition = "false", EditConditionHides))
diff --git a/Engine/Plugins/Runtime/StateTree/Source/StateTreeModule/Public/StateTreePropertyBindings.h b/Engine/Plugins/Runtime/StateTree/Source/StateTreeModule/Public/StateTreePropertyBindings.h
--- a/Engine/Plugins/Runtime/StateTree/Source/StateTreeModule/Public/StateTreePropertyBindings.h
+++ b/Engine/Plugins/Runtime/StateTree/Source/StateTreeModule/Public/StateTreePropertyBindings.h
@@ -416,7 +416,7 @@ struct STATETREEMODULE_API FStateTreePropertyBindings
 	 * Resolves paths to indirections.
 	 * @return True if resolve succeeded.
 	 */
-	bool ResolvePaths();
+	[[nodiscard]] bool ResolvePaths();
 
 	/**
 	 * @return True if bindings have been resolved.
@@ -439,8 +439,8 @@ struct STATETREEMODULE_API FStateTreePropertyBindings
 	void DebugPrintInternalLayout(FString& OutString) const;
 
 protected:
-	bool ResolvePath(const UStruct* Struct, const FStateTreePropertyPath& Path, FStateTreePropertyAccess& OutAccess);
-	bool ValidateCopy(FStateTreePropCopy& Copy) const;
+	[[nodiscard]] bool ResolvePath(const UStruct* Struct, const FStateTreePropertyPath& Path, FStateTreePropertyAccess& OutAccess);
+	[[nodiscard]] bool ValidateCopy(FStateTreePropCopy& Copy) const;
 	void PerformCopy(const FStateTreePropCopy& Copy, const FProperty* SourceProperty, const uint8* SourceAddress, const FProperty* TargetProperty, uint8* TargetAddress) const;
 	uint8* GetAddress(FStateTreeDataView InStructView, const FStateTreePropertyAccess& InAccess) const;
 	FString GetPathAsString(const FStateTreePropertyPath& Path, const int32 HighlightedSegment = INDEX_NONE, const TCHAR* HighlightPrefix = nullptr, const TCHAR* HighlightPostfix = nullptr);
diff --git a/Engine/Plugins/Runtime/StateTree/Source/StateTreeModule/Public/StateTreeTypes.h b/Engine/Plugins/Runtime/StateTree/Source/StateTreeModule/Public/StateTreeTypes.h
--- a/Engine/Plugins/Runtime/StateTree/Source/StateTreeModule/Public/StateTreeTypes.h
+++ b/Engine/Plugins/Runtime/StateTree/Source/StateTreeModule/Public/StateTreeTypes.h
@@ -2,7 +2,6 @@
 
 #pragma once
 
-#include "CoreMinimal.h"
 #include "StateTreeTypes.generated.h"
 
 STATETREEMODULE_API DECLARE_LOG_CATEGORY_EXTERN(LogStateTree, Warning, All);
@@ -517,117 +516,3 @@ struct STATETREEMODULE_API FStateTreeExternalDataDesc
 
 #define STATETREE_INSTANCEDATA_PROPERTY(Struct, Member) \
 		decltype(Struct::Member){}, Struct::StaticStruct(), TEXT(#Member)
-
-UENUM()
-enum class EStateTreeLinkerStatus : uint8
-{
-	Succeeded,
-	Failed,
-};
-/**
- * The StateTree linker is used to resolved references to various StateTree data at load time.
- * @see TStateTreeExternalDataHandle<> for example usage.
- */
-struct FStateTreeLinker
-{
-	/** Sets base index for all external data handles. */
-	void SetExternalDataBaseIndex(const int32 InExternalDataBaseIndex) { ExternalDataBaseIndex = InExternalDataBaseIndex; }
-
-	/** Sets currently linked item's instance data type and index. */ 
-	void SetCurrentInstanceDataType(const UStruct* Struct, const int32 Index)
-	{
-		CurrentInstanceStruct = Struct;
-		CurrentInstanceIndex = Index;
-	}
-
-	EStateTreeLinkerStatus GetStatus() const { return Status; }
-	
-	/**
-	 * Links reference to an external UObject.
-	 * @param Handle Reference to TStateTreeExternalDataHandle<> with UOBJECT type to link to.
-	 */
-	template <typename T>
-	typename TEnableIf<TIsDerivedFrom<typename T::DataType, UObject>::IsDerived, void>::Type LinkExternalData(T& Handle)
-	{
-		LinkExternalData(Handle, T::DataType::StaticClass(), T::DataRequirement);
-	}
-
-	/**
-	 * Links reference to an external UObject.
-	 * @param Handle Reference to TStateTreeExternalDataHandle<> with USTRUCT type to link to.
-	 */
-	template <typename T>
-	typename TEnableIf<!TIsDerivedFrom<typename T::DataType, UObject>::IsDerived, void>::Type LinkExternalData(T& Handle)
-	{
-		LinkExternalData(Handle, T::DataType::StaticStruct(), T::DataRequirement);
-	}
-
-	/**
-	 * Links reference to an external Object or Struct.
-	 * This function should only be used when TStateTreeExternalDataHandle<> cannot be used, i.e. the Struct is based on some data.
-	 * @param Handle Reference to link to.
-	 * @param Struct Expected type of the Object or Struct to link to.
-	 * @param Requirement Describes if the external data is expected to be required or optional.
-	 */
-	void LinkExternalData(FStateTreeExternalDataHandle& Handle, const UStruct* Struct, const EStateTreeExternalDataRequirement Requirement)
-	{
-		const FStateTreeExternalDataDesc Desc(Struct, Requirement);
-		int32 Index = ExternalDataDescs.Find(Desc);
-		if (Index == INDEX_NONE)
-		{
-			Index = ExternalDataDescs.Add(Desc);
-			check(FStateTreeExternalDataHandle::IsValidIndex(Index + ExternalDataBaseIndex));
-			ExternalDataDescs[Index].Handle.DataViewIndex = (uint8)(Index + ExternalDataBaseIndex);
-		}
-		Handle.DataViewIndex = (uint8)(Index + ExternalDataBaseIndex);
-	}
-	/**
-	 * Links reference to a property in instance data.
-	 * Usage:
-	 * 	  Linker.LinkRuntimeDataProperty(HitPointsHandle, STATETREE_INSTANCEDATA_PROPERTY(FHitPointLayout, HitPoints));
-	 *
-	 * @param Handle Reference to TStateTreeExternalDataHandle<> with USTRUCT type to link to.
-	 * @param DummyProperty Do not use directly.
-	 * @param ScriptStruct Do not use directly.
-	 * @param PropertyName Do not use directly.
-	 */
-	template <typename T, typename S>
-	void LinkInstanceDataProperty(T& Handle, const S& DummyProperty, const UScriptStruct* ScriptStruct, const TCHAR* PropertyName)
-	{
-		static_assert(TIsSame<typename T::DataType, S>::Value, "Expecting linked handle to have same type as the instance data struct member.");
-		LinkInstanceDataPropertyInternal(Handle, ScriptStruct, PropertyName);
-	}
-
-	/** @return linked external data descriptors. */
-	TConstArrayView<FStateTreeExternalDataDesc> GetExternalDataDescs() const { return ExternalDataDescs; }
-
-protected:
-
-
-	void LinkInstanceDataPropertyInternal(FStateTreeInstanceDataPropertyHandle& Handle, const UScriptStruct* ScriptStruct, const TCHAR* PropertyName)
-	{
-		check(CurrentInstanceStruct != nullptr);
-		check(CurrentInstanceIndex != INDEX_NONE);
-
-		FProperty* Property = ScriptStruct->FindPropertyByName(FName(PropertyName));
-		if (Property == nullptr)
-		{
-			Handle = FStateTreeInstanceDataPropertyHandle();
-			Status = EStateTreeLinkerStatus::Failed;
-			return;
-		}
-
-		check(CurrentInstanceIndex < MAX_uint8);
-		check(Property->GetOffset_ForInternal() < MAX_uint16);
-		
-		Handle.DataViewIndex = (uint8)CurrentInstanceIndex;
-		Handle.Type = EStateTreePropertyIndirection::Offset;
-		Handle.PropertyOffset = (uint16)Property->GetOffset_ForInternal();
-	}
-
-	EStateTreeLinkerStatus Status = EStateTreeLinkerStatus::Succeeded;
-	const UStruct* CurrentInstanceStruct = nullptr;
-	int32 CurrentInstanceIndex = INDEX_NONE;
-	int32 ExternalDataBaseIndex = 0;
-	TArray<FStateTreeExternalDataDesc> ExternalDataDescs;
-};
diff --git a/Engine/Plugins/Runtime/StateTree/Source/StateTreeTestSuite/Private/StateTreeTestTypes.h b/Engine/Plugins/Runtime/StateTree/Source/StateTreeTestSuite/Private/StateTreeTestTypes.h
--- a/Engine/Plugins/Runtime/StateTree/Source/StateTreeTestSuite/Private/StateTreeTestTypes.h
+++ b/Engine/Plugins/Runtime/StateTree/Source/StateTreeTestSuite/Private/StateTreeTestTypes.h
@@ -2,10 +2,10 @@
 
 #pragma once
 
-#include "AITestsCommon.h"
 #include "StateTreeTaskBase.h"
 #include "StateTreeEvaluatorBase.h"
 #include "StateTreeExecutionContext.h"
+#include "StateTreeLinker.h"
 #include "StateTreeTestTypes.generated.h"
 
 
