[StateTree] updated GetPropertyCompatibility for property to ObjectPtr to restrict to the class hierarchy
#rnx

diff --git a/Engine/Plugins/Runtime/StateTree/Source/StateTreeEditorModule/Private/StateTreePropertyBindingCompiler.cpp b/Engine/Plugins/Runtime/StateTree/Source/StateTreeEditorModule/Private/StateTreePropertyBindingCompiler.cpp
--- a/Engine/Plugins/Runtime/StateTree/Source/StateTreeEditorModule/Private/StateTreePropertyBindingCompiler.cpp
+++ b/Engine/Plugins/Runtime/StateTree/Source/StateTreeEditorModule/Private/StateTreePropertyBindingCompiler.cpp
@@ -460,10 +460,13 @@ EPropertyAccessCompatibility FStateTreePropertyBindingCompiler::GetPropertyCompa
 		return EPropertyAccessCompatibility::Incompatible;
 	}
 
-	// Special case for object properties
+	// Special case for object properties since InPropertyA->SameType(InPropertyB) requires both properties to be of the exact same class.
+	// In our case we want to be able to bind a source property if its class is a child of the target property class.
 	if (InPropertyA->IsA<FObjectPropertyBase>() && InPropertyB->IsA<FObjectPropertyBase>())
 	{
-		return EPropertyAccessCompatibility::Compatible;
+		const FObjectPropertyBase* SourceProperty = CastField<FObjectPropertyBase>(InPropertyA);
+		const FObjectPropertyBase* TargetProperty = CastField<FObjectPropertyBase>(InPropertyB);
+		return (SourceProperty->PropertyClass->IsChildOf(TargetProperty->PropertyClass)) ? EPropertyAccessCompatibility::Compatible : EPropertyAccessCompatibility::Incompatible;
 	}
 
 	// Extract underlying types for enums
