[StateTree] allow transition to "Tree Succeeded" or "Tree Failed" from StateLink combo box
#rnx

diff --git a/Engine/Plugins/Runtime/StateTree/Source/StateTreeEditorModule/Private/Customizations/StateTreeStateLinkDetails.cpp b/Engine/Plugins/Runtime/StateTree/Source/StateTreeEditorModule/Private/Customizations/StateTreeStateLinkDetails.cpp
--- a/Engine/Plugins/Runtime/StateTree/Source/StateTreeEditorModule/Private/Customizations/StateTreeStateLinkDetails.cpp
+++ b/Engine/Plugins/Runtime/StateTree/Source/StateTreeEditorModule/Private/Customizations/StateTreeStateLinkDetails.cpp
@@ -156,6 +156,12 @@ TSharedRef<SWidget> FStateTreeStateLinkDetails::OnGetStateContent() const
 
 		FUIAction NotSetItemAction(FExecuteAction::CreateSP(const_cast<FStateTreeStateLinkDetails*>(this), &FStateTreeStateLinkDetails::OnStateComboChange, ComboNotSet));
 		MenuBuilder.AddMenuEntry(LOCTEXT("TransitionBlock", "Block Transition"), LOCTEXT("TransitionBlockTooltip", "Will not transition to any state, but will block other transitions to trigger if the condition passes."), FSlateIcon(), NotSetItemAction);
+
+		FUIAction SucceededItemAction(FExecuteAction::CreateSP(const_cast<FStateTreeStateLinkDetails*>(this), &FStateTreeStateLinkDetails::OnStateComboChange, ComboSucceeded));
+		MenuBuilder.AddMenuEntry(LOCTEXT("TransitionTreeSucceeded", "Tree Succeeded"), LOCTEXT("TransitionTreeSuccessTooltip", "Complete tree with success."), FSlateIcon(), SucceededItemAction);
+
+		FUIAction FailedItemAction(FExecuteAction::CreateSP(const_cast<FStateTreeStateLinkDetails*>(this), &FStateTreeStateLinkDetails::OnStateComboChange, ComboFailed));
+		MenuBuilder.AddMenuEntry(LOCTEXT("TransitionTreeFailed", "Tree Failed"), LOCTEXT("TransitionTreeFailedTooltip", "Complete tree with failure."), FSlateIcon(), FailedItemAction);
 	}
 
 	if (CachedNames.Num() > 0)
@@ -189,6 +195,14 @@ FText FStateTreeStateLinkDetails::GetCurrentStateDesc() const
 	{
 		return LOCTEXT("TransitionNextState", "Next State");
 	}
+	else if (TransitionType == EStateTreeTransitionType::Succeeded)
+	{
+		return LOCTEXT("TransitionTreeSucceeded", "Tree Succeeded");
+	}
+	else if (TransitionType == EStateTreeTransitionType::Failed)
+	{
+		return LOCTEXT("TransitionTreeFailed", "Tree Failed");
+	}
 	else if (TransitionType == EStateTreeTransitionType::GotoState)
 	{
 		FName OldName;
