[StateTree]
+ Added StateTree parameters usable for data bindings and that could be overriden on context initialization by a StateTreeReference
+ Added StateTreeReference struct to hold a reference to a StateTree asset along with a list of values to parameterized the tree.
+ Added named external data items that are defined bt the Schema and for which values must be provided at runtime through the execution context.
+ Added delegate OnPostCompile after successful compilation. The StateTreeReference listens to it to validate its parameters
+ EditorData now contains its own version of the schema and parameters. On successful compilation they are copied over the StateTree own properties.
#rnx

diff --git a/Engine/Plugins/Runtime/StateTree/Source/StateTreeEditorModule/Private/Customizations/StateTreeEditorNodeDetails.cpp b/Engine/Plugins/Runtime/StateTree/Source/StateTreeEditorModule/Private/Customizations/StateTreeEditorNodeDetails.cpp
--- a/Engine/Plugins/Runtime/StateTree/Source/StateTreeEditorModule/Private/Customizations/StateTreeEditorNodeDetails.cpp
+++ b/Engine/Plugins/Runtime/StateTree/Source/StateTreeEditorModule/Private/Customizations/StateTreeEditorNodeDetails.cpp
@@ -809,7 +809,7 @@ EVisibility FStateTreeEditorNodeDetails::IsNameVisible() const
 
 	if (ScriptStruct != nullptr && ScriptStruct->IsChildOf(FStateTreeTaskBase::StaticStruct()))
 	{
-		const UStateTreeSchema* Schema = StateTree ? StateTree->GetSchema() : nullptr;
+		const UStateTreeSchema* Schema = EditorData ? EditorData->Schema : nullptr;
 		if (Schema && Schema->AllowMultipleTasks() == false)
 		{
 			// Single task states use the state name as task name.
@@ -977,7 +977,7 @@ TSharedRef<SWidget> FStateTreeEditorNodeDetails::GeneratePicker()
 	ClassCache->GetClasses(BaseClass, ObjectNodes);
 
 	const FSlateIcon ScripStructIcon = FSlateIconFinder::FindIconForClass(UScriptStruct::StaticClass());
-	const UStateTreeSchema* Schema = StateTree ? StateTree->GetSchema() : nullptr;
+	const UStateTreeSchema* Schema = EditorData ? EditorData->Schema : nullptr;
 	
 	for (const TSharedPtr<FStateTreeNodeClassData>& Data : StructNodes)
 	{
diff --git a/Engine/Plugins/Runtime/StateTree/Source/StateTreeEditorModule/Private/Customizations/StateTreeStateDetails.cpp b/Engine/Plugins/Runtime/StateTree/Source/StateTreeEditorModule/Private/Customizations/StateTreeStateDetails.cpp
--- a/Engine/Plugins/Runtime/StateTree/Source/StateTreeEditorModule/Private/Customizations/StateTreeStateDetails.cpp
+++ b/Engine/Plugins/Runtime/StateTree/Source/StateTreeEditorModule/Private/Customizations/StateTreeStateDetails.cpp
@@ -1,34 +1,12 @@
 // Copyright Epic Games, Inc. All Rights Reserved.
 
 #include "StateTreeStateDetails.h"
-#include "Misc/MessageDialog.h"
-#include "UObject/ObjectMacros.h"
-#include "UObject/Class.h"
-#include "UObject/UnrealType.h"
-#include "Layout/Visibility.h"
-#include "Misc/Attribute.h"
 #include "Widgets/DeclarativeSyntaxSupport.h"
-#include "Widgets/SWidget.h"
-#include "Widgets/Text/STextBlock.h"
-#include "Widgets/Layout/SBox.h"
-#include "Widgets/Input/SComboBox.h"
-#include "Widgets/Input/SVectorInputBox.h"
-#include "Widgets/Input/SNumericEntryBox.h"
-#include "Widgets/Input/SButton.h"
-#include "Framework/MultiBox/MultiBoxBuilder.h"
-#include "ComponentVisualizer.h"
-#include "Editor/UnrealEdEngine.h"
-#include "UnrealEdGlobals.h"
-#include "IDetailCustomNodeBuilder.h"
 #include "IDetailChildrenBuilder.h"
 #include "DetailWidgetRow.h"
 #include "DetailLayoutBuilder.h"
 #include "DetailCategoryBuilder.h"
 #include "PropertyCustomizationHelpers.h"
-#include "ScopedTransaction.h"
-#include "Modules/ModuleManager.h"
-#include "ISettingsModule.h"
-#include "Editor.h"
 #include "IPropertyUtilities.h"
 #include "StateTreeEditorData.h"
 #include "StateTreeViewModel.h"
@@ -45,22 +23,22 @@ TSharedRef<IDetailCustomization> FStateTreeStateDetails::MakeInstance()
 
 void FStateTreeStateDetails::CustomizeDetails(IDetailLayoutBuilder& DetailBuilder)
 {
-	// Find StateTree associated with this panel.
-	UStateTree* StateTree = nullptr;
+	// Find StateTreeEditorData associated with this panel.
+	const UStateTreeEditorData* EditorData = nullptr;
 	TArray<TWeakObjectPtr<UObject>> Objects;
 	DetailBuilder.GetObjectsBeingCustomized(Objects);
 	for (TWeakObjectPtr<UObject>& WeakObject : Objects)
 	{
-		if (UObject* Object = WeakObject.Get())
+		if (const UObject* Object = WeakObject.Get())
 		{
-			if (UStateTree* OuterStateTree = Object->GetTypedOuter<UStateTree>())
+			const UStateTreeEditorData* OuterEditorData = Object->GetTypedOuter<UStateTreeEditorData>();
 			{
-				StateTree = OuterStateTree;
+				EditorData = OuterEditorData;
 				break;
 			}
 		}
 	}
-	const UStateTreeSchema* Schema = StateTree ? StateTree->GetSchema() : nullptr;
+	const UStateTreeSchema* Schema = EditorData ? EditorData->Schema : nullptr;
 
 	TSharedPtr<IPropertyHandle> TasksProperty = DetailBuilder.GetProperty(TEXT("Tasks"));
 	TSharedPtr<IPropertyHandle> SingleTaskProperty = DetailBuilder.GetProperty(TEXT("SingleTask"));
diff --git a/Engine/Plugins/Runtime/StateTree/Source/StateTreeEditorModule/Private/StateTreeCompiler.cpp b/Engine/Plugins/Runtime/StateTree/Source/StateTreeEditorModule/Private/StateTreeCompiler.cpp
--- a/Engine/Plugins/Runtime/StateTree/Source/StateTreeEditorModule/Private/StateTreeCompiler.cpp
+++ b/Engine/Plugins/Runtime/StateTree/Source/StateTreeEditorModule/Private/StateTreeCompiler.cpp
@@ -10,7 +10,6 @@
 #include "StateTreeConditionBase.h"
 #include "StateTreeState.h"
 #include "StateTreeExecutionContext.h"
-#include "CoreMinimal.h"
 #include "StateTreePropertyBindingCompiler.h"
 
 
@@ -32,6 +31,27 @@ bool FStateTreeCompiler::Compile(UStateTree& InStateTree)
 		return false;
 	}
 
+	// Copy schema the EditorData
+	StateTree->Schema = DuplicateObject(TreeData->Schema, StateTree);
+
+	// Copy parameters from EditorData	
+	StateTree->Parameters = TreeData->Parameters;
+	
+	// Mark all parameters as binding source
+	for (FStateTreeParameterDesc& Desc : StateTree->Parameters.Parameters)
+	{
+		Desc.DataViewIndex = BindingsCompiler.AddSourceStruct({Desc.Name, Desc.Parameter.GetScriptStruct(), Desc.ID});
+	}	
+	
+	// Mark all named external values as binding source
+	if (StateTree->Schema)
+	{
+		for (FStateTreeExternalDataDesc& Desc : StateTree->Schema->GetMutableNamedExternalDataDescs())
+		{
+			Desc.Handle.DataViewIndex = BindingsCompiler.AddSourceStruct({Desc.Name, Desc.Struct, Desc.ID});
+		} 
+	}
+	
 	if (!CreateStates())
 	{
 		StateTree->ResetCompiled();
diff --git a/Engine/Plugins/Runtime/StateTree/Source/StateTreeEditorModule/Private/StateTreeEditor.cpp b/Engine/Plugins/Runtime/StateTree/Source/StateTreeEditorModule/Private/StateTreeEditor.cpp
--- a/Engine/Plugins/Runtime/StateTree/Source/StateTreeEditorModule/Private/StateTreeEditor.cpp
+++ b/Engine/Plugins/Runtime/StateTree/Source/StateTreeEditorModule/Private/StateTreeEditor.cpp
@@ -195,6 +195,7 @@ void FStateTreeEditor::InitEditor( const EToolkitMode::Type Mode, const TSharedP
 
 	UE::StateTree::Delegates::OnIdentifierChanged.AddSP(this, &FStateTreeEditor::OnIdentifierChanged);
 	UE::StateTree::Delegates::OnSchemaChanged.AddSP(this, &FStateTreeEditor::OnSchemaChanged);
+	UE::StateTree::Delegates::OnParametersChanged.AddSP(this, &FStateTreeEditor::OnParametersChanged);
 }
 
 FName FStateTreeEditor::GetToolkitFName() const
@@ -275,7 +276,7 @@ TSharedRef<SDockTab> FStateTreeEditor::SpawnTab_AssetDetails(const FSpawnTabArgs
 	DetailsViewArgs.NameAreaSettings = FDetailsViewArgs::HideNameArea;
 
 	AssetDetailsView = PropertyEditorModule.CreateDetailView(DetailsViewArgs);
-	AssetDetailsView->SetObject(StateTree);
+	AssetDetailsView->SetObject(StateTree ? StateTree->EditorData : nullptr);
 	AssetDetailsView->OnFinishedChangingProperties().AddSP(this, &FStateTreeEditor::OnAssetFinishedChangingProperties);
 
 	TSharedRef<SDockTab> SpawnedTab = SNew(SDockTab)
@@ -354,7 +355,7 @@ void FStateTreeEditor::HandleModelSelectionChanged(const TArray<UStateTreeState*
 
 void FStateTreeEditor::SaveAsset_Execute()
 {
-	// Remember the treview expnasion state
+	// Remember the treeview expansion state
 	if (StateTreeView)
 	{
 		StateTreeView->SavePersistentExpandedStates();
@@ -392,6 +393,19 @@ void FStateTreeEditor::OnSchemaChanged(const UStateTree& InStateTree)
 	}
 }
 
+void FStateTreeEditor::OnParametersChanged(const UStateTree& InStateTree)
+{
+	if (StateTree == &InStateTree)
+	{
+		// Accessible structs might be different after modifying parameters so forcing refresh
+		// so the FStateTreeBindingExtension can rebuild the list of bindable structs
+		if (SelectionDetailsView.IsValid())
+		{
+			SelectionDetailsView->ForceRefresh();
+		}
+	}
+}
+
 void FStateTreeEditor::OnAssetFinishedChangingProperties(const FPropertyChangedEvent& PropertyChangedEvent)
 {
 	// Make sure nodes get updates when properties are changed.
@@ -557,7 +571,7 @@ namespace UE::StateTree::Editor::Internal
 			return;
 		}
 		
-		const UStateTreeSchema* Schema = StateTree.GetSchema();
+		const UStateTreeSchema* Schema = TreeData->Schema;
 		if (!Schema)
 		{
 			return;
@@ -706,6 +720,8 @@ void FStateTreeEditor::Compile()
 	{
 		// Success
 		StateTree->LastCompiledEditorDataHash = EditorDataHash;
+
+		UE::StateTree::Delegates::OnPostCompile.Broadcast(*StateTree);
 	}
 	else
 	{
diff --git a/Engine/Plugins/Runtime/StateTree/Source/StateTreeEditorModule/Private/StateTreeEditor.h b/Engine/Plugins/Runtime/StateTree/Source/StateTreeEditorModule/Private/StateTreeEditor.h
--- a/Engine/Plugins/Runtime/StateTree/Source/StateTreeEditorModule/Private/StateTreeEditor.h
+++ b/Engine/Plugins/Runtime/StateTree/Source/StateTreeEditorModule/Private/StateTreeEditor.h
@@ -71,6 +71,7 @@ private:
 	void HandleModelAssetChanged();
 	void OnIdentifierChanged(const UStateTree& StateTree);
 	void OnSchemaChanged(const UStateTree& StateTree);
+	void OnParametersChanged(const UStateTree& StateTree);
 
 	FText GetStatisticsText() const;
 
diff --git a/Engine/Plugins/Runtime/StateTree/Source/StateTreeEditorModule/Private/StateTreeEditorData.cpp b/Engine/Plugins/Runtime/StateTree/Source/StateTreeEditorModule/Private/StateTreeEditorData.cpp
--- a/Engine/Plugins/Runtime/StateTree/Source/StateTreeEditorModule/Private/StateTreeEditorData.cpp
+++ b/Engine/Plugins/Runtime/StateTree/Source/StateTreeEditorModule/Private/StateTreeEditorData.cpp
@@ -2,16 +2,62 @@
 
 #include "StateTreeEditorData.h"
 #include "StateTreeConditionBase.h"
+#include "StateTreeDelegates.h"
 #include "StateTreeEvaluatorBase.h"
 #include "StateTreeTaskBase.h"
-#include "CoreMinimal.h"
-
 
 void UStateTreeEditorData::PostInitProperties()
 {
 	Super::PostInitProperties();
 }
 
+#if WITH_EDITOR
+void UStateTreeEditorData::PostEditChangeProperty(struct FPropertyChangedEvent& PropertyChangedEvent)
+{
+	Super::PostEditChangeProperty(PropertyChangedEvent);
+	
+	if (PropertyChangedEvent.MemberProperty && PropertyChangedEvent.Property)
+	{
+		const UStateTree* StateTree = GetTypedOuter<UStateTree>();
+		checkf(StateTree, TEXT("UStateTreeEditorData should only be allocated within a UStateTree"));
+		
+		const FName MemberName = PropertyChangedEvent.MemberProperty->GetFName();
+		if (MemberName == GET_MEMBER_NAME_CHECKED(UStateTreeEditorData, Schema))
+		{			
+			UE::StateTree::Delegates::OnSchemaChanged.Broadcast(*StateTree);
+		}
+		else if (MemberName == GET_MEMBER_NAME_CHECKED(UStateTreeEditorData, Parameters))
+		{
+			TArray<FStateTreeParameterDesc>& Params = Parameters.Parameters;
+			const FName PropertyName = PropertyChangedEvent.Property->GetFName();
+			
+			if (PropertyChangedEvent.ChangeType == EPropertyChangeType::Duplicate)
+			{
+				// Ensure unique ID on duplicated items.
+				const int32 ArrayIndex = PropertyChangedEvent.GetArrayIndex(PropertyName.ToString());
+				if (Params.IsValidIndex(ArrayIndex))
+				{
+					Params[ArrayIndex].Name = FName(Params[ArrayIndex].Name.ToString() + TEXT(" Duplicate"));
+					Params[ArrayIndex].ID = FGuid::NewGuid();
+				}
+			}
+			else if (PropertyChangedEvent.ChangeType == EPropertyChangeType::ArrayAdd
+					|| PropertyChangedEvent.ChangeType == EPropertyChangeType::ValueSet)
+			{
+				// Create unique ID on added items or ensure on set ones (maybe a copy/paste) 
+				const int32 ArrayIndex = PropertyChangedEvent.GetArrayIndex(PropertyName.ToString());
+				if (Params.IsValidIndex(ArrayIndex))
+				{
+					Params[ArrayIndex].ID = FGuid::NewGuid();
+				}
+			}
+
+			UE::StateTree::Delegates::OnParametersChanged.Broadcast(*StateTree);
+		}
+	}
+}
+#endif // WITH_EDITOR
+
 void UStateTreeEditorData::GetAccessibleStructs(const FGuid TargetStructID, TArray<FStateTreeBindableStructDesc>& OutStructDescs) const
 {
 	// Find the states that are updated before the current state.
@@ -27,6 +73,24 @@ void UStateTreeEditorData::GetAccessibleStructs(const FGuid TargetStructID, TArr
 
 void UStateTreeEditorData::GetAccessibleStructs(const TConstArrayView<const UStateTreeState*> Path, const FGuid TargetStructID, TArray<FStateTreeBindableStructDesc>& OutStructDescs) const
 {
+	const UStateTree* StateTree = GetTypedOuter<UStateTree>();
+	checkf(StateTree, TEXT("UStateTreeEditorData should only be allocated within a UStateTree"));
+	
+	// All parameters are accessible
+	for (const FStateTreeParameterDesc& Desc : Parameters.Parameters)
+	{
+		OutStructDescs.Emplace(Desc.Name, Desc.Parameter.GetScriptStruct(), Desc.ID);
+	}
+
+	// All named external data items declared by the schema are accessible
+	if (Schema != nullptr)
+	{
+		for (const FStateTreeExternalDataDesc& Desc : Schema->GetNamedExternalDataDescs())
+		{
+			OutStructDescs.Emplace(Desc.Name, Desc.Struct, Desc.ID);
+		}	
+	}
+
 	TArray<FStateTreeBindableStructDesc> EvalDescs;
 	TArray<FStateTreeBindableStructDesc> TaskDescs;
 
@@ -129,6 +193,24 @@ void UStateTreeEditorData::GetAllStructIDs(TMap<FGuid, const UStruct*>& AllStruc
 {
 	AllStructs.Reset();
 
+	const UStateTree* StateTree = GetTypedOuter<UStateTree>();
+	checkf(StateTree, TEXT("UStateTreeEditorData should only be allocated within a UStateTree"));
+
+	// All parameters
+	for (const FStateTreeParameterDesc& Desc : Parameters.Parameters)
+	{
+		AllStructs.Emplace(Desc.ID, Desc.Parameter.GetScriptStruct());
+	}
+
+	// All named external data items declared by the schema
+	if (Schema != nullptr)
+	{
+		for (const FStateTreeExternalDataDesc& Desc : Schema->GetNamedExternalDataDescs())
+		{
+			AllStructs.Emplace(Desc.ID, Desc.Struct);
+		}	
+	}
+
 	VisitHierarchy([&AllStructs](const UStateTreeState& State, const FGuid& ID, const FName& Name, const EStateTreeNodeType NodeType, const UScriptStruct* NodeStruct, const UStruct* InstanceStruct) -> bool
 		{
 			AllStructs.Add(ID, InstanceStruct);
diff --git a/Engine/Plugins/Runtime/StateTree/Source/StateTreeEditorModule/Public/StateTreeCompiler.h b/Engine/Plugins/Runtime/StateTree/Source/StateTreeEditorModule/Public/StateTreeCompiler.h
--- a/Engine/Plugins/Runtime/StateTree/Source/StateTreeEditorModule/Public/StateTreeCompiler.h
+++ b/Engine/Plugins/Runtime/StateTree/Source/StateTreeEditorModule/Public/StateTreeCompiler.h
@@ -2,7 +2,6 @@
 
 #pragma once
 
-#include "CoreMinimal.h"
 #include "StateTreeTypes.h"
 #include "StateTreePropertyBindingCompiler.h"
 #include "StateTreeCompilerLog.h"
@@ -20,14 +19,12 @@ struct FStateTreeStateLink;
 struct STATETREEEDITORMODULE_API FStateTreeCompiler
 {
 public:
-
-	FStateTreeCompiler(FStateTreeCompilerLog& InLog)
+	explicit FStateTreeCompiler(FStateTreeCompilerLog& InLog)
 		: Log(InLog)
 	{
 	}
 	
 	bool Compile(UStateTree& InStateTree);
-
 private:
 
 	bool ResolveTransitionState(const UStateTreeState& SourceState, const FStateTreeStateLink& Link, FStateTreeHandle& OutTransitionHandle) const;
@@ -55,7 +52,7 @@ private:
 	struct FExecutionPath
 	{
 		FExecutionPath() = default;
-		FExecutionPath(const TConstArrayView<const UStateTreeState*> InPath) : Path(InPath) {}
+		explicit FExecutionPath(const TConstArrayView<const UStateTreeState*> InPath) : Path(InPath) {}
 		
 		TArray<const UStateTreeState*> Path;
 	};
diff --git a/Engine/Plugins/Runtime/StateTree/Source/StateTreeEditorModule/Public/StateTreeEditorData.h b/Engine/Plugins/Runtime/StateTree/Source/StateTreeEditorModule/Public/StateTreeEditorData.h
--- a/Engine/Plugins/Runtime/StateTree/Source/StateTreeEditorModule/Public/StateTreeEditorData.h
+++ b/Engine/Plugins/Runtime/StateTree/Source/StateTreeEditorModule/Public/StateTreeEditorData.h
@@ -6,6 +6,8 @@
 #include "StateTreeEditorPropertyBindings.h"
 #include "StateTreeEditorData.generated.h"
 
+class UStateTreeSchema;
+
 UENUM()
 enum class EStateTreeNodeType : uint8
 {
@@ -32,6 +34,10 @@ public:
 	virtual FStateTreeEditorPropertyBindings* GetPropertyEditorBindings() override { return &EditorBindings; }
 	// ~IStateTreeEditorPropertyBindingsOwner
 
+#if WITH_EDITOR
+	virtual void PostEditChangeProperty(struct FPropertyChangedEvent& PropertyChangedEvent) override;
+#endif
+
 	/** Returns parent state of a struct, or nullptr if not found. */
 	const UStateTreeState* GetStateByStructID(const FGuid TargetStructID) const;
 
@@ -90,12 +96,19 @@ public:
 		EditorBindings.AddPropertyBinding(SourcePath, TargetPath);
 	}
 	// ~StateTree Builder API
+
+	/** Schema describing which inputs, evaluators, and tasks a StateTree can contain */	
+	UPROPERTY(EditDefaultsOnly, Category = Common, Instanced)
+	TObjectPtr<UStateTreeSchema> Schema = nullptr;
 	
+	/** Public parameters that could be used for bindings within the Tree. */
+	UPROPERTY(EditDefaultsOnly, Category = Common, meta=(ShowOnlyInnerProperties))
+	FStateTreeParameters Parameters;
 	
 	UPROPERTY(meta = (ExcludeFromHash))
 	FStateTreeEditorPropertyBindings EditorBindings;
 
-	// Top level States aka Routines.
+	/** Top level States. */
 	UPROPERTY()
 	TArray<UStateTreeState*> SubTrees;
 };
diff --git a/Engine/Plugins/Runtime/StateTree/Source/StateTreeModule/Private/StateTree.cpp b/Engine/Plugins/Runtime/StateTree/Source/StateTreeModule/Private/StateTree.cpp
--- a/Engine/Plugins/Runtime/StateTree/Source/StateTreeModule/Private/StateTree.cpp
+++ b/Engine/Plugins/Runtime/StateTree/Source/StateTreeModule/Private/StateTree.cpp
@@ -1,7 +1,6 @@
 // Copyright Epic Games, Inc. All Rights Reserved.
 
 #include "StateTree.h"
-#include "StateTreeDelegates.h"
 #include "StateTreeLinker.h"
 #include "StateTreeNodeBase.h"
 
@@ -17,6 +16,8 @@ void UStateTree::ResetCompiled()
 	States.Reset();
 	Transitions.Reset();
 
+	Schema = nullptr;
+	Parameters.Reset();
 	Nodes.Reset();
 	Instances.Reset();
 	InstanceObjects.Reset();
@@ -29,17 +30,6 @@ void UStateTree::ResetCompiled()
 	InstanceDataDefaultValue.Reset();
 }
 
-void UStateTree::PostEditChangeProperty(struct FPropertyChangedEvent& PropertyChangedEvent)
-{
-	if (PropertyChangedEvent.MemberProperty && PropertyChangedEvent.Property)
-	{
-		if (PropertyChangedEvent.MemberProperty->GetFName() == GET_MEMBER_NAME_CHECKED(UStateTree, Schema))
-		{
-			UE::StateTree::Delegates::OnSchemaChanged.Broadcast(*this);
-		}
-	}
-}
-
 void UStateTree::GetAssetRegistryTags(TArray<FAssetRegistryTag>& OutTags) const
 {
 	static const FName SchemaTag(TEXT("Schema"));
diff --git a/Engine/Plugins/Runtime/StateTree/Source/StateTreeModule/Private/StateTreeDelegates.cpp b/Engine/Plugins/Runtime/StateTree/Source/StateTreeModule/Private/StateTreeDelegates.cpp
--- a/Engine/Plugins/Runtime/StateTree/Source/StateTreeModule/Private/StateTreeDelegates.cpp
+++ b/Engine/Plugins/Runtime/StateTree/Source/StateTreeModule/Private/StateTreeDelegates.cpp
@@ -1,7 +1,6 @@
 // Copyright Epic Games, Inc. All Rights Reserved.
 
 #include "StateTreeDelegates.h"
-#include "CoreMinimal.h"
 
 namespace UE::StateTree::Delegates
 {
@@ -9,6 +8,8 @@ namespace UE::StateTree::Delegates
 #if WITH_EDITOR
 FOnIdentifierChanged OnIdentifierChanged;
 FOnSchemaChanged OnSchemaChanged;
+FOnParametersChanged OnParametersChanged;
+FOnPostCompile OnPostCompile;
 #endif
 	
 }; // UE::StateTree::Delegates
diff --git a/Engine/Plugins/Runtime/StateTree/Source/StateTreeModule/Private/StateTreeExecutionContext.cpp b/Engine/Plugins/Runtime/StateTree/Source/StateTreeModule/Private/StateTreeExecutionContext.cpp
--- a/Engine/Plugins/Runtime/StateTree/Source/StateTreeModule/Private/StateTreeExecutionContext.cpp
+++ b/Engine/Plugins/Runtime/StateTree/Source/StateTreeModule/Private/StateTreeExecutionContext.cpp
@@ -5,16 +5,13 @@
 #include "StateTreeTaskBase.h"
 #include "StateTreeEvaluatorBase.h"
 #include "StateTreeConditionBase.h"
-#include "CoreMinimal.h"
-#include "Engine/World.h"
+#include "Containers/StaticArray.h"
 #include "VisualLogger/VisualLogger.h"
 #include "ProfilingDebugging/CsvProfiler.h"
-#include "Algo/Reverse.h"
 #include "Logging/LogScopedVerbosityOverride.h"
-#include "ProfilingDebugging/CsvProfiler.h"
 
-#define STATETREE_LOG(Verbosity, Format, ...) UE_VLOG_UELOG(GetOwner(), LogStateTree, Verbosity, TEXT("%s: ") Format, *GetInstanceDescription(), ##__VA_ARGS__)
-#define STATETREE_CLOG(Condition, Verbosity, Format, ...) UE_CVLOG_UELOG((Condition), GetOwner(), LogStateTree, Verbosity, TEXT("%s: ") Format, *GetInstanceDescription(), ##__VA_ARGS__)
+#define STATETREE_LOG(Verbosity, Format, ...) UE_VLOG_UELOG(GetOwner(), LogStateTree, Verbosity, TEXT("%s") Format, *GetInstanceDescription(), ##__VA_ARGS__)
+#define STATETREE_CLOG(Condition, Verbosity, Format, ...) UE_CVLOG_UELOG((Condition), GetOwner(), LogStateTree, Verbosity, TEXT("%s") Format, *GetInstanceDescription(), ##__VA_ARGS__)
 
 namespace UE::StateTree
 {
@@ -53,12 +50,66 @@ bool FStateTreeExecutionContext::Init(UObject& InOwner, const UStateTree& InStat
 	// Initialize data views for all possible items.
 	DataViews.SetNum(StateTree->GetNumDataViews());
 
+	// Set data views associated to the parameters using the default values
+	SetDefaultParameters();
+
 	// Initialize array to keep track of which states have been updated.
 	VisitedStates.Init(false, StateTree->States.Num());
 
 	return true;
 }
 
+void FStateTreeExecutionContext::SetDefaultParameters()
+{
+	for (const FStateTreeParameterDesc& Desc : StateTree->GetParameterDescs())
+	{
+		check(DataViews.IsValidIndex(Desc.DataViewIndex));
+		DataViews[Desc.DataViewIndex] = FStateTreeDataView(Desc.Parameter);
+	}
+}
+
+void FStateTreeExecutionContext::SetParameters(const FStateTreeParameters& Parameters)
+{
+	const TConstArrayView<FStateTreeParameterDesc> TreeDescs = StateTree->GetParameterDescs();
+	const int32 NumTreeDescs = TreeDescs.Num();
+	if (NumTreeDescs == 0)
+	{
+		return;
+	}
+
+	// The current implementation golden path is based on an expected 1:1 match between the 2 parameter lists
+	// In case they mismatch then we search for matching parameter. 
+	int32 TreeDescIndex = 0;
+	for (const FStateTreeParameterDesc& ExternalDesc : Parameters.Parameters)
+	{
+		if (!TreeDescs[TreeDescIndex].IsMatching(ExternalDesc))
+		{
+			// The parameter was not found where we expected it to be, search for it starting from last matching index
+			bool bMatchFound = false;
+			for (int32 Index = 0; Index < NumTreeDescs; Index++)
+			{
+				TreeDescIndex = ++TreeDescIndex % NumTreeDescs;
+				if (TreeDescs[TreeDescIndex].IsMatching(ExternalDesc))
+				{
+					bMatchFound = true;
+					break;
+				}
+			}
+
+			if (!bMatchFound)
+			{
+				STATETREE_LOG(Warning, TEXT("Parameter '%s' doesn't match any expected parameters of StateTree '%s'."), *LexToString(ExternalDesc), *GetNameSafe(StateTree));
+				continue;
+			}
+		}
+
+		const uint16 DataViewIndex = TreeDescs[TreeDescIndex].DataViewIndex;
+		check(DataViewIndex != INDEX_NONE && DataViews.IsValidIndex(DataViewIndex));
+		DataViews[DataViewIndex] = FStateTreeDataView(ExternalDesc.Parameter);
+		TreeDescIndex = ++TreeDescIndex % NumTreeDescs;
+	}
+}
+
 void FStateTreeExecutionContext::Reset()
 {
 	InternalInstanceData.Reset();
@@ -227,7 +278,7 @@ EStateTreeRunStatus FStateTreeExecutionContext::Tick(const float DeltaTime, FSta
 
 	// The state selection is repeated up to MaxIteration time. This allows failed EnterState() to potentially find a new state immediately.
 	// This helps event driven StateTrees to not require another event/tick to find a suitable state.
-	static const int32 MaxIterations = 5;
+	static constexpr int32 MaxIterations = 5;
 	for (int32 Iter = 0; Iter < MaxIterations; Iter++)
 	{
 		// Trigger conditional transitions or state succeed/failed transitions. First tick transition is handled here too.
@@ -973,7 +1024,7 @@ FString FStateTreeExecutionContext::GetSafeStateName(const FStateTreeHandle Stat
 	return TEXT("(Unknown)");
 }
 
-FString FStateTreeExecutionContext::DebugGetStatePath(const FStateTreeActiveStates& ActiveStates, int32 ActiveStateIndex) const
+FString FStateTreeExecutionContext::DebugGetStatePath(const FStateTreeActiveStates& ActiveStates, const int32 ActiveStateIndex) const
 {
 	FString StatePath;
 	if (!ensureMsgf(ActiveStates.IsValidIndex(ActiveStateIndex), TEXT("Provided index must be valid")))
diff --git a/Engine/Plugins/Runtime/StateTree/Source/StateTreeModule/Private/StateTreeReference.cpp b/Engine/Plugins/Runtime/StateTree/Source/StateTreeModule/Private/StateTreeReference.cpp
new file mode 100644
--- /dev/null
+++ b/Engine/Plugins/Runtime/StateTree/Source/StateTreeModule/Private/StateTreeReference.cpp
@@ -0,0 +1,143 @@
+﻿// Copyright Epic Games, Inc. All Rights Reserved.
+
+#include "StateTreeReference.h"
+#include "StateTree.h"
+#include "StateTreeDelegates.h"
+#if WITH_EDITOR
+#include "Editor.h"
+#endif
+
+#if WITH_EDITOR
+void UStateTreeReferenceWrapper::PostEditChangeProperty(struct FPropertyChangedEvent& PropertyChangedEvent)
+{
+	UObject::PostEditChangeProperty(PropertyChangedEvent);
+
+	if (PropertyChangedEvent.MemberProperty && PropertyChangedEvent.Property)
+	{
+		if (PropertyChangedEvent.Property->GetFName() == GET_MEMBER_NAME_CHECKED(FStateTreeReference, StateTree))
+		{
+			StateTreeReference.SyncParameters();
+		}
+	}
+}
+
+void UStateTreeReferenceWrapper::PostEditChangeChainProperty(FPropertyChangedChainEvent& PropertyChangedEvent)
+{
+	UObject::PostEditChangeChainProperty(PropertyChangedEvent);
+}
+
+void UStateTreeReferenceWrapper::PostInitProperties()
+{
+	UObject::PostInitProperties();
+
+	// Registers a delegate to be notified when the associated StateTree asset get successfully recompiled
+	// to make sure that the parameters in the StateTreeReference are still valid.
+	PostCompileHandle = UE::StateTree::Delegates::OnPostCompile.AddLambda([this](const UStateTree& InStateTree)
+	{
+		if (StateTreeReference.StateTree == &InStateTree)
+		{
+			if (StateTreeReference.SyncParameters())
+			{
+				MarkPackageDirty();
+			}
+		}
+	});
+
+	// Make sure any modifications to the parameters since the last sync (PostLoad) won't cause runtime issues.
+	// This shouldn't be necessary if only parameters values could be modified but until that is enforce this will
+	// prevent parameters inconsistencies at runtime.
+	PIEHandle = FEditorDelegates::PreBeginPIE.AddLambda([this](const bool bBegan)
+	{
+		const bool bModified = StateTreeReference.SyncParameters();
+		UE_CLOG(bModified, LogStateTree, Warning, TEXT("Parameters for StateTree '%s' stored in %s were auto-fixed to be usable at runtime."),
+			*GetNameSafe(StateTreeReference.StateTree), *GetName());
+	});
+}
+
+void UStateTreeReferenceWrapper::PostLoad()
+{
+	UObject::PostLoad();
+
+	// This might modify the object but we don't want to dirty on load
+	StateTreeReference.SyncParameters();
+}
+
+void UStateTreeReferenceWrapper::BeginDestroy()
+{
+	UObject::BeginDestroy();
+
+	// Unregister all our delegates
+	UE::StateTree::Delegates::OnPostCompile.Remove(PostCompileHandle);
+	FEditorDelegates::PreBeginPIE.Remove(PIEHandle);
+}
+
+bool FStateTreeReference::SyncParameters()
+{
+	TArray<FStateTreeParameterDesc>& Params = Parameters.Parameters;
+	if (StateTree == nullptr)
+	{
+		const bool bHadParameters = !Params.IsEmpty();
+		Params.Reset();
+		return bHadParameters;
+	}
+
+	// First implementation is to make sure that the StateTreeReference holds values for each parameter of the StateTree.
+	// This is the easiest way to expose them to the user but that also means that user can't specify only a subset of parameters
+	// and use default value for the others.
+	const TConstArrayView<FStateTreeParameterDesc> AssetParameters = StateTree ? StateTree->GetParameterDescs() : TConstArrayView<FStateTreeParameterDesc>();
+	bool bAreParametersSynced = AssetParameters.Num() == Params.Num();
+	if (bAreParametersSynced)
+	{
+		for (int32 i = 0; i < Params.Num(); ++i)
+		{
+			if (!Params[i].IsMatching(AssetParameters[i]))
+			{
+				// Ok to mismatch only by the name if both ID and Type match. Simply update it since
+				// it must also match when passed to the ExecutionContext at runtime.
+				if (Params[i].ID == AssetParameters[i].ID && Params[i].IsSameType(AssetParameters[i]))
+				{
+					UE_LOG(LogStateTree, Warning, TEXT("StateTree parameter '%s' name was updated to match StateTree parameter '%s'."),
+						*LexToString(Params[i]), *LexToString(AssetParameters[i]));
+					
+					Params[i].Name = AssetParameters[i].Name;
+					checkf(Params[i].IsMatching(AssetParameters[i]),
+						TEXT("After fixing the name the parameters should match."
+							 " If this fails it indicates that same ID+Type is no longer the only required condition."));
+				}
+				else
+				{
+					bAreParametersSynced = false;
+					break;	
+				}
+			}
+		}
+	}
+
+	if (!bAreParametersSynced)
+	{
+		TArray<FStateTreeParameterDesc> DeprecatedParameters = MoveTemp(Params);
+		Params = AssetParameters;
+
+		for (int32 i = 0; i < Params.Num(); ++i)
+		{
+			FStateTreeParameterDesc& ParameterDesc = Params[i];
+
+			// Find parameter corresponding index in the deprecated parameters using 'Type' and 'Id'.
+			// Not using IsMatching() here since 'Name' might have change in addition to parameters added/removed since the last sync.
+			const int32 IndexInDeprecated = DeprecatedParameters.IndexOfByPredicate([&ParameterDesc](const FStateTreeParameterDesc& DeprecatedDesc)
+				{
+					return DeprecatedDesc.ID == ParameterDesc.ID && DeprecatedDesc.IsSameType(ParameterDesc);
+				});
+
+			if (IndexInDeprecated != INDEX_NONE)
+			{
+				// Only update value
+				ParameterDesc.Parameter = DeprecatedParameters[IndexInDeprecated].Parameter;
+				DeprecatedParameters.RemoveAtSwap(IndexInDeprecated);
+			}
+		}
+	}
+
+	return !bAreParametersSynced;
+}
+#endif // WITH_EDITOR
diff --git a/Engine/Plugins/Runtime/StateTree/Source/StateTreeModule/Public/StateTree.h b/Engine/Plugins/Runtime/StateTree/Source/StateTreeModule/Public/StateTree.h
--- a/Engine/Plugins/Runtime/StateTree/Source/StateTreeModule/Public/StateTree.h
+++ b/Engine/Plugins/Runtime/StateTree/Source/StateTreeModule/Public/StateTree.h
@@ -32,9 +32,11 @@ public:
 	/** @return List of external data required by the state tree */
 	TConstArrayView<FStateTreeExternalDataDesc> GetExternalDataDescs() const { return ExternalDataDescs; }
 
-	/** @return Schema describing which inputs, evaluators, and tasks a StateTree can contain */
-	const UStateTreeSchema* GetSchema() const { return Schema; }
-	void SetSchema(UStateTreeSchema* InSchema) { Schema = InSchema; }
+	/** @return List of named external data enforced by the schema that must be provided through the execution context. */
+	TConstArrayView<FStateTreeExternalDataDesc> GetNamedExternalDataDescs() const { return Schema ? Schema->GetNamedExternalDataDescs() : TConstArrayView<FStateTreeExternalDataDesc>(); }
+
+	/** @return List of parameters of the state tree. Default parameter values can be overridden at runtime by the execution context. */
+	TConstArrayView<FStateTreeParameterDesc> GetParameterDescs() const { return Parameters.Parameters; }
 
 	/** @return true if the tree asset can be used at runtime. */
 	bool IsReadyToRun() const;
@@ -66,7 +68,6 @@ protected:
 	virtual void Serialize(FStructuredArchiveRecord Record) override;
 	
 #if WITH_EDITOR
-	virtual void PostEditChangeProperty(struct FPropertyChangedEvent& PropertyChangedEvent) override;
 	virtual void GetAssetRegistryTags(TArray<FAssetRegistryTag>& OutTags) const override;
 #endif
 
@@ -74,9 +75,17 @@ private:
 
 	// Properties
 
-	UPROPERTY(EditDefaultsOnly, Category = Common, Instanced)
+	UPROPERTY(Instanced)
 	TObjectPtr<UStateTreeSchema> Schema = nullptr;
 
+	/**
+	 * Parameters that could be used for bindings within the Tree.
+	 * Default values are stored within the asset but StateTreeReference can be used to parameterized the tree.
+	 * @see FStateTreeReference
+	 */
+	UPROPERTY()
+	FStateTreeParameters Parameters;	
+
 	/** Evaluators, Tasks, and Condition items */
 	UPROPERTY()
 	TArray<FInstancedStruct> Nodes;
@@ -118,9 +127,3 @@ private:
 #endif
 };
 
-
-
-
-
-
-
diff --git a/Engine/Plugins/Runtime/StateTree/Source/StateTreeModule/Public/StateTreeDelegates.h b/Engine/Plugins/Runtime/StateTree/Source/StateTreeModule/Public/StateTreeDelegates.h
--- a/Engine/Plugins/Runtime/StateTree/Source/StateTreeModule/Public/StateTreeDelegates.h
+++ b/Engine/Plugins/Runtime/StateTree/Source/StateTreeModule/Public/StateTreeDelegates.h
@@ -2,7 +2,7 @@
 
 #pragma once
 
-#include "CoreMinimal.h"
+#include "Delegates/Delegate.h"
 
 class UStateTree;
 
@@ -10,13 +10,29 @@ namespace UE::StateTree::Delegates
 {
 	
 #if WITH_EDITOR
-	// Called when linkable name in a StateTree has changed.
+	/** Called when linkable name in a StateTree has changed. */
 	DECLARE_MULTICAST_DELEGATE_OneParam(FOnIdentifierChanged, const UStateTree& /*StateTree*/);
 	extern STATETREEMODULE_API FOnIdentifierChanged OnIdentifierChanged;
 
-	// Called when schema of the StateTree has changed.
+	/**
+	 * Called when schema of the StateTree EditorData has changed.
+	 * This is used to refresh the asset editor.
+	 * Note that this is NOT called when updating the StateTree schema from the EditorData on successful compilation.
+	 */
 	DECLARE_MULTICAST_DELEGATE_OneParam(FOnSchemaChanged, const UStateTree& /*StateTree*/);
 	extern STATETREEMODULE_API FOnSchemaChanged OnSchemaChanged;
+
+	/**
+	 * Called when parameters of the StateTree EditorData changed.
+	 * This should mainly used by the asset editor to maintain consistency in the UI for manipulations on the EditorData
+	 * until the tree gets compiled.
+	 */
+	DECLARE_MULTICAST_DELEGATE_OneParam(FOnParametersChanged, const UStateTree& /*StateTree*/);
+	extern STATETREEMODULE_API FOnParametersChanged OnParametersChanged;
+	
+	/** Called when compilation succeeds */
+	DECLARE_MULTICAST_DELEGATE_OneParam(FOnPostCompile, const UStateTree& /*StateTree*/);
+	extern STATETREEMODULE_API FOnPostCompile OnPostCompile;
 #endif
 
 }; // UE::StateTree::Delegates
diff --git a/Engine/Plugins/Runtime/StateTree/Source/StateTreeModule/Public/StateTreeExecutionContext.h b/Engine/Plugins/Runtime/StateTree/Source/StateTreeModule/Public/StateTreeExecutionContext.h
--- a/Engine/Plugins/Runtime/StateTree/Source/StateTreeModule/Public/StateTreeExecutionContext.h
+++ b/Engine/Plugins/Runtime/StateTree/Source/StateTreeModule/Public/StateTreeExecutionContext.h
@@ -4,7 +4,6 @@
 
 #include "StateTree.h"
 #include "InstancedStruct.h"
-#include "Containers/StaticArray.h"
 #include "StateTreePropertyBindings.h"
 #include "StateTreeInstanceData.h"
 #include "StateTreeExecutionContext.generated.h"
@@ -65,6 +64,15 @@ public:
 	/** Initializes the StateTree instance to be used with specific owner and StateTree asset. */
 	bool Init(UObject& InOwner, const UStateTree& InStateTree, const EStateTreeStorage InStorageType);
 
+	/** Updates data views of all parameters by using the default values defined in the StateTree asset. */
+	void SetDefaultParameters();
+
+	/**
+	 * Updates data views of the matching parameters by replacing the default values defined in the StateTree asset by the provided values.
+	 * Note: caller is responsible to make sure external parameters lifetime matches the context.
+	 */
+	void SetParameters(const FStateTreeParameters& Parameters);
+	
 	/** Resets the instance to initial empty state. Note: Does not call ExitState(). */
 	void Reset();
 
@@ -100,6 +108,13 @@ public:
 		return StateTree->ExternalDataDescs;
 	}
 
+	/** @return Array view to named external data descriptors associated with this context. Note: Init() must be called before calling this method. */
+	TConstArrayView<FStateTreeExternalDataDesc> GetNamedExternalDataDescs() const
+	{
+		check(StateTree);
+		return StateTree->GetNamedExternalDataDescs();
+	}
+
 	/** @return True if all required external data pointers are set. */ 
 	bool AreExternalDataViewsValid() const
 	{
@@ -111,7 +126,7 @@ public:
 			
 			if (DataDesc.Requirement == EStateTreeExternalDataRequirement::Required)
 			{
-				// Required items must have valid pointer and expected type.  
+				// Required items must have valid pointer of the expected type.  
 				if (!DataView.IsValid() || !DataView.GetStruct()->IsChildOf(DataDesc.Struct))
 				{
 					bResult = false;
@@ -126,7 +141,18 @@ public:
 					bResult = false;
 					break;
 				}
-			 
+			}
+		}
+
+		for (const FStateTreeExternalDataDesc& DataDesc : StateTree->GetNamedExternalDataDescs())
+		{
+			const FStateTreeDataView& DataView = DataViews[DataDesc.Handle.DataViewIndex];
+
+			// Items must have valid pointer of the expected type.  
+			if (!DataView.IsValid() || !DataView.GetStruct()->IsChildOf(DataDesc.Struct))
+			{
+				bResult = false;
+				break;
 			}
 		}
 		return bResult;
diff --git a/Engine/Plugins/Runtime/StateTree/Source/StateTreeModule/Public/StateTreeParameterTypes.h b/Engine/Plugins/Runtime/StateTree/Source/StateTreeModule/Public/StateTreeParameterTypes.h
new file mode 100644
--- /dev/null
+++ b/Engine/Plugins/Runtime/StateTree/Source/StateTreeModule/Public/StateTreeParameterTypes.h
@@ -0,0 +1,98 @@
+﻿// Copyright Epic Games, Inc. All Rights Reserved.
+
+#pragma once
+
+#include "StateTreeParameterTypes.generated.h"
+
+/**
+ * This base class was used to filter the whole family of structs wrapping types supported
+ * by StateTree parameters. If we decide to allow any struct then this is not required.
+ *		UPROPERTY(..., meta = (BaseStruct = "StateTreeParameterType", ExcludeBaseStruct))
+ *	vs	UPROPERTY(...)
+ */
+USTRUCT(meta = (Hidden))
+struct STATETREEMODULE_API FStateTreeParameterType
+{
+	GENERATED_BODY()
+};
+
+USTRUCT(DisplayName = "Bool")
+struct STATETREEMODULE_API FStateTreeParameterBool : public FStateTreeParameterType
+{
+	GENERATED_BODY()
+
+	UPROPERTY(EditAnywhere, Category = Parameter)
+	bool bValue = false;
+};
+
+USTRUCT(DisplayName = "Int")
+struct STATETREEMODULE_API FStateTreeParameterInt : public FStateTreeParameterType
+{
+	GENERATED_BODY()
+
+	UPROPERTY(EditAnywhere, Category = Parameter)
+	int32 Value = 0;
+};
+
+USTRUCT(DisplayName = "Float")
+struct STATETREEMODULE_API FStateTreeParameterFloat : public FStateTreeParameterType
+{
+	GENERATED_BODY()
+
+	UPROPERTY(EditAnywhere, Category = Parameter)
+	float Value = 0.f;
+};
+
+USTRUCT(DisplayName = "Real")
+struct STATETREEMODULE_API FStateTreeParameterReal : public FStateTreeParameterType
+{
+	GENERATED_BODY()
+
+	UPROPERTY(EditAnywhere, Category = Parameter)
+	double Value = 0.0;
+};
+
+USTRUCT(DisplayName = "Vector")
+struct STATETREEMODULE_API FStateTreeParameterVector : public FStateTreeParameterType
+{
+	GENERATED_BODY()
+
+	UPROPERTY(EditAnywhere, Category = Parameter)
+	FVector Value = FVector::ZeroVector;
+};
+
+USTRUCT(DisplayName = "Quat")
+struct STATETREEMODULE_API FStateTreeParameterQuat : public FStateTreeParameterType
+{
+	GENERATED_BODY()
+
+	UPROPERTY(EditAnywhere, Category = Parameter)
+	FQuat Value = FQuat::Identity;
+};
+
+USTRUCT(DisplayName = "Transform")
+struct STATETREEMODULE_API FStateTreeParameterTransform : public FStateTreeParameterType
+{
+	GENERATED_BODY()
+
+	UPROPERTY(EditAnywhere, Category = Parameter)
+	FTransform Value = FTransform::Identity;
+};
+
+USTRUCT(DisplayName = "String")
+struct STATETREEMODULE_API FStateTreeParameterString : public FStateTreeParameterType
+{
+	GENERATED_BODY()
+
+	UPROPERTY(EditAnywhere, Category = Parameter)
+	FString Value;
+};
+
+USTRUCT(DisplayName = "Name")
+struct STATETREEMODULE_API FStateTreeParameterName : public FStateTreeParameterType
+{
+	GENERATED_BODY()
+
+	UPROPERTY(EditAnywhere, Category = Parameter)
+	FName Value;
+};
diff --git a/Engine/Plugins/Runtime/StateTree/Source/StateTreeModule/Public/StateTreePropertyBindings.h b/Engine/Plugins/Runtime/StateTree/Source/StateTreeModule/Public/StateTreePropertyBindings.h
--- a/Engine/Plugins/Runtime/StateTree/Source/StateTreeModule/Public/StateTreePropertyBindings.h
+++ b/Engine/Plugins/Runtime/StateTree/Source/StateTreeModule/Public/StateTreePropertyBindings.h
@@ -153,7 +153,11 @@ struct STATETREEMODULE_API FStateTreeBindableStructDesc
 {
 	GENERATED_BODY()
 
+	FStateTreeBindableStructDesc() = default;
+
 #if WITH_EDITORONLY_DATA
+	FStateTreeBindableStructDesc(const FName InName, const UStruct* InStruct, const FGuid InGuid) : Struct(InStruct), Name(InName), ID(InGuid) {}
+
 	bool operator==(const FStateTreeBindableStructDesc& RHS) const
 	{
 		return ID == RHS.ID && Struct == RHS.Struct; // Not checking name, it's cosmetic.
diff --git a/Engine/Plugins/Runtime/StateTree/Source/StateTreeModule/Public/StateTreeReference.h b/Engine/Plugins/Runtime/StateTree/Source/StateTreeModule/Public/StateTreeReference.h
new file mode 100644
--- /dev/null
+++ b/Engine/Plugins/Runtime/StateTree/Source/StateTreeModule/Public/StateTreeReference.h
@@ -0,0 +1,58 @@
+﻿// Copyright Epic Games, Inc. All Rights Reserved.
+
+#pragma once
+
+#include "UObject/Object.h"
+#include "StateTreeTypes.h"
+#include "StateTreeReference.generated.h"
+
+class UStateTree;
+
+/**
+ * Struct to hold reference to a StateTree asset along with values to parameterized it.
+ */
+USTRUCT()
+struct FStateTreeReference
+{
+	GENERATED_BODY()
+
+	UPROPERTY(EditDefaultsOnly, Category="")
+	TObjectPtr<UStateTree> StateTree;
+
+	UPROPERTY(EditDefaultsOnly, Category="")
+	FStateTreeParameters Parameters;
+
+#if WITH_EDITOR
+	/**
+	 * Make sure that parameters are still compatible with those available in the selected StateTree asset.
+	 * @return true when parameters were 'fixed' to be in sync, false if they were already synced (i.e. unchanged).
+	 */
+	bool SyncParameters();
+#endif
+};
+
+
+/**
+ * Utility class wrapping a StateTreeReference to handle Editor actions to keep parameters synced.
+ */
+UCLASS(EditInlineNew)
+class STATETREEMODULE_API UStateTreeReferenceWrapper : public UObject
+{
+	GENERATED_BODY()
+
+public:
+	UPROPERTY(EditDefaultsOnly, Category="")
+	FStateTreeReference StateTreeReference;
+
+protected:
+#if WITH_EDITOR
+	virtual void PostEditChangeProperty(struct FPropertyChangedEvent& PropertyChangedEvent) override;
+	virtual void PostEditChangeChainProperty(FPropertyChangedChainEvent& PropertyChangedEvent) override;
+	virtual void PostInitProperties() override;
+	virtual void PostLoad() override;
+	virtual void BeginDestroy() override;
+	
+	FDelegateHandle PostCompileHandle;
+	FDelegateHandle PIEHandle;
+#endif
+};
diff --git a/Engine/Plugins/Runtime/StateTree/Source/StateTreeModule/Public/StateTreeSchema.h b/Engine/Plugins/Runtime/StateTree/Source/StateTreeModule/Public/StateTreeSchema.h
--- a/Engine/Plugins/Runtime/StateTree/Source/StateTreeModule/Public/StateTreeSchema.h
+++ b/Engine/Plugins/Runtime/StateTree/Source/StateTreeModule/Public/StateTreeSchema.h
@@ -2,6 +2,7 @@
 
 #pragma once
 
+#include "StateTreeTypes.h"
 #include "StateTreeSchema.generated.h"
 
 /**
@@ -39,6 +40,12 @@ public:
 	 */
 	bool IsChildOfBlueprintBase(const UClass* InClass) const;
 
+	/** @return List of named external data enforced by the schema. They must be provided at runtime through the execution context. */
+	virtual TConstArrayView<FStateTreeExternalDataDesc> GetNamedExternalDataDescs() const	{ return {}; }
+
+	/** @return List of mutable named external data enforced by the schema. They must be provided at runtime through the execution context. */
+	virtual TArrayView<FStateTreeExternalDataDesc> GetMutableNamedExternalDataDescs() { return {}; }
+
 #if WITH_EDITOR
 	
 	/** @return True if enter conditions are allowed. */
diff --git a/Engine/Plugins/Runtime/StateTree/Source/StateTreeModule/Public/StateTreeTypes.h b/Engine/Plugins/Runtime/StateTree/Source/StateTreeModule/Public/StateTreeTypes.h
--- a/Engine/Plugins/Runtime/StateTree/Source/StateTreeModule/Public/StateTreeTypes.h
+++ b/Engine/Plugins/Runtime/StateTree/Source/StateTreeModule/Public/StateTreeTypes.h
@@ -2,6 +2,7 @@
 
 #pragma once
 
+#include "InstancedStruct.h"
 #include "StateTreeTypes.generated.h"
 
 STATETREEMODULE_API DECLARE_LOG_CATEGORY_EXTERN(LogStateTree, Warning, All);
@@ -177,8 +178,8 @@ struct STATETREEMODULE_API FStateTreeActiveStates
 		{
 			return FStateTreeHandle::Invalid;			
 		}
-		
-		FStateTreeHandle Ret = States[NumStates - 1];
+
+		const FStateTreeHandle Ret = States[NumStates - 1];
 		NumStates--;
 		return Ret;
 	}
@@ -406,7 +407,8 @@ struct STATETREEMODULE_API FStateTreeExternalDataHandle
 	static bool IsValidIndex(const int32 Index) { return Index >= 0 && Index < (int32)IndexNone; }
 
 	bool IsValid() const { return DataViewIndex != IndexNone; }
-	
+
+	UPROPERTY()
 	uint8 DataViewIndex = IndexNone;
 };
 
@@ -482,6 +484,78 @@ struct TStateTreeInstanceDataPropertyHandle : FStateTreeInstanceDataPropertyHand
 	typedef T DataType;
 };
 
+/**
+ * Describes a parameter of the state tree that could be used for bindings.
+ * Can also be provided externally by a StateTreeReference to parameterized the tree.
+ * @see FStateTreeReference
+ */
+USTRUCT()
+struct STATETREEMODULE_API FStateTreeParameterDesc
+{
+	GENERATED_BODY()
+
+	/** Indicates that parameters hold the same data type */
+	bool IsSameType(const FStateTreeParameterDesc& RHS) const { return Parameter.GetScriptStruct() == RHS.Parameter.GetScriptStruct(); }
+
+	/** Indicates that parameters hold the same data type and have the same identifier. They might have different values. */
+	bool IsMatching(const FStateTreeParameterDesc& RHS) const
+	{				
+#if WITH_EDITORONLY_DATA
+		return Name == RHS.Name && ID == RHS.ID && IsSameType(RHS);
+#else
+		return Name == RHS.Name && IsSameType(RHS);
+#endif			
+	}
+	
+	friend FString LexToString(const FStateTreeParameterDesc& Desc)
+	{
+		return FString::Printf(TEXT("{%s} %s"),
+			Desc.Parameter.GetScriptStruct() != nullptr ? *Desc.Parameter.GetScriptStruct()->GetName() : TEXT("null type"),
+			*Desc.Name.ToString());
+	}
+
+	bool IsValid() const
+	{
+		return DataViewIndex != InvalidIndex;
+	}
+
+	/** The type of the parameter. */
+	UPROPERTY(EditDefaultsOnly, Category = Common)
+	FInstancedStruct Parameter;
+
+	/** Name of the parameter. */
+	UPROPERTY(EditDefaultsOnly, Category = Common)
+	FName Name;
+
+	/** The runtime data's data view index in the StateTreeExecutionContext, and source struct index in property binding. */
+	UPROPERTY()
+	uint16 DataViewIndex = InvalidIndex;
+
+	static constexpr uint16 InvalidIndex = MAX_uint16;
+
+#if WITH_EDITORONLY_DATA
+	/** Unique identifier. */
+	UPROPERTY()
+	FGuid ID;
+#endif
+};
+
+
+/**
+ * Container for StateTree parameters.
+ * Could be used for parameter definitions (within the StateTreeAsset) and parameterization (StateTreeReference).
+ */
+USTRUCT()
+struct FStateTreeParameters
+{
+	GENERATED_BODY()
+
+	void Reset() { Parameters.Reset(); }
+	
+	UPROPERTY(EditDefaultsOnly, Category= Common)
+	TArray<FStateTreeParameterDesc> Parameters;
+};
+
 
 /**
  * Describes an external data. The data can point to a struct or object.
@@ -495,6 +569,14 @@ struct STATETREEMODULE_API FStateTreeExternalDataDesc
 	FStateTreeExternalDataDesc() = default;
 	FStateTreeExternalDataDesc(const UStruct* InStruct, const EStateTreeExternalDataRequirement InRequirement) : Struct(InStruct), Requirement(InRequirement) {}
 
+	FStateTreeExternalDataDesc(const FName InName, const UStruct* InStruct, const FGuid InGuid)
+		: Struct(InStruct)
+		, Name(InName)
+#if WITH_EDITORONLY_DATA
+		, ID(InGuid)
+#endif
+	{}
+
 	bool operator==(const FStateTreeExternalDataDesc& Other) const
 	{
 		return Struct == Other.Struct && Requirement == Other.Requirement;
@@ -504,6 +586,14 @@ struct STATETREEMODULE_API FStateTreeExternalDataDesc
 	UPROPERTY();
 	const UStruct* Struct = nullptr;
 
+	/**
+	 * Name of the external data. Used only for bindable external data (enforced by the schema).
+	 * External data linked explicitly by the nodes (i.e. LinkExternalData) are identified only
+	 * by their type since they are used for unique instance of a given type.  
+	 */
+	UPROPERTY(VisibleAnywhere, Category = Common)
+	FName Name;
+	
 	/** Handle/Index to the StateTreeExecutionContext data views array */
 	UPROPERTY();
 	FStateTreeExternalDataHandle Handle;
@@ -511,6 +601,12 @@ struct STATETREEMODULE_API FStateTreeExternalDataDesc
 	/** Describes if the data is required or not. */
 	UPROPERTY();
 	EStateTreeExternalDataRequirement Requirement = EStateTreeExternalDataRequirement::Required;
+
+#if WITH_EDITORONLY_DATA
+	/** Unique identifier. Used only for bindable external data. */
+	UPROPERTY()
+	FGuid ID;
+#endif
 };
 
 
diff --git a/Engine/Plugins/Runtime/StateTree/Source/StateTreeTestSuite/Private/StateTreeTest.cpp b/Engine/Plugins/Runtime/StateTree/Source/StateTreeTestSuite/Private/StateTreeTest.cpp
--- a/Engine/Plugins/Runtime/StateTree/Source/StateTreeTestSuite/Private/StateTreeTest.cpp
+++ b/Engine/Plugins/Runtime/StateTree/Source/StateTreeTestSuite/Private/StateTreeTest.cpp
@@ -1,14 +1,11 @@
 // Copyright Epic Games, Inc. All Rights Reserved.
 
 #include "StateTreeTest.h"
-#include "CoreMinimal.h"
 #include "AITestsCommon.h"
 #include "StateTreeEditorData.h"
 #include "StateTreeState.h"
 #include "StateTreeCompiler.h"
 #include "Conditions/StateTreeCondition_Common.h"
-#include "StateTreeEvaluatorBase.h"
-#include "StateTreeTaskBase.h"
 #include "StateTree.h"
 #include "StateTreeTestTypes.h"
 #include "StateTreeExecutionContext.h"
@@ -27,8 +24,7 @@ namespace UE::StateTree::Tests
 		UStateTreeEditorData* EditorData = NewObject<UStateTreeEditorData>(StateTree);
 		check(EditorData);
 		StateTree->EditorData = EditorData;
-		UStateTreeSchema* Schema = NewObject<UStateTreeTestSchema>();
-		StateTree->SetSchema(Schema);
+		EditorData->Schema = NewObject<UStateTreeTestSchema>();
 		return *StateTree;
 	}
 
