StateTree: Allow to bind StateTree source structs directly.

diff --git a/Engine/Plugins/Runtime/StateTree/Source/StateTreeEditorModule/Private/Customizations/StateTreeBindingExtension.cpp b/Engine/Plugins/Runtime/StateTree/Source/StateTreeEditorModule/Private/Customizations/StateTreeBindingExtension.cpp
--- a/Engine/Plugins/Runtime/StateTree/Source/StateTreeEditorModule/Private/Customizations/StateTreeBindingExtension.cpp
+++ b/Engine/Plugins/Runtime/StateTree/Source/StateTreeEditorModule/Private/Customizations/StateTreeBindingExtension.cpp
@@ -90,6 +90,24 @@ EStateTreePropertyUsage ParsePropertyUsage(TSharedPtr<const IPropertyHandle> InP
 	return EStateTreePropertyUsage::Invalid;
 }
 
+FText GetSectionNameFromDataSource(const EStateTreeBindableStructSource Source)
+{
+	switch(Source)
+	{
+		case EStateTreeBindableStructSource::Task:
+			return LOCTEXT("Tasks", "Tasks");
+		case EStateTreeBindableStructSource::Evaluator:
+			return LOCTEXT("Evaluators", "Evaluators");
+		case EStateTreeBindableStructSource::TreeData:
+			return LOCTEXT("Tree Data", "Tree Data");
+		case EStateTreeBindableStructSource::StateParameter:
+			return LOCTEXT("State Parameter", "State Parameter");
+		case EStateTreeBindableStructSource::TreeParameter:
+			return LOCTEXT("Tree Parameter", "Tree Parameter");
+	}
+	return FText::GetEmpty();
+}
+
 
 FOnStateTreeBindingChanged STATETREEEDITORMODULE_API OnStateTreeBindingChanged;
 
@@ -145,7 +163,11 @@ void FStateTreeBindingExtension::ExtendWidgetRow(FDetailWidgetRow& InWidgetRow,
 			for (FStateTreeBindableStructDesc& StructDesc : AccessibleStructs)
 			{
 				const UStruct* Struct = StructDesc.Struct;
-				Context.Emplace(const_cast<UStruct*>(Struct), nullptr, FText::FromString(StructDesc.Name.ToString()));
+
+				FBindingContextStruct& ContextStruct = Context.AddDefaulted_GetRef();
+				ContextStruct.DisplayText = FText::FromString(StructDesc.Name.ToString());
+				ContextStruct.Struct = const_cast<UStruct*>(Struct);
+				ContextStruct.Section = UE::StateTree::PropertyBinding::GetSectionNameFromDataSource(StructDesc.DataSource);
 			}
 		}
 	}
@@ -207,6 +229,19 @@ void FStateTreeBindingExtension::ExtendWidgetRow(FDetailWidgetRow& InWidgetRow,
 			return bCanBind;
 		});
 
+	Args.OnCanBindToContextStruct = FOnCanBindToContextStruct::CreateLambda([InPropertyHandle](const UStruct* InStruct)
+		{
+			if (const FStructProperty* StructProperty = CastField<FStructProperty>(InPropertyHandle->GetProperty()))
+			{
+				return StructProperty->Struct == InStruct;
+			}
+			if (const FObjectProperty* ObjectProperty = CastField<FObjectProperty>(InPropertyHandle->GetProperty()))
+			{
+				return InStruct != nullptr && InStruct->IsChildOf(ObjectProperty->PropertyClass);
+			}
+			return false;
+		});
+
 	Args.OnCanAcceptPropertyOrChildren = FOnCanBindProperty::CreateLambda([](FProperty* InProperty)
 		{
 			// Make only editor visible properties visible for binding.
@@ -223,13 +258,18 @@ void FStateTreeBindingExtension::ExtendWidgetRow(FDetailWidgetRow& InWidgetRow,
 			IPropertyAccessEditor& PropertyAccessEditor = IModularFeatures::Get().GetModularFeature<IPropertyAccessEditor>("PropertyAccessEditor");
 			if (EditorBindings && OwnerObject)
 			{
-				if (TargetPath.IsValid() && InBindingChain.Num() > 1)	// Assume at least: [0] struct index, [1] a property.
+				if (TargetPath.IsValid() && InBindingChain.Num() > 0)
 				{
+					// First item in the binding chain is the index in AccessibleStructs.
 					FStateTreeEditorPropertyPath SourcePath;
 					const int32 SourceStructIndex = InBindingChain[0].ArrayIndex;
-					TArray<FBindingChainElement> SourceBindingChain(InBindingChain.GetData() + 1, InBindingChain.Num() - 1);
+					
+					TArray<FBindingChainElement> SourceBindingChain = InBindingChain;
+					SourceBindingChain.RemoveAt(0); // remove struct index.
+
 					check(SourceStructIndex >= 0 && SourceStructIndex < AccessibleStructs.Num());
 
+					// If SourceBindingChain is empty at this stage, it means that the binding points to the source struct itself.
 					SourcePath.StructID = AccessibleStructs[SourceStructIndex].ID;
 					PropertyAccessEditor.MakeStringPath(SourceBindingChain, SourcePath.Path);
 
diff --git a/Engine/Plugins/Runtime/StateTree/Source/StateTreeEditorModule/Private/StateTreePropertyBindingCompiler.cpp b/Engine/Plugins/Runtime/StateTree/Source/StateTreeEditorModule/Private/StateTreePropertyBindingCompiler.cpp
--- a/Engine/Plugins/Runtime/StateTree/Source/StateTreeEditorModule/Private/StateTreePropertyBindingCompiler.cpp
+++ b/Engine/Plugins/Runtime/StateTree/Source/StateTreeEditorModule/Private/StateTreePropertyBindingCompiler.cpp
@@ -69,7 +69,17 @@ bool FStateTreePropertyBindingCompiler::CompileBatch(const FStateTreeBindableStr
 		NewBinding.SourcePathIndex = SourceResult.PathIndex;
 		NewBinding.TargetPathIndex = TargetResult.PathIndex;
 		NewBinding.SourceStructIndex = SourceStructIdx;
-		NewBinding.CopyType = GetCopyType(SourceResult.LeafProperty, SourceResult.LeafArrayIndex, TargetResult.LeafProperty, TargetResult.LeafArrayIndex);
+		NewBinding.CopyType = GetCopyType(SourceStruct.Struct, SourceResult.LeafProperty, SourceResult.LeafArrayIndex,
+			TargetStruct.Struct, TargetResult.LeafProperty, TargetResult.LeafArrayIndex);
+
+		if (NewBinding.CopyType == EStateTreePropertyCopyType::None)
+		{
+			Log->Reportf(EMessageSeverity::Error, TargetStruct,
+				TEXT("Could not resolve copy type for properties '%s:%s' and '%s:%s'."),
+				*SourceStruct.Name.ToString(), *EditorBinding.SourcePath.ToString(),
+				*TargetStruct.Name.ToString(), *EditorBinding.TargetPath.ToString());
+			return false;
+		}
 	}
 
 	const int32 BindingsEnd = PropertyBindings->PropertyBindings.Num();
@@ -101,8 +111,30 @@ int32 FStateTreePropertyBindingCompiler::GetSourceStructIndexByID(const FGuid& I
 	return SourceStructs.IndexOfByPredicate([ID](const FStateTreeBindableStructDesc& Structs) { return (Structs.ID == ID); });
 }
 
-EStateTreePropertyCopyType FStateTreePropertyBindingCompiler::GetCopyType(const FProperty* SourceProperty, const int32 SourceArrayIndex, const FProperty* TargetProperty, const int32 TargetArrayIndex)
+EStateTreePropertyCopyType FStateTreePropertyBindingCompiler::GetCopyType(const UStruct* SourceStruct, const FProperty* SourceProperty, const int32 SourceArrayIndex,
+																		  const UStruct* TargetStruct, const FProperty* TargetProperty, const int32 TargetArrayIndex)
 {
+	if (SourceProperty == nullptr)
+	{
+		// Copy directly from the source struct, target must be.
+		if (const FStructProperty* StructProperty = CastField<FStructProperty>(TargetProperty))
+		{
+			if (StructProperty->Struct == SourceStruct)
+			{
+				return EStateTreePropertyCopyType::CopyStruct;
+			}
+		}
+		else if (const FObjectPropertyBase* ObjectProperty = CastField<FObjectPropertyBase>(TargetProperty))
+		{
+			if (SourceStruct->IsChildOf(ObjectProperty->PropertyClass))
+			{
+				return EStateTreePropertyCopyType::CopyObject;
+			}
+		}
+		
+		return EStateTreePropertyCopyType::None;
+	}
+	
 	if (const FArrayProperty* SourceArrayProperty = CastField<FArrayProperty>(SourceProperty))
 	{
 		// use the array's inner property if we are not trying to copy the whole array
@@ -294,10 +326,18 @@ bool FStateTreePropertyBindingCompiler::ResolvePropertyPath(const FStateTreeBind
 		return false;
 	}
 
+	// If the path is empty, we're pointing directly at the source struct.
+	if (InPath.Path.Num() == 0)
+	{
+		OutLeafProperty = nullptr;
+		OutLeafArrayIndex = INDEX_NONE;
+		return true;
+	}
+
 	const UStruct* CurrentStruct = InStructDesc.Struct;
 	const FProperty* LeafProperty = nullptr;
 	int32 LeafArrayIndex = INDEX_NONE;
-	bool bResult = InPath.Path.Num() > 0;
+	bool bResult = true;
 
 	for (int32 SegmentIndex = 0; SegmentIndex < InPath.Path.Num(); SegmentIndex++)
 	{
diff --git a/Engine/Plugins/Runtime/StateTree/Source/StateTreeEditorModule/Public/StateTreePropertyBindingCompiler.h b/Engine/Plugins/Runtime/StateTree/Source/StateTreeEditorModule/Public/StateTreePropertyBindingCompiler.h
--- a/Engine/Plugins/Runtime/StateTree/Source/StateTreeEditorModule/Public/StateTreePropertyBindingCompiler.h
+++ b/Engine/Plugins/Runtime/StateTree/Source/StateTreeEditorModule/Public/StateTreePropertyBindingCompiler.h
@@ -99,6 +99,6 @@ protected:
 		int32 LeafArrayIndex = INDEX_NONE;
 	};
 
-	EStateTreePropertyCopyType GetCopyType(const FProperty* SourceProperty, const int32 SourceArrayIndex, const FProperty* TargetProperty, const int32 TargetArrayIndex);
+	EStateTreePropertyCopyType GetCopyType(const UStruct* SourceStruct, const FProperty* SourceProperty, const int32 SourceArrayIndex, const UStruct* TargetStruct, const FProperty* TargetProperty, const int32 TargetArrayIndex);
 	[[nodiscard]] bool ResolvePropertyPath(const FStateTreeBindableStructDesc& InOwnerStructDesc, const FStateTreeBindableStructDesc& InStructDesc, const FStateTreeEditorPropertyPath& InPath, FResolvedPathResult& OutResult);
 };
diff --git a/Engine/Plugins/Runtime/StateTree/Source/StateTreeModule/Private/StateTreePropertyBindings.cpp b/Engine/Plugins/Runtime/StateTree/Source/StateTreeModule/Private/StateTreePropertyBindings.cpp
--- a/Engine/Plugins/Runtime/StateTree/Source/StateTreeModule/Private/StateTreePropertyBindings.cpp
+++ b/Engine/Plugins/Runtime/StateTree/Source/StateTreeModule/Private/StateTreePropertyBindings.cpp
@@ -184,6 +184,11 @@ bool FStateTreePropertyBindings::ValidateCopy(FStateTreePropCopy& Copy) const
 	const FProperty* SourceProperty = PropertyAccesses[Copy.SourceAccessIndex].LeafProperty;
 	const FProperty* TargetProperty = PropertyAccesses[Copy.TargetAccessIndex].LeafProperty;
 
+	if (SourceProperty == nullptr)
+	{
+		return TargetProperty->IsA<FStructProperty>() || TargetProperty->IsA<FObjectPropertyBase>();
+	}
+	
 	if (!SourceProperty || !TargetProperty)
 	{
 		return false;
@@ -297,6 +302,7 @@ uint8* FStateTreePropertyBindings::GetAddress(FStateTreeDataView InStructView, c
 	check(InStructView.IsValid());
 
 	uint8* Address = InStructView.GetMutableMemory();
+	
 	for (int32 i = InAccess.IndirectionsBegin; Address != nullptr && i < InAccess.IndirectionsEnd; i++)
 	{
 		const FStateTreePropertyIndirection& Indirection = PropertyIndirections[i];
@@ -352,7 +358,7 @@ uint8* FStateTreePropertyBindings::GetAddress(FStateTreeDataView InStructView, c
 
 void FStateTreePropertyBindings::PerformCopy(const FStateTreePropCopy& Copy, const FProperty* SourceProperty, const uint8* SourceAddress, const FProperty* TargetProperty, uint8* TargetAddress) const
 {
-	check(SourceProperty);
+	// Source property can be null
 	check(SourceAddress);
 	check(TargetProperty);
 	check(TargetAddress);
@@ -363,16 +369,25 @@ void FStateTreePropertyBindings::PerformCopy(const FStateTreePropCopy& Copy, con
 		FMemory::Memcpy(TargetAddress, SourceAddress, Copy.CopySize);
 		break;
 	case EStateTreePropertyCopyType::CopyComplex:
-		SourceProperty->CopyCompleteValue(TargetAddress, SourceAddress);
+		TargetProperty->CopyCompleteValue(TargetAddress, SourceAddress);
 		break;
 	case EStateTreePropertyCopyType::CopyBool:
 		static_cast<const FBoolProperty*>(TargetProperty)->SetPropertyValue(TargetAddress, static_cast<const FBoolProperty*>(SourceProperty)->GetPropertyValue(SourceAddress));
 		break;
 	case EStateTreePropertyCopyType::CopyStruct:
+		// If SourceProperty == nullptr (pointing to the struct source directly), the GetAddress() did the right thing and is pointing the the beginning of the struct. 
 		static_cast<const FStructProperty*>(TargetProperty)->Struct->CopyScriptStruct(TargetAddress, SourceAddress);
 		break;
 	case EStateTreePropertyCopyType::CopyObject:
-		static_cast<const FObjectPropertyBase*>(TargetProperty)->SetObjectPropertyValue(TargetAddress, static_cast<const FObjectPropertyBase*>(SourceProperty)->GetObjectPropertyValue(SourceAddress));
+		if (SourceProperty == nullptr)
+		{
+			// Source is pointing at object directly.
+			static_cast<const FObjectPropertyBase*>(TargetProperty)->SetObjectPropertyValue(TargetAddress, (UObject*)SourceAddress);
+		}
+		else
+		{
+			static_cast<const FObjectPropertyBase*>(TargetProperty)->SetObjectPropertyValue(TargetAddress, static_cast<const FObjectPropertyBase*>(SourceProperty)->GetObjectPropertyValue(SourceAddress));
+		}
 		break;
 	case EStateTreePropertyCopyType::CopyName:
 		static_cast<const FNameProperty*>(TargetProperty)->SetPropertyValue(TargetAddress, static_cast<const FNameProperty*>(SourceProperty)->GetPropertyValue(SourceAddress));
@@ -388,7 +403,7 @@ void FStateTreePropertyBindings::PerformCopy(const FStateTreePropCopy& Copy, con
 		const int32 MinSize = FMath::Min(SourceArrayHelper.Num(), TargetArrayHelper.Num());
 		for (int32 ElementIndex = 0; ElementIndex < MinSize; ++ElementIndex)
 		{
-			SourceArrayProperty->Inner->CopySingleValue(TargetArrayHelper.GetRawPtr(ElementIndex), SourceArrayHelper.GetRawPtr(ElementIndex));
+			TargetArrayProperty->Inner->CopySingleValue(TargetArrayHelper.GetRawPtr(ElementIndex), SourceArrayHelper.GetRawPtr(ElementIndex));
 		}
 		break;
 	}
@@ -482,8 +497,11 @@ void FStateTreePropertyBindings::CopyTo(TConstArrayView<FStateTreeDataView> Sour
 		}
 		const FStateTreePropertyAccess& SourceAccess = PropertyAccesses[Copy.SourceAccessIndex];
 		const FStateTreePropertyAccess& TargetAccess = PropertyAccesses[Copy.TargetAccessIndex];
-		check(SourceStructViews[Copy.SourceStructIndex].GetStruct() == SourceStructs[Copy.SourceStructIndex].Struct);
-		const uint8* SourceAddress = GetAddress(SourceStructViews[Copy.SourceStructIndex], SourceAccess);
+
+		const FStateTreeDataView SourceStructView = SourceStructViews[Copy.SourceStructIndex];
+		check(SourceStructView.GetStruct() == SourceStructs[Copy.SourceStructIndex].Struct || SourceStructView.GetStruct()->IsChildOf(SourceStructs[Copy.SourceStructIndex].Struct));
+		
+		const uint8* SourceAddress = GetAddress(SourceStructView, SourceAccess);
 		uint8* TargetAddress = GetAddress(TargetStructView, TargetAccess);
 		if (SourceAddress && TargetAddress)
 		{
diff --git a/Engine/Plugins/Runtime/StateTree/Source/StateTreeModule/Public/StateTreePropertyBindings.h b/Engine/Plugins/Runtime/StateTree/Source/StateTreeModule/Public/StateTreePropertyBindings.h
--- a/Engine/Plugins/Runtime/StateTree/Source/StateTreeModule/Public/StateTreePropertyBindings.h
+++ b/Engine/Plugins/Runtime/StateTree/Source/StateTreeModule/Public/StateTreePropertyBindings.h
@@ -144,6 +144,16 @@ protected:
 };
 
 
+UENUM()
+enum class EStateTreeBindableStructSource : uint8
+{
+	TreeData,
+	TreeParameter,
+	StateParameter, 
+	Task,
+	Evaluator,
+};
+
 /**
  * Descriptor for a struct or class that can be a binding source or target.
  * Each struct has unique identifier, which is used to distinguish them, and name that is mostly for debugging and UI.
@@ -172,6 +182,9 @@ struct STATETREEMODULE_API FStateTreeBindableStructDesc
 	UPROPERTY()
 	FName Name;
 
+	UPROPERTY()
+	EStateTreeBindableStructSource DataSource = EStateTreeBindableStructSource::TreeData;
+	
 #if WITH_EDITORONLY_DATA
 	/** Unique identifier of the struct. */
 	UPROPERTY()
@@ -518,7 +531,7 @@ struct STATETREEMODULE_API FStateTreeEditorPropertyPath
 	UPROPERTY()
 	TArray<FString> Path;
 
-	bool IsValid() const { return StructID.IsValid() && Path.Num() > 0; }
+	bool IsValid() const { return StructID.IsValid(); }
 
 	bool operator==(const FStateTreeEditorPropertyPath& RHS) const;
 };
