Fix for uninitialized struct member error in city sample

diff --git a/Engine/Plugins/Runtime/StateTree/Source/StateTreeModule/Public/StateTreeReference.h b/Engine/Plugins/Runtime/StateTree/Source/StateTreeModule/Public/StateTreeReference.h
--- a/Engine/Plugins/Runtime/StateTree/Source/StateTreeModule/Public/StateTreeReference.h
+++ b/Engine/Plugins/Runtime/StateTree/Source/StateTreeModule/Public/StateTreeReference.h
@@ -17,7 +17,7 @@ struct FStateTreeReference
 	GENERATED_BODY()
 
 	UPROPERTY(EditDefaultsOnly, Category="")
-	TObjectPtr<UStateTree> StateTree;
+	TObjectPtr<UStateTree> StateTree = nullptr;
 
 	UPROPERTY(EditDefaultsOnly, Category="")
 	FStateTreeParameters Parameters;
