StateTree: fix PVS warnings
- The 'TargetProperty' pointer was utilized before it was verified against nullpt
- The modification of the 'TreeDescIndex' variable is unsequenced relative to another operation on the same variable. This may lead to undefined behavior.

diff --git a/Engine/Plugins/Runtime/StateTree/Source/StateTreeModule/Private/StateTreeExecutionContext.cpp b/Engine/Plugins/Runtime/StateTree/Source/StateTreeModule/Private/StateTreeExecutionContext.cpp
--- a/Engine/Plugins/Runtime/StateTree/Source/StateTreeModule/Private/StateTreeExecutionContext.cpp
+++ b/Engine/Plugins/Runtime/StateTree/Source/StateTreeModule/Private/StateTreeExecutionContext.cpp
@@ -88,7 +88,7 @@ void FStateTreeExecutionContext::SetParameters(const FStateTreeParameters& Param
 			bool bMatchFound = false;
 			for (int32 Index = 0; Index < NumTreeDescs; Index++)
 			{
-				TreeDescIndex = ++TreeDescIndex % NumTreeDescs;
+				TreeDescIndex = (TreeDescIndex + 1) % NumTreeDescs;
 				if (TreeDescs[TreeDescIndex].IsMatching(ExternalDesc))
 				{
 					bMatchFound = true;
@@ -106,7 +106,7 @@ void FStateTreeExecutionContext::SetParameters(const FStateTreeParameters& Param
 		const uint16 DataViewIndex = TreeDescs[TreeDescIndex].DataViewIndex;
 		check(DataViewIndex != INDEX_NONE && DataViews.IsValidIndex(DataViewIndex));
 		DataViews[DataViewIndex] = FStateTreeDataView(ExternalDesc.Parameter);
-		TreeDescIndex = ++TreeDescIndex % NumTreeDescs;
+		TreeDescIndex = (TreeDescIndex + 1) % NumTreeDescs;
 	}
 }
 
diff --git a/Engine/Plugins/Runtime/StateTree/Source/StateTreeModule/Private/StateTreePropertyBindings.cpp b/Engine/Plugins/Runtime/StateTree/Source/StateTreeModule/Private/StateTreePropertyBindings.cpp
--- a/Engine/Plugins/Runtime/StateTree/Source/StateTreeModule/Private/StateTreePropertyBindings.cpp
+++ b/Engine/Plugins/Runtime/StateTree/Source/StateTreeModule/Private/StateTreePropertyBindings.cpp
@@ -184,14 +184,15 @@ bool FStateTreePropertyBindings::ValidateCopy(FStateTreePropCopy& Copy) const
 	const FProperty* SourceProperty = PropertyAccesses[Copy.SourceAccessIndex].LeafProperty;
 	const FProperty* TargetProperty = PropertyAccesses[Copy.TargetAccessIndex].LeafProperty;
 
-	if (SourceProperty == nullptr)
+	if (!TargetProperty)
 	{
-		return TargetProperty->IsA<FStructProperty>() || TargetProperty->IsA<FObjectPropertyBase>();
+		return false;
 	}
-	
-	if (!SourceProperty || !TargetProperty)
+
+	// If source property is nullptr, we're binding directly to the source.
+	if (SourceProperty == nullptr)
 	{
-		return false;
+		return TargetProperty->IsA<FStructProperty>() || TargetProperty->IsA<FObjectPropertyBase>();
 	}
 
 	// Extract underlying types for enums
