StateTree: Fix crash in undo
- Fix crash when undoing a move or delete
- Cleaned up the logic for move and delete

diff --git a/Engine/Plugins/Runtime/StateTree/Source/StateTreeEditorModule/Private/SStateTreeView.cpp b/Engine/Plugins/Runtime/StateTree/Source/StateTreeEditorModule/Private/SStateTreeView.cpp
--- a/Engine/Plugins/Runtime/StateTree/Source/StateTreeEditorModule/Private/SStateTreeView.cpp
+++ b/Engine/Plugins/Runtime/StateTree/Source/StateTreeEditorModule/Private/SStateTreeView.cpp
@@ -179,7 +179,7 @@ void SStateTreeView::Tick(const FGeometry& AllottedGeometry, const double InCurr
 {
 	if (bItemsDirty)
 	{
-		UpdateTree();
+		UpdateTree(/*bExpandPersistent*/true);
 	}
 
 	if (RequestedRenameState && !TreeView->IsPendingRefresh())
diff --git a/Engine/Plugins/Runtime/StateTree/Source/StateTreeEditorModule/Private/StateTreeState.cpp b/Engine/Plugins/Runtime/StateTree/Source/StateTreeEditorModule/Private/StateTreeState.cpp
--- a/Engine/Plugins/Runtime/StateTree/Source/StateTreeEditorModule/Private/StateTreeState.cpp
+++ b/Engine/Plugins/Runtime/StateTree/Source/StateTreeEditorModule/Private/StateTreeState.cpp
@@ -40,7 +40,6 @@ FStateTreeTransition::FStateTreeTransition(const EStateTreeTransitionEvent InEve
 UStateTreeState::UStateTreeState(const FObjectInitializer& ObjectInitializer)
 	: Super(ObjectInitializer)
 	, ID(FGuid::NewGuid())
-	, Parent(nullptr)
 {
 }
 
@@ -121,11 +120,24 @@ void UStateTreeState::PostEditChangeChainProperty(FPropertyChangedChainEvent& Pr
 					}
 				}
 				// TODO: Transition conditions.
+				
 			}
 		}
 	}
 }
 
+void UStateTreeState::PostLoad()
+{
+	Super::PostLoad();
+
+	// Make sure state has transactional flags to make it work with undo (to fix a bug where root states were created without this flag).
+	if (!HasAnyFlags(RF_Transactional))
+	{
+		SetFlags(RF_Transactional);
+	}
+}
+
+
 #endif
 
 UStateTreeState* UStateTreeState::GetNextSiblingState() const
diff --git a/Engine/Plugins/Runtime/StateTree/Source/StateTreeEditorModule/Private/StateTreeViewModel.cpp b/Engine/Plugins/Runtime/StateTree/Source/StateTreeEditorModule/Private/StateTreeViewModel.cpp
--- a/Engine/Plugins/Runtime/StateTree/Source/StateTreeEditorModule/Private/StateTreeViewModel.cpp
+++ b/Engine/Plugins/Runtime/StateTree/Source/StateTreeEditorModule/Private/StateTreeViewModel.cpp
@@ -18,77 +18,8 @@
 #define LOCTEXT_NAMESPACE "StateTreeEditor"
 
 
-namespace FStateTreeViewUtilities
+namespace UE::StateTree::Editor
 {
-	// Remove state from the array. Recurses into children if no match found.
-	bool RemoveRecursive(TArray<UStateTreeState*>& Array, UStateTreeState* StateToRemove)
-	{
-		if (!StateToRemove)
-		{
-			return false;
-		}
-
-		int32 ItemIndex = Array.Find(StateToRemove);
-		if (ItemIndex != INDEX_NONE)
-		{
-			Array.RemoveAt(ItemIndex);
-			StateToRemove->Parent = nullptr;
-			return true;
-		}
-
-		// Did not successfully remove an item. Try all the children.
-		for (int32 i = 0; i < Array.Num(); ++i)
-		{
-			if (RemoveRecursive(Array[i]->Children, StateToRemove))
-			{
-				return true;
-			}
-		}
-
-		return false;
-	}
-
-	// Insert state in the array relative to target state. Recurses into children if no match found.
-	bool InsertRecursive(UStateTreeState* ParentState, TArray<UStateTreeState*>& ParentArray, UStateTreeState* TargetState, UStateTreeState* StateToInsert, int32 RelativeLocation)
-	{
-		if (!StateToInsert || !TargetState)
-		{
-			return false;
-		}
-
-		const int32 TargetIndex = ParentArray.Find(TargetState);
-		if (TargetIndex != INDEX_NONE)
-		{
-			if (RelativeLocation == -1)
-			{
-				ParentArray.Insert(StateToInsert, TargetIndex);
-				StateToInsert->Parent = ParentState;
-			}
-			else if (RelativeLocation == 1)
-			{
-				ParentArray.Insert(StateToInsert, TargetIndex + 1);
-				StateToInsert->Parent = ParentState;
-			}
-			else
-			{
-				ensure(RelativeLocation == 0);
-				ParentArray[TargetIndex]->Children.Insert(StateToInsert, 0);
-				StateToInsert->Parent = ParentArray[TargetIndex];
-			}
-			return true;
-		}
-
-		for (int32 i = 0; i < ParentArray.Num(); ++i)
-		{
-			if (InsertRecursive(ParentArray[i], ParentArray[i]->Children, TargetState, StateToInsert, RelativeLocation))
-			{
-				return true;
-			}
-		}
-
-		return false;
-	}
-
 	// Removes states from the array which are children of any other state.
 	void RemoveContainedChildren(TArray<UStateTreeState*>& States)
 	{
@@ -168,7 +99,7 @@ void FStateTreeViewModel::Init(UStateTreeEditorData* InTreeData)
 	UE::StateTree::Delegates::OnIdentifierChanged.AddSP(this, &FStateTreeViewModel::HandleIdentifierChanged);
 }
 
-void FStateTreeViewModel::HandleIdentifierChanged(const UStateTree& StateTree)
+void FStateTreeViewModel::HandleIdentifierChanged(const UStateTree& StateTree) const
 {
 	const UStateTree* OuterStateTree = TreeData ? Cast<UStateTree>(TreeData->GetOuter()) : nullptr;
 	if (OuterStateTree == &StateTree)
@@ -177,17 +108,17 @@ void FStateTreeViewModel::HandleIdentifierChanged(const UStateTree& StateTree)
 	}
 }
 
-void FStateTreeViewModel::NotifyAssetChangedExternally()
+void FStateTreeViewModel::NotifyAssetChangedExternally() const
 {
 	OnAssetChanged.Broadcast();
 }
 
-void FStateTreeViewModel::NotifyStatesChangedExternally(const TSet<UStateTreeState*>& ChangedStates, const FPropertyChangedEvent& PropertyChangedEvent)
+void FStateTreeViewModel::NotifyStatesChangedExternally(const TSet<UStateTreeState*>& ChangedStates, const FPropertyChangedEvent& PropertyChangedEvent) const
 {
 	OnStatesChanged.Broadcast(ChangedStates, PropertyChangedEvent);
 }
 
-TArray<UStateTreeState*>* FStateTreeViewModel::GetSubTrees()
+TArray<UStateTreeState*>* FStateTreeViewModel::GetSubTrees() const
 {
 	return TreeData != nullptr ? &TreeData->SubTrees : nullptr;
 }
@@ -212,18 +143,18 @@ void FStateTreeViewModel::ClearSelection()
 {
 	SelectedStates.Reset();
 
-	TArray<UStateTreeState*> SelectedStatesArr;
+	const TArray<UStateTreeState*> SelectedStatesArr;
 	OnSelectionChanged.Broadcast(SelectedStatesArr);
 }
 
-void FStateTreeViewModel::SetSelection(UStateTreeState* SelectedState)
+void FStateTreeViewModel::SetSelection(UStateTreeState* Selected)
 {
 	SelectedStates.Reset();
 
-	SelectedStates.Add(SelectedState);
+	SelectedStates.Add(Selected);
 
 	TArray<UStateTreeState*> SelectedStatesArr;
-	SelectedStatesArr.Add(SelectedState);
+	SelectedStatesArr.Add(Selected);
 	OnSelectionChanged.Broadcast(SelectedStatesArr);
 }
 
@@ -258,12 +189,10 @@ bool FStateTreeViewModel::IsChildOfSelection(const UStateTreeState* State) const
 			{
 				return true;
 			}
-			else
+			
+			if (UE::StateTree::Editor::IsChildOf(SelectedState, State))
 			{
-				if (FStateTreeViewUtilities::IsChildOf(SelectedState, State))
-				{
-					return true;
-				}
+				return true;
 			}
 		}
 	}
@@ -290,7 +219,7 @@ bool FStateTreeViewModel::HasSelection() const
 void FStateTreeViewModel::GetPersistentExpandedStates(TSet<UStateTreeState*>& OutExpandedStates)
 {
 	OutExpandedStates.Reset();
-	if (!TreeData)
+	if (TreeData == nullptr)
 	{
 		return;
 	}
@@ -314,7 +243,7 @@ void FStateTreeViewModel::GetExpandedStatesRecursive(UStateTreeState* State, TSe
 
 void FStateTreeViewModel::SetPersistentExpandedStates(TSet<UStateTreeState*>& InExpandedStates)
 {
-	if (!TreeData)
+	if (TreeData == nullptr)
 	{
 		return;
 	}
@@ -323,7 +252,7 @@ void FStateTreeViewModel::SetPersistentExpandedStates(TSet<UStateTreeState*>& In
 
 	for (UStateTreeState* State : InExpandedStates)
 	{
-		if (State)
+		if (State != nullptr)
 		{
 			State->bExpanded = true;
 		}
@@ -333,7 +262,7 @@ void FStateTreeViewModel::SetPersistentExpandedStates(TSet<UStateTreeState*>& In
 
 void FStateTreeViewModel::AddState(UStateTreeState* AfterState)
 {
-	if (!TreeData)
+	if (TreeData == nullptr)
 	{
 		return;
 	}
@@ -343,28 +272,34 @@ void FStateTreeViewModel::AddState(UStateTreeState* AfterState)
 	UStateTreeState* NewState = NewObject<UStateTreeState>(TreeData, FName(), RF_Transactional);
 	UStateTreeState* ParentState = nullptr;
 
-	if (AfterState)
+	if (AfterState != nullptr)
 	{
 		ParentState = AfterState->Parent;
-		if (ParentState)
+		if (ParentState != nullptr)
 		{
 			ParentState->Modify();
-			NewState->Parent = ParentState;
 		}
 		else
 		{
 			TreeData->Modify();
-			NewState->Parent = nullptr;
 		}
 
-		TArray<UStateTreeState*>& ArrayToAddTo = ParentState ? ParentState->Children : TreeData->SubTrees;
-		FStateTreeViewUtilities::InsertRecursive(ParentState, ArrayToAddTo, AfterState, NewState, 1); // Insert after
-	}
-	else
-	{
-		TreeData->Modify();
-		NewState->Parent = nullptr;
-		TreeData->SubTrees.Add(NewState);
+		TArray<UStateTreeState*>& ParentArray = ParentState ? ParentState->Children : TreeData->SubTrees;
+
+		const int32 TargetIndex = ParentArray.Find(AfterState);
+		if (TargetIndex != INDEX_NONE)
+		{
+			// Insert After
+			ParentArray.Insert(NewState, TargetIndex + 1);
+			NewState->Parent = ParentState;
+		}
+		else
+		{
+			// Fallback, should never happen.
+			ensureMsgf(false, TEXT("%s: Failed to find specified target state %s on state %s while adding new state."), *GetNameSafe(TreeData->GetOuter()), *GetNameSafe(AfterState), *GetNameSafe(ParentState));
+			ParentArray.Add(NewState);
+			NewState->Parent = ParentState;
+		}
 	}
 
 	OnStateAdded.Broadcast(ParentState, NewState);
@@ -372,17 +307,17 @@ void FStateTreeViewModel::AddState(UStateTreeState* AfterState)
 
 void FStateTreeViewModel::AddChildState(UStateTreeState* ParentState)
 {
-	if (!TreeData || !ParentState)
+	if (TreeData == nullptr || ParentState == nullptr)
 	{
 		return;
 	}
 
 	const FScopedTransaction Transaction(LOCTEXT("AddChildStateTransaction", "Add Child State"));
 
-	TreeData->Modify();
-
-	UStateTreeState* NewState = NewObject<UStateTreeState>(TreeData, FName(), RF_Transactional);
+	UStateTreeState* NewState = NewObject<UStateTreeState>(ParentState, FName(), RF_Transactional);
 
+	ParentState->Modify();
+	
 	ParentState->Children.Add(NewState);
 	NewState->Parent = ParentState;
 
@@ -391,7 +326,7 @@ void FStateTreeViewModel::AddChildState(UStateTreeState* ParentState)
 
 void FStateTreeViewModel::RenameState(UStateTreeState* State, FName NewName)
 {
-	if (!State)
+	if (State == nullptr)
 	{
 		return;
 	}
@@ -410,7 +345,7 @@ void FStateTreeViewModel::RenameState(UStateTreeState* State, FName NewName)
 
 void FStateTreeViewModel::RemoveSelectedStates()
 {
-	if (!TreeData)
+	if (TreeData == nullptr)
 	{
 		return;
 	}
@@ -419,7 +354,7 @@ void FStateTreeViewModel::RemoveSelectedStates()
 	GetSelectedStates(States);
 
 	// Remove items whose parent also exists in the selection.
-	FStateTreeViewUtilities::RemoveContainedChildren(States);
+	UE::StateTree::Editor::RemoveContainedChildren(States);
 
 	if (States.Num() > 0)
 	{
@@ -427,22 +362,32 @@ void FStateTreeViewModel::RemoveSelectedStates()
 
 		TSet<UStateTreeState*> AffectedParents;
 
-		for (UStateTreeState* State : States)
+		for (UStateTreeState* StateToRemove : States)
 		{
-			if (State)
+			if (StateToRemove)
 			{
-				if (UStateTreeState* ParentState = State->Parent)
+				StateToRemove->Modify();
+
+				UStateTreeState* ParentState = StateToRemove->Parent;
+				if (ParentState != nullptr)
 				{
 					AffectedParents.Add(ParentState);
 					ParentState->Modify();
-					FStateTreeViewUtilities::RemoveRecursive(ParentState->Children, State);
 				}
 				else
 				{
 					AffectedParents.Add(nullptr);
 					TreeData->Modify();
-					FStateTreeViewUtilities::RemoveRecursive(TreeData->SubTrees, State);
 				}
+				
+				TArray<UStateTreeState*>& ArrayToRemoveFrom = ParentState ? ParentState->Children : TreeData->SubTrees;
+				const int32 ItemIndex = ArrayToRemoveFrom.Find(StateToRemove);
+				if (ItemIndex != INDEX_NONE)
+				{
+					ArrayToRemoveFrom.RemoveAt(ItemIndex);
+					StateToRemove->Parent = nullptr;
+				}
+
 			}
 		}
 
@@ -454,22 +399,22 @@ void FStateTreeViewModel::RemoveSelectedStates()
 
 void FStateTreeViewModel::MoveSelectedStatesBefore(UStateTreeState* TargetState)
 {
-	MoveSelectedStates(TargetState, -1);
+	MoveSelectedStates(TargetState, FStateTreeViewModelInsert::Before);
 }
 
 void FStateTreeViewModel::MoveSelectedStatesAfter(UStateTreeState* TargetState)
 {
-	MoveSelectedStates(TargetState, 1);
+	MoveSelectedStates(TargetState, FStateTreeViewModelInsert::After);
 }
 
 void FStateTreeViewModel::MoveSelectedStatesInto(UStateTreeState* TargetState)
 {
-	MoveSelectedStates(TargetState, 0);
+	MoveSelectedStates(TargetState, FStateTreeViewModelInsert::Into);
 }
 
-void FStateTreeViewModel::MoveSelectedStates(UStateTreeState* TargetState, int32 RelativeLocation)
+void FStateTreeViewModel::MoveSelectedStates(UStateTreeState* TargetState, const FStateTreeViewModelInsert RelativeLocation)
 {
-	if (!TreeData || !TargetState)
+	if (TreeData == nullptr || TargetState == nullptr)
 	{
 		return;
 	}
@@ -477,25 +422,41 @@ void FStateTreeViewModel::MoveSelectedStates(UStateTreeState* TargetState, int32
 	TArray<UStateTreeState*> States;
 	GetSelectedStates(States);
 
-	UStateTreeState* TargetParent = TargetState->Parent;
-
 	// Remove child items whose parent also exists in the selection.
-	FStateTreeViewUtilities::RemoveContainedChildren(States);
-	FStateTreeViewUtilities::RemoveRecursive(States, TargetState);
+	UE::StateTree::Editor::RemoveContainedChildren(States);
+
+	// Remove states which contain target state as child.
+	States.RemoveAll([TargetState](const UStateTreeState* State)
+	{
+		return UE::StateTree::Editor::IsChildOf(State, TargetState);
+	});
 
-	if (States.Num() > 0 && TargetState)
+	if (States.Num() > 0 && TargetState != nullptr)
 	{
 		const FScopedTransaction Transaction(LOCTEXT("MoveTransaction", "Move"));
 
 		TSet<UStateTreeState*> AffectedParents;
 		TSet<UStateTreeState*> AffectedStates;
 
-		AffectedParents.Add(TargetParent);
+		UStateTreeState* TargetParent = TargetState->Parent;
+		if (RelativeLocation == FStateTreeViewModelInsert::Into)
+		{
+			AffectedParents.Add(TargetState);
+		}
+		else
+		{
+			AffectedParents.Add(TargetParent);
+		}
+		
 		for (int32 i = States.Num() - 1; i >= 0; i--)
 		{
 			if (UStateTreeState* State = States[i])
 			{
-				AffectedParents.Add(State->Parent);
+				State->Modify();
+				if (State->Parent)
+				{
+					AffectedParents.Add(State->Parent);
+				}
 			}
 		}
 
@@ -516,13 +477,53 @@ void FStateTreeViewModel::MoveSelectedStates(UStateTreeState* TargetState, int32
 		{
 			if (UStateTreeState* SelectedState = States[i])
 			{
-				UStateTreeState* SelectedParent = SelectedState->Parent;
 				AffectedStates.Add(SelectedState);
 
+				UStateTreeState* SelectedParent = SelectedState->Parent;
+
+				// Remove from current parent
 				TArray<UStateTreeState*>& ArrayToRemoveFrom = SelectedParent ? SelectedParent->Children : TreeData->SubTrees;
-				TArray<UStateTreeState*>& ArrayToMoveTo = TargetParent ? TargetParent->Children : TreeData->SubTrees;
-				FStateTreeViewUtilities::RemoveRecursive(ArrayToRemoveFrom, SelectedState);
-				FStateTreeViewUtilities::InsertRecursive(TargetParent, ArrayToMoveTo, TargetState, SelectedState, RelativeLocation);
+				const int32 ItemIndex = ArrayToRemoveFrom.Find(SelectedState);
+				if (ItemIndex != INDEX_NONE)
+				{
+					ArrayToRemoveFrom.RemoveAt(ItemIndex);
+					SelectedState->Parent = nullptr;
+				}
+
+				// Insert to new parent
+				if (RelativeLocation == FStateTreeViewModelInsert::Into)
+				{
+					// Into
+					TargetState->Children.Insert(SelectedState, /*Index*/0);
+					SelectedState->Parent = TargetState;
+				}
+				else
+				{
+					TArray<UStateTreeState*>& ArrayToMoveTo = TargetParent ? TargetParent->Children : TreeData->SubTrees;
+					const int32 TargetIndex = ArrayToMoveTo.Find(TargetState);
+					if (TargetIndex != INDEX_NONE)
+					{
+						if (RelativeLocation == FStateTreeViewModelInsert::Before)
+						{
+							// Before
+							ArrayToMoveTo.Insert(SelectedState, TargetIndex);
+							SelectedState->Parent = TargetParent;
+						}
+						else if (RelativeLocation == FStateTreeViewModelInsert::After)
+						{
+							// After
+							ArrayToMoveTo.Insert(SelectedState, TargetIndex + 1);
+							SelectedState->Parent = TargetParent;
+						}
+					}
+					else
+					{
+						// Fallback, should never happen.
+						ensureMsgf(false, TEXT("%s: Failed to find specified target state %s on state %s while moving a state."), *GetNameSafe(TreeData->GetOuter()), *GetNameSafe(TargetState), *GetNameSafe(SelectedParent));
+						ArrayToMoveTo.Add(SelectedState);
+						SelectedState->Parent = TargetParent;
+					}
+				}
 			}
 		}
 
diff --git a/Engine/Plugins/Runtime/StateTree/Source/StateTreeEditorModule/Private/StateTreeViewModel.h b/Engine/Plugins/Runtime/StateTree/Source/StateTreeEditorModule/Private/StateTreeViewModel.h
--- a/Engine/Plugins/Runtime/StateTree/Source/StateTreeEditorModule/Private/StateTreeViewModel.h
+++ b/Engine/Plugins/Runtime/StateTree/Source/StateTreeEditorModule/Private/StateTreeViewModel.h
@@ -12,6 +12,13 @@ class UStateTreeState;
 class UStateTree;
 class FMenuBuilder;
 
+enum class FStateTreeViewModelInsert : uint8
+{
+	Before,
+	After,
+	Into,
+};
+
 /**
  * ModelView for editing StateTreeEditorData.
  */
@@ -27,13 +34,13 @@ public:
 	DECLARE_MULTICAST_DELEGATE_OneParam(FOnSelectionChanged, const TArray<UStateTreeState*>& /*SelectedStates*/);
 
 	FStateTreeViewModel();
-	~FStateTreeViewModel();
+	virtual ~FStateTreeViewModel() override;
 
 	void Init(UStateTreeEditorData* InTreeData);
 
 	//~ FEditorUndoClient
-	void PostUndo(bool bSuccess);
-	void PostRedo(bool bSuccess);
+	virtual void PostUndo(bool bSuccess) override;
+	virtual void PostRedo(bool bSuccess) override;
 
 	// Selection handling.
 	void ClearSelection();
@@ -45,14 +52,14 @@ public:
 	bool HasSelection() const;
 
 	// Returns array of subtrees to edit.
-	TArray<UStateTreeState*>* GetSubTrees();
+	TArray<UStateTreeState*>* GetSubTrees() const;
 	int32 GetSubTreeCount() const;
 
 	// Gets and sets StateTree view expansion state store in the asset.
 	void GetPersistentExpandedStates(TSet<UStateTreeState*>& OutExpandedStates);
 	void SetPersistentExpandedStates(TSet<UStateTreeState*>& InExpandedStates);
 
-	// State manupulation commands
+	// State manipulation commands
 	void AddState(UStateTreeState* AfterState);
 	void AddChildState(UStateTreeState* ParentState);
 	void RenameState(UStateTreeState* State, FName NewName);
@@ -62,8 +69,8 @@ public:
 	void MoveSelectedStatesInto(UStateTreeState* TargetState);
 
 	// Force to update the view externally.
-	void NotifyAssetChangedExternally();
-	void NotifyStatesChangedExternally(const TSet<UStateTreeState*>& ChangedStates, const FPropertyChangedEvent& PropertyChangedEvent);
+	void NotifyAssetChangedExternally() const;
+	void NotifyStatesChangedExternally(const TSet<UStateTreeState*>& ChangedStates, const FPropertyChangedEvent& PropertyChangedEvent) const;
 
 	// Called when the whole asset is updated (i.e. undo/redo).
 	FOnAssetChanged& GetOnAssetChanged() { return OnAssetChanged; }
@@ -85,12 +92,11 @@ public:
 
 protected:
 	void GetExpandedStatesRecursive(UStateTreeState* State, TSet<UStateTreeState*>& ExpandedStates);
-	void MoveSelectedStates(UStateTreeState* TargetState, int32 RelativeLocation);
+	void MoveSelectedStates(UStateTreeState* TargetState, const FStateTreeViewModelInsert RelativeLocation);
 
-	void HandleIdentifierChanged(const UStateTree& StateTree);
-	void HandleParameterLayoutChanged(const UStateTree& StateTree);
+	void HandleIdentifierChanged(const UStateTree& StateTree) const;
 
-	UStateTreeEditorData* TreeData;	// TODO: is this safe? should we used TWeakObjectPtr here istead?
+	UStateTreeEditorData* TreeData;	// TODO: is this safe? should we used TWeakObjectPtr here instead?
 	TSet<FWeakObjectPtr> SelectedStates;
 	FOnAssetChanged OnAssetChanged;
 	FOnStatesChanged OnStatesChanged;
diff --git a/Engine/Plugins/Runtime/StateTree/Source/StateTreeEditorModule/Public/StateTreeEditorData.h b/Engine/Plugins/Runtime/StateTree/Source/StateTreeEditorModule/Public/StateTreeEditorData.h
--- a/Engine/Plugins/Runtime/StateTree/Source/StateTreeEditorModule/Public/StateTreeEditorData.h
+++ b/Engine/Plugins/Runtime/StateTree/Source/StateTreeEditorModule/Public/StateTreeEditorData.h
@@ -72,7 +72,7 @@ public:
 	 */
 	UStateTreeState& AddSubTree(const FName Name)
 	{
-		UStateTreeState* SubTreeState = NewObject<UStateTreeState>(this);
+		UStateTreeState* SubTreeState = NewObject<UStateTreeState>(this, FName(), RF_Transactional);
 		check(SubTreeState);
 		SubTreeState->Name = Name;
 		SubTrees.Add(SubTreeState);
diff --git a/Engine/Plugins/Runtime/StateTree/Source/StateTreeEditorModule/Public/StateTreeState.h b/Engine/Plugins/Runtime/StateTree/Source/StateTreeEditorModule/Public/StateTreeState.h
--- a/Engine/Plugins/Runtime/StateTree/Source/StateTreeEditorModule/Public/StateTreeState.h
+++ b/Engine/Plugins/Runtime/StateTree/Source/StateTreeEditorModule/Public/StateTreeState.h
@@ -150,6 +150,7 @@ public:
 
 #if WITH_EDITOR
 	virtual void PostEditChangeChainProperty(FPropertyChangedChainEvent& PropertyChangedEvent) override;
+	virtual void PostLoad() override;;
 #endif
 
 	UStateTreeState* GetNextSiblingState() const;
@@ -159,7 +160,7 @@ public:
 	/** Adds child state with specified name. */
 	UStateTreeState& AddChildState(const FName ChildName, const EStateTreeStateType StateType = EStateTreeStateType::State)
 	{
-		UStateTreeState* ChildState = NewObject<UStateTreeState>(this);
+		UStateTreeState* ChildState = NewObject<UStateTreeState>(this, FName(), RF_Transactional);
 		check(ChildState);
 		ChildState->Name = ChildName;
 		ChildState->Parent = this;
@@ -230,14 +231,14 @@ public:
 	{
 		return Transitions.Emplace_GetRef(InEvent, InType, InState);
 	}
-
+ 
 	// ~StateTree Builder API
 
 	UPROPERTY(EditDefaultsOnly, Category = "State")
 	FName Name;
 
 	UPROPERTY(EditDefaultsOnly, Category = "State")
-	EStateTreeStateType Type;
+	EStateTreeStateType Type = EStateTreeStateType::State;
 
 	UPROPERTY(EditDefaultsOnly, Category = "State", Meta=(DirectStatesOnly))
 	FStateTreeStateLink LinkedState;
@@ -265,9 +266,8 @@ public:
 	TArray<UStateTreeState*> Children;
 
 	UPROPERTY(meta = (ExcludeFromHash))
-	bool bExpanded;
+	bool bExpanded = true;
 
 	UPROPERTY(meta = (ExcludeFromHash))
-	UStateTreeState* Parent;
-
+	UStateTreeState* Parent = nullptr;
 };
