StateTree: Fix for crash when moving states

diff --git a/Engine/Plugins/Runtime/StateTree/Source/StateTreeEditorModule/Private/StateTreeEditor.cpp b/Engine/Plugins/Runtime/StateTree/Source/StateTreeEditorModule/Private/StateTreeEditor.cpp
--- a/Engine/Plugins/Runtime/StateTree/Source/StateTreeEditorModule/Private/StateTreeEditor.cpp
+++ b/Engine/Plugins/Runtime/StateTree/Source/StateTreeEditorModule/Private/StateTreeEditor.cpp
@@ -110,6 +110,28 @@ void FStateTreeEditor::InitEditor( const EToolkitMode::Type Mode, const TSharedP
 		StateTree->EditorData = EditorData;
 	}
 
+	// @todo: Temporary fix
+	// Make sure all states are transactional
+	for (UStateTreeState* SubTree : EditorData->SubTrees)
+	{
+		TArray<UStateTreeState*> Stack;
+
+		Stack.Add(SubTree);
+		while (!Stack.IsEmpty())
+		{
+			if (UStateTreeState* State = Stack.Pop())
+			{
+				State->SetFlags(RF_Transactional);
+				
+				for (UStateTreeState* ChildState : State->Children)
+				{
+					Stack.Add(ChildState);
+				}
+			}
+		}
+	}
+
+
 	StateTreeViewModel = MakeShareable(new FStateTreeViewModel());
 	StateTreeViewModel->Init(EditorData);
 
diff --git a/Engine/Plugins/Runtime/StateTree/Source/StateTreeEditorModule/Private/StateTreeViewModel.cpp b/Engine/Plugins/Runtime/StateTree/Source/StateTreeEditorModule/Private/StateTreeViewModel.cpp
--- a/Engine/Plugins/Runtime/StateTree/Source/StateTreeEditorModule/Private/StateTreeViewModel.cpp
+++ b/Engine/Plugins/Runtime/StateTree/Source/StateTreeEditorModule/Private/StateTreeViewModel.cpp
@@ -460,6 +460,12 @@ void FStateTreeViewModel::MoveSelectedStates(UStateTreeState* TargetState, const
 			}
 		}
 
+		if (RelativeLocation == 0)
+		{
+			// Move into
+			TargetState->Modify();
+		}
+		
 		for (UStateTreeState* Parent : AffectedParents)
 		{
 			if (Parent)
