Fix invalid operands to binary expression ('const FStateTreeViewModelInsert' and 'int')
#rnx

diff --git a/Engine/Plugins/Runtime/StateTree/Source/StateTreeEditorModule/Private/StateTreeViewModel.cpp b/Engine/Plugins/Runtime/StateTree/Source/StateTreeEditorModule/Private/StateTreeViewModel.cpp
--- a/Engine/Plugins/Runtime/StateTree/Source/StateTreeEditorModule/Private/StateTreeViewModel.cpp
+++ b/Engine/Plugins/Runtime/StateTree/Source/StateTreeEditorModule/Private/StateTreeViewModel.cpp
@@ -460,7 +460,7 @@ void FStateTreeViewModel::MoveSelectedStates(UStateTreeState* TargetState, const
 			}
 		}
 
-		if (RelativeLocation == 0)
+		if (RelativeLocation == FStateTreeViewModelInsert::Before)
 		{
 			// Move into
 			TargetState->Modify();
