StateTree Editor: Corrected earlier fix.

diff --git a/Engine/Plugins/Runtime/StateTree/Source/StateTreeEditorModule/Private/StateTreeViewModel.cpp b/Engine/Plugins/Runtime/StateTree/Source/StateTreeEditorModule/Private/StateTreeViewModel.cpp
--- a/Engine/Plugins/Runtime/StateTree/Source/StateTreeEditorModule/Private/StateTreeViewModel.cpp
+++ b/Engine/Plugins/Runtime/StateTree/Source/StateTreeEditorModule/Private/StateTreeViewModel.cpp
@@ -460,7 +460,7 @@ void FStateTreeViewModel::MoveSelectedStates(UStateTreeState* TargetState, const
 			}
 		}
 
-		if (RelativeLocation == FStateTreeViewModelInsert::Before)
+		if (RelativeLocation == FStateTreeViewModelInsert::Into)
 		{
 			// Move into
 			TargetState->Modify();
