StateTree: Parametrized states
- Added specific "subtree" states which can be linked to via a "linked" state
- Subtrees can only bind data back to the their root state (or tree "global" data)
- Added property bag based parameters to subtree and linked states
- Update (transient) property bag structs before resolving property paths
- Linked states reflect their parameters from the state they link to
- Added property binding to and from parameters
- Relaxed the property binding to fail if the source data is not available
- Allow enter conditions to bind to tasks, and fail the expression if trying to access unaccessible data
- Added icon for linked and subtree states
- Added source type for all binding source structs
- Added UStateTreeEditorData::VisitHierarchy to simplify iterating over all states

diff --git a/Engine/Plugins/Runtime/StateTree/Source/StateTreeEditorModule/Private/Customizations/StateTreeBindingExtension.cpp b/Engine/Plugins/Runtime/StateTree/Source/StateTreeEditorModule/Private/Customizations/StateTreeBindingExtension.cpp
--- a/Engine/Plugins/Runtime/StateTree/Source/StateTreeEditorModule/Private/Customizations/StateTreeBindingExtension.cpp
+++ b/Engine/Plugins/Runtime/StateTree/Source/StateTreeEditorModule/Private/Customizations/StateTreeBindingExtension.cpp
@@ -72,8 +72,15 @@ EStateTreePropertyUsage ParsePropertyUsage(TSharedPtr<const IPropertyHandle> InP
 {
 	static const FName CategoryName(TEXT("Category"));
 	
-	const FString Category = InPropertyHandle->GetMetaData(CategoryName);
-
+	FString Category = InPropertyHandle->GetMetaData(CategoryName);
+	if (Category.IsEmpty())
+	{
+		if (const FString* InstanceMetaData = InPropertyHandle->GetInstanceMetaData(CategoryName))
+		{
+			Category = *InstanceMetaData;
+		}
+	}
+	
 	if (Category == TEXT("Input"))
 	{
 		return EStateTreePropertyUsage::Input;
@@ -92,7 +99,7 @@ EStateTreePropertyUsage ParsePropertyUsage(TSharedPtr<const IPropertyHandle> InP
 
 FText GetSectionNameFromDataSource(const EStateTreeBindableStructSource Source)
 {
-	switch(Source)
+	switch (Source)
 	{
 		case EStateTreeBindableStructSource::Task:
 			return LOCTEXT("Tasks", "Tasks");
@@ -105,6 +112,7 @@ FText GetSectionNameFromDataSource(const EStateTreeBindableStructSource Source)
 		case EStateTreeBindableStructSource::TreeParameter:
 			return LOCTEXT("Tree Parameter", "Tree Parameter");
 	}
+
 	return FText::GetEmpty();
 }
 
@@ -363,4 +371,3 @@ void FStateTreeBindingExtension::ExtendWidgetRow(FDetailWidgetRow& InWidgetRow,
 }
 
 #undef LOCTEXT_NAMESPACE
-
diff --git a/Engine/Plugins/Runtime/StateTree/Source/StateTreeEditorModule/Private/Customizations/StateTreeStateDetails.cpp b/Engine/Plugins/Runtime/StateTree/Source/StateTreeEditorModule/Private/Customizations/StateTreeStateDetails.cpp
--- a/Engine/Plugins/Runtime/StateTree/Source/StateTreeEditorModule/Private/Customizations/StateTreeStateDetails.cpp
+++ b/Engine/Plugins/Runtime/StateTree/Source/StateTreeEditorModule/Private/Customizations/StateTreeStateDetails.cpp
@@ -47,6 +47,7 @@ void FStateTreeStateDetails::CustomizeDetails(IDetailLayoutBuilder& DetailBuilde
 	TSharedPtr<IPropertyHandle> TransitionsProperty = DetailBuilder.GetProperty(TEXT("Transitions"));
 	TSharedPtr<IPropertyHandle> TypeProperty = DetailBuilder.GetProperty(TEXT("Type"));
 	TSharedPtr<IPropertyHandle> LinkedStateProperty = DetailBuilder.GetProperty(TEXT("LinkedState"));
+	TSharedPtr<IPropertyHandle> ParametersProperty = DetailBuilder.GetProperty(TEXT("Parameters"));
 
 	uint8 StateTypeValue = 0;
 	TypeProperty->GetValue(StateTypeValue);
@@ -60,8 +61,13 @@ void FStateTreeStateDetails::CustomizeDetails(IDetailLayoutBuilder& DetailBuilde
 		LinkedStateProperty->MarkHiddenByCustomization();
 	}
 	
+	if (!(StateType == EStateTreeStateType::Subtree || StateType == EStateTreeStateType::Linked))
+	{
+		ParametersProperty->MarkHiddenByCustomization();
+	}
+	
 	const FName EvalCategoryName(TEXT("Evaluators"));
-	if (StateType == EStateTreeStateType::State && Schema && Schema->AllowEvaluators())
+	if ((StateType == EStateTreeStateType::State || StateType == EStateTreeStateType::Subtree) && Schema && Schema->AllowEvaluators())
 	{
 		MakeArrayCategory(DetailBuilder, EvalCategoryName, LOCTEXT("StateDetailsEvaluators", "Evaluators"), 1, EvaluatorsProperty);
 	}
@@ -80,7 +86,7 @@ void FStateTreeStateDetails::CustomizeDetails(IDetailLayoutBuilder& DetailBuilde
 		DetailBuilder.EditCategory(EnterConditionsCategoryName).SetCategoryVisibility(false);
 	}
 
-	if (StateType == EStateTreeStateType::State)
+	if ((StateType == EStateTreeStateType::State || StateType == EStateTreeStateType::Subtree))
 	{
 		if (Schema && Schema->AllowMultipleTasks())
 		{
diff --git a/Engine/Plugins/Runtime/StateTree/Source/StateTreeEditorModule/Private/Customizations/StateTreeStateLinkDetails.cpp b/Engine/Plugins/Runtime/StateTree/Source/StateTreeEditorModule/Private/Customizations/StateTreeStateLinkDetails.cpp
--- a/Engine/Plugins/Runtime/StateTree/Source/StateTreeEditorModule/Private/Customizations/StateTreeStateLinkDetails.cpp
+++ b/Engine/Plugins/Runtime/StateTree/Source/StateTreeEditorModule/Private/Customizations/StateTreeStateLinkDetails.cpp
@@ -29,10 +29,13 @@ void FStateTreeStateLinkDetails::CustomizeHeader(TSharedRef<class IPropertyHandl
 	IDProperty = StructProperty->GetChildHandle(TEXT("ID"));
 	TypeProperty = StructProperty->GetChildHandle(TEXT("Type"));
 
-	static const FName NAME_DirectStatesOnly = "DirectStatesOnly";
 	if (const FProperty* MetaDataProperty = StructProperty->GetMetaDataProperty())
 	{
-		bAllowMetaStates = !MetaDataProperty->HasMetaData(NAME_DirectStatesOnly);
+		static const FName NAME_DirectStatesOnly = "DirectStatesOnly";
+		static const FName NAME_SubtreesOnly = "SubtreesOnly";
+		
+		bDirectStatesOnly = MetaDataProperty->HasMetaData(NAME_DirectStatesOnly);
+		bSubtreesOnly = MetaDataProperty->HasMetaData(NAME_SubtreesOnly);
 	}
 	
 	CacheStates();
@@ -75,8 +78,17 @@ void FStateTreeStateLinkDetails::CacheStates(const UStateTreeState* State)
 		return;
 	}
 
-	CachedNames.Add(State->Name);
-	CachedIDs.Add(State->ID);
+	bool bShouldAdd = true;
+	if (bSubtreesOnly && State->Type != EStateTreeStateType::Subtree)
+	{
+		bShouldAdd = false;
+	}
+
+	if (bShouldAdd)
+	{
+		CachedNames.Add(State->Name);
+		CachedIDs.Add(State->ID);
+	}
 
 	for (UStateTreeState* ChildState : State->Children)
 	{
@@ -149,7 +161,7 @@ TSharedRef<SWidget> FStateTreeStateLinkDetails::OnGetStateContent() const
 {
 	FMenuBuilder MenuBuilder(true, NULL);
 
-	if (bAllowMetaStates)
+	if (!bDirectStatesOnly)
 	{
 		FUIAction NextItemAction(FExecuteAction::CreateSP(const_cast<FStateTreeStateLinkDetails*>(this), &FStateTreeStateLinkDetails::OnStateComboChange, ComboNextState));
 		MenuBuilder.AddMenuEntry(LOCTEXT("TransitionNextState", "Next State"), LOCTEXT("TransitionNextTooltip", "Goto next sibling State."), FSlateIcon(), NextItemAction);
diff --git a/Engine/Plugins/Runtime/StateTree/Source/StateTreeEditorModule/Private/Customizations/StateTreeStateLinkDetails.h b/Engine/Plugins/Runtime/StateTree/Source/StateTreeEditorModule/Private/Customizations/StateTreeStateLinkDetails.h
--- a/Engine/Plugins/Runtime/StateTree/Source/StateTreeEditorModule/Private/Customizations/StateTreeStateLinkDetails.h
+++ b/Engine/Plugins/Runtime/StateTree/Source/StateTreeEditorModule/Private/Customizations/StateTreeStateLinkDetails.h
@@ -49,8 +49,10 @@ private:
 	TSharedPtr<IPropertyHandle> IDProperty;
 	TSharedPtr<IPropertyHandle> TypeProperty;
 
-	// If set, allow to select meta states like Next or (tree) Succeeded.
-	bool bAllowMetaStates = true;
+	// If set, hide selecting meta states like Next or (tree) Succeeded.
+	bool bDirectStatesOnly = false;
+	// If set, allow to select only states marked as subtrees.
+	bool bSubtreesOnly = false;
 	
 	class IPropertyUtilities* PropUtils;
 	TSharedPtr<IPropertyHandle> StructProperty;
diff --git a/Engine/Plugins/Runtime/StateTree/Source/StateTreeEditorModule/Private/Customizations/StateTreeStateParametersDetails.cpp b/Engine/Plugins/Runtime/StateTree/Source/StateTreeEditorModule/Private/Customizations/StateTreeStateParametersDetails.cpp
new file mode 100644
--- /dev/null
+++ b/Engine/Plugins/Runtime/StateTree/Source/StateTreeEditorModule/Private/Customizations/StateTreeStateParametersDetails.cpp
@@ -0,0 +1,168 @@
+﻿// Copyright Epic Games, Inc. All Rights Reserved.
+
+#include "StateTreeStateParametersDetails.h"
+#include "Widgets/DeclarativeSyntaxSupport.h"
+#include "DetailWidgetRow.h"
+#include "DetailLayoutBuilder.h"
+#include "IDetailPropertyRow.h"
+#include "IDetailChildrenBuilder.h"
+#include "Widgets/Text/STextBlock.h"
+#include "PropertyBagDetails.h"
+#include "StateTree.h"
+#include "StateTreeEditorData.h"
+#include "StateTreePropertyHelpers.h"
+#include "StateTreeBindingExtension.h"
+
+#define LOCTEXT_NAMESPACE "StateTreeEditor"
+
+
+//----------------------------------------------------------------//
+// FStateTreeStateParametersInstanceDataDetails
+//----------------------------------------------------------------//
+
+class FStateTreeStateParametersInstanceDataDetails : public FPropertyBagInstanceDataDetails
+{
+public:
+	FStateTreeStateParametersInstanceDataDetails(TSharedPtr<IPropertyHandle> InValueProperty, TSharedPtr<IPropertyHandle> InStructProperty, IPropertyUtilities* InPropUtils, const bool bInFixedLayout, FGuid InID, UStateTreeEditorData* InEditorData)
+		: FPropertyBagInstanceDataDetails(InValueProperty, InStructProperty, InPropUtils, bInFixedLayout)
+		, EditorData(InEditorData)
+	{
+		EditorPropBindings = EditorData ? EditorData->GetPropertyEditorBindings() : nullptr;
+		ID = InID;
+	}
+
+	virtual void OnChildRowAdded(IDetailPropertyRow& ChildRow) override
+	{
+		if (!EditorPropBindings
+			|| !bFixedLayout		// No binding for parameter definitions.
+			|| !ID.IsValid())
+		{
+			FPropertyBagInstanceDataDetails::OnChildRowAdded(ChildRow);
+			return;
+		}
+
+		TSharedPtr<IPropertyHandle> ChildPropHandle = ChildRow.GetPropertyHandle();
+		const FProperty* Property = ChildPropHandle->GetProperty();
+
+		// Set the category to Parameter so that the binding extension will pick it up.
+		static const FName CategoryName(TEXT("Category"));
+		ChildPropHandle->SetInstanceMetaData(CategoryName, TEXT("Parameter"));
+	
+		// Conditionally control visibility of the value field of bound properties.
+
+		// Pass the node ID to binding extension. Since the properties are added using AddChildStructure(), we break the hierarchy and cannot access parent.
+		ChildPropHandle->SetInstanceMetaData(UE::StateTree::PropertyBinding::StateTreeNodeIDName, LexToString(ID));
+
+		FStateTreeEditorPropertyPath Path(ID, *Property->GetFName().ToString());
+		TSharedPtr<SWidget> NameWidget;
+		TSharedPtr<SWidget> ValueWidget;
+		FDetailWidgetRow Row;
+		ChildRow.GetDefaultWidgets(NameWidget, ValueWidget, Row);
+
+		auto IsValueVisible = TAttribute<EVisibility>::Create([Path, this]() -> EVisibility
+			{
+				return EditorPropBindings->HasPropertyBinding(Path) ? EVisibility::Collapsed : EVisibility::Visible;
+			});
+
+		ChildRow
+			.CustomWidget(/*bShowChildren*/true)
+			.NameContent()
+			[
+				NameWidget.ToSharedRef()
+			]
+			.ValueContent()
+			[
+				SNew(SBox)
+				.Visibility(IsValueVisible)
+				[
+					ValueWidget.ToSharedRef()
+				]
+			];
+	}
+
+	UStateTreeEditorData* EditorData;
+	FStateTreeEditorPropertyBindings* EditorPropBindings;
+	FGuid ID;
+};
+
+
+//----------------------------------------------------------------//
+// FStateTreeStateParametersDetails
+//----------------------------------------------------------------//
+
+TSharedRef<IPropertyTypeCustomization> FStateTreeStateParametersDetails::MakeInstance()
+{
+	return MakeShareable(new FStateTreeStateParametersDetails);
+}
+
+void FStateTreeStateParametersDetails::CustomizeHeader(TSharedRef<class IPropertyHandle> StructPropertyHandle, class FDetailWidgetRow& HeaderRow, IPropertyTypeCustomizationUtils& StructCustomizationUtils)
+{
+	StructProperty = StructPropertyHandle;
+	PropUtils = StructCustomizationUtils.GetPropertyUtilities().Get();
+
+	ParametersProperty = StructProperty->GetChildHandle(TEXT("Parameters"));
+	FixedLayoutProperty = StructProperty->GetChildHandle(TEXT("bFixedLayout"));
+	IDProperty = StructProperty->GetChildHandle(TEXT("ID"));
+	ParametersValueProperty = ParametersProperty->GetChildHandle(TEXT("Value"));
+	check(ParametersProperty.IsValid());
+	check(FixedLayoutProperty.IsValid());
+	check(IDProperty.IsValid());
+	check(ParametersValueProperty.IsValid());
+
+	FindOuterObjects();
+	
+	bFixedLayout = false;
+	FixedLayoutProperty->GetValue(bFixedLayout);
+
+	TSharedPtr<SWidget> ValueWidget = SNullWidget::NullWidget;
+	if (!bFixedLayout)
+	{
+		ValueWidget = FPropertyBagDetails::MakeAddPropertyWidget(ParametersProperty, PropUtils);
+	}
+	
+	HeaderRow
+		.NameContent()
+		[
+			ParametersProperty->CreatePropertyNameWidget()
+		]
+		.ValueContent()
+		[
+			ValueWidget.ToSharedRef()
+		]
+		.ShouldAutoExpand(true);
+}
+
+void FStateTreeStateParametersDetails::CustomizeChildren(TSharedRef<class IPropertyHandle> StructPropertyHandle, class IDetailChildrenBuilder& StructBuilder, IPropertyTypeCustomizationUtils& StructCustomizationUtils)
+{
+	FGuid ID;
+	UE::StateTree::PropertyHelpers::GetStructValue<FGuid>(IDProperty, ID);
+
+	// Show the Value (FInstancedStruct) as child rows.
+	TSharedRef<FStateTreeStateParametersInstanceDataDetails> InstanceDetails = MakeShareable(new FStateTreeStateParametersInstanceDataDetails(ParametersValueProperty, ParametersProperty, PropUtils, bFixedLayout, ID, EditorData));
+	StructBuilder.AddCustomBuilder(InstanceDetails);
+}
+
+void FStateTreeStateParametersDetails::FindOuterObjects()
+{
+	check(StructProperty);
+	
+	EditorData = nullptr;
+	StateTree = nullptr;
+
+	TArray<UObject*> OuterObjects;
+	StructProperty->GetOuterObjects(OuterObjects);
+	for (UObject* Outer : OuterObjects)
+	{
+		UStateTreeEditorData* OuterEditorData = Outer->GetTypedOuter<UStateTreeEditorData>();
+		UStateTree* OuterStateTree = OuterEditorData ? OuterEditorData->GetTypedOuter<UStateTree>() : nullptr;
+		if (OuterEditorData && OuterStateTree)
+		{
+			StateTree = OuterStateTree;
+			EditorData = OuterEditorData;
+			break;
+		}
+	}
+}
+
+
+#undef LOCTEXT_NAMESPACE
diff --git a/Engine/Plugins/Runtime/StateTree/Source/StateTreeEditorModule/Private/Customizations/StateTreeStateParametersDetails.h b/Engine/Plugins/Runtime/StateTree/Source/StateTreeEditorModule/Private/Customizations/StateTreeStateParametersDetails.h
new file mode 100644
--- /dev/null
+++ b/Engine/Plugins/Runtime/StateTree/Source/StateTreeEditorModule/Private/Customizations/StateTreeStateParametersDetails.h
@@ -0,0 +1,41 @@
+﻿// Copyright Epic Games, Inc. All Rights Reserved.
+
+#pragma once
+
+#include "IPropertyTypeCustomization.h"
+
+class IPropertyHandle;
+class UStateTree;
+class UStateTreeEditorData;
+
+/**
+* Type customization for FStateTreeStateParameters.
+*/
+
+class FStateTreeStateParametersDetails : public IPropertyTypeCustomization
+{
+public:
+	/** Makes a new instance of this detail layout class for a specific detail view requesting it */
+	static TSharedRef<IPropertyTypeCustomization> MakeInstance();
+
+	/** IPropertyTypeCustomization interface */
+	virtual void CustomizeHeader(TSharedRef<IPropertyHandle> StructPropertyHandle, class FDetailWidgetRow& HeaderRow, IPropertyTypeCustomizationUtils& StructCustomizationUtils) override;
+	virtual void CustomizeChildren(TSharedRef<IPropertyHandle> StructPropertyHandle, class IDetailChildrenBuilder& StructBuilder, IPropertyTypeCustomizationUtils& StructCustomizationUtils) override;
+
+private:
+
+	void FindOuterObjects();
+
+	class IPropertyUtilities* PropUtils = nullptr;
+
+	TSharedPtr<IPropertyHandle> ParametersProperty;
+	TSharedPtr<IPropertyHandle> ParametersValueProperty;
+	TSharedPtr<IPropertyHandle> FixedLayoutProperty;
+	TSharedPtr<IPropertyHandle> IDProperty;
+	TSharedPtr<IPropertyHandle> StructProperty;
+
+	bool bFixedLayout = false;
+
+	UStateTreeEditorData* EditorData = nullptr;
+	UStateTree* StateTree = nullptr;
+};
diff --git a/Engine/Plugins/Runtime/StateTree/Source/StateTreeEditorModule/Private/SStateTreeViewRow.cpp b/Engine/Plugins/Runtime/StateTree/Source/StateTreeEditorModule/Private/SStateTreeViewRow.cpp
--- a/Engine/Plugins/Runtime/StateTree/Source/StateTreeEditorModule/Private/SStateTreeViewRow.cpp
+++ b/Engine/Plugins/Runtime/StateTree/Source/StateTreeEditorModule/Private/SStateTreeViewRow.cpp
@@ -115,7 +115,7 @@ void SStateTreeViewRow::Construct(const FArguments& InArgs, const TSharedRef<STa
                         .Visibility(this, &SStateTreeViewRow::GetSelectorVisibility)
                         [
                             SNew(STextBlock)
-                            .Text(FEditorFontGlyphs::Level_Down)
+                            .Text(this, &SStateTreeViewRow::GetSelectorDesc)
                             .TextStyle(FStateTreeEditorStyle::Get(), "StateTree.Icon")
                         ]
                     ]
@@ -421,7 +421,7 @@ FSlateColor SStateTreeViewRow::GetTitleColor() const
 		{
 			return FLinearColor(FColor(236, 134, 39));
 		}
-		if (IsRoutine())
+		if (IsRootState() || State->Type == EStateTreeStateType::Subtree)
 		{
 			return FLinearColor(FColor(17, 117, 131));
 		}
@@ -452,11 +452,45 @@ EVisibility SStateTreeViewRow::GetSelectorVisibility() const
 {
 	if (const UStateTreeState* State = WeakState.Get())
 	{
-		return State->Children.Num() > 0 ? EVisibility::Visible : EVisibility::Collapsed;
+		if (State->Type == EStateTreeStateType::State || State->Type == EStateTreeStateType::Group)
+		{
+			return State->Children.Num() > 0 ? EVisibility::Visible : EVisibility::Collapsed;
+		}
+		else if (IsRootState() || State->Type == EStateTreeStateType::Subtree)
+		{
+			return EVisibility::Visible;
+		}
+		else if (State->Type == EStateTreeStateType::Linked)
+		{
+			return EVisibility::Visible;
+		}
 	}
+	
 	return EVisibility::Collapsed;
 }
 
+
+FText SStateTreeViewRow::GetSelectorDesc() const
+{
+	if (const UStateTreeState* State = WeakState.Get())
+	{
+		if (State->Type == EStateTreeStateType::State || State->Type == EStateTreeStateType::Group)
+		{
+			return FEditorFontGlyphs::Level_Down; // Selector
+		}
+		else if (IsRootState() || State->Type == EStateTreeStateType::Subtree)
+		{
+			return FEditorFontGlyphs::Cogs;
+		}
+		else if (State->Type == EStateTreeStateType::Linked)
+		{
+			return FEditorFontGlyphs::Link;
+		}
+	}
+
+	return FText::GetEmpty(); 
+}
+
 EVisibility SStateTreeViewRow::GetEvaluatorsVisibility() const
 {
 	if (const UStateTreeState* State = WeakState.Get())
@@ -825,7 +859,7 @@ FText SStateTreeViewRow::GetConditionalTransitionsDesc() const
 	return FText::Join(FText::FromString(TEXT(", ")), DescItems);
 }
 
-bool SStateTreeViewRow::IsRoutine() const
+bool SStateTreeViewRow::IsRootState() const
 {
 	// Routines can be identified by not having parent state.
 	const UStateTreeState* State = WeakState.Get();
diff --git a/Engine/Plugins/Runtime/StateTree/Source/StateTreeEditorModule/Private/SStateTreeViewRow.h b/Engine/Plugins/Runtime/StateTree/Source/StateTreeEditorModule/Private/SStateTreeViewRow.h
--- a/Engine/Plugins/Runtime/StateTree/Source/StateTreeEditorModule/Private/SStateTreeViewRow.h
+++ b/Engine/Plugins/Runtime/StateTree/Source/StateTreeEditorModule/Private/SStateTreeViewRow.h
@@ -30,6 +30,7 @@ private:
 
 	EVisibility GetConditionVisibility() const;
 	EVisibility GetSelectorVisibility() const;
+	FText GetSelectorDesc() const;
 
 	EVisibility GetEvaluatorsVisibility() const;
 	FText GetEvaluatorsDesc() const;
@@ -61,7 +62,7 @@ private:
 
 	bool HasParentTransitionForEvent(const UStateTreeState& State, const EStateTreeTransitionEvent Event) const;
 
-	bool IsRoutine() const;
+	bool IsRootState() const;
 	bool IsSelected() const;
 
 	bool VerifyNodeTextChanged(const FText& NewLabel, FText& OutErrorMessage);
diff --git a/Engine/Plugins/Runtime/StateTree/Source/StateTreeEditorModule/Private/StateTreeCompiler.cpp b/Engine/Plugins/Runtime/StateTree/Source/StateTreeEditorModule/Private/StateTreeCompiler.cpp
--- a/Engine/Plugins/Runtime/StateTree/Source/StateTreeEditorModule/Private/StateTreeCompiler.cpp
+++ b/Engine/Plugins/Runtime/StateTree/Source/StateTreeEditorModule/Private/StateTreeCompiler.cpp
@@ -12,7 +12,6 @@
 #include "StateTreeExecutionContext.h"
 #include "StateTreePropertyBindingCompiler.h"
 
-
 bool FStateTreeCompiler::Compile(UStateTree& InStateTree)
 {
 	StateTree = &InStateTree;
@@ -40,7 +39,7 @@ bool FStateTreeCompiler::Compile(UStateTree& InStateTree)
 	// Mark all parameters as binding source
 	for (FStateTreeParameterDesc& Desc : StateTree->Parameters.Parameters)
 	{
-		Desc.DataViewIndex = BindingsCompiler.AddSourceStruct({Desc.Name, Desc.Parameter.GetScriptStruct(), Desc.ID});
+		Desc.DataViewIndex = BindingsCompiler.AddSourceStruct({Desc.Name, Desc.Parameter.GetScriptStruct(), EStateTreeBindableStructSource::TreeParameter, Desc.ID});
 	}	
 	
 	// Mark all named external values as binding source
@@ -48,7 +47,7 @@ bool FStateTreeCompiler::Compile(UStateTree& InStateTree)
 	{
 		for (FStateTreeExternalDataDesc& Desc : StateTree->Schema->GetMutableNamedExternalDataDescs())
 		{
-			Desc.Handle.DataViewIndex = BindingsCompiler.AddSourceStruct({Desc.Name, Desc.Struct, Desc.ID});
+			Desc.Handle.DataViewIndex = BindingsCompiler.AddSourceStruct({Desc.Name, Desc.Struct, EStateTreeBindableStructSource::TreeData, Desc.ID});
 		} 
 	}
 	
@@ -58,19 +57,13 @@ bool FStateTreeCompiler::Compile(UStateTree& InStateTree)
 		return false;
 	}
 
-	if (!CreateExecutionInfos())
-	{
-		StateTree->ResetCompiled();
-		return false;
-	}
-	
 	if (!CreateStateEvaluators())
 	{
 		StateTree->ResetCompiled();
 		return false;
 	}
 
-	if (!CreateStateTasks())
+	if (!CreateStateTasksAndParameters())
 	{
 		StateTree->ResetCompiled();
 		return false;
@@ -120,6 +113,7 @@ bool FStateTreeCompiler::CreateStates()
 	// Create item for the runtime execution state
 	StateTree->Instances.Add(FInstancedStruct::Make<FStateTreeExecutionState>());
 
+	// Create main tree (omit subtrees)
 	for (UStateTreeState* SubTree : TreeData->SubTrees)
 	{
 		if (SubTree != nullptr)
@@ -130,7 +124,28 @@ bool FStateTreeCompiler::CreateStates()
 			}
 		}
 	}
-	
+
+	// Create Subtrees
+	for (UStateTreeState* SubTree : TreeData->SubTrees)
+	{
+		TArray<UStateTreeState*> Stack;
+		Stack.Push(SubTree);
+		while (!Stack.IsEmpty())
+		{
+			if (UStateTreeState* State = Stack.Pop())
+			{
+				if (State->Type == EStateTreeStateType::Subtree)
+				{
+					if (!CreateStateRecursive(*State, FStateTreeHandle::Invalid))
+					{
+						return false;
+					}
+				}
+				Stack.Append(State->Children);
+			}
+		}
+	}
+
 	return true;
 }
 
@@ -139,10 +154,12 @@ bool FStateTreeCompiler::CreateStateRecursive(UStateTreeState& State, const FSta
 	FStateTreeCompilerLogStateScope LogStateScope(&State, Log);
 
 	const int32 StateIdx = StateTree->States.AddDefaulted();
-	FCompactStateTreeState& BakedState = StateTree->States[StateIdx];
-	BakedState.Name = State.Name;
-	BakedState.Parent = Parent;
+	FCompactStateTreeState& CompactState = StateTree->States[StateIdx];
+	CompactState.Name = State.Name;
+	CompactState.Parent = Parent;
 
+	CompactState.Type = State.Type; 
+	
 	SourceStates.Add(&State);
 	IDToState.Add(State.ID, StateIdx);
 
@@ -150,10 +167,10 @@ bool FStateTreeCompiler::CreateStateRecursive(UStateTreeState& State, const FSta
 
 	// Child states
 	check(StateTree->States.Num() <= int32(MAX_uint16));
-	BakedState.ChildrenBegin = uint16(StateTree->States.Num());
+	CompactState.ChildrenBegin = uint16(StateTree->States.Num());
 	for (UStateTreeState* Child : State.Children)
 	{
-		if (Child)
+		if (Child != nullptr && Child->Type != EStateTreeStateType::Subtree)
 		{
 			if (!CreateStateRecursive(*Child, FStateTreeHandle((uint16)StateIdx)))
 			{
@@ -162,111 +179,7 @@ bool FStateTreeCompiler::CreateStateRecursive(UStateTreeState& State, const FSta
 		}
 	}
 	check(StateTree->States.Num() <= int32(MAX_uint16));
-	StateTree->States[StateIdx].ChildrenEnd = uint16(StateTree->States.Num()); // Cannot use BakedState here, it may be invalid due to array resize.
-
-	return true;
-}
-
-FString FStateTreeCompiler::GetExecutionPathString(const TConstArrayView<const UStateTreeState*> Path)
-{
-	FString PathStr;
-	const UStateTreeState* PrevState = nullptr;
-	for (int32 Index = 0; Index < Path.Num(); Index++)
-	{
-		const UStateTreeState* State = Path[Index];
-		check(State);
-		if (PrevState != nullptr)
-		{
-			if (PrevState != State->Parent)
-			{
-				// Linked
-				PathStr += TEXT(">");
-			}
-			else
-			{
-				PathStr += TEXT("/");
-			}
-		}
-		PathStr += State->Name.ToString();
-		PrevState = State;
-	}
-
-	return PathStr;
-}
-
-bool FStateTreeCompiler::IsPathLinked(const TConstArrayView<const UStateTreeState*> Path)
-{
-	bool bIsLinked = false; 
-	const UStateTreeState* PrevState = nullptr;
-	for (int32 Index = 0; Index < Path.Num(); Index++)
-	{
-		const UStateTreeState* State = Path[Index];
-		check(State);
-		// If the previous state is not the parent state in tree, then the connection must have come from a linked state (a state linking to an arbitrary state in the tree).
-		if (PrevState && PrevState != State->Parent)
-		{
-			bIsLinked = true;
-			break;
-		}
-		PrevState = State;
-	}
-	return bIsLinked;
-}
-
-bool FStateTreeCompiler::CreateExecutionInfos()
-{
-	for (UStateTreeState* SubTree : TreeData->SubTrees)
-	{
-		if (SubTree != nullptr)
-		{
-			TArray<const UStateTreeState*> Path;
-			if (!CreateExecutionInfosRecursive(*SubTree, Path))
-			{
-				return false;
-			}
-		}
-	}
-
-	return true;
-}
-
-bool FStateTreeCompiler::CreateExecutionInfosRecursive(UStateTreeState& State, TArray<const UStateTreeState*>& Path)
-{
-	Path.Add(&State);
-	
-	if (Path.Num() > FStateTreeActiveStates::MaxStates)
-	{
-		Log.Reportf(EMessageSeverity::Error,
-		TEXT("Reached maximum execution depth %d at: '%s'."), FStateTreeActiveStates::MaxStates, *GetExecutionPathString(Path));
-		return false;
-	}
-
-	FStateExecutionInfo& ExecInfo = ExecutionInfos.FindOrAdd(&State);
-	ExecInfo.ExecutionPaths.Emplace(Path);
-
-	if (State.Type == EStateTreeStateType::Linked)
-	{
-		if (UStateTreeState* LinkedState = GetState(State.LinkedState.ID))
-		{
-			if (!CreateExecutionInfosRecursive(*LinkedState, Path))
-			{
-				return false;
-			}
-		}
-	}
-	
-	for (UStateTreeState* Child : State.Children)
-	{
-		if (Child)
-		{
-			if (!CreateExecutionInfosRecursive(*Child, Path))
-			{
-				return false;
-			}
-		}
-	}
-
-	Path.Pop();
+	StateTree->States[StateIdx].ChildrenEnd = uint16(StateTree->States.Num()); // Cannot use CompactState here, it may be invalid due to array resize.
 
 	return true;
 }
@@ -296,19 +209,66 @@ bool FStateTreeCompiler::CreateConditions(UStateTreeState& State, TConstArrayVie
 	return true;
 }
 
-bool FStateTreeCompiler::CreateStateTasks()
+bool FStateTreeCompiler::CreateStateTasksAndParameters()
 {
 	for (int32 i = 0; i < StateTree->States.Num(); i++)
 	{
-		FCompactStateTreeState& BakedState = StateTree->States[i];
+		FCompactStateTreeState& CompactState = StateTree->States[i];
 		UStateTreeState* SourceState = SourceStates[i];
 		check(SourceState != nullptr);
 
 		FStateTreeCompilerLogStateScope LogStateScope(SourceState, Log);
+
+		// Create parameters
+		if (SourceState->Type == EStateTreeStateType::Linked || SourceState->Type == EStateTreeStateType::Subtree)
+		{
+			// Both linked and subtree has instance data describing their parameters.
+			// This allows to resolve the binding paths and lets us have bindable parameters when transitioned into a parameterized subtree directly.
+			FInstancedStruct& Instance = StateTree->Instances.AddDefaulted_GetRef();
+			const int32 InstanceIndex = StateTree->Instances.Num() - 1;
+
+			check(InstanceIndex <= int32(MAX_uint16));
+			CompactState.ParameterInstanceIndex = InstanceIndex;
+		
+			Instance.InitializeAs<FCompactStateTreeParameters>();
+			FCompactStateTreeParameters& CompactParams = Instance.GetMutable<FCompactStateTreeParameters>();
+
+			CompactParams.Parameters = SourceState->Parameters.Parameters;
+
+			if (SourceState->Type == EStateTreeStateType::Subtree)
+			{
+				// Register a binding source
+				CompactState.ParameterDataViewIndex = BindingsCompiler.AddSourceStruct({ SourceState->Name, SourceState->Parameters.Parameters.GetPropertyBagStruct(), EStateTreeBindableStructSource::StateParameter, SourceState->Parameters.ID });
+			}
+			else if (SourceState->Type == EStateTreeStateType::Linked)
+			{
+				// Binding target
+				FStateTreeBindableStructDesc StructDesc;
+				StructDesc.ID = SourceState->Parameters.ID;
+				StructDesc.Name = SourceState->Name;
+				StructDesc.DataSource = EStateTreeBindableStructSource::StateParameter;
+				StructDesc.Struct = SourceState->Parameters.Parameters.GetPropertyBagStruct();
+
+				// Check that the bindings for this struct are still all valid.
+				TArray<FStateTreeEditorPropertyBinding> Bindings;
+				if (!GetAndValidateBindings(*SourceState, StructDesc, Bindings))
+				{
+					return false;
+				}
+
+				int32 BatchIndex = INDEX_NONE;
+				if (!BindingsCompiler.CompileBatch(StructDesc, Bindings, BatchIndex))
+				{
+					return false;
+				}
+				check(BatchIndex < int32(MAX_uint16));
+				CompactParams.BindingsBatch = BatchIndex == INDEX_NONE ? FStateTreeHandle::Invalid : FStateTreeHandle(uint16(BatchIndex));
+			}
+		}
 		
 		// Create tasks
 		check(StateTree->Nodes.Num() <= int32(MAX_uint16));
-		BakedState.TasksBegin = uint16(StateTree->Nodes.Num());
+		CompactState.TasksBegin = uint16(StateTree->Nodes.Num());
 
 		for (FStateTreeEditorNode& TaskNode : SourceState->Tasks)
 		{
@@ -323,9 +283,9 @@ bool FStateTreeCompiler::CreateStateTasks()
 			return false;
 		}
 	
-		const int32 TasksNum = StateTree->Nodes.Num() - int32(BakedState.TasksBegin);
+		const int32 TasksNum = StateTree->Nodes.Num() - int32(CompactState.TasksBegin);
 		check(TasksNum <= int32(MAX_uint8));
-		BakedState.TasksNum = uint8(TasksNum);
+		CompactState.TasksNum = uint8(TasksNum);
 	}
 	
 	return true;
@@ -335,7 +295,7 @@ bool FStateTreeCompiler::CreateStateEvaluators()
 {
 	for (int32 i = 0; i < StateTree->States.Num(); i++)
 	{
-		FCompactStateTreeState& BakedState = StateTree->States[i];
+		FCompactStateTreeState& CompactState = StateTree->States[i];
 		UStateTreeState* SourceState = SourceStates[i];
 		check(SourceState != nullptr);
 
@@ -343,7 +303,7 @@ bool FStateTreeCompiler::CreateStateEvaluators()
 		
 		// Collect evaluators
 		check(StateTree->Nodes.Num() <= int32(MAX_uint16));
-		BakedState.EvaluatorsBegin = uint16(StateTree->Nodes.Num());
+		CompactState.EvaluatorsBegin = uint16(StateTree->Nodes.Num());
 
 		for (FStateTreeEditorNode& EvalNode : SourceState->Evaluators)
 		{
@@ -353,9 +313,9 @@ bool FStateTreeCompiler::CreateStateEvaluators()
 			}
 		}
 		
-		const int32 EvaluatorsNum = StateTree->Nodes.Num() - int32(BakedState.EvaluatorsBegin);
+		const int32 EvaluatorsNum = StateTree->Nodes.Num() - int32(CompactState.EvaluatorsBegin);
 		check(EvaluatorsNum <= int32(MAX_uint8));
-		BakedState.EvaluatorsNum = uint8(EvaluatorsNum);
+		CompactState.EvaluatorsNum = uint8(EvaluatorsNum);
 	}
 
 	return true;
@@ -365,21 +325,21 @@ bool FStateTreeCompiler::CreateStateTransitions()
 {
 	for (int32 i = 0; i < StateTree->States.Num(); i++)
 	{
-		FCompactStateTreeState& BakedState = StateTree->States[i];
+		FCompactStateTreeState& CompactState = StateTree->States[i];
 		UStateTreeState* SourceState = SourceStates[i];
 		check(SourceState != nullptr);
 
 		FStateTreeCompilerLogStateScope LogStateScope(SourceState, Log);
 		
 		// Enter conditions.
-		BakedState.EnterConditionsBegin = uint16(StateTree->Nodes.Num());
+		CompactState.EnterConditionsBegin = uint16(StateTree->Nodes.Num());
 		if (!CreateConditions(*SourceState, SourceState->EnterConditions))
 		{
 			Log.Reportf(EMessageSeverity::Error,
 				TEXT("Failed to create state enter condition."));
 			return false;
 		}
-		BakedState.EnterConditionsNum = uint8(uint16(StateTree->Nodes.Num()) - BakedState.EnterConditionsBegin);
+		CompactState.EnterConditionsNum = uint8(uint16(StateTree->Nodes.Num()) - CompactState.EnterConditionsBegin);
 
 		// Linked state
 		if (SourceState->Type == EStateTreeStateType::Linked)
@@ -403,9 +363,9 @@ bool FStateTreeCompiler::CreateStateTransitions()
 				return false;
 			}
 			
-			BakedState.LinkedState = GetStateHandle(SourceState->LinkedState.ID);
+			CompactState.LinkedState = GetStateHandle(SourceState->LinkedState.ID);
 			
-			if (!BakedState.LinkedState.IsValid())
+			if (!CompactState.LinkedState.IsValid())
 			{
 				Log.Reportf(EMessageSeverity::Error,
 					TEXT("Failed to resolve linked state '%s'."),
@@ -415,7 +375,7 @@ bool FStateTreeCompiler::CreateStateTransitions()
 		}
 		
 		// Transitions
-		BakedState.TransitionsBegin = uint16(StateTree->Transitions.Num());
+		CompactState.TransitionsBegin = uint16(StateTree->Transitions.Num());
 		for (FStateTreeTransition& Transition : SourceState->Transitions)
 		{
 			FCompactStateTransition& BakedTransition = StateTree->Transitions.AddDefaulted_GetRef();
@@ -439,7 +399,7 @@ bool FStateTreeCompiler::CreateStateTransitions()
 			}
 			BakedTransition.ConditionsNum = uint8(uint16(StateTree->Nodes.Num()) - BakedTransition.ConditionsBegin);
 		}
-		BakedState.TransitionsNum = uint8(uint16(StateTree->Transitions.Num()) - BakedState.TransitionsBegin);
+		CompactState.TransitionsNum = uint8(uint16(StateTree->Transitions.Num()) - CompactState.TransitionsBegin);
 	}
 
 	// @todo: Add test to check that all success/failure transition is possible (see editor).
@@ -494,6 +454,7 @@ bool FStateTreeCompiler::CreateCondition(UStateTreeState& State, const FStateTre
 	FStateTreeBindableStructDesc StructDesc;
 	StructDesc.ID = CondNode.ID;
 	StructDesc.Name = CondNode.Node.GetScriptStruct()->GetFName();
+	StructDesc.DataSource = EStateTreeBindableStructSource::Condition;
 
 	// Check that item has valid instance initialized.
 	if (!CondNode.Instance.IsValid() && CondNode.InstanceObject == nullptr)
@@ -584,6 +545,7 @@ bool FStateTreeCompiler::CreateTask(UStateTreeState& State, const FStateTreeEdit
 	FStateTreeBindableStructDesc StructDesc;
 	StructDesc.ID = TaskNode.ID;
 	StructDesc.Name = TaskNode.Node.GetScriptStruct()->GetFName();
+	StructDesc.DataSource = EStateTreeBindableStructSource::Task;
 
 	// Check that item has valid instance initialized.
 	if (!TaskNode.Instance.IsValid() && TaskNode.InstanceObject == nullptr)
@@ -672,6 +634,7 @@ bool FStateTreeCompiler::CreateEvaluator(UStateTreeState& State, const FStateTre
 	FStateTreeBindableStructDesc StructDesc;
     StructDesc.ID = EvalNode.ID;
     StructDesc.Name = EvalNode.Node.GetScriptStruct()->GetFName();
+	StructDesc.DataSource = EStateTreeBindableStructSource::Evaluator;
 
     // Check that item has valid instance initialized.
     if (!EvalNode.Instance.IsValid() && EvalNode.InstanceObject == nullptr)
@@ -785,42 +748,29 @@ bool FStateTreeCompiler::GetAndValidateBindings(UStateTreeState& State, const FS
 		if (SourceStructIdx == INDEX_NONE)
 		{
 			Log.Reportf(EMessageSeverity::Error, TargetStruct,
-						TEXT("Failed to find binding source '%s:%s'."),
-						*TargetStruct.Name.ToString(), *Binding.TargetPath.ToString());
+						TEXT("Failed to find binding source property '%s' for target '%s:%s'."),
+						*Binding.SourcePath.ToString(), *TargetStruct.Name.ToString(), *Binding.TargetPath.ToString());
 			return false;
 		}
 		const FStateTreeBindableStructDesc& SourceStruct = BindingsCompiler.GetSourceStructDesc(SourceStructIdx);
 
 		// Source must be accessible by the target struct via all execution paths.
 		TArray<FStateTreeBindableStructDesc> AccessibleStructs;
-		const FStateExecutionInfo& ExecInfo = ExecutionInfos.FindChecked(&State);
-		for (const FExecutionPath& ExecPath : ExecInfo.ExecutionPaths)
-		{
-			AccessibleStructs.Reset();
-			TreeData->GetAccessibleStructs(ExecPath.Path, Binding.TargetPath.StructID, AccessibleStructs);
+		TreeData->GetAccessibleStructs(Binding.TargetPath.StructID, AccessibleStructs);
 
-			const bool SourceAccessible = AccessibleStructs.ContainsByPredicate([SourceStructID](const FStateTreeBindableStructDesc& Structs)
-				{
-					return (Structs.ID == SourceStructID);
-				});
-
-			if (!SourceAccessible)
+		const bool bSourceAccessible = AccessibleStructs.ContainsByPredicate([SourceStructID](const FStateTreeBindableStructDesc& Structs)
 			{
-				Log.Reportf(EMessageSeverity::Error, TargetStruct,
-							TEXT("Property '%s:%s' cannot be bound to '%s:%s', because the binding source '%s' is not updated before '%s' in the tree."),
-							*SourceStruct.Name.ToString(), *Binding.SourcePath.ToString(),
-							*TargetStruct.Name.ToString(), *Binding.TargetPath.ToString(),
-							*SourceStruct.Name.ToString(), *TargetStruct.Name.ToString());
+				return (Structs.ID == SourceStructID);
+			});
 
-				if (IsPathLinked(ExecPath.Path))
-				{
-					Log.Reportf(EMessageSeverity::Error, TargetStruct,
-								TEXT("The binding source is not updated when executing via linked state: %s."),
-								*GetExecutionPathString(ExecPath.Path));
-				}
-				
-				return false;
-			}
+		if (!bSourceAccessible)
+		{
+			Log.Reportf(EMessageSeverity::Error, TargetStruct,
+						TEXT("Property '%s:%s' cannot be bound to '%s:%s', because the binding source '%s' is not updated before '%s' in the tree."),
+						*SourceStruct.Name.ToString(), *Binding.SourcePath.ToString(),
+						*TargetStruct.Name.ToString(), *Binding.TargetPath.ToString(),
+						*SourceStruct.Name.ToString(), *TargetStruct.Name.ToString());
+			return false;
 		}
 		
 		// Special case fo AnyEnum. StateTreeBindingExtension allows AnyEnums to bind to other enum types.
diff --git a/Engine/Plugins/Runtime/StateTree/Source/StateTreeEditorModule/Private/StateTreeEditor.cpp b/Engine/Plugins/Runtime/StateTree/Source/StateTreeEditorModule/Private/StateTreeEditor.cpp
--- a/Engine/Plugins/Runtime/StateTree/Source/StateTreeEditorModule/Private/StateTreeEditor.cpp
+++ b/Engine/Plugins/Runtime/StateTree/Source/StateTreeEditorModule/Private/StateTreeEditor.cpp
@@ -218,6 +218,7 @@ void FStateTreeEditor::InitEditor( const EToolkitMode::Type Mode, const TSharedP
 	UE::StateTree::Delegates::OnIdentifierChanged.AddSP(this, &FStateTreeEditor::OnIdentifierChanged);
 	UE::StateTree::Delegates::OnSchemaChanged.AddSP(this, &FStateTreeEditor::OnSchemaChanged);
 	UE::StateTree::Delegates::OnParametersChanged.AddSP(this, &FStateTreeEditor::OnParametersChanged);
+	UE::StateTree::Delegates::OnStateParametersChanged.AddSP(this, &FStateTreeEditor::OnStateParametersChanged);
 }
 
 FName FStateTreeEditor::GetToolkitFName() const
@@ -428,6 +429,31 @@ void FStateTreeEditor::OnParametersChanged(const UStateTree& InStateTree)
 	}
 }
 
+void FStateTreeEditor::OnStateParametersChanged(const UStateTree& InStateTree, const FGuid ChangedStateID)
+{
+	if (StateTree == &InStateTree)
+	{
+		if (const UStateTreeEditorData* TreeData = Cast<UStateTreeEditorData>(StateTree->EditorData))
+		{
+			TreeData->VisitHierarchy([&ChangedStateID](UStateTreeState& State, UStateTreeState* /*ParentState*/)
+			{
+				if (State.Type == EStateTreeStateType::Linked && State.LinkedState.ID == ChangedStateID)
+				{
+					State.UpdateParametersFromLinkedState();
+				}
+				return EStateTreeVisitor::Continue;
+			});
+		}
+
+		// Accessible structs might be different after modifying parameters so forcing refresh
+		// so the FStateTreeBindingExtension can rebuild the list of bindable structs
+		if (SelectionDetailsView.IsValid())
+		{
+			SelectionDetailsView->ForceRefresh();
+		}
+	}
+}
+
 void FStateTreeEditor::OnAssetFinishedChangingProperties(const FPropertyChangedEvent& PropertyChangedEvent)
 {
 	// Make sure nodes get updates when properties are changed.
@@ -496,60 +522,28 @@ namespace UE::StateTree::Editor::Internal
 
 		// Create ID to state name map.
 		TMap<FGuid, FName> IDToName;
-		TArray<UStateTreeState*> Stack;
 
-		for (UStateTreeState* SubTree : TreeData->SubTrees)
+		TreeData->VisitHierarchy([&IDToName](const UStateTreeState& State, UStateTreeState* /*ParentState*/)
 		{
-			if (SubTree)
-			{
-				Stack.Reset();
-				Stack.Add(SubTree);
-				while (Stack.Num() > 0)
-				{
-					UStateTreeState* CurState = Stack.Pop();
-					IDToName.Add(CurState->ID, CurState->Name);
-					for (UStateTreeState* ChildState : CurState->Children)
-					{
-						if (ChildState)
-						{
-							Stack.Append(CurState->Children);
-						}
-					}
-				}
-			}
-		}
-
+			IDToName.Add(State.ID, State.Name);
+			return EStateTreeVisitor::Continue;
+		});
+		
 		// Fix changed names.
-		for (UStateTreeState* SubTree : TreeData->SubTrees)
+		TreeData->VisitHierarchy([&IDToName](UStateTreeState& State, UStateTreeState* /*ParentState*/)
 		{
-			if (SubTree)
+			if (State.Type == EStateTreeStateType::Linked)
 			{
-				Stack.Reset();
-				Stack.Add(SubTree);
-				while (Stack.Num() > 0)
-				{
-					UStateTreeState* CurState = Stack.Pop();
-
-					if (CurState->Type == EStateTreeStateType::Linked)
-					{
-						FixChangedStateLinkName(CurState->LinkedState, IDToName);
-					}
+				FixChangedStateLinkName(State.LinkedState, IDToName);
+			}
 					
-					for (FStateTreeTransition& Transition : CurState->Transitions)
-					{
-						FixChangedStateLinkName(Transition.State, IDToName);
-					}
-
-					for (UStateTreeState* ChildState : CurState->Children)
-					{
-						if (ChildState)
-						{
-							Stack.Append(CurState->Children);
-						}
-					}
-				}
+			for (FStateTreeTransition& Transition : State.Transitions)
+			{
+				FixChangedStateLinkName(Transition.State, IDToName);
 			}
-		}
+
+			return EStateTreeVisitor::Continue;
+		});
 	}
 
 	void UpdateParents(UStateTree& StateTree)
@@ -560,29 +554,11 @@ namespace UE::StateTree::Editor::Internal
 			return;
 		}
 
-		TArray<UStateTreeState*> Stack;
-
-		for (UStateTreeState* SubTree : TreeData->SubTrees)
+		TreeData->VisitHierarchy([](UStateTreeState& State, UStateTreeState* ParentState)
 		{
-			if (SubTree)
-			{
-				SubTree->Parent = nullptr;
-				Stack.Reset();
-				Stack.Add(SubTree);
-				while (Stack.Num() > 0)
-				{
-					UStateTreeState* CurState = Stack.Pop();
-					for (UStateTreeState* ChildState : CurState->Children)
-					{
-						if (ChildState)
-						{
-							ChildState->Parent = CurState;
-							Stack.Append(CurState->Children);
-						}
-					}
-				}
-			}
-		}
+			State.Parent = ParentState;
+			return EStateTreeVisitor::Continue;
+		});
 	}
 
 	void ApplySchema(UStateTree& StateTree)
@@ -599,67 +575,48 @@ namespace UE::StateTree::Editor::Internal
 			return;
 		}
 
-		TArray<UStateTreeState*> Stack;
-
-		for (UStateTreeState* SubTree : TreeData->SubTrees)
+		TreeData->VisitHierarchy([&StateTree, Schema](UStateTreeState& State, UStateTreeState* /*ParentState*/)
 		{
-			if (SubTree)
+			// Clear enter conditions if not allowed.
+			if (Schema->AllowEnterConditions() == false && State.EnterConditions.Num() > 0)
 			{
-				Stack.Reset();
-				Stack.Add(SubTree);
-				while (Stack.Num() > 0)
+				UE_LOG(LogStateTree, Warning, TEXT("%s: Resetting Enter Conditions in state %s due to current schema restrictions."), *GetNameSafe(&StateTree), *GetNameSafe(&State));
+				State.EnterConditions.Reset();
+			}
+
+			// Clear evaluators if not allowed.
+			if (Schema->AllowEvaluators() == false && State.Evaluators.Num() > 0)
+			{
+				UE_LOG(LogStateTree, Warning, TEXT("%s: Resetting Evaluators in state %s due to current schema restrictions."), *GetNameSafe(&StateTree), *GetNameSafe(&State));
+				State.Evaluators.Reset();
+			}
+
+			// Keep single and many tasks based on what is allowed.
+			if (Schema->AllowMultipleTasks() == false)
+			{
+				if (State.Tasks.Num() > 0)
 				{
-					UStateTreeState* CurState = Stack.Pop();
-
-					// Clear enter conditions if not allowed.
-					if (Schema->AllowEnterConditions() == false && CurState->EnterConditions.Num() > 0)
-					{
-						UE_LOG(LogStateTree, Warning, TEXT("%s: Resetting Enter Conditions in state %s due to current schema restrictions."), *GetNameSafe(&StateTree), *GetNameSafe(CurState));
-						CurState->EnterConditions.Reset();
-					}
-
-					// Clear evaluators if not allowed.
-					if (Schema->AllowEvaluators() == false && CurState->Evaluators.Num() > 0)
-					{
-						UE_LOG(LogStateTree, Warning, TEXT("%s: Resetting Evaluators in state %s due to current schema restrictions."), *GetNameSafe(&StateTree), *GetNameSafe(CurState));
-						CurState->Evaluators.Reset();
-					}
-
-					// Keep single and many tasks based on what is allowed.
-					if (Schema->AllowMultipleTasks() == false)
-					{
-						if (CurState->Tasks.Num() > 0)
-						{
-							CurState->Tasks.Reset();
-							UE_LOG(LogStateTree, Warning, TEXT("%s: Resetting Tasks in state %s due to current schema restrictions."), *GetNameSafe(&StateTree), *GetNameSafe(CurState));
-						}
-						
-						// Task name is the same as state name.
-						if (FStateTreeTaskBase* Task = CurState->SingleTask.Node.GetMutablePtr<FStateTreeTaskBase>())
-						{
-							Task->Name = CurState->Name;
-						}
-					}
-					else
-					{
-						if (CurState->SingleTask.Node.IsValid())
-						{
-							CurState->SingleTask.Reset();
-							UE_LOG(LogStateTree, Warning, TEXT("%s: Resetting Single Task in state %s due to current schema restrictions."), *GetNameSafe(&StateTree), *GetNameSafe(CurState));
-						}
-					}
-					
-					for (UStateTreeState* ChildState : CurState->Children)
-					{
-						if (ChildState)
-						{
-							Stack.Append(CurState->Children);
-						}
-					}
+					State.Tasks.Reset();
+					UE_LOG(LogStateTree, Warning, TEXT("%s: Resetting Tasks in state %s due to current schema restrictions."), *GetNameSafe(&StateTree), *GetNameSafe(&State));
+				}
+				
+				// Task name is the same as state name.
+				if (FStateTreeTaskBase* Task = State.SingleTask.Node.GetMutablePtr<FStateTreeTaskBase>())
+				{
+					Task->Name = State.Name;
 				}
 			}
-		}
-
+			else
+			{
+				if (State.SingleTask.Node.IsValid())
+				{
+					State.SingleTask.Reset();
+					UE_LOG(LogStateTree, Warning, TEXT("%s: Resetting Single Task in state %s due to current schema restrictions."), *GetNameSafe(&StateTree), *GetNameSafe(&State));
+				}
+			}
+			
+			return EStateTreeVisitor::Continue;
+		});
 	}
 
 	void RemoveUnusedBindings(UStateTree& StateTree)
@@ -675,6 +632,24 @@ namespace UE::StateTree::Editor::Internal
 		TreeData->GetPropertyEditorBindings()->RemoveUnusedBindings(AllStructIDs);
 	}
 
+	void UpdateLinkedStateParameters(UStateTree& StateTree)
+	{
+		UStateTreeEditorData* TreeData = Cast<UStateTreeEditorData>(StateTree.EditorData);
+		if (!TreeData)
+		{
+			return;
+		}
+
+		TreeData->VisitHierarchy([](UStateTreeState& State, UStateTreeState* /*ParentState*/)
+		{
+			if (State.Type == EStateTreeStateType::Linked)
+			{
+				State.UpdateParametersFromLinkedState();
+			}
+			return EStateTreeVisitor::Continue;
+		});
+	}
+
 }
 
 void FStateTreeEditor::BindCommands()
@@ -703,15 +678,14 @@ void FStateTreeEditor::RegisterToolbar()
 	}
 
 	const FStateTreeEditorCommands& Commands = FStateTreeEditorCommands::Get();
-	FToolMenuInsert InsertAfterAssetSection("Asset", EToolMenuInsertType::After);
-	{
-		FToolMenuSection& Section = ToolBar->AddSection("Compile", TAttribute<FText>(), InsertAfterAssetSection);
-		Section.AddEntry(FToolMenuEntry::InitToolBarButton(
-			Commands.Compile,
-			TAttribute<FText>(),
-			TAttribute<FText>(),
-			TAttribute<FSlateIcon>(this, &FStateTreeEditor::GetCompileStatusImage)));
-	}
+	const FToolMenuInsert InsertAfterAssetSection("Asset", EToolMenuInsertType::After);
+
+	FToolMenuSection& Section = ToolBar->AddSection("Compile", TAttribute<FText>(), InsertAfterAssetSection);
+	Section.AddEntry(FToolMenuEntry::InitToolBarButton(
+		Commands.Compile,
+		TAttribute<FText>(),
+		TAttribute<FText>(),
+		TAttribute<FSlateIcon>(this, &FStateTreeEditor::GetCompileStatusImage)));
 }
 
 void FStateTreeEditor::Compile()
@@ -809,6 +783,7 @@ void FStateTreeEditor::UpdateAsset()
 	UE::StateTree::Editor::Internal::ApplySchema(*StateTree);
 	UE::StateTree::Editor::Internal::RemoveUnusedBindings(*StateTree);
 	UE::StateTree::Editor::Internal::ValidateLinkedStates(*StateTree);
+	UE::StateTree::Editor::Internal::UpdateLinkedStateParameters(*StateTree);
 
 	EditorDataHash = 0;
 	if (StateTree->EditorData != nullptr)
diff --git a/Engine/Plugins/Runtime/StateTree/Source/StateTreeEditorModule/Private/StateTreeEditor.h b/Engine/Plugins/Runtime/StateTree/Source/StateTreeEditorModule/Private/StateTreeEditor.h
--- a/Engine/Plugins/Runtime/StateTree/Source/StateTreeEditorModule/Private/StateTreeEditor.h
+++ b/Engine/Plugins/Runtime/StateTree/Source/StateTreeEditorModule/Private/StateTreeEditor.h
@@ -72,6 +72,7 @@ private:
 	void OnIdentifierChanged(const UStateTree& StateTree);
 	void OnSchemaChanged(const UStateTree& StateTree);
 	void OnParametersChanged(const UStateTree& StateTree);
+	void OnStateParametersChanged(const UStateTree& StateTree, const FGuid StateID);
 
 	FText GetStatisticsText() const;
 
diff --git a/Engine/Plugins/Runtime/StateTree/Source/StateTreeEditorModule/Private/StateTreeEditorData.cpp b/Engine/Plugins/Runtime/StateTree/Source/StateTreeEditorModule/Private/StateTreeEditorData.cpp
--- a/Engine/Plugins/Runtime/StateTree/Source/StateTreeEditorModule/Private/StateTreeEditorData.cpp
+++ b/Engine/Plugins/Runtime/StateTree/Source/StateTreeEditorModule/Private/StateTreeEditorData.cpp
@@ -66,6 +66,12 @@ void UStateTreeEditorData::GetAccessibleStructs(const FGuid TargetStructID, TArr
 	for (const UStateTreeState* State = CurrentState; State != nullptr; State = State->Parent)
 	{
 		Path.Insert(State, 0);
+
+		// Stop at subtree root.
+		if (State->Type == EStateTreeStateType::Subtree)
+		{
+			break;
+		}
 	}
 	
 	GetAccessibleStructs(Path, TargetStructID, OutStructDescs);
@@ -79,7 +85,7 @@ void UStateTreeEditorData::GetAccessibleStructs(const TConstArrayView<const USta
 	// All parameters are accessible
 	for (const FStateTreeParameterDesc& Desc : Parameters.Parameters)
 	{
-		OutStructDescs.Emplace(Desc.Name, Desc.Parameter.GetScriptStruct(), Desc.ID);
+		OutStructDescs.Emplace(Desc.Name, Desc.Parameter.GetScriptStruct(), EStateTreeBindableStructSource::TreeParameter, Desc.ID);
 	}
 
 	// All named external data items declared by the schema are accessible
@@ -87,10 +93,22 @@ void UStateTreeEditorData::GetAccessibleStructs(const TConstArrayView<const USta
 	{
 		for (const FStateTreeExternalDataDesc& Desc : Schema->GetNamedExternalDataDescs())
 		{
-			OutStructDescs.Emplace(Desc.Name, Desc.Struct, Desc.ID);
+			OutStructDescs.Emplace(Desc.Name, Desc.Struct, EStateTreeBindableStructSource::TreeData, Desc.ID);
 		}	
 	}
 
+	// State parameters
+	const UStateTreeState* RootState = Path.Num() > 0 ? Path[0] : nullptr;
+	if (RootState != nullptr
+		&& RootState->Type == EStateTreeStateType::Subtree
+		&& RootState->Parameters.Parameters.GetPropertyBagStruct() != nullptr)
+	{
+		if (const UScriptStruct* PropertyBagStruct = RootState->Parameters.Parameters.GetPropertyBagStruct())
+		{
+			OutStructDescs.Emplace(RootState->Name, PropertyBagStruct, EStateTreeBindableStructSource::StateParameter, RootState->Parameters.ID);
+		}
+	}
+
 	TArray<FStateTreeBindableStructDesc> EvalDescs;
 	TArray<FStateTreeBindableStructDesc> TaskDescs;
 
@@ -101,35 +119,27 @@ void UStateTreeEditorData::GetAccessibleStructs(const TConstArrayView<const USta
 			continue;
 		}
 		
-		VisitState(*State, [&OutStructDescs, &EvalDescs, &TaskDescs, TargetStructID]
-			(const UStateTreeState& State, const FGuid& ID, const FName& Name, const EStateTreeNodeType NodeType, const UScriptStruct* NodeStruct, const UStruct* InstanceStruct) -> bool
+		VisitStateNodes(*State, [&OutStructDescs, &EvalDescs, &TaskDescs, TargetStructID]
+			(const UStateTreeState& State, const FGuid& ID, const FName& Name, const EStateTreeNodeType NodeType, const UScriptStruct* NodeStruct, const UStruct* InstanceStruct)
 			{
 				if (ID == TargetStructID)
 				{
-					if (NodeType == EStateTreeNodeType::EnterCondition)
+					if (NodeType == EStateTreeNodeType::EnterCondition
+						|| NodeType == EStateTreeNodeType::Task
+						|| NodeType == EStateTreeNodeType::TransitionCondition
+						|| NodeType == EStateTreeNodeType::StateParameters)
 					{
-						// Enter conditions can only see evaluators
+						// These can only see both evaluators and tasks
 						OutStructDescs.Append(EvalDescs);
+						OutStructDescs.Append(TaskDescs);
 					}
 					else if (NodeType == EStateTreeNodeType::Evaluator)
 					{
 						// Evaluators can only see other evaluators
 						OutStructDescs.Append(EvalDescs);
 					}
-					else if (NodeType == EStateTreeNodeType::Task)
-					{
-						// Tasks can see evals and tasks.
-						OutStructDescs.Append(EvalDescs);
-						OutStructDescs.Append(TaskDescs);
-					}
-					else if (NodeType == EStateTreeNodeType::TransitionCondition)
-					{
-						// Transitions can see evals and tasks.
-						OutStructDescs.Append(EvalDescs);
-						OutStructDescs.Append(TaskDescs);
-					}
 					
-					return false; // Stop visit
+					return EStateTreeVisitor::Break;
 				}
 
 				// Not at target yet, collect all evaluators and tasks visible so far.
@@ -139,6 +149,7 @@ void UStateTreeEditorData::GetAccessibleStructs(const TConstArrayView<const USta
 					Desc.Struct = InstanceStruct;
 					Desc.Name = Name;
 					Desc.ID = ID;
+					Desc.DataSource = EStateTreeBindableStructSource::Evaluator;
 				}
 				else if (NodeStruct->IsChildOf(FStateTreeTaskBase::StaticStruct()))
 				{
@@ -146,17 +157,23 @@ void UStateTreeEditorData::GetAccessibleStructs(const TConstArrayView<const USta
 					Desc.Struct = InstanceStruct;
 					Desc.Name = Name;
 					Desc.ID = ID;
+					Desc.DataSource = EStateTreeBindableStructSource::Task;
 				}
-				return true; // Continue
+				return EStateTreeVisitor::Continue;
 			});
 	}
+
+	OutStructDescs.StableSort([](const FStateTreeBindableStructDesc& A, const FStateTreeBindableStructDesc& B)
+	{
+		return (uint8)A.DataSource < (uint8)B.DataSource;
+	});
 }
 
 bool UStateTreeEditorData::GetStructByID(const FGuid StructID, FStateTreeBindableStructDesc& OutStructDesc) const
 {
 	bool bResult = false;
 
-	VisitHierarchy([&bResult, &OutStructDesc, StructID](const UStateTreeState& State, const FGuid& ID, const FName& Name, const EStateTreeNodeType NodeType, const UScriptStruct* NodeStruct, const UStruct* InstanceStruct) -> bool
+	VisitHierarchyNodes([&bResult, &OutStructDesc, StructID](const UStateTreeState& State, const FGuid& ID, const FName& Name, const EStateTreeNodeType NodeType, const UScriptStruct* NodeStruct, const UStruct* InstanceStruct)
 		{
 			if (ID == StructID)
 			{
@@ -164,9 +181,9 @@ bool UStateTreeEditorData::GetStructByID(const FGuid StructID, FStateTreeBindabl
 				OutStructDesc.Name = Name;
 				OutStructDesc.ID = ID;
 				bResult = true;
-				return false; // Stop visit
+				return EStateTreeVisitor::Break;
 			}
-			return true; // Continue
+			return EStateTreeVisitor::Continue;
 		});
 
 	return bResult;
@@ -176,19 +193,37 @@ const UStateTreeState* UStateTreeEditorData::GetStateByStructID(const FGuid Targ
 {
 	const UStateTreeState* Result = nullptr;
 
-	VisitHierarchy([&Result, TargetStructID](const UStateTreeState& State, const FGuid& ID, const FName& Name, const EStateTreeNodeType NodeType, const UScriptStruct* NodeStruct, const UStruct* InstanceStruct) -> bool
+	VisitHierarchyNodes([&Result, TargetStructID](const UStateTreeState& State, const FGuid& ID, const FName& Name, const EStateTreeNodeType NodeType, const UScriptStruct* NodeStruct, const UStruct* InstanceStruct)
 		{
 			if (ID == TargetStructID)
 			{
 				Result = &State;
-				return false; // Stop visit
+				return EStateTreeVisitor::Break;
 			}
-			return true; // Continue
+			return EStateTreeVisitor::Continue;
 		});
 
 	return Result;
 }
 
+const UStateTreeState* UStateTreeEditorData::GetStateByID(const FGuid StateID) const
+{
+	const UStateTreeState* Result = nullptr;
+	
+	VisitHierarchy([&Result, &StateID](const UStateTreeState& State, UStateTreeState* /*ParentState*/)
+	{
+		if (State.ID == StateID)
+		{
+			Result = &State;
+			return EStateTreeVisitor::Break;
+		}
+
+		return EStateTreeVisitor::Continue;
+	});
+
+	return Result;
+}
+
 void UStateTreeEditorData::GetAllStructIDs(TMap<FGuid, const UStruct*>& AllStructs) const
 {
 	AllStructs.Reset();
@@ -211,23 +246,23 @@ void UStateTreeEditorData::GetAllStructIDs(TMap<FGuid, const UStruct*>& AllStruc
 		}	
 	}
 
-	VisitHierarchy([&AllStructs](const UStateTreeState& State, const FGuid& ID, const FName& Name, const EStateTreeNodeType NodeType, const UScriptStruct* NodeStruct, const UStruct* InstanceStruct) -> bool
+	VisitHierarchyNodes([&AllStructs](const UStateTreeState& State, const FGuid& ID, const FName& Name, const EStateTreeNodeType NodeType, const UScriptStruct* NodeStruct, const UStruct* InstanceStruct)
 		{
 			AllStructs.Add(ID, InstanceStruct);
-			return true; // Continue
+			return EStateTreeVisitor::Continue;
 		});
 }
 
-bool UStateTreeEditorData::VisitState(const UStateTreeState& State, TFunctionRef<bool(const UStateTreeState& State, const FGuid& ID, const FName& Name, const EStateTreeNodeType NodeType, const UScriptStruct* NodeStruct, const UStruct* InstanceStruct)> InFunc) const
+EStateTreeVisitor UStateTreeEditorData::VisitStateNodes(const UStateTreeState& State, TFunctionRef<EStateTreeVisitor(const UStateTreeState& State, const FGuid& ID, const FName& Name, const EStateTreeNodeType NodeType, const UScriptStruct* NodeStruct, const UStruct* InstanceStruct)> InFunc) const
 {
 	bool bContinue = true;
-	
+
 	// Evaluators
 	for (const FStateTreeEditorNode& Node : State.Evaluators)
 	{
 		if (const FStateTreeEvaluatorBase* Evaluator = Node.Node.GetPtr<FStateTreeEvaluatorBase>())
 		{
-			if (!InFunc(State, Node.ID, Evaluator->Name, EStateTreeNodeType::Evaluator, Node.Node.GetScriptStruct(), Evaluator->GetInstanceDataType()))
+			if (InFunc(State, Node.ID, Evaluator->Name, EStateTreeNodeType::Evaluator, Node.Node.GetScriptStruct(), Evaluator->GetInstanceDataType()) == EStateTreeVisitor::Break)
 			{
 				bContinue = false;
 				break;
@@ -241,7 +276,7 @@ bool UStateTreeEditorData::VisitState(const UStateTreeState& State, TFunctionRef
 		{
 			if (const FStateTreeConditionBase* Cond = Node.Node.GetPtr<FStateTreeConditionBase>())
 			{
-				if (!InFunc(State, Node.ID, Node.Node.GetScriptStruct()->GetFName(), EStateTreeNodeType::EnterCondition, Node.Node.GetScriptStruct(), Cond->GetInstanceDataType()))
+				if (InFunc(State, Node.ID, Node.Node.GetScriptStruct()->GetFName(), EStateTreeNodeType::EnterCondition, Node.Node.GetScriptStruct(), Cond->GetInstanceDataType()) == EStateTreeVisitor::Break)
 				{
 					bContinue = false;
 					break;
@@ -256,7 +291,7 @@ bool UStateTreeEditorData::VisitState(const UStateTreeState& State, TFunctionRef
 		{
 			if (const FStateTreeTaskBase* Task = Node.Node.GetPtr<FStateTreeTaskBase>())
 			{
-				if (!InFunc(State, Node.ID, Task->Name, EStateTreeNodeType::Task, Node.Node.GetScriptStruct(), Task->GetInstanceDataType()))
+				if (InFunc(State, Node.ID, Task->Name, EStateTreeNodeType::Task, Node.Node.GetScriptStruct(), Task->GetInstanceDataType()) == EStateTreeVisitor::Break)
 				{
 					bContinue = false;
 					break;
@@ -268,7 +303,7 @@ bool UStateTreeEditorData::VisitState(const UStateTreeState& State, TFunctionRef
 	{
 		if (const FStateTreeTaskBase* Task = State.SingleTask.Node.GetPtr<FStateTreeTaskBase>())
 		{
-			if (!InFunc(State, State.SingleTask.ID, Task->Name, EStateTreeNodeType::Task, State.SingleTask.Node.GetScriptStruct(), Task->GetInstanceDataType()))
+			if (InFunc(State, State.SingleTask.ID, Task->Name, EStateTreeNodeType::Task, State.SingleTask.Node.GetScriptStruct(), Task->GetInstanceDataType()) == EStateTreeVisitor::Break)
 			{
 				bContinue = false;
 			}
@@ -284,7 +319,7 @@ bool UStateTreeEditorData::VisitState(const UStateTreeState& State, TFunctionRef
 			{
 				if (const FStateTreeConditionBase* Cond = Node.Node.GetPtr<FStateTreeConditionBase>())
 				{
-					if (!InFunc(State, Node.ID, Node.Node.GetScriptStruct()->GetFName(), EStateTreeNodeType::TransitionCondition, Node.Node.GetScriptStruct(), Cond->GetInstanceDataType()))
+					if (InFunc(State, Node.ID, Node.Node.GetScriptStruct()->GetFName(), EStateTreeNodeType::TransitionCondition, Node.Node.GetScriptStruct(), Cond->GetInstanceDataType()) == EStateTreeVisitor::Break)
 					{
 						bContinue = false;
 						break;
@@ -293,39 +328,53 @@ bool UStateTreeEditorData::VisitState(const UStateTreeState& State, TFunctionRef
 			}
 		}
 	}
+	if (bContinue)
+	{
+		// Bindable state parameters
+		if (State.Type == EStateTreeStateType::Linked && State.Parameters.Parameters.IsValid())
+		{
+			if (InFunc(State, State.Parameters.ID, State.Name, EStateTreeNodeType::StateParameters, nullptr, State.Parameters.Parameters.GetPropertyBagStruct()) == EStateTreeVisitor::Break)
+			{
+				bContinue = false;
+			}
+		}
+	}
 
-	return bContinue;
+	return bContinue ? EStateTreeVisitor::Continue : EStateTreeVisitor::Break;
 }
 
-void UStateTreeEditorData::VisitHierarchy(TFunctionRef<bool(const UStateTreeState& State, const FGuid& ID, const FName& Name, const EStateTreeNodeType NodeType, const UScriptStruct* NodeStruct, const UStruct* InstanceStruct)> InFunc) const
+void UStateTreeEditorData::VisitHierarchy(TFunctionRef<EStateTreeVisitor(UStateTreeState& State, UStateTreeState* ParentState)> InFunc) const
 {
-	TArray<const UStateTreeState*> Stack;
+	using FStatePair = TTuple<UStateTreeState*, UStateTreeState*>; 
+	TArray<FStatePair> Stack;
 	bool bContinue = true;
 
-	for (const UStateTreeState* SubTree : SubTrees)
+	for (UStateTreeState* SubTree : SubTrees)
 	{
 		if (!SubTree)
 		{
 			continue;
 		}
 
-		Stack.Add(SubTree);
+		Stack.Add( FStatePair(nullptr, SubTree));
 
 		while (!Stack.IsEmpty() && bContinue)
 		{
-			const UStateTreeState* State = Stack[0];
+			FStatePair Current = Stack[0];
+			UStateTreeState* ParentState = Current.Get<0>();
+			UStateTreeState* State = Current.Get<1>();
 			check(State);
 
 			Stack.RemoveAt(0);
 
-			bContinue = VisitState(*State, InFunc);
+			bContinue = InFunc(*State, ParentState) == EStateTreeVisitor::Continue;
 			
 			if (bContinue)
 			{
 				// Children
-				for (const UStateTreeState* ChildState : State->Children)
+				for (UStateTreeState* ChildState : State->Children)
 				{
-					Stack.Add(ChildState);
+					Stack.Add(FStatePair(State, ChildState));
 				}
 			}
 		}
@@ -337,3 +386,11 @@ void UStateTreeEditorData::VisitHierarchy(TFunctionRef<bool(const UStateTreeStat
 	}
 }
 
+void UStateTreeEditorData::VisitHierarchyNodes(TFunctionRef<EStateTreeVisitor(const UStateTreeState& State, const FGuid& ID, const FName& Name, const EStateTreeNodeType NodeType, const UScriptStruct* NodeStruct, const UStruct* InstanceStruct)> InFunc) const
+{
+	VisitHierarchy([this, &InFunc](const UStateTreeState& State, UStateTreeState* /*ParentState*/)
+	{
+		return VisitStateNodes(State, InFunc);
+	});
+}
+
diff --git a/Engine/Plugins/Runtime/StateTree/Source/StateTreeEditorModule/Private/StateTreeEditorModule.cpp b/Engine/Plugins/Runtime/StateTree/Source/StateTreeEditorModule/Private/StateTreeEditorModule.cpp
--- a/Engine/Plugins/Runtime/StateTree/Source/StateTreeEditorModule/Private/StateTreeEditorModule.cpp
+++ b/Engine/Plugins/Runtime/StateTree/Source/StateTreeEditorModule/Private/StateTreeEditorModule.cpp
@@ -9,6 +9,7 @@
 #include "Customizations/StateTreeStateLinkDetails.h"
 #include "Customizations/StateTreeStateDetails.h"
 #include "Customizations/StateTreeEditorNodeDetails.h"
+#include "Customizations/StateTreeStateParametersDetails.h"
 #include "Customizations/StateTreeAnyEnumDetails.h"
 #include "StateTreeEditor.h"
 #include "StateTree.h"
@@ -50,6 +51,7 @@ void FStateTreeEditorModule::StartupModule()
 	PropertyModule.RegisterCustomPropertyTypeLayout("StateTreeTransition", FOnGetPropertyTypeCustomizationInstance::CreateStatic(&FStateTreeTransitionDetails::MakeInstance));
 	PropertyModule.RegisterCustomPropertyTypeLayout("StateTreeStateLink", FOnGetPropertyTypeCustomizationInstance::CreateStatic(&FStateTreeStateLinkDetails::MakeInstance));
 	PropertyModule.RegisterCustomPropertyTypeLayout("StateTreeEditorNode", FOnGetPropertyTypeCustomizationInstance::CreateStatic(&FStateTreeEditorNodeDetails::MakeInstance));
+	PropertyModule.RegisterCustomPropertyTypeLayout("StateTreeStateParameters", FOnGetPropertyTypeCustomizationInstance::CreateStatic(&FStateTreeStateParametersDetails::MakeInstance));
 	PropertyModule.RegisterCustomPropertyTypeLayout("StateTreeAnyEnum", FOnGetPropertyTypeCustomizationInstance::CreateStatic(&FStateTreeAnyEnumDetails::MakeInstance));
 	PropertyModule.RegisterCustomClassLayout("StateTreeState", FOnGetDetailCustomizationInstance::CreateStatic(&FStateTreeStateDetails::MakeInstance));
 
@@ -84,9 +86,8 @@ void FStateTreeEditorModule::ShutdownModule()
 		FPropertyEditorModule& PropertyModule = FModuleManager::LoadModuleChecked<FPropertyEditorModule>("PropertyEditor");
 		PropertyModule.UnregisterCustomPropertyTypeLayout("StateTreeTransition");
 		PropertyModule.UnregisterCustomPropertyTypeLayout("StateTreeStateLink");
-		PropertyModule.UnregisterCustomPropertyTypeLayout("StateTreeEvaluatorItem");
-		PropertyModule.UnregisterCustomPropertyTypeLayout("StateTreeTaskItem");
-		PropertyModule.UnregisterCustomPropertyTypeLayout("StateTreeConditionItem");
+		PropertyModule.UnregisterCustomPropertyTypeLayout("StateTreeEditorNode");
+		PropertyModule.UnregisterCustomPropertyTypeLayout("StateTreeStateParameters");
 		PropertyModule.UnregisterCustomPropertyTypeLayout("StateTreeAnyEnum");
 		PropertyModule.NotifyCustomizationModuleChanged();
 	}
diff --git a/Engine/Plugins/Runtime/StateTree/Source/StateTreeEditorModule/Private/StateTreeState.cpp b/Engine/Plugins/Runtime/StateTree/Source/StateTreeEditorModule/Private/StateTreeState.cpp
--- a/Engine/Plugins/Runtime/StateTree/Source/StateTreeEditorModule/Private/StateTreeState.cpp
+++ b/Engine/Plugins/Runtime/StateTree/Source/StateTreeEditorModule/Private/StateTreeState.cpp
@@ -2,6 +2,7 @@
 
 #include "StateTreeState.h"
 #include "StateTree.h"
+#include "StateTreeEditorData.h"
 #include "StateTreeConditionBase.h"
 #include "StateTreeEvaluatorBase.h"
 #include "StateTreeTaskBase.h"
@@ -70,16 +71,58 @@ void UStateTreeState::PostEditChangeChainProperty(FPropertyChangedChainEvent& Pr
 		if (Property->GetOwnerClass() == UStateTreeState::StaticClass()
 			&& Property->GetFName() == GET_MEMBER_NAME_CHECKED(UStateTreeState, Type))
 		{
+			// Remove any tasks and evaluators when they are not used.
 			if (Type == EStateTreeStateType::Group || Type == EStateTreeStateType::Linked)
 			{
-				// Remove any tasks and evaluators.
 				Tasks.Reset();
 				Evaluators.Reset();
 			}
+
+			// If transitioning from linked state, reset the linked state.
 			if (Type != EStateTreeStateType::Linked)
 			{
 				LinkedState = FStateTreeStateLink();
 			}
+
+			if (Type == EStateTreeStateType::Linked)
+			{
+				// Linked parameter layout is fixed, and copied from the linked target state.
+				Parameters.bFixedLayout = true;
+				UpdateParametersFromLinkedState();
+			}
+			else if (Type == EStateTreeStateType::Subtree)
+			{
+				// Subtree parameter layout can be edited
+				Parameters.bFixedLayout = false;
+			}
+			else
+			{
+				Parameters.Reset();
+			}
+		}
+
+		if (Property->GetOwnerClass() == UStateTreeState::StaticClass()
+			&& Property->GetFName() == GET_MEMBER_NAME_CHECKED(UStateTreeState, LinkedState))
+		{
+			// When switching to new state, update the parameters.
+			if (Type == EStateTreeStateType::Linked)
+			{
+				UpdateParametersFromLinkedState();
+			}
+		}
+
+		if (Property->GetOwnerClass() == UStateTreeState::StaticClass()
+			&& Property->GetFName() == GET_MEMBER_NAME_CHECKED(UStateTreeState, Parameters))
+		{
+			if (Type == EStateTreeStateType::Subtree)
+			{
+				// Broadcast subtree parameter edits so that the linked states can adapt.
+				const UStateTree* StateTree = GetTypedOuter<UStateTree>();
+				if (ensure(StateTree))
+				{
+					UE::StateTree::Delegates::OnStateParametersChanged.Broadcast(*StateTree, ID);
+				}
+			}
 		}
 
 		if (MemberProperty)
@@ -137,6 +180,16 @@ void UStateTreeState::PostLoad()
 	}
 }
 
+void UStateTreeState::UpdateParametersFromLinkedState()
+{
+	if (const UStateTreeEditorData* TreeData = GetTypedOuter<UStateTreeEditorData>())
+	{
+		if (const UStateTreeState* LinkTargetState = TreeData->GetStateByID(LinkedState.ID))
+		{
+			Parameters.Parameters.MigrateToNewBagInstance(LinkTargetState->Parameters.Parameters);
+		}
+	}
+}
 
 #endif
 
diff --git a/Engine/Plugins/Runtime/StateTree/Source/StateTreeEditorModule/Public/StateTreeCompiler.h b/Engine/Plugins/Runtime/StateTree/Source/StateTreeEditorModule/Public/StateTreeCompiler.h
--- a/Engine/Plugins/Runtime/StateTree/Source/StateTreeEditorModule/Public/StateTreeCompiler.h
+++ b/Engine/Plugins/Runtime/StateTree/Source/StateTreeEditorModule/Public/StateTreeCompiler.h
@@ -30,15 +30,11 @@ private:
 	bool ResolveTransitionState(const UStateTreeState& SourceState, const FStateTreeStateLink& Link, FStateTreeHandle& OutTransitionHandle) const;
 	FStateTreeHandle GetStateHandle(const FGuid& StateID) const;
 	UStateTreeState* GetState(const FGuid& StateID);
-	static FString GetExecutionPathString(const TConstArrayView<const UStateTreeState*> Path);
-	static bool IsPathLinked(const TConstArrayView<const UStateTreeState*> Path);
 
-	bool CreateExecutionInfos();
-	bool CreateExecutionInfosRecursive(UStateTreeState& State, TArray<const UStateTreeState*>& Path);
 	bool CreateStates();
 	bool CreateStateRecursive(UStateTreeState& State, const FStateTreeHandle Parent);
 	
-	bool CreateStateTasks();
+	bool CreateStateTasksAndParameters();
 	bool CreateStateEvaluators();
 	bool CreateStateTransitions();
 	
@@ -49,24 +45,10 @@ private:
 	bool GetAndValidateBindings(UStateTreeState& State, const FStateTreeBindableStructDesc& TargetStruct, TArray<FStateTreeEditorPropertyBinding>& OutBindings) const;
 	bool IsPropertyAnyEnum(const FStateTreeBindableStructDesc& Struct, FStateTreeEditorPropertyPath Path) const;
 
-	struct FExecutionPath
-	{
-		FExecutionPath() = default;
-		explicit FExecutionPath(const TConstArrayView<const UStateTreeState*> InPath) : Path(InPath) {}
-		
-		TArray<const UStateTreeState*> Path;
-	};
-
-	struct FStateExecutionInfo
-	{
-		TArray<FExecutionPath> ExecutionPaths;
-	};
-
 	FStateTreeCompilerLog& Log;
 	UStateTree* StateTree = nullptr;
 	UStateTreeEditorData* TreeData = nullptr;
 	TMap<FGuid, int32> IDToState;
 	TArray<UStateTreeState*> SourceStates;
-	TMap<UStateTreeState*, FStateExecutionInfo> ExecutionInfos;
 	FStateTreePropertyBindingCompiler BindingsCompiler;
 };
diff --git a/Engine/Plugins/Runtime/StateTree/Source/StateTreeEditorModule/Public/StateTreeEditorData.h b/Engine/Plugins/Runtime/StateTree/Source/StateTreeEditorModule/Public/StateTreeEditorData.h
--- a/Engine/Plugins/Runtime/StateTree/Source/StateTreeEditorModule/Public/StateTreeEditorData.h
+++ b/Engine/Plugins/Runtime/StateTree/Source/StateTreeEditorModule/Public/StateTreeEditorData.h
@@ -15,8 +15,17 @@ enum class EStateTreeNodeType : uint8
 	Evaluator,
 	Task,
 	TransitionCondition,
+	StateParameters,
 };
 
+UENUM()
+enum class EStateTreeVisitor : uint8
+{
+	Continue,
+	Break,
+};
+
+
 /**
  * Edit time data for StateTree asset. This data gets baked into runtime format before being used by the StateTreeInstance.
  */
@@ -41,20 +50,29 @@ public:
 	/** Returns parent state of a struct, or nullptr if not found. */
 	const UStateTreeState* GetStateByStructID(const FGuid TargetStructID) const;
 
+	/** Returns state based on its ID, or nullptr if not found. */
+	const UStateTreeState* GetStateByID(const FGuid StateID) const;
+
 	/** Gets the IDs of all bindable structs in the StateTree. */
 	void GetAllStructIDs(TMap<FGuid, const UStruct*>& AllStructs) const;
 
+	/**
+	* Iterates over all structs that are related to binding
+	* @param InFunc function called at each node, should return true if visiting is continued or false to stop.
+	*/
+	void VisitHierarchy(TFunctionRef<EStateTreeVisitor(UStateTreeState& State, UStateTreeState* ParentState)> InFunc) const;
+
 	/**
 	 * Iterates over all structs that are related to binding
 	 * @param InFunc function called at each node, should return true if visiting is continued or false to stop.
 	 */
-	void VisitHierarchy(TFunctionRef<bool(const UStateTreeState& State, const FGuid& ID, const FName& Name, const EStateTreeNodeType NodeType, const UScriptStruct* NodeStruct, const UStruct* InstanceStruct)> InFunc) const;
+	void VisitHierarchyNodes(TFunctionRef<EStateTreeVisitor(const UStateTreeState& State, const FGuid& ID, const FName& Name, const EStateTreeNodeType NodeType, const UScriptStruct* NodeStruct, const UStruct* InstanceStruct)> InFunc) const;
 
 	/**
 	 * Iterates over all nodes in a given state.
 	 * @param InFunc function called at each node, should return true if visiting is continued or false to stop.
 	 */
-	bool VisitState(const UStateTreeState& State, TFunctionRef<bool(const UStateTreeState& State, const FGuid& ID, const FName& Name, const EStateTreeNodeType NodeType, const UScriptStruct* NodeStruct, const UStruct* InstanceStruct)> InFunc) const;
+	EStateTreeVisitor VisitStateNodes(const UStateTreeState& State, TFunctionRef<EStateTreeVisitor(const UStateTreeState& State, const FGuid& ID, const FName& Name, const EStateTreeNodeType NodeType, const UScriptStruct* NodeStruct, const UStruct* InstanceStruct)> InFunc) const;
 
 	/**
 	 * Returns array of nodes along the execution path, up to the TargetStruct.
diff --git a/Engine/Plugins/Runtime/StateTree/Source/StateTreeEditorModule/Public/StateTreeState.h b/Engine/Plugins/Runtime/StateTree/Source/StateTreeEditorModule/Public/StateTreeState.h
--- a/Engine/Plugins/Runtime/StateTree/Source/StateTreeEditorModule/Public/StateTreeState.h
+++ b/Engine/Plugins/Runtime/StateTree/Source/StateTreeEditorModule/Public/StateTreeState.h
@@ -6,6 +6,7 @@
 #include "StateTreeNodeBase.h"
 #include "Misc/Guid.h"
 #include "InstancedStruct.h"
+#include "PropertyBag.h"
 #include "StateTreeState.generated.h"
 
 class UStateTreeState;
@@ -128,14 +129,31 @@ struct STATETREEEDITORMODULE_API FStateTreeTransition
 	TArray<FStateTreeEditorNode> Conditions;
 };
 
-UENUM()
-enum class EStateTreeStateType : uint8
+USTRUCT()
+struct STATETREEEDITORMODULE_API FStateTreeStateParameters
 {
-	State,	// A State containing tasks and evaluators.
-	Group,	// A State containing just sub states. 
-	Linked,	// A State that is linked to another state in the tree (the execution continues on the linked state).
-};
+	GENERATED_BODY()
+
+	FStateTreeStateParameters()
+	{
+		ID = FGuid::NewGuid();
+	}
+	
+	void Reset()
+	{
+		Parameters.Reset();
+		bFixedLayout = false;
+	}
+	
+	UPROPERTY(EditDefaultsOnly, Category = Parameters)
+	FInstancedPropertyBag Parameters;
+
+	UPROPERTY(EditDefaultsOnly, Category = Parameters)
+	bool bFixedLayout = false;
 
+	UPROPERTY(EditDefaultsOnly, Category = Parameters)
+	FGuid ID;
+};
 
 /**
  * Editor representation of a state in StateTree
@@ -151,10 +169,11 @@ public:
 #if WITH_EDITOR
 	virtual void PostEditChangeChainProperty(FPropertyChangedChainEvent& PropertyChangedEvent) override;
 	virtual void PostLoad() override;;
+	void UpdateParametersFromLinkedState();
 #endif
 
 	UStateTreeState* GetNextSiblingState() const;
-
+	
 	// StateTree Builder API
 	
 	/** Adds child state with specified name. */
@@ -240,9 +259,12 @@ public:
 	UPROPERTY(EditDefaultsOnly, Category = "State")
 	EStateTreeStateType Type = EStateTreeStateType::State;
 
-	UPROPERTY(EditDefaultsOnly, Category = "State", Meta=(DirectStatesOnly))
+	UPROPERTY(EditDefaultsOnly, Category = "State", Meta=(DirectStatesOnly, SubtreesOnly))
 	FStateTreeStateLink LinkedState;
 
+	UPROPERTY(EditDefaultsOnly, Category = "State")
+	FStateTreeStateParameters Parameters;
+
 	UPROPERTY()
 	FGuid ID;
 
diff --git a/Engine/Plugins/Runtime/StateTree/Source/StateTreeModule/Private/StateTree.cpp b/Engine/Plugins/Runtime/StateTree/Source/StateTreeModule/Private/StateTree.cpp
--- a/Engine/Plugins/Runtime/StateTree/Source/StateTreeModule/Private/StateTree.cpp
+++ b/Engine/Plugins/Runtime/StateTree/Source/StateTreeModule/Private/StateTree.cpp
@@ -75,7 +75,53 @@ bool UStateTree::Link()
 	ExternalDataDescs.Reset();
 	NumDataViews = 0;
 	ExternalDataBaseIndex = 0;
+	
+	// Update property bag structs before resolving binding.
+	TArrayView<FStateTreeBindableStructDesc> SourceStructs = PropertyBindings.GetSourceStructs();
+	TArrayView<FStateTreePropCopyBatch> CopyBatches = PropertyBindings.GetCopyBatches();
+	for (const FCompactStateTreeState& State : States)
+	{
+		if (State.Type == EStateTreeStateType::Subtree)
+		{
+			if (State.ParameterInstanceIndex != INDEX_NONE
+				|| State.ParameterDataViewIndex != INDEX_NONE)
+			{
+				UE_LOG(LogStateTree, Error, TEXT("%s: Data for state '%s' is malformed. Please recompile the StateTree asset."), *GetName(), *State.Name.ToString());
+				return false;
+			}
+
+			// Subtree is a bind source, update bag struct.
+			const FCompactStateTreeParameters& Params = Instances[State.ParameterInstanceIndex].GetMutable<FCompactStateTreeParameters>();
+			FStateTreeBindableStructDesc& Desc = SourceStructs[State.ParameterDataViewIndex];
+			Desc.Struct = Params.Parameters.GetPropertyBagStruct();
+		}
+		else if (State.Type == EStateTreeStateType::Linked && State.LinkedState.IsValid())
+		{
+			const FCompactStateTreeState& LinkedState = States[State.LinkedState.Index];
+
+			if (State.ParameterInstanceIndex != INDEX_NONE
+				|| LinkedState.ParameterInstanceIndex != INDEX_NONE)
+			{
+				UE_LOG(LogStateTree, Error, TEXT("%s: Data for state '%s' is malformed. Please recompile the StateTree asset."), *GetName(), *State.Name.ToString());
+				return false;
+			}
+
+			const FCompactStateTreeParameters& Params = Instances[State.ParameterInstanceIndex].GetMutable<FCompactStateTreeParameters>();
 
+			// Check that the bag in linked state matches.
+			const FCompactStateTreeParameters& LinkedStateParams = Instances[LinkedState.ParameterInstanceIndex].GetMutable<FCompactStateTreeParameters>();
+
+			if (LinkedStateParams.Parameters.GetPropertyBagStruct() != Params.Parameters.GetPropertyBagStruct())
+			{
+				UE_LOG(LogStateTree, Error, TEXT("%s: The parameters on state '%s' does not match the linked state parameters in state '%s'. Please recompile the StateTree asset."), *GetName(), *State.Name.ToString(), *LinkedState.Name.ToString());
+				return false;
+			}
+
+			FStateTreePropCopyBatch& Batch = CopyBatches[Params.BindingsBatch.Index];
+			Batch.TargetStruct.Struct = Params.Parameters.GetPropertyBagStruct();
+		}
+	}
+	
 	// Resolves property paths used by bindings a store property pointers
 	if (!PropertyBindings.ResolvePaths())
 	{
diff --git a/Engine/Plugins/Runtime/StateTree/Source/StateTreeModule/Private/StateTreeDelegates.cpp b/Engine/Plugins/Runtime/StateTree/Source/StateTreeModule/Private/StateTreeDelegates.cpp
--- a/Engine/Plugins/Runtime/StateTree/Source/StateTreeModule/Private/StateTreeDelegates.cpp
+++ b/Engine/Plugins/Runtime/StateTree/Source/StateTreeModule/Private/StateTreeDelegates.cpp
@@ -9,6 +9,7 @@ namespace UE::StateTree::Delegates
 FOnIdentifierChanged OnIdentifierChanged;
 FOnSchemaChanged OnSchemaChanged;
 FOnParametersChanged OnParametersChanged;
+FOnStateParametersChanged OnStateParametersChanged;
 FOnPostCompile OnPostCompile;
 #endif
 	
diff --git a/Engine/Plugins/Runtime/StateTree/Source/StateTreeModule/Private/StateTreeExecutionContext.cpp b/Engine/Plugins/Runtime/StateTree/Source/StateTreeModule/Private/StateTreeExecutionContext.cpp
--- a/Engine/Plugins/Runtime/StateTree/Source/StateTreeModule/Private/StateTreeExecutionContext.cpp
+++ b/Engine/Plugins/Runtime/StateTree/Source/StateTreeModule/Private/StateTreeExecutionContext.cpp
@@ -120,6 +120,44 @@ void FStateTreeExecutionContext::Reset()
 	Owner = nullptr;
 }
 
+void FStateTreeExecutionContext::UpdateLinkedStateParameters(FStateTreeInstanceData& InstanceData, const FCompactStateTreeState& State)
+{
+	check(State.ParameterDataViewIndex != INDEX_NONE);
+	check(State.ParameterInstanceIndex != INDEX_NONE);
+	
+	const FStateTreeDataView StateParamsInstance = GetInstanceData(InstanceData, false, State.ParameterInstanceIndex);
+	const FCompactStateTreeParameters& StateParams = StateParamsInstance.GetMutable<FCompactStateTreeParameters>();
+
+	// Parameters property bag
+	const FStateTreeDataView ParametersView(StateParams.Parameters.GetMutableValue());
+	if (StateParams.BindingsBatch.IsValid())
+	{
+		StateTree->PropertyBindings.CopyTo(DataViews, StateParams.BindingsBatch.Index, ParametersView);
+	}
+
+	// Set the parameters as the input parameters for the linked state.
+	check(State.LinkedState.IsValid());
+	const FCompactStateTreeState& LinkedState = StateTree->States[State.LinkedState.Index];
+	DataViews[LinkedState.ParameterDataViewIndex] = ParametersView;
+}
+
+void FStateTreeExecutionContext::UpdateSubtreeStateParameters(FStateTreeInstanceData& InstanceData, const FCompactStateTreeState& State)
+{
+	check(State.ParameterDataViewIndex != INDEX_NONE);
+	check(State.ParameterInstanceIndex != INDEX_NONE);
+	
+	// Usually the subtree parameter view is set by the linked state. If it's not (i.e. transitioned into a parametrized subtree), we'll set the view default params.
+	if (DataViews[State.ParameterDataViewIndex].IsValid())
+	{
+		return;
+	}
+
+	// Set view to default parameters.
+	const FStateTreeDataView ParamInstanceView = GetInstanceData(InstanceData, false, State.ParameterInstanceIndex);
+	const FCompactStateTreeParameters& Params = ParamInstanceView.GetMutable<FCompactStateTreeParameters>();
+	DataViews[State.ParameterDataViewIndex] = FStateTreeDataView(Params.Parameters.GetMutableValue());
+}
+
 EStateTreeRunStatus FStateTreeExecutionContext::Start(FStateTreeInstanceData* ExternalInstanceData)
 {
 	CSV_SCOPED_TIMING_STAT_EXCLUSIVE(StateTree_Start);
@@ -362,7 +400,16 @@ EStateTreeRunStatus FStateTreeExecutionContext::EnterState(FStateTreeInstanceDat
 				ANSI_TO_TCHAR(__FUNCTION__), *GetStateStatusString(Exec), *GetNameSafe(Owner), *GetNameSafe(StateTree));
 			break;
 		}
-		
+
+		if (State.Type == EStateTreeStateType::Linked)
+		{
+			UpdateLinkedStateParameters(InstanceData, State);
+		}
+		else if (State.Type == EStateTreeStateType::Subtree)
+		{
+			UpdateSubtreeStateParameters(InstanceData, State);
+		}
+
 		bOnTargetBranch = bOnTargetBranch || CurrentHandle == Transition.TargetState;
 		const bool bWasActive = PreviousHandle == CurrentHandle;
 		if (bWasActive && !bOnTargetBranch)
@@ -470,6 +517,15 @@ void FStateTreeExecutionContext::ExitState(FStateTreeInstanceData& InstanceData,
 		const FStateTreeHandle NextHandle = Transition.NextActiveStates.GetStateSafe(Index);
 		const FCompactStateTreeState& State = StateTree->States[CurrentHandle.Index];
 
+		if (State.Type == EStateTreeStateType::Linked)
+		{
+			UpdateLinkedStateParameters(InstanceData, State);
+		}
+		else if (State.Type == EStateTreeStateType::Subtree)
+		{
+			UpdateSubtreeStateParameters(InstanceData, State);
+		}
+
 		const bool bRemainsActive = NextHandle == CurrentHandle;
 		bOnTargetBranch = bOnTargetBranch || NextHandle == Transition.TargetState;
 		if (bRemainsActive && !bOnTargetBranch)
@@ -624,6 +680,15 @@ void FStateTreeExecutionContext::TickEvaluators(FStateTreeInstanceData& Instance
 
 		const FCompactStateTreeState& State = StateTree->States[StateHandle.Index];
 
+		if (State.Type == EStateTreeStateType::Linked)
+		{
+			UpdateLinkedStateParameters(InstanceData, State);
+		}
+		else if (State.Type == EStateTreeStateType::Subtree)
+		{
+			UpdateSubtreeStateParameters(InstanceData, State);
+		}
+		
 		STATETREE_CLOG(State.EvaluatorsNum > 0, Verbose, TEXT("%*sTicking Evaluators of state '%s' %s"), Index*UE::StateTree::DebugIndentSize, TEXT(""), *DebugGetStatePath(ActiveStates, Index), *UEnum::GetValueAsString(EvalType));
 
 		// Tick evaluators
@@ -668,6 +733,15 @@ EStateTreeRunStatus FStateTreeExecutionContext::TickTasks(FStateTreeInstanceData
 
 		STATETREE_CLOG(State.TasksNum > 0, Verbose, TEXT("%*sTicking Tasks of state '%s'"), Index*UE::StateTree::DebugIndentSize, TEXT(""), *DebugGetStatePath(Exec.ActiveStates, Index));
 
+		if (State.Type == EStateTreeStateType::Linked)
+		{
+			UpdateLinkedStateParameters(InstanceData, State);
+		}
+		else if (State.Type == EStateTreeStateType::Subtree)
+		{
+			UpdateSubtreeStateParameters(InstanceData, State);
+		}
+
 		// Tick Tasks
 		for (int32 TaskIndex = 0; TaskIndex < int32(State.TasksNum); TaskIndex++)
 		{
@@ -734,7 +808,12 @@ bool FStateTreeExecutionContext::TestAllConditions(FStateTreeInstanceData& Insta
 		// Copy bound properties.
 		if (Cond.BindingsBatch.IsValid())
 		{
-			StateTree->PropertyBindings.CopyTo(DataViews, Cond.BindingsBatch.Index, CondInstanceView);
+			if (!StateTree->PropertyBindings.CopyTo(DataViews, Cond.BindingsBatch.Index, CondInstanceView))
+			{
+				// If the source data cannot be accessed, the whole expression evaluates to false.
+				Values[0] = false;
+				break;
+			}
 		}
 
 		const bool bValue = Cond.TestCondition(*this);
diff --git a/Engine/Plugins/Runtime/StateTree/Source/StateTreeModule/Private/StateTreePropertyBindings.cpp b/Engine/Plugins/Runtime/StateTree/Source/StateTreeModule/Private/StateTreePropertyBindings.cpp
--- a/Engine/Plugins/Runtime/StateTree/Source/StateTreeModule/Private/StateTreePropertyBindings.cpp
+++ b/Engine/Plugins/Runtime/StateTree/Source/StateTreeModule/Private/StateTreePropertyBindings.cpp
@@ -45,8 +45,10 @@ bool FStateTreePropertyBindings::ResolvePaths()
 			const UStruct* TargetStruct = Batch.TargetStruct.Struct;
 			if (!SourceStruct || !TargetStruct)
 			{
+				Copy.Type = EStateTreePropertyCopyType::None;
 				bBindingsResolved = false;
-				break;
+				bResult = false;
+				continue;
 			}
 
 			// Resolve paths and validate the copy. Stops on first failure.
@@ -300,9 +302,11 @@ bool FStateTreePropertyBindings::ValidateCopy(FStateTreePropCopy& Copy) const
 
 uint8* FStateTreePropertyBindings::GetAddress(FStateTreeDataView InStructView, const FStateTreePropertyAccess& InAccess) const
 {
-	check(InStructView.IsValid());
-
 	uint8* Address = InStructView.GetMutableMemory();
+	if (Address == nullptr)
+	{
+		return nullptr;
+	}
 	
 	for (int32 i = InAccess.IndirectionsBegin; Address != nullptr && i < InAccess.IndirectionsEnd; i++)
 	{
@@ -469,17 +473,17 @@ void FStateTreePropertyBindings::PerformCopy(const FStateTreePropCopy& Copy, con
 	}
 }
 
-void FStateTreePropertyBindings::CopyTo(TConstArrayView<FStateTreeDataView> SourceStructViews, const int32 TargetBatchIndex, FStateTreeDataView TargetStructView) const
+bool FStateTreePropertyBindings::CopyTo(TConstArrayView<FStateTreeDataView> SourceStructViews, const int32 TargetBatchIndex, FStateTreeDataView TargetStructView) const
 {
 	// This is made ensure so that the programmers have the change to catch it (it's usually programming error not to call ResolvePaths(), and it wont spam log for others.
 	if (!ensureMsgf(bBindingsResolved, TEXT("Bindings must be resolved successfully before copying. See ResolvePaths()")))
 	{
-		return;
+		return false;
 	}
 
 	if (TargetBatchIndex == INDEX_NONE)
 	{
-		return;
+		return false;
 	}
 
 	check(CopyBatches.IsValidIndex(TargetBatchIndex));
@@ -488,6 +492,8 @@ void FStateTreePropertyBindings::CopyTo(TConstArrayView<FStateTreeDataView> Sour
 	check(TargetStructView.IsValid());
 	check(TargetStructView.GetStruct() == Batch.TargetStruct.Struct);
 
+	bool bResult = true;
+	
 	for (int32 i = Batch.BindingsBegin; i != Batch.BindingsEnd; i++)
 	{
 		const FStateTreePropCopy& Copy = PropertyCopies[i];
@@ -496,19 +502,28 @@ void FStateTreePropertyBindings::CopyTo(TConstArrayView<FStateTreeDataView> Sour
 		{
 			continue;
 		}
+		
+		const FStateTreeDataView SourceStructView = SourceStructViews[Copy.SourceStructIndex];
+		check(SourceStructView.GetStruct() == SourceStructs[Copy.SourceStructIndex].Struct || (SourceStructView.GetStruct() && SourceStructView.GetStruct()->IsChildOf(SourceStructs[Copy.SourceStructIndex].Struct)));
+			
 		const FStateTreePropertyAccess& SourceAccess = PropertyAccesses[Copy.SourceAccessIndex];
 		const FStateTreePropertyAccess& TargetAccess = PropertyAccesses[Copy.TargetAccessIndex];
-
-		const FStateTreeDataView SourceStructView = SourceStructViews[Copy.SourceStructIndex];
-		check(SourceStructView.GetStruct() == SourceStructs[Copy.SourceStructIndex].Struct || SourceStructView.GetStruct()->IsChildOf(SourceStructs[Copy.SourceStructIndex].Struct));
 		
 		const uint8* SourceAddress = GetAddress(SourceStructView, SourceAccess);
 		uint8* TargetAddress = GetAddress(TargetStructView, TargetAccess);
-		if (SourceAddress && TargetAddress)
+		
+		if (SourceAddress != nullptr && TargetAddress != nullptr)
 		{
 			PerformCopy(Copy, SourceAccess.LeafProperty, SourceAddress, TargetAccess.LeafProperty, TargetAddress);
 		}
+		else
+		{
+			bResult = false;
+			break;
+		}
 	}
+
+	return bResult;
 }
 
 void FStateTreePropertyBindings::DebugPrintInternalLayout(FString& OutString) const
diff --git a/Engine/Plugins/Runtime/StateTree/Source/StateTreeModule/Public/StateTreeDelegates.h b/Engine/Plugins/Runtime/StateTree/Source/StateTreeModule/Public/StateTreeDelegates.h
--- a/Engine/Plugins/Runtime/StateTree/Source/StateTreeModule/Public/StateTreeDelegates.h
+++ b/Engine/Plugins/Runtime/StateTree/Source/StateTreeModule/Public/StateTreeDelegates.h
@@ -29,7 +29,14 @@ namespace UE::StateTree::Delegates
 	 */
 	DECLARE_MULTICAST_DELEGATE_OneParam(FOnParametersChanged, const UStateTree& /*StateTree*/);
 	extern STATETREEMODULE_API FOnParametersChanged OnParametersChanged;
-	
+
+	/**
+	 * Called when parameters of a StateTree State changed.
+	 * This should mainly used by the asset editor to maintain consistency in the UI for manipulations.
+	 */
+	DECLARE_MULTICAST_DELEGATE_TwoParams(FOnStateParametersChanged, const UStateTree& /*StateTree*/, const FGuid /*StateID*/);
+	extern STATETREEMODULE_API FOnStateParametersChanged OnStateParametersChanged;
+
 	/** Called when compilation succeeds */
 	DECLARE_MULTICAST_DELEGATE_OneParam(FOnPostCompile, const UStateTree& /*StateTree*/);
 	extern STATETREEMODULE_API FOnPostCompile OnPostCompile;
diff --git a/Engine/Plugins/Runtime/StateTree/Source/StateTreeModule/Public/StateTreeExecutionContext.h b/Engine/Plugins/Runtime/StateTree/Source/StateTreeModule/Public/StateTreeExecutionContext.h
--- a/Engine/Plugins/Runtime/StateTree/Source/StateTreeModule/Public/StateTreeExecutionContext.h
+++ b/Engine/Plugins/Runtime/StateTree/Source/StateTreeModule/Public/StateTreeExecutionContext.h
@@ -384,6 +384,12 @@ protected:
 		return InstanceData.Get<FStateTreeExecutionState>(0);
 	}
 
+	/** Sets up parameter data view for a linked state and copies bound properties. */
+	void UpdateLinkedStateParameters(FStateTreeInstanceData& InstanceData, const FCompactStateTreeState& State);
+
+	/** Sets up parameter data view for subtree state. */
+	void UpdateSubtreeStateParameters(FStateTreeInstanceData& InstanceData, const FCompactStateTreeState& State);
+
 	/** @return StateTree node at specified index. */
 	template <typename T>
 	const T& GetNode(const int32 Index) const
diff --git a/Engine/Plugins/Runtime/StateTree/Source/StateTreeModule/Public/StateTreePropertyBindings.h b/Engine/Plugins/Runtime/StateTree/Source/StateTreeModule/Public/StateTreePropertyBindings.h
--- a/Engine/Plugins/Runtime/StateTree/Source/StateTreeModule/Public/StateTreePropertyBindings.h
+++ b/Engine/Plugins/Runtime/StateTree/Source/StateTreeModule/Public/StateTreePropertyBindings.h
@@ -149,9 +149,10 @@ enum class EStateTreeBindableStructSource : uint8
 {
 	TreeData,
 	TreeParameter,
-	StateParameter, 
+	StateParameter,
 	Task,
 	Evaluator,
+	Condition,
 };
 
 /**
@@ -166,7 +167,10 @@ struct STATETREEMODULE_API FStateTreeBindableStructDesc
 	FStateTreeBindableStructDesc() = default;
 
 #if WITH_EDITORONLY_DATA
-	FStateTreeBindableStructDesc(const FName InName, const UStruct* InStruct, const FGuid InGuid) : Struct(InStruct), Name(InName), ID(InGuid) {}
+	FStateTreeBindableStructDesc(const FName InName, const UStruct* InStruct, const EStateTreeBindableStructSource InDataSource, const FGuid InGuid)
+		: Struct(InStruct), Name(InName), DataSource(InDataSource), ID(InGuid)
+	{
+	}
 
 	bool operator==(const FStateTreeBindableStructDesc& RHS) const
 	{
@@ -445,13 +449,19 @@ struct STATETREEMODULE_API FStateTreePropertyBindings
 	 */
 	int32 GetSourceStructNum() const { return SourceStructs.Num(); }
 
+	TArrayView<FStateTreeBindableStructDesc> GetSourceStructs() { return  SourceStructs; };
+
+	TArrayView<FStateTreePropCopyBatch> GetCopyBatches() { return CopyBatches; };
+
+	
 	/**
 	 * Copies a batch of properties from source structs to target struct.
 	 * @param SourceStructViews Views to structs where properties are copied from.
 	 * @param TargetBatchIndex Batch index to copy (see FStateTreePropertyBindingCompiler).
 	 * @param TargetStructView View to struct where properties are copied to.
+	 * @return true if all copies succeeded (a copy can fail e.g. if source or destination struct view is invalid).
 	 */
-	void CopyTo(TConstArrayView<FStateTreeDataView> SourceStructViews, const int32 TargetBatchIndex, FStateTreeDataView TargetStructView) const;
+	bool CopyTo(TConstArrayView<FStateTreeDataView> SourceStructViews, const int32 TargetBatchIndex, FStateTreeDataView TargetStructView) const;
 
 	void DebugPrintInternalLayout(FString& OutString) const;
 
diff --git a/Engine/Plugins/Runtime/StateTree/Source/StateTreeModule/Public/StateTreeTypes.h b/Engine/Plugins/Runtime/StateTree/Source/StateTreeModule/Public/StateTreeTypes.h
--- a/Engine/Plugins/Runtime/StateTree/Source/StateTreeModule/Public/StateTreeTypes.h
+++ b/Engine/Plugins/Runtime/StateTree/Source/StateTreeModule/Public/StateTreeTypes.h
@@ -3,6 +3,7 @@
 #pragma once
 
 #include "InstancedStruct.h"
+#include "PropertyBag.h"
 #include "StateTreeTypes.generated.h"
 
 STATETREEMODULE_API DECLARE_LOG_CATEGORY_EXTERN(LogStateTree, Warning, All);
@@ -58,6 +59,15 @@ enum class EStateTreeConditionOperand : uint8
 	Or,					/// Combine results with OR.
 };
 
+UENUM()
+enum class EStateTreeStateType : uint8
+{
+	State,		// A State containing tasks and evaluators.
+	Group,		// A State containing just sub states.
+	Linked,		// A State that is linked to another state in the tree (the execution continues on the linked state).
+	Subtree,	// A subtree that can be linked to.
+};
+
 namespace UE::StateTree
 {
 	constexpr int32 MaxConditionIndent = 4; 
@@ -363,6 +373,11 @@ struct STATETREEMODULE_API FCompactStateTreeState
 	UPROPERTY()
 	uint16 EvaluatorsBegin = 0;							// Index to first evaluator
 
+	UPROPERTY()
+	uint16 ParameterInstanceIndex = INDEX_NONE;			// Index to state instance data
+	UPROPERTY()
+	uint16 ParameterDataViewIndex = INDEX_NONE;			// Data view index of the input parameters
+
 	UPROPERTY()
 	uint8 EnterConditionsNum = 0;						// Number of enter conditions
 	UPROPERTY()
@@ -371,6 +386,20 @@ struct STATETREEMODULE_API FCompactStateTreeState
 	uint8 TasksNum = 0;									// Number of tasks
 	UPROPERTY()
 	uint8 EvaluatorsNum = 0;							// Number of evaluators
+	UPROPERTY()
+	EStateTreeStateType Type = EStateTreeStateType::State; 
+};
+
+USTRUCT()
+struct STATETREEMODULE_API FCompactStateTreeParameters
+{
+	GENERATED_BODY()
+
+	UPROPERTY()
+	FStateTreeHandle BindingsBatch = FStateTreeHandle::Invalid;
+
+	UPROPERTY()
+	FInstancedPropertyBag Parameters;
 };
 
 /** An offset into the StateTree runtime storage type to get a struct view to a specific Task, Evaluator, or Condition. */
