StateTree: Fix warning about FStateTreeStateParameters::ID is not initialized properly

diff --git a/Engine/Plugins/Runtime/StateTree/Source/StateTreeEditorModule/Private/StateTreeState.cpp b/Engine/Plugins/Runtime/StateTree/Source/StateTreeEditorModule/Private/StateTreeState.cpp
--- a/Engine/Plugins/Runtime/StateTree/Source/StateTreeEditorModule/Private/StateTreeState.cpp
+++ b/Engine/Plugins/Runtime/StateTree/Source/StateTreeEditorModule/Private/StateTreeState.cpp
@@ -42,6 +42,7 @@ UStateTreeState::UStateTreeState(const FObjectInitializer& ObjectInitializer)
 	: Super(ObjectInitializer)
 	, ID(FGuid::NewGuid())
 {
+	Parameters.ID = FGuid::NewGuid();
 }
 
 #if WITH_EDITOR
diff --git a/Engine/Plugins/Runtime/StateTree/Source/StateTreeEditorModule/Public/StateTreeState.h b/Engine/Plugins/Runtime/StateTree/Source/StateTreeEditorModule/Public/StateTreeState.h
--- a/Engine/Plugins/Runtime/StateTree/Source/StateTreeEditorModule/Public/StateTreeState.h
+++ b/Engine/Plugins/Runtime/StateTree/Source/StateTreeEditorModule/Public/StateTreeState.h
@@ -134,11 +134,6 @@ struct STATETREEEDITORMODULE_API FStateTreeStateParameters
 {
 	GENERATED_BODY()
 
-	FStateTreeStateParameters()
-	{
-		ID = FGuid::NewGuid();
-	}
-	
 	void Reset()
 	{
 		Parameters.Reset();
