[StateTree ] added Delay and DebugText tasks
#rnx

diff --git a/Engine/Plugins/Runtime/StateTree/Source/StateTreeModule/Private/Tasks/StateTreeDebugTextTask.cpp b/Engine/Plugins/Runtime/StateTree/Source/StateTreeModule/Private/Tasks/StateTreeDebugTextTask.cpp
new file mode 100644
--- /dev/null
+++ b/Engine/Plugins/Runtime/StateTree/Source/StateTreeModule/Private/Tasks/StateTreeDebugTextTask.cpp
@@ -0,0 +1,74 @@
+﻿// Copyright Epic Games, Inc. All Rights Reserved.
+
+#include "StateTreeDebugTextTask.h"
+#include "StateTreeExecutionContext.h"
+#include "StateTreeLinker.h"
+#include "GameFramework/Actor.h"
+#include "DrawDebugHelpers.h"
+
+bool FStateTreeDebugTextTask::Link(FStateTreeLinker& Linker)
+{
+	Linker.LinkExternalData(ReferenceActorHandle);
+	return true;
+}
+
+EStateTreeRunStatus FStateTreeDebugTextTask::EnterState(FStateTreeExecutionContext& Context, const EStateTreeStateChangeType ChangeType, const FStateTreeTransitionResult& Transition) const
+{
+	if (!bEnabled)
+	{
+		return EStateTreeRunStatus::Succeeded;
+	}
+	
+	const UWorld* World = Context.GetWorld();
+
+	AActor* ReferenceActor = Context.GetExternalDataPtr(ReferenceActorHandle);
+	if (World == nullptr && ReferenceActor != nullptr)
+	{
+		World = ReferenceActor->GetWorld();
+	}
+	else if (ReferenceActor == nullptr)
+	{
+		ReferenceActor = Cast<AActor>(Context.GetOwner());
+	}
+
+	// Reference actor is not required (offset will be used as a global world location)
+	// but a valid world is required.
+	if (World == nullptr)
+	{
+		return EStateTreeRunStatus::Failed;
+	}
+
+	DrawDebugString(World, Offset, Text, ReferenceActor, TextColor,	/*Duration*/-1,	/*DrawShadows*/true, FontScale);
+	
+	return EStateTreeRunStatus::Running;
+}
+
+void FStateTreeDebugTextTask::ExitState(FStateTreeExecutionContext& Context, const EStateTreeStateChangeType ChangeType, const FStateTreeTransitionResult& Transition) const
+{
+	if (!bEnabled)
+	{
+		return;
+	}
+
+	const UWorld* World = Context.GetWorld();
+
+	AActor* ReferenceActor = Context.GetExternalDataPtr(ReferenceActorHandle);
+	if (World == nullptr && ReferenceActor != nullptr)
+	{
+		World = ReferenceActor->GetWorld();
+	}	
+	else if (ReferenceActor == nullptr)
+	{
+		ReferenceActor = Cast<AActor>(Context.GetOwner());
+	}
+	
+	// Reference actor is not required (offset was used as a global world location)
+	// but a valid world is required.
+	if (World == nullptr)
+	{
+		return;
+	}
+
+	// Drawing an empty text will remove the HUD DebugText entries associated to the target actor
+	DrawDebugString(World, Offset, "",	ReferenceActor);
+}
diff --git a/Engine/Plugins/Runtime/StateTree/Source/StateTreeModule/Private/Tasks/StateTreeDebugTextTask.h b/Engine/Plugins/Runtime/StateTree/Source/StateTreeModule/Private/Tasks/StateTreeDebugTextTask.h
new file mode 100644
--- /dev/null
+++ b/Engine/Plugins/Runtime/StateTree/Source/StateTreeModule/Private/Tasks/StateTreeDebugTextTask.h
@@ -0,0 +1,48 @@
+﻿// Copyright Epic Games, Inc. All Rights Reserved.
+
+#pragma once
+
+#include "StateTreeTaskBase.h"
+#include "StateTreeDebugTextTask.generated.h"
+
+USTRUCT()
+struct STATETREEMODULE_API FStateTreeDebugTextTaskInstanceData
+{
+	GENERATED_BODY()
+};
+
+/**
+ * Draws debug text on the HUD associated to the player controller.
+ */
+USTRUCT(meta = (DisplayName = "Debug Text Task"))
+struct STATETREEMODULE_API FStateTreeDebugTextTask : public FStateTreeTaskCommonBase
+{
+	GENERATED_BODY()
+
+	typedef FStateTreeDebugTextTaskInstanceData InstanceDataType;
+	
+	FStateTreeDebugTextTask() = default;
+
+	virtual const UStruct* GetInstanceDataType() const override { return InstanceDataType::StaticStruct(); }
+
+	virtual bool Link(FStateTreeLinker& Linker) override;
+	virtual EStateTreeRunStatus EnterState(FStateTreeExecutionContext& Context, const EStateTreeStateChangeType ChangeType, const FStateTreeTransitionResult& Transition) const override;
+	virtual void ExitState(FStateTreeExecutionContext& Context, const EStateTreeStateChangeType ChangeType, const FStateTreeTransitionResult& Transition) const override;
+
+	TStateTreeExternalDataHandle<AActor, EStateTreeExternalDataRequirement::Optional> ReferenceActorHandle;
+	
+	UPROPERTY(EditAnywhere, Category = Parameter)
+	FString Text;
+
+	UPROPERTY(EditAnywhere, Category = Parameter)
+	FColor TextColor = FColor::White;
+
+	UPROPERTY(EditAnywhere, Category = Parameter, meta=(ClampMin = 0, UIMin = 0))
+	float FontScale = 1.0f;
+	
+	UPROPERTY(EditAnywhere, Category = Parameter)
+	FVector Offset = FVector::ZeroVector;
+
+	UPROPERTY(EditAnywhere, Category = Parameter)
+	bool bEnabled = true;
+};
diff --git a/Engine/Plugins/Runtime/StateTree/Source/StateTreeModule/Private/Tasks/StateTreeDelayTask.cpp b/Engine/Plugins/Runtime/StateTree/Source/StateTreeModule/Private/Tasks/StateTreeDelayTask.cpp
new file mode 100644
--- /dev/null
+++ b/Engine/Plugins/Runtime/StateTree/Source/StateTreeModule/Private/Tasks/StateTreeDelayTask.cpp
@@ -0,0 +1,28 @@
+﻿// Copyright Epic Games, Inc. All Rights Reserved.
+
+#include "StateTreeDelayTask.h"
+#include "StateTreeExecutionContext.h"
+#include "StateTreeLinker.h"
+
+bool FStateTreeDelayTask::Link(FStateTreeLinker& Linker)
+{
+	Linker.LinkInstanceDataProperty(DurationHandle, STATETREE_INSTANCEDATA_PROPERTY(InstanceDataType, Duration));
+	Linker.LinkInstanceDataProperty(TimeHandle, STATETREE_INSTANCEDATA_PROPERTY(InstanceDataType, Time));
+	return true;
+}
+
+EStateTreeRunStatus FStateTreeDelayTask::EnterState(FStateTreeExecutionContext& Context, const EStateTreeStateChangeType ChangeType, const FStateTreeTransitionResult& Transition) const
+{
+	float& Time = Context.GetInstanceData(TimeHandle);
+	Time = 0.f;
+	return EStateTreeRunStatus::Running;
+}
+
+EStateTreeRunStatus FStateTreeDelayTask::Tick(FStateTreeExecutionContext& Context, const float DeltaTime) const
+{
+	float& Time = Context.GetInstanceData(TimeHandle);
+	const float Duration = Context.GetInstanceData(DurationHandle);
+	Time += DeltaTime;
+	
+	return Duration <= 0.0f ? EStateTreeRunStatus::Running : (Time < Duration ? EStateTreeRunStatus::Running : EStateTreeRunStatus::Succeeded);
+}
diff --git a/Engine/Plugins/Runtime/StateTree/Source/StateTreeModule/Private/Tasks/StateTreeDelayTask.h b/Engine/Plugins/Runtime/StateTree/Source/StateTreeModule/Private/Tasks/StateTreeDelayTask.h
new file mode 100644
--- /dev/null
+++ b/Engine/Plugins/Runtime/StateTree/Source/StateTreeModule/Private/Tasks/StateTreeDelayTask.h
@@ -0,0 +1,42 @@
+﻿// Copyright Epic Games, Inc. All Rights Reserved.
+
+#pragma once
+
+#include "StateTreeTaskBase.h"
+#include "StateTreeDelayTask.generated.h"
+
+USTRUCT()
+struct STATETREEMODULE_API FStateTreeDelayTaskInstanceData
+{
+	GENERATED_BODY()
+	
+	/** Delay before the task ends. Default (0 or any negative) will run indefinitely so it requires a transition in the state tree to stop it. */
+	UPROPERTY(EditAnywhere, Category = Parameter)
+	float Duration = 0.f;
+
+	/** Accumulated time used to stop task if duration is set */
+	UPROPERTY()
+	float Time = 0.f;
+};
+
+/**
+ * Simple task to wait indefinitely or for a given time (in seconds) before succeeding.
+ */
+USTRUCT(meta = (DisplayName = "Delay Task"))
+struct STATETREEMODULE_API FStateTreeDelayTask : public FStateTreeTaskCommonBase
+{
+	GENERATED_BODY()
+
+	typedef FStateTreeDelayTaskInstanceData InstanceDataType;
+	
+	FStateTreeDelayTask() = default;
+
+	virtual const UStruct* GetInstanceDataType() const override { return InstanceDataType::StaticStruct(); }
+
+	virtual bool Link(FStateTreeLinker& Linker) override;
+	virtual EStateTreeRunStatus EnterState(FStateTreeExecutionContext& Context, const EStateTreeStateChangeType ChangeType, const FStateTreeTransitionResult& Transition) const override;
+	virtual EStateTreeRunStatus Tick(FStateTreeExecutionContext& Context, const float DeltaTime) const override;;
+
+	TStateTreeInstanceDataPropertyHandle<float> DurationHandle;
+	TStateTreeInstanceDataPropertyHandle<float> TimeHandle;
+};
