StateTree fixes
- Fixed unset index check in UStateTree::Link() and FStateTreeExecutionContext
- Changed FStateTreeExecutionContext to call ExitState() from leaf-to-root
- Fixed handling of empty sources in FStateTreePropertyBindings::CopyTo()

diff --git a/Engine/Plugins/Runtime/StateTree/Source/StateTreeModule/Private/StateTree.cpp b/Engine/Plugins/Runtime/StateTree/Source/StateTreeModule/Private/StateTree.cpp
--- a/Engine/Plugins/Runtime/StateTree/Source/StateTreeModule/Private/StateTree.cpp
+++ b/Engine/Plugins/Runtime/StateTree/Source/StateTreeModule/Private/StateTree.cpp
@@ -83,8 +83,8 @@ bool UStateTree::Link()
 	{
 		if (State.Type == EStateTreeStateType::Subtree)
 		{
-			if (State.ParameterInstanceIndex != INDEX_NONE
-				|| State.ParameterDataViewIndex != INDEX_NONE)
+			if (State.ParameterInstanceIndex == MAX_uint16
+				|| State.ParameterDataViewIndex == MAX_uint16)
 			{
 				UE_LOG(LogStateTree, Error, TEXT("%s: Data for state '%s' is malformed. Please recompile the StateTree asset."), *GetName(), *State.Name.ToString());
 				return false;
@@ -99,8 +99,8 @@ bool UStateTree::Link()
 		{
 			const FCompactStateTreeState& LinkedState = States[State.LinkedState.Index];
 
-			if (State.ParameterInstanceIndex != INDEX_NONE
-				|| LinkedState.ParameterInstanceIndex != INDEX_NONE)
+			if (State.ParameterInstanceIndex == MAX_uint16
+				|| LinkedState.ParameterInstanceIndex == MAX_uint16)
 			{
 				UE_LOG(LogStateTree, Error, TEXT("%s: Data for state '%s' is malformed. Please recompile the StateTree asset."), *GetName(), *State.Name.ToString());
 				return false;
diff --git a/Engine/Plugins/Runtime/StateTree/Source/StateTreeModule/Private/StateTreeExecutionContext.cpp b/Engine/Plugins/Runtime/StateTree/Source/StateTreeModule/Private/StateTreeExecutionContext.cpp
--- a/Engine/Plugins/Runtime/StateTree/Source/StateTreeModule/Private/StateTreeExecutionContext.cpp
+++ b/Engine/Plugins/Runtime/StateTree/Source/StateTreeModule/Private/StateTreeExecutionContext.cpp
@@ -104,7 +104,7 @@ void FStateTreeExecutionContext::SetParameters(const FStateTreeParameters& Param
 		}
 
 		const uint16 DataViewIndex = TreeDescs[TreeDescIndex].DataViewIndex;
-		check(DataViewIndex != INDEX_NONE && DataViews.IsValidIndex(DataViewIndex));
+		check(DataViewIndex != FStateTreeParameterDesc::InvalidIndex && DataViews.IsValidIndex(DataViewIndex));
 		DataViews[DataViewIndex] = FStateTreeDataView(ExternalDesc.Parameter);
 		TreeDescIndex = (TreeDescIndex + 1) % NumTreeDescs;
 	}
@@ -122,8 +122,7 @@ void FStateTreeExecutionContext::Reset()
 
 void FStateTreeExecutionContext::UpdateLinkedStateParameters(FStateTreeInstanceData& InstanceData, const FCompactStateTreeState& State)
 {
-	check(State.ParameterDataViewIndex != INDEX_NONE);
-	check(State.ParameterInstanceIndex != INDEX_NONE);
+	check(State.ParameterInstanceIndex != MAX_uint16);
 	
 	const FStateTreeDataView StateParamsInstance = GetInstanceData(InstanceData, false, State.ParameterInstanceIndex);
 	const FCompactStateTreeParameters& StateParams = StateParamsInstance.GetMutable<FCompactStateTreeParameters>();
@@ -138,13 +137,14 @@ void FStateTreeExecutionContext::UpdateLinkedStateParameters(FStateTreeInstanceD
 	// Set the parameters as the input parameters for the linked state.
 	check(State.LinkedState.IsValid());
 	const FCompactStateTreeState& LinkedState = StateTree->States[State.LinkedState.Index];
+	check(LinkedState.ParameterDataViewIndex != MAX_uint16);
 	DataViews[LinkedState.ParameterDataViewIndex] = ParametersView;
 }
 
 void FStateTreeExecutionContext::UpdateSubtreeStateParameters(FStateTreeInstanceData& InstanceData, const FCompactStateTreeState& State)
 {
-	check(State.ParameterDataViewIndex != INDEX_NONE);
-	check(State.ParameterInstanceIndex != INDEX_NONE);
+	check(State.ParameterDataViewIndex != MAX_uint16);
+	check(State.ParameterInstanceIndex != MAX_uint16);
 	
 	// Usually the subtree parameter view is set by the linked state. If it's not (i.e. transitioned into a parametrized subtree), we'll set the view default params.
 	if (DataViews[State.ParameterDataViewIndex].IsValid())
@@ -388,6 +388,7 @@ EStateTreeRunStatus FStateTreeExecutionContext::EnterState(FStateTreeInstanceDat
 	Exec.EnterStateFailedTaskIndex = MAX_uint16; // This will make all tasks to be accepted.
 	Exec.ActiveStates.Reset();
 	
+	// Do property copy on all states, propagating the results from last tick.
 	for (int32 Index = 0; Index < Transition.NextActiveStates.Num() && Result != EStateTreeRunStatus::Failed; Index++)
 	{
 		const FStateTreeHandle CurrentHandle = Transition.NextActiveStates[Index];
@@ -412,28 +413,13 @@ EStateTreeRunStatus FStateTreeExecutionContext::EnterState(FStateTreeInstanceDat
 
 		bOnTargetBranch = bOnTargetBranch || CurrentHandle == Transition.TargetState;
 		const bool bWasActive = PreviousHandle == CurrentHandle;
-		if (bWasActive && !bOnTargetBranch)
-		{
-			// States which will keep on being active and were not part of the transition will not get enter/exit state.
-			// Must update data views.
-			for (int32 EvalIndex = 0; EvalIndex < int32(State.EvaluatorsNum); EvalIndex++)
-			{
-				const FStateTreeEvaluatorBase& Eval = GetNode<FStateTreeEvaluatorBase>(int32(State.EvaluatorsBegin) + EvalIndex);
-				DataViews[Eval.DataViewIndex] = GetInstanceData(InstanceData, Eval.bInstanceIsObject, Eval.InstanceIndex);
-			}
-			for (int32 TaskIndex = 0; TaskIndex < int32(State.TasksNum); TaskIndex++)
-			{
-				const FStateTreeTaskBase& Task = GetNode<FStateTreeTaskBase>(int32(State.TasksBegin) + TaskIndex);
-				DataViews[Task.DataViewIndex] = GetInstanceData(InstanceData, Task.bInstanceIsObject, Task.InstanceIndex);
-			}
-			continue;
-		}
+		const bool bIsEnteringState = !bWasActive || bOnTargetBranch;
 
 		CurrentTransition.CurrentState = CurrentHandle;
 		
 		const EStateTreeStateChangeType ChangeType = bWasActive ? EStateTreeStateChangeType::Sustained : EStateTreeStateChangeType::Changed;
 
-		STATETREE_LOG(Log, TEXT("%*sEnter state '%s' %s"), Index*UE::StateTree::DebugIndentSize, TEXT(""), *DebugGetStatePath(Transition.NextActiveStates, Index), *UEnum::GetValueAsString(ChangeType));
+		STATETREE_CLOG(bIsEnteringState, Log, TEXT("%*sEnter state '%s' %s"), Index*UE::StateTree::DebugIndentSize, TEXT(""), *DebugGetStatePath(Transition.NextActiveStates, Index), *UEnum::GetValueAsString(ChangeType));
 		
 		for (int32 EvalIndex = State.EvaluatorsBegin; EvalIndex < (int32)(State.EvaluatorsBegin + State.EvaluatorsNum); EvalIndex++)
 		{
@@ -446,8 +432,10 @@ EStateTreeRunStatus FStateTreeExecutionContext::EnterState(FStateTreeInstanceDat
 			{
 				StateTree->PropertyBindings.CopyTo(DataViews, Eval.BindingsBatch.Index, EvalInstanceView);
 			}
-			STATETREE_LOG(Verbose, TEXT("%*s  Notify Evaluator '%s'."), Index*UE::StateTree::DebugIndentSize, TEXT(""), *Eval.Name.ToString());
+
+			if (bIsEnteringState)
 			{
+				STATETREE_LOG(Verbose, TEXT("%*s  Notify Evaluator '%s'."), Index*UE::StateTree::DebugIndentSize, TEXT(""), *Eval.Name.ToString());
 				QUICK_SCOPE_CYCLE_COUNTER(StateTree_Eval_EnterState);
 				CSV_SCOPED_TIMING_STAT_EXCLUSIVE(StateTree_Eval_EnterState);
 				Eval.EnterState(*this, ChangeType, CurrentTransition);
@@ -466,10 +454,10 @@ EStateTreeRunStatus FStateTreeExecutionContext::EnterState(FStateTreeInstanceDat
 			{
 				StateTree->PropertyBindings.CopyTo(DataViews, Task.BindingsBatch.Index, TaskInstanceView);
 			}
-		
-			STATETREE_LOG(Verbose, TEXT("%*s  Notify Task '%s'"), Index*UE::StateTree::DebugIndentSize, TEXT(""), *Task.Name.ToString());
 
+			if (bIsEnteringState)
 			{
+				STATETREE_LOG(Verbose, TEXT("%*s  Notify Task '%s'"), Index*UE::StateTree::DebugIndentSize, TEXT(""), *Task.Name.ToString());
 				QUICK_SCOPE_CYCLE_COUNTER(StateTree_Task_EnterState);
 				CSV_SCOPED_TIMING_STAT_EXCLUSIVE(StateTree_Task_EnterState);
 				const EStateTreeRunStatus Status = Task.EnterState(*this, ChangeType, CurrentTransition);
@@ -507,10 +495,13 @@ void FStateTreeExecutionContext::ExitState(FStateTreeInstanceData& InstanceData,
 	// EnterState called. That is, a transition is handled as "replan from this state".
 	bool bOnTargetBranch = false;
 
-	FStateTreeTransitionResult CurrentTransition = Transition;
-
-	// It would be more symmetrical, if the evals, tasks and states were executed in reverse order, but we need
-	// to do it like this because of the property copies.
+	FStateTreeHandle ExitedStates[FStateTreeActiveStates::MaxStates];
+	EStateTreeStateChangeType ExitedStateChangeType[FStateTreeActiveStates::MaxStates];
+	int32 ExitedStateActiveIndex[FStateTreeActiveStates::MaxStates];
+	int32 NumExitedStates = 0;
+	
+	// Do property copy on all states, propagating the results from last tick.
+	// Collect the states that need to be called, the actual call is done below in reverse order.
 	for (int32 Index = 0; Index < Transition.CurrentActiveStates.Num(); Index++)
 	{
 		const FStateTreeHandle CurrentHandle = Transition.CurrentActiveStates[Index];
@@ -528,29 +519,19 @@ void FStateTreeExecutionContext::ExitState(FStateTreeInstanceData& InstanceData,
 
 		const bool bRemainsActive = NextHandle == CurrentHandle;
 		bOnTargetBranch = bOnTargetBranch || NextHandle == Transition.TargetState;
-		if (bRemainsActive && !bOnTargetBranch)
-		{
-			// States which will keep on being active and were not part of the transition will not get enter/exit state.
-			// Must update item views.
-			for (int32 EvalIndex = 0; EvalIndex < int32(State.EvaluatorsNum); EvalIndex++)
-			{
-				const FStateTreeEvaluatorBase& Eval = GetNode<FStateTreeEvaluatorBase>(int32(State.EvaluatorsBegin) + EvalIndex);
-				DataViews[Eval.DataViewIndex] = GetInstanceData(InstanceData, Eval.bInstanceIsObject, Eval.InstanceIndex);
-			}
-			for (int32 TaskIndex = 0; TaskIndex < int32(State.TasksNum); TaskIndex++)
-			{
-				const FStateTreeTaskBase& Task = GetNode<FStateTreeTaskBase>(int32(State.TasksBegin) + TaskIndex);
-				DataViews[Task.DataViewIndex] = GetInstanceData(InstanceData, Task.bInstanceIsObject, Task.InstanceIndex);
-			}
-			continue;
-		}
-
-		CurrentTransition.CurrentState = CurrentHandle;
-		
 		const EStateTreeStateChangeType ChangeType = bRemainsActive ? EStateTreeStateChangeType::Sustained : EStateTreeStateChangeType::Changed;
 
-		STATETREE_LOG(Log, TEXT("%*sExit state '%s' %s"), Index*UE::StateTree::DebugIndentSize, TEXT(""), *DebugGetStatePath(Transition.CurrentActiveStates, Index), *UEnum::GetValueAsString(ChangeType));
+		if (!bRemainsActive || bOnTargetBranch)
+		{
+			// Should call ExitState() on this state.
+			check (NumExitedStates < FStateTreeActiveStates::MaxStates);
+			ExitedStates[NumExitedStates] = CurrentHandle;
+			ExitedStateChangeType[NumExitedStates] = ChangeType;
+			ExitedStateActiveIndex[NumExitedStates] = Index;
+			NumExitedStates++;
+		}
 
+		// Do property copies, ExitState() is called below.
 		for (int32 EvalIndex = State.EvaluatorsBegin; EvalIndex < (int32)(State.EvaluatorsBegin + State.EvaluatorsNum); EvalIndex++)
 		{
 			const FStateTreeEvaluatorBase& Eval = GetNode<FStateTreeEvaluatorBase>(EvalIndex);
@@ -562,30 +543,43 @@ void FStateTreeExecutionContext::ExitState(FStateTreeInstanceData& InstanceData,
 			{
 				StateTree->PropertyBindings.CopyTo(DataViews, Eval.BindingsBatch.Index, EvalInstanceView);
 			}
-			STATETREE_LOG(Verbose, TEXT("%*s  Notify Evaluator '%s'."), Index*UE::StateTree::DebugIndentSize, TEXT(""), *Eval.Name.ToString());
+		}
+
+		for (int32 TaskIndex = State.TasksBegin; TaskIndex < (int32)(State.TasksBegin + State.TasksNum); TaskIndex++)
+		{
+			const FStateTreeTaskBase& Task = GetNode<FStateTreeTaskBase>(TaskIndex);
+			const FStateTreeDataView TaskInstanceView = GetInstanceData(InstanceData, Task.bInstanceIsObject, Task.InstanceIndex);
+			DataViews[Task.DataViewIndex] = TaskInstanceView;
+
+			// Copy bound properties.
+			if (Task.BindingsBatch.IsValid())
 			{
-				QUICK_SCOPE_CYCLE_COUNTER(StateTree_Eval_ExitState);
-				CSV_SCOPED_TIMING_STAT_EXCLUSIVE(StateTree_Eval_ExitState);
-				Eval.ExitState(*this, ChangeType, CurrentTransition);
+				StateTree->PropertyBindings.CopyTo(DataViews, Task.BindingsBatch.Index, TaskInstanceView);
 			}
 		}
+	}
 
-		// Deactivate tasks on current State
-		for (int32 TaskIndex = State.TasksBegin; TaskIndex < (int32)(State.TasksBegin + State.TasksNum); TaskIndex++)
+	// Call in reverse order.
+	FStateTreeTransitionResult CurrentTransition = Transition;
+
+	for (int32 Index = NumExitedStates - 1; Index >= 0; Index--)
+	{
+		const FStateTreeHandle CurrentHandle = ExitedStates[Index];
+		const FCompactStateTreeState& State = StateTree->States[CurrentHandle.Index];
+		const EStateTreeStateChangeType ChangeType = ExitedStateChangeType[Index];
+
+		STATETREE_LOG(Log, TEXT("%*sExit state '%s' %s"), Index*UE::StateTree::DebugIndentSize, TEXT(""), *DebugGetStatePath(Transition.CurrentActiveStates, ExitedStateActiveIndex[Index]), *UEnum::GetValueAsString(ChangeType));
+
+		CurrentTransition.CurrentState = CurrentHandle;
+
+		// Tasks
+		for (int32 TaskIndex = (int32)(State.TasksBegin + State.TasksNum) - 1; TaskIndex >= State.TasksBegin; TaskIndex--)
 		{
-			// Call task completed only if EnterState() was called. This ensures that we have matching calls to Enter/ExitState.
+			// Call task completed only if EnterState() was called.
 			// The task order in the tree (BF) allows us to use the comparison.
 			if (TaskIndex <= (int32)Exec.EnterStateFailedTaskIndex)
 			{
 				const FStateTreeTaskBase& Task = GetNode<FStateTreeTaskBase>(TaskIndex);
-				const FStateTreeDataView TaskInstanceView = GetInstanceData(InstanceData, Task.bInstanceIsObject, Task.InstanceIndex);
-				DataViews[Task.DataViewIndex] = TaskInstanceView;
-
-				// Copy bound properties.
-				if (Task.BindingsBatch.IsValid())
-				{
-					StateTree->PropertyBindings.CopyTo(DataViews, Task.BindingsBatch.Index, TaskInstanceView);
-				}
 
 				STATETREE_LOG(Verbose, TEXT("%*s  Notify Task '%s'"), Index*UE::StateTree::DebugIndentSize, TEXT(""), *Task.Name.ToString());
 				{
@@ -595,6 +589,19 @@ void FStateTreeExecutionContext::ExitState(FStateTreeInstanceData& InstanceData,
 				}
 			}
 		}
+
+		// Evaluators
+		for (int32 EvalIndex = (int32)(State.EvaluatorsBegin + State.EvaluatorsNum) - 1; EvalIndex >= State.EvaluatorsBegin; EvalIndex--)
+		{
+			const FStateTreeEvaluatorBase& Eval = GetNode<FStateTreeEvaluatorBase>(EvalIndex);
+
+			STATETREE_LOG(Verbose, TEXT("%*s  Notify Evaluator '%s'."), Index*UE::StateTree::DebugIndentSize, TEXT(""), *Eval.Name.ToString());
+			{
+				QUICK_SCOPE_CYCLE_COUNTER(StateTree_Eval_ExitState);
+				CSV_SCOPED_TIMING_STAT_EXCLUSIVE(StateTree_Eval_ExitState);
+				Eval.ExitState(*this, ChangeType, CurrentTransition);
+			}
+		}
 	}
 }
 
diff --git a/Engine/Plugins/Runtime/StateTree/Source/StateTreeModule/Private/StateTreePropertyBindings.cpp b/Engine/Plugins/Runtime/StateTree/Source/StateTreeModule/Private/StateTreePropertyBindings.cpp
--- a/Engine/Plugins/Runtime/StateTree/Source/StateTreeModule/Private/StateTreePropertyBindings.cpp
+++ b/Engine/Plugins/Runtime/StateTree/Source/StateTreeModule/Private/StateTreePropertyBindings.cpp
@@ -504,22 +504,22 @@ bool FStateTreePropertyBindings::CopyTo(TConstArrayView<FStateTreeDataView> Sour
 		}
 		
 		const FStateTreeDataView SourceStructView = SourceStructViews[Copy.SourceStructIndex];
-		check(SourceStructView.GetStruct() == SourceStructs[Copy.SourceStructIndex].Struct || (SourceStructView.GetStruct() && SourceStructView.GetStruct()->IsChildOf(SourceStructs[Copy.SourceStructIndex].Struct)));
-			
-		const FStateTreePropertyAccess& SourceAccess = PropertyAccesses[Copy.SourceAccessIndex];
-		const FStateTreePropertyAccess& TargetAccess = PropertyAccesses[Copy.TargetAccessIndex];
-		
-		const uint8* SourceAddress = GetAddress(SourceStructView, SourceAccess);
-		uint8* TargetAddress = GetAddress(TargetStructView, TargetAccess);
-		
-		if (SourceAddress != nullptr && TargetAddress != nullptr)
+		if (SourceStructView.IsValid())
 		{
+			check(SourceStructView.GetStruct() == SourceStructs[Copy.SourceStructIndex].Struct || (SourceStructView.GetStruct() && SourceStructView.GetStruct()->IsChildOf(SourceStructs[Copy.SourceStructIndex].Struct)));
+				
+			const FStateTreePropertyAccess& SourceAccess = PropertyAccesses[Copy.SourceAccessIndex];
+			const FStateTreePropertyAccess& TargetAccess = PropertyAccesses[Copy.TargetAccessIndex];
+			
+			const uint8* SourceAddress = GetAddress(SourceStructView, SourceAccess);
+			uint8* TargetAddress = GetAddress(TargetStructView, TargetAccess);
+			
+			check (SourceAddress != nullptr && TargetAddress != nullptr);
 			PerformCopy(Copy, SourceAccess.LeafProperty, SourceAddress, TargetAccess.LeafProperty, TargetAddress);
 		}
 		else
 		{
 			bResult = false;
-			break;
 		}
 	}
 
diff --git a/Engine/Plugins/Runtime/StateTree/Source/StateTreeModule/Public/StateTreeTypes.h b/Engine/Plugins/Runtime/StateTree/Source/StateTreeModule/Public/StateTreeTypes.h
--- a/Engine/Plugins/Runtime/StateTree/Source/StateTreeModule/Public/StateTreeTypes.h
+++ b/Engine/Plugins/Runtime/StateTree/Source/StateTreeModule/Public/StateTreeTypes.h
@@ -374,9 +374,9 @@ struct STATETREEMODULE_API FCompactStateTreeState
 	uint16 EvaluatorsBegin = 0;							// Index to first evaluator
 
 	UPROPERTY()
-	uint16 ParameterInstanceIndex = INDEX_NONE;			// Index to state instance data
+	uint16 ParameterInstanceIndex = MAX_uint16;			// Index to state instance data
 	UPROPERTY()
-	uint16 ParameterDataViewIndex = INDEX_NONE;			// Data view index of the input parameters
+	uint16 ParameterDataViewIndex = MAX_uint16;			// Data view index of the input parameters
 
 	UPROPERTY()
 	uint8 EnterConditionsNum = 0;						// Number of enter conditions
