[StateTree] replaced StateTreeParameters by InstancedPropertyBag
#rnx

diff --git a/Engine/Plugins/Runtime/StateTree/Source/StateTreeEditorModule/Private/StateTreeCompiler.cpp b/Engine/Plugins/Runtime/StateTree/Source/StateTreeEditorModule/Private/StateTreeCompiler.cpp
--- a/Engine/Plugins/Runtime/StateTree/Source/StateTreeEditorModule/Private/StateTreeCompiler.cpp
+++ b/Engine/Plugins/Runtime/StateTree/Source/StateTreeEditorModule/Private/StateTreeCompiler.cpp
@@ -34,13 +34,11 @@ bool FStateTreeCompiler::Compile(UStateTree& InStateTree)
 	StateTree->Schema = DuplicateObject(TreeData->Schema, StateTree);
 
 	// Copy parameters from EditorData	
-	StateTree->Parameters = TreeData->Parameters;
+	StateTree->Parameters = TreeData->RootParameters.Parameters;
 	
-	// Mark all parameters as binding source
-	for (FStateTreeParameterDesc& Desc : StateTree->Parameters.Parameters)
-	{
-		Desc.DataViewIndex = BindingsCompiler.AddSourceStruct({Desc.Name, Desc.Parameter.GetScriptStruct(), EStateTreeBindableStructSource::TreeParameter, Desc.ID});
-	}	
+	// Mark parameters as binding source
+	StateTree->DefaultParametersDataViewIndex = BindingsCompiler.AddSourceStruct(
+		{ TEXT("Parameters"), StateTree->Parameters.GetPropertyBagStruct(), EStateTreeBindableStructSource::TreeParameter, TreeData->RootParameters.ID });
 	
 	// Mark all named external values as binding source
 	if (StateTree->Schema)
diff --git a/Engine/Plugins/Runtime/StateTree/Source/StateTreeEditorModule/Private/StateTreeEditorData.cpp b/Engine/Plugins/Runtime/StateTree/Source/StateTreeEditorModule/Private/StateTreeEditorData.cpp
--- a/Engine/Plugins/Runtime/StateTree/Source/StateTreeEditorModule/Private/StateTreeEditorData.cpp
+++ b/Engine/Plugins/Runtime/StateTree/Source/StateTreeEditorModule/Private/StateTreeEditorData.cpp
@@ -9,6 +9,8 @@
 void UStateTreeEditorData::PostInitProperties()
 {
 	Super::PostInitProperties();
+
+	RootParameters.ID = FGuid::NewGuid();
 }
 
 #if WITH_EDITOR
@@ -26,32 +28,8 @@ void UStateTreeEditorData::PostEditChangeProperty(struct FPropertyChangedEvent&
 		{			
 			UE::StateTree::Delegates::OnSchemaChanged.Broadcast(*StateTree);
 		}
-		else if (MemberName == GET_MEMBER_NAME_CHECKED(UStateTreeEditorData, Parameters))
+		else if (MemberName == GET_MEMBER_NAME_CHECKED(UStateTreeEditorData, RootParameters))
 		{
-			TArray<FStateTreeParameterDesc>& Params = Parameters.Parameters;
-			const FName PropertyName = PropertyChangedEvent.Property->GetFName();
-			
-			if (PropertyChangedEvent.ChangeType == EPropertyChangeType::Duplicate)
-			{
-				// Ensure unique ID on duplicated items.
-				const int32 ArrayIndex = PropertyChangedEvent.GetArrayIndex(PropertyName.ToString());
-				if (Params.IsValidIndex(ArrayIndex))
-				{
-					Params[ArrayIndex].Name = FName(Params[ArrayIndex].Name.ToString() + TEXT(" Duplicate"));
-					Params[ArrayIndex].ID = FGuid::NewGuid();
-				}
-			}
-			else if (PropertyChangedEvent.ChangeType == EPropertyChangeType::ArrayAdd
-					|| PropertyChangedEvent.ChangeType == EPropertyChangeType::ValueSet)
-			{
-				// Create unique ID on added items or ensure on set ones (maybe a copy/paste) 
-				const int32 ArrayIndex = PropertyChangedEvent.GetArrayIndex(PropertyName.ToString());
-				if (Params.IsValidIndex(ArrayIndex))
-				{
-					Params[ArrayIndex].ID = FGuid::NewGuid();
-				}
-			}
-
 			UE::StateTree::Delegates::OnParametersChanged.Broadcast(*StateTree);
 		}
 	}
@@ -83,9 +61,9 @@ void UStateTreeEditorData::GetAccessibleStructs(const TConstArrayView<const USta
 	checkf(StateTree, TEXT("UStateTreeEditorData should only be allocated within a UStateTree"));
 	
 	// All parameters are accessible
-	for (const FStateTreeParameterDesc& Desc : Parameters.Parameters)
+	if (const UScriptStruct* PropertyBagStruct = RootParameters.Parameters.GetPropertyBagStruct())
 	{
-		OutStructDescs.Emplace(Desc.Name, Desc.Parameter.GetScriptStruct(), EStateTreeBindableStructSource::TreeParameter, Desc.ID);
+		OutStructDescs.Emplace(TEXT("Parameters"), PropertyBagStruct, EStateTreeBindableStructSource::TreeParameter, RootParameters.ID);
 	}
 
 	// All named external data items declared by the schema are accessible
@@ -232,10 +210,7 @@ void UStateTreeEditorData::GetAllStructIDs(TMap<FGuid, const UStruct*>& AllStruc
 	checkf(StateTree, TEXT("UStateTreeEditorData should only be allocated within a UStateTree"));
 
 	// All parameters
-	for (const FStateTreeParameterDesc& Desc : Parameters.Parameters)
-	{
-		AllStructs.Emplace(Desc.ID, Desc.Parameter.GetScriptStruct());
-	}
+	AllStructs.Emplace(RootParameters.ID, RootParameters.Parameters.GetPropertyBagStruct());
 
 	// All named external data items declared by the schema
 	if (Schema != nullptr)
diff --git a/Engine/Plugins/Runtime/StateTree/Source/StateTreeEditorModule/Public/StateTreeEditorData.h b/Engine/Plugins/Runtime/StateTree/Source/StateTreeEditorModule/Public/StateTreeEditorData.h
--- a/Engine/Plugins/Runtime/StateTree/Source/StateTreeEditorModule/Public/StateTreeEditorData.h
+++ b/Engine/Plugins/Runtime/StateTree/Source/StateTreeEditorModule/Public/StateTreeEditorData.h
@@ -118,10 +118,10 @@ public:
 	/** Schema describing which inputs, evaluators, and tasks a StateTree can contain */	
 	UPROPERTY(EditDefaultsOnly, Category = Common, Instanced)
 	TObjectPtr<UStateTreeSchema> Schema = nullptr;
-	
+
 	/** Public parameters that could be used for bindings within the Tree. */
-	UPROPERTY(EditDefaultsOnly, Category = Common, meta=(ShowOnlyInnerProperties))
-	FStateTreeParameters Parameters;
+	UPROPERTY(EditDefaultsOnly, Category = Common)
+	FStateTreeStateParameters RootParameters;
 	
 	UPROPERTY(meta = (ExcludeFromHash))
 	FStateTreeEditorPropertyBindings EditorBindings;
diff --git a/Engine/Plugins/Runtime/StateTree/Source/StateTreeModule/Private/StateTree.cpp b/Engine/Plugins/Runtime/StateTree/Source/StateTreeModule/Private/StateTree.cpp
--- a/Engine/Plugins/Runtime/StateTree/Source/StateTreeModule/Private/StateTree.cpp
+++ b/Engine/Plugins/Runtime/StateTree/Source/StateTreeModule/Private/StateTree.cpp
@@ -79,6 +79,12 @@ bool UStateTree::Link()
 	// Update property bag structs before resolving binding.
 	TArrayView<FStateTreeBindableStructDesc> SourceStructs = PropertyBindings.GetSourceStructs();
 	TArrayView<FStateTreePropCopyBatch> CopyBatches = PropertyBindings.GetCopyBatches();
+
+	if (DefaultParametersDataViewIndex != INDEX_NONE)
+	{
+		SourceStructs[DefaultParametersDataViewIndex].Struct = Parameters.GetPropertyBagStruct();
+	}
+	
 	for (const FCompactStateTreeState& State : States)
 	{
 		if (State.Type == EStateTreeStateType::Subtree)
diff --git a/Engine/Plugins/Runtime/StateTree/Source/StateTreeModule/Private/StateTreeExecutionContext.cpp b/Engine/Plugins/Runtime/StateTree/Source/StateTreeModule/Private/StateTreeExecutionContext.cpp
--- a/Engine/Plugins/Runtime/StateTree/Source/StateTreeModule/Private/StateTreeExecutionContext.cpp
+++ b/Engine/Plugins/Runtime/StateTree/Source/StateTreeModule/Private/StateTreeExecutionContext.cpp
@@ -61,52 +61,21 @@ bool FStateTreeExecutionContext::Init(UObject& InOwner, const UStateTree& InStat
 
 void FStateTreeExecutionContext::SetDefaultParameters()
 {
-	for (const FStateTreeParameterDesc& Desc : StateTree->GetParameterDescs())
+	if (ensureMsgf(StateTree != nullptr, TEXT("Execution context must be initialized before calling %s"), ANSI_TO_TCHAR(__FUNCTION__)))
 	{
-		check(DataViews.IsValidIndex(Desc.DataViewIndex));
-		DataViews[Desc.DataViewIndex] = FStateTreeDataView(Desc.Parameter);
+		check(DataViews.IsValidIndex(StateTree->DefaultParametersDataViewIndex));
+		DataViews[StateTree->DefaultParametersDataViewIndex] = FStateTreeDataView(StateTree->GetDefaultParameters().GetMutableValue());	
 	}
 }
 
-void FStateTreeExecutionContext::SetParameters(const FStateTreeParameters& Parameters)
+void FStateTreeExecutionContext::SetParameters(const FInstancedPropertyBag& Parameters)
 {
-	const TConstArrayView<FStateTreeParameterDesc> TreeDescs = StateTree->GetParameterDescs();
-	const int32 NumTreeDescs = TreeDescs.Num();
-	if (NumTreeDescs == 0)
+	if (ensureMsgf(StateTree != nullptr, TEXT("Execution context must be initialized before calling %s"), ANSI_TO_TCHAR(__FUNCTION__))
+		&& ensureMsgf(StateTree->GetDefaultParameters().GetPropertyBagStruct() == Parameters.GetPropertyBagStruct(),
+			TEXT("Parameters must be of the same struct type. Make sure to migrate the provided parameters to the same type as the StateTree default parameters.")))
 	{
-		return;
-	}
-
-	// The current implementation golden path is based on an expected 1:1 match between the 2 parameter lists
-	// In case they mismatch then we search for matching parameter. 
-	int32 TreeDescIndex = 0;
-	for (const FStateTreeParameterDesc& ExternalDesc : Parameters.Parameters)
-	{
-		if (!TreeDescs[TreeDescIndex].IsMatching(ExternalDesc))
-		{
-			// The parameter was not found where we expected it to be, search for it starting from last matching index
-			bool bMatchFound = false;
-			for (int32 Index = 0; Index < NumTreeDescs; Index++)
-			{
-				TreeDescIndex = (TreeDescIndex + 1) % NumTreeDescs;
-				if (TreeDescs[TreeDescIndex].IsMatching(ExternalDesc))
-				{
-					bMatchFound = true;
-					break;
-				}
-			}
-
-			if (!bMatchFound)
-			{
-				STATETREE_LOG(Warning, TEXT("Parameter '%s' doesn't match any expected parameters of StateTree '%s'."), *LexToString(ExternalDesc), *GetNameSafe(StateTree));
-				continue;
-			}
-		}
-
-		const uint16 DataViewIndex = TreeDescs[TreeDescIndex].DataViewIndex;
-		check(DataViewIndex != FStateTreeParameterDesc::InvalidIndex && DataViews.IsValidIndex(DataViewIndex));
-		DataViews[DataViewIndex] = FStateTreeDataView(ExternalDesc.Parameter);
-		TreeDescIndex = (TreeDescIndex + 1) % NumTreeDescs;
+		check(DataViews.IsValidIndex(StateTree->DefaultParametersDataViewIndex));
+		DataViews[StateTree->DefaultParametersDataViewIndex] = FStateTreeDataView(Parameters.GetMutableValue());
 	}
 }
 
diff --git a/Engine/Plugins/Runtime/StateTree/Source/StateTreeModule/Private/StateTreeReference.cpp b/Engine/Plugins/Runtime/StateTree/Source/StateTreeModule/Private/StateTreeReference.cpp
--- a/Engine/Plugins/Runtime/StateTree/Source/StateTreeModule/Private/StateTreeReference.cpp
+++ b/Engine/Plugins/Runtime/StateTree/Source/StateTreeModule/Private/StateTreeReference.cpp
@@ -73,69 +73,18 @@ void UStateTreeReferenceWrapper::BeginDestroy()
 
 bool FStateTreeReference::SyncParameters()
 {
-	TArray<FStateTreeParameterDesc>& Params = Parameters.Parameters;
 	if (StateTree == nullptr)
 	{
-		const bool bHadParameters = !Params.IsEmpty();
-		Params.Reset();
+		const bool bHadParameters = Parameters.IsValid();
+		Parameters.Reset();
 		return bHadParameters;
 	}
 
-	// First implementation is to make sure that the StateTreeReference holds values for each parameter of the StateTree.
-	// This is the easiest way to expose them to the user but that also means that user can't specify only a subset of parameters
-	// and use default value for the others.
-	const TConstArrayView<FStateTreeParameterDesc> AssetParameters = StateTree ? StateTree->GetParameterDescs() : TConstArrayView<FStateTreeParameterDesc>();
-	bool bAreParametersSynced = AssetParameters.Num() == Params.Num();
-	if (bAreParametersSynced)
-	{
-		for (int32 i = 0; i < Params.Num(); ++i)
-		{
-			if (!Params[i].IsMatching(AssetParameters[i]))
-			{
-				// Ok to mismatch only by the name if both ID and Type match. Simply update it since
-				// it must also match when passed to the ExecutionContext at runtime.
-				if (Params[i].ID == AssetParameters[i].ID && Params[i].IsSameType(AssetParameters[i]))
-				{
-					UE_LOG(LogStateTree, Warning, TEXT("StateTree parameter '%s' name was updated to match StateTree parameter '%s'."),
-						*LexToString(Params[i]), *LexToString(AssetParameters[i]));
-					
-					Params[i].Name = AssetParameters[i].Name;
-					checkf(Params[i].IsMatching(AssetParameters[i]),
-						TEXT("After fixing the name the parameters should match."
-							 " If this fails it indicates that same ID+Type is no longer the only required condition."));
-				}
-				else
-				{
-					bAreParametersSynced = false;
-					break;	
-				}
-			}
-		}
-	}
-
+	const FInstancedPropertyBag& DefaultParameters = StateTree->GetDefaultParameters();
+	const bool bAreParametersSynced = DefaultParameters.GetPropertyBagStruct() == Parameters.GetPropertyBagStruct();
 	if (!bAreParametersSynced)
 	{
-		TArray<FStateTreeParameterDesc> DeprecatedParameters = MoveTemp(Params);
-		Params = AssetParameters;
-
-		for (int32 i = 0; i < Params.Num(); ++i)
-		{
-			FStateTreeParameterDesc& ParameterDesc = Params[i];
-
-			// Find parameter corresponding index in the deprecated parameters using 'Type' and 'Id'.
-			// Not using IsMatching() here since 'Name' might have change in addition to parameters added/removed since the last sync.
-			const int32 IndexInDeprecated = DeprecatedParameters.IndexOfByPredicate([&ParameterDesc](const FStateTreeParameterDesc& DeprecatedDesc)
-				{
-					return DeprecatedDesc.ID == ParameterDesc.ID && DeprecatedDesc.IsSameType(ParameterDesc);
-				});
-
-			if (IndexInDeprecated != INDEX_NONE)
-			{
-				// Only update value
-				ParameterDesc.Parameter = DeprecatedParameters[IndexInDeprecated].Parameter;
-				DeprecatedParameters.RemoveAtSwap(IndexInDeprecated);
-			}
-		}
+		Parameters.MigrateToNewBagInstance(DefaultParameters);	
 	}
 
 	return !bAreParametersSynced;
diff --git a/Engine/Plugins/Runtime/StateTree/Source/StateTreeModule/Public/StateTree.h b/Engine/Plugins/Runtime/StateTree/Source/StateTreeModule/Public/StateTree.h
--- a/Engine/Plugins/Runtime/StateTree/Source/StateTreeModule/Public/StateTree.h
+++ b/Engine/Plugins/Runtime/StateTree/Source/StateTreeModule/Public/StateTree.h
@@ -35,8 +35,8 @@ public:
 	/** @return List of named external data enforced by the schema that must be provided through the execution context. */
 	TConstArrayView<FStateTreeExternalDataDesc> GetNamedExternalDataDescs() const { return Schema ? Schema->GetNamedExternalDataDescs() : TConstArrayView<FStateTreeExternalDataDesc>(); }
 
-	/** @return List of parameters of the state tree. Default parameter values can be overridden at runtime by the execution context. */
-	TConstArrayView<FStateTreeParameterDesc> GetParameterDescs() const { return Parameters.Parameters; }
+	/** @return List of default parameters of the state tree. Default parameter values can be overridden at runtime by the execution context. */
+	const FInstancedPropertyBag& GetDefaultParameters() const { return Parameters; }
 
 	/** @return true if the tree asset can be used at runtime. */
 	bool IsReadyToRun() const;
@@ -84,7 +84,7 @@ private:
 	 * @see FStateTreeReference
 	 */
 	UPROPERTY()
-	FStateTreeParameters Parameters;	
+	FInstancedPropertyBag Parameters;
 
 	/** Evaluators, Tasks, and Condition items */
 	UPROPERTY()
@@ -111,6 +111,10 @@ private:
 	UPROPERTY(Transient)
 	int32 ExternalDataBaseIndex = 0;
 
+	/** Data view index of the tree parameters */
+	UPROPERTY()
+	int32 DefaultParametersDataViewIndex = INDEX_NONE;
+	
 	UPROPERTY()
 	FStateTreePropertyBindings PropertyBindings;
 
diff --git a/Engine/Plugins/Runtime/StateTree/Source/StateTreeModule/Public/StateTreeExecutionContext.h b/Engine/Plugins/Runtime/StateTree/Source/StateTreeModule/Public/StateTreeExecutionContext.h
--- a/Engine/Plugins/Runtime/StateTree/Source/StateTreeModule/Public/StateTreeExecutionContext.h
+++ b/Engine/Plugins/Runtime/StateTree/Source/StateTreeModule/Public/StateTreeExecutionContext.h
@@ -64,14 +64,14 @@ public:
 	/** Initializes the StateTree instance to be used with specific owner and StateTree asset. */
 	bool Init(UObject& InOwner, const UStateTree& InStateTree, const EStateTreeStorage InStorageType);
 
-	/** Updates data views of all parameters by using the default values defined in the StateTree asset. */
+	/** Updates data view of the parameters by using the default values defined in the StateTree asset. */
 	void SetDefaultParameters();
 
 	/**
-	 * Updates data views of the matching parameters by replacing the default values defined in the StateTree asset by the provided values.
+	 * Updates data view of the parameters by replacing the default values defined in the StateTree asset by the provided values.
 	 * Note: caller is responsible to make sure external parameters lifetime matches the context.
 	 */
-	void SetParameters(const FStateTreeParameters& Parameters);
+	void SetParameters(const FInstancedPropertyBag& Parameters);
 	
 	/** Resets the instance to initial empty state. Note: Does not call ExitState(). */
 	void Reset();
diff --git a/Engine/Plugins/Runtime/StateTree/Source/StateTreeModule/Public/StateTreeReference.h b/Engine/Plugins/Runtime/StateTree/Source/StateTreeModule/Public/StateTreeReference.h
--- a/Engine/Plugins/Runtime/StateTree/Source/StateTreeModule/Public/StateTreeReference.h
+++ b/Engine/Plugins/Runtime/StateTree/Source/StateTreeModule/Public/StateTreeReference.h
@@ -19,8 +19,8 @@ struct FStateTreeReference
 	UPROPERTY(EditDefaultsOnly, Category="")
 	TObjectPtr<UStateTree> StateTree = nullptr;
 
-	UPROPERTY(EditDefaultsOnly, Category="")
-	FStateTreeParameters Parameters;
+	UPROPERTY(EditDefaultsOnly, Category="", meta=(FixedLayout))
+	FInstancedPropertyBag Parameters;
 
 #if WITH_EDITOR
 	/**
@@ -46,7 +46,7 @@ public:
 
 protected:
 #if WITH_EDITOR
-	virtual void PostEditChangeProperty(struct FPropertyChangedEvent& PropertyChangedEvent) override;
+	virtual void PostEditChangeProperty(FPropertyChangedEvent& PropertyChangedEvent) override;
 	virtual void PostEditChangeChainProperty(FPropertyChangedChainEvent& PropertyChangedEvent) override;
 	virtual void PostInitProperties() override;
 	virtual void PostLoad() override;
diff --git a/Engine/Plugins/Runtime/StateTree/Source/StateTreeModule/Public/StateTreeTypes.h b/Engine/Plugins/Runtime/StateTree/Source/StateTreeModule/Public/StateTreeTypes.h
--- a/Engine/Plugins/Runtime/StateTree/Source/StateTreeModule/Public/StateTreeTypes.h
+++ b/Engine/Plugins/Runtime/StateTree/Source/StateTreeModule/Public/StateTreeTypes.h
@@ -513,78 +513,6 @@ struct TStateTreeInstanceDataPropertyHandle : FStateTreeInstanceDataPropertyHand
 	typedef T DataType;
 };
 
-/**
- * Describes a parameter of the state tree that could be used for bindings.
- * Can also be provided externally by a StateTreeReference to parameterized the tree.
- * @see FStateTreeReference
- */
-USTRUCT()
-struct STATETREEMODULE_API FStateTreeParameterDesc
-{
-	GENERATED_BODY()
-
-	/** Indicates that parameters hold the same data type */
-	bool IsSameType(const FStateTreeParameterDesc& RHS) const { return Parameter.GetScriptStruct() == RHS.Parameter.GetScriptStruct(); }
-
-	/** Indicates that parameters hold the same data type and have the same identifier. They might have different values. */
-	bool IsMatching(const FStateTreeParameterDesc& RHS) const
-	{				
-#if WITH_EDITORONLY_DATA
-		return Name == RHS.Name && ID == RHS.ID && IsSameType(RHS);
-#else
-		return Name == RHS.Name && IsSameType(RHS);
-#endif			
-	}
-	
-	friend FString LexToString(const FStateTreeParameterDesc& Desc)
-	{
-		return FString::Printf(TEXT("{%s} %s"),
-			Desc.Parameter.GetScriptStruct() != nullptr ? *Desc.Parameter.GetScriptStruct()->GetName() : TEXT("null type"),
-			*Desc.Name.ToString());
-	}
-
-	bool IsValid() const
-	{
-		return DataViewIndex != InvalidIndex;
-	}
-
-	/** The type of the parameter. */
-	UPROPERTY(EditDefaultsOnly, Category = Common)
-	FInstancedStruct Parameter;
-
-	/** Name of the parameter. */
-	UPROPERTY(EditDefaultsOnly, Category = Common)
-	FName Name;
-
-	/** The runtime data's data view index in the StateTreeExecutionContext, and source struct index in property binding. */
-	UPROPERTY()
-	uint16 DataViewIndex = InvalidIndex;
-
-	static constexpr uint16 InvalidIndex = MAX_uint16;
-
-#if WITH_EDITORONLY_DATA
-	/** Unique identifier. */
-	UPROPERTY()
-	FGuid ID;
-#endif
-};
-
-
-/**
- * Container for StateTree parameters.
- * Could be used for parameter definitions (within the StateTreeAsset) and parameterization (StateTreeReference).
- */
-USTRUCT()
-struct FStateTreeParameters
-{
-	GENERATED_BODY()
-
-	void Reset() { Parameters.Reset(); }
-	
-	UPROPERTY(EditDefaultsOnly, Category= Common)
-	TArray<FStateTreeParameterDesc> Parameters;
-};
-
 
 /**
  * Describes an external data. The data can point to a struct or object.
