[StateTree] NamedExternalDataDescs are now copied from schema to the StateTree at compilation. Schema derived classes need to override only one method to provide their descriptors.
#rnx

diff --git a/Engine/Plugins/Runtime/StateTree/Source/StateTreeEditorModule/Private/StateTreeCompiler.cpp b/Engine/Plugins/Runtime/StateTree/Source/StateTreeEditorModule/Private/StateTreeCompiler.cpp
--- a/Engine/Plugins/Runtime/StateTree/Source/StateTreeEditorModule/Private/StateTreeCompiler.cpp
+++ b/Engine/Plugins/Runtime/StateTree/Source/StateTreeEditorModule/Private/StateTreeCompiler.cpp
@@ -43,7 +43,8 @@ bool FStateTreeCompiler::Compile(UStateTree& InStateTree)
 	// Mark all named external values as binding source
 	if (StateTree->Schema)
 	{
-		for (FStateTreeExternalDataDesc& Desc : StateTree->Schema->GetMutableNamedExternalDataDescs())
+		StateTree->NamedExternalDataDescs = StateTree->Schema->GetNamedExternalDataDescs();
+		for (FStateTreeExternalDataDesc& Desc : StateTree->NamedExternalDataDescs)
 		{
 			Desc.Handle.DataViewIndex = BindingsCompiler.AddSourceStruct({Desc.Name, Desc.Struct, EStateTreeBindableStructSource::TreeData, Desc.ID});
 		} 
diff --git a/Engine/Plugins/Runtime/StateTree/Source/StateTreeModule/Private/StateTree.cpp b/Engine/Plugins/Runtime/StateTree/Source/StateTreeModule/Private/StateTree.cpp
--- a/Engine/Plugins/Runtime/StateTree/Source/StateTreeModule/Private/StateTree.cpp
+++ b/Engine/Plugins/Runtime/StateTree/Source/StateTreeModule/Private/StateTree.cpp
@@ -22,6 +22,7 @@ void UStateTree::ResetCompiled()
 	Instances.Reset();
 	InstanceObjects.Reset();
 	ExternalDataDescs.Reset();
+	NamedExternalDataDescs.Reset();
 	PropertyBindings.Reset();
 
 	NumDataViews = 0;
diff --git a/Engine/Plugins/Runtime/StateTree/Source/StateTreeModule/Public/StateTree.h b/Engine/Plugins/Runtime/StateTree/Source/StateTreeModule/Public/StateTree.h
--- a/Engine/Plugins/Runtime/StateTree/Source/StateTreeModule/Public/StateTree.h
+++ b/Engine/Plugins/Runtime/StateTree/Source/StateTreeModule/Public/StateTree.h
@@ -33,7 +33,7 @@ public:
 	TConstArrayView<FStateTreeExternalDataDesc> GetExternalDataDescs() const { return ExternalDataDescs; }
 
 	/** @return List of named external data enforced by the schema that must be provided through the execution context. */
-	TConstArrayView<FStateTreeExternalDataDesc> GetNamedExternalDataDescs() const { return Schema ? Schema->GetNamedExternalDataDescs() : TConstArrayView<FStateTreeExternalDataDesc>(); }
+	TConstArrayView<FStateTreeExternalDataDesc> GetNamedExternalDataDescs() const { return NamedExternalDataDescs; }
 
 	/** @return List of default parameters of the state tree. Default parameter values can be overridden at runtime by the execution context. */
 	const FInstancedPropertyBag& GetDefaultParameters() const { return Parameters; }
@@ -101,10 +101,14 @@ private:
 	UPROPERTY(Transient)
 	FStateTreeInstanceData InstanceDataDefaultValue;
 
-	/** List of external data required by the state tree, creating during linking. */
+	/** List of external data required by the state tree, created during linking. */
 	UPROPERTY(Transient)
 	TArray<FStateTreeExternalDataDesc> ExternalDataDescs;
 
+	/** List of names external data enforced by the schema, created at compilation. */
+	UPROPERTY(Transient)
+	TArray<FStateTreeExternalDataDesc> NamedExternalDataDescs;
+	
 	UPROPERTY(Transient)
 	int32 NumDataViews = 0;
 
diff --git a/Engine/Plugins/Runtime/StateTree/Source/StateTreeModule/Public/StateTreeSchema.h b/Engine/Plugins/Runtime/StateTree/Source/StateTreeModule/Public/StateTreeSchema.h
--- a/Engine/Plugins/Runtime/StateTree/Source/StateTreeModule/Public/StateTreeSchema.h
+++ b/Engine/Plugins/Runtime/StateTree/Source/StateTreeModule/Public/StateTreeSchema.h
@@ -41,10 +41,7 @@ public:
 	bool IsChildOfBlueprintBase(const UClass* InClass) const;
 
 	/** @return List of named external data enforced by the schema. They must be provided at runtime through the execution context. */
-	virtual TConstArrayView<FStateTreeExternalDataDesc> GetNamedExternalDataDescs() const	{ return {}; }
-
-	/** @return List of mutable named external data enforced by the schema. They must be provided at runtime through the execution context. */
-	virtual TArrayView<FStateTreeExternalDataDesc> GetMutableNamedExternalDataDescs() { return {}; }
+	virtual TConstArrayView<FStateTreeExternalDataDesc> GetNamedExternalDataDescs() const { return {}; }
 
 #if WITH_EDITOR
 	
