[StateTree] validate Parameter view index instead of using a check since trees need to be recompiled to set it but we don't want to force it to be able to use older asset.
#rnx

diff --git a/Engine/Plugins/Runtime/StateTree/Source/StateTreeModule/Private/StateTreeExecutionContext.cpp b/Engine/Plugins/Runtime/StateTree/Source/StateTreeModule/Private/StateTreeExecutionContext.cpp
--- a/Engine/Plugins/Runtime/StateTree/Source/StateTreeModule/Private/StateTreeExecutionContext.cpp
+++ b/Engine/Plugins/Runtime/StateTree/Source/StateTreeModule/Private/StateTreeExecutionContext.cpp
@@ -61,9 +61,9 @@ bool FStateTreeExecutionContext::Init(UObject& InOwner, const UStateTree& InStat
 
 void FStateTreeExecutionContext::SetDefaultParameters()
 {
-	if (ensureMsgf(StateTree != nullptr, TEXT("Execution context must be initialized before calling %s"), ANSI_TO_TCHAR(__FUNCTION__)))
+	if (ensureMsgf(StateTree != nullptr, TEXT("Execution context must be initialized before calling %s"), ANSI_TO_TCHAR(__FUNCTION__))
+		&& DataViews.IsValidIndex(StateTree->DefaultParametersDataViewIndex))
 	{
-		check(DataViews.IsValidIndex(StateTree->DefaultParametersDataViewIndex));
 		DataViews[StateTree->DefaultParametersDataViewIndex] = FStateTreeDataView(StateTree->GetDefaultParameters().GetMutableValue());	
 	}
 }
@@ -72,9 +72,9 @@ void FStateTreeExecutionContext::SetParameters(const FInstancedPropertyBag& Para
 {
 	if (ensureMsgf(StateTree != nullptr, TEXT("Execution context must be initialized before calling %s"), ANSI_TO_TCHAR(__FUNCTION__))
 		&& ensureMsgf(StateTree->GetDefaultParameters().GetPropertyBagStruct() == Parameters.GetPropertyBagStruct(),
-			TEXT("Parameters must be of the same struct type. Make sure to migrate the provided parameters to the same type as the StateTree default parameters.")))
+			TEXT("Parameters must be of the same struct type. Make sure to migrate the provided parameters to the same type as the StateTree default parameters."))
+		&& DataViews.IsValidIndex(StateTree->DefaultParametersDataViewIndex))
 	{
-		check(DataViews.IsValidIndex(StateTree->DefaultParametersDataViewIndex));
 		DataViews[StateTree->DefaultParametersDataViewIndex] = FStateTreeDataView(Parameters.GetMutableValue());
 	}
 }
