[StateTree] compiled NamedExterrnalData descriptors need to be serialized
#rnx

diff --git a/Engine/Plugins/Runtime/StateTree/Source/StateTreeModule/Public/StateTree.h b/Engine/Plugins/Runtime/StateTree/Source/StateTreeModule/Public/StateTree.h
--- a/Engine/Plugins/Runtime/StateTree/Source/StateTreeModule/Public/StateTree.h
+++ b/Engine/Plugins/Runtime/StateTree/Source/StateTreeModule/Public/StateTree.h
@@ -106,7 +106,7 @@ private:
 	TArray<FStateTreeExternalDataDesc> ExternalDataDescs;
 
 	/** List of names external data enforced by the schema, created at compilation. */
-	UPROPERTY(Transient)
+	UPROPERTY()
 	TArray<FStateTreeExternalDataDesc> NamedExternalDataDescs;
 	
 	UPROPERTY(Transient)
