StateTree: shared instance data
- Separated condition instance data into a shared instance data (mutable but no persistent state)
- Made BP TestCondition() const, conditions should not hold state
- Added macro to define runtime data POD for faster init
- Added custom serialization version for UStateTree, older assets need recompile
- Updated memory reporting to be a bit more accurate, separated shared and unique data

diff --git a/Engine/Plugins/Runtime/StateTree/Source/StateTreeEditorModule/Private/StateTreeCompiler.cpp b/Engine/Plugins/Runtime/StateTree/Source/StateTreeEditorModule/Private/StateTreeCompiler.cpp
--- a/Engine/Plugins/Runtime/StateTree/Source/StateTreeEditorModule/Private/StateTreeCompiler.cpp
+++ b/Engine/Plugins/Runtime/StateTree/Source/StateTreeEditorModule/Private/StateTreeCompiler.cpp
@@ -475,8 +475,8 @@ bool FStateTreeCompiler::CreateCondition(UStateTreeState& State, const FStateTre
 	if (CondNode.Instance.IsValid())
 	{
 		// Struct instance
-		FInstancedStruct& Instance = StateTree->Instances.AddDefaulted_GetRef();
-		const int32 InstanceIndex = StateTree->Instances.Num() - 1;
+		FInstancedStruct& Instance = StateTree->SharedInstances.AddDefaulted_GetRef();
+		const int32 InstanceIndex = StateTree->SharedInstances.Num() - 1;
 		
 		Instance = CondNode.Instance;
 	
@@ -496,7 +496,7 @@ bool FStateTreeCompiler::CreateCondition(UStateTreeState& State, const FStateTre
 		UObject* Instance = DuplicateObject(CondNode.InstanceObject, StateTree);
 		
 		StateTree->InstanceObjects.Add(Instance);
-		const int32 InstanceIndex = StateTree->InstanceObjects.Num() - 1;
+		const int32 InstanceIndex = StateTree->SharedInstanceObjects.Num() - 1;
 		
 		// Create binding source struct descriptor.
 		StructDesc.Struct = Instance->GetClass();
diff --git a/Engine/Plugins/Runtime/StateTree/Source/StateTreeEditorModule/Private/StateTreeEditor.cpp b/Engine/Plugins/Runtime/StateTree/Source/StateTreeEditorModule/Private/StateTreeEditor.cpp
--- a/Engine/Plugins/Runtime/StateTree/Source/StateTreeEditorModule/Private/StateTreeEditor.cpp
+++ b/Engine/Plugins/Runtime/StateTree/Source/StateTreeEditorModule/Private/StateTreeEditor.cpp
@@ -348,12 +348,14 @@ FText FStateTreeEditor::GetStatisticsText() const
 	}
 
 	const FStateTreeInstanceData& InstanceDataDefaultValue = StateTree->GetInstanceDataDefaultValue();
-	const uint64 LayoutSize = InstanceDataDefaultValue.IsValid() ? InstanceDataDefaultValue.GetLayout()->GetLayoutInstanceSize() : 0;
-	
-	const FText SizeText = FText::AsMemory(LayoutSize);
-	const FText NumNodesText = FText::AsNumber(StateTree->GetNumInstances());
+	const FText RuntimeSizeText = FText::AsMemory(InstanceDataDefaultValue.GetEstimatedMemoryUsage());
+	const FText RuntimeNumNodesText = FText::AsNumber(InstanceDataDefaultValue.GetNumItems());
+
+	const FStateTreeInstanceData& SharedInstanceData = StateTree->GetSharedInstanceData();
+	const FText SharedSizeText = FText::AsMemory(SharedInstanceData.GetEstimatedMemoryUsage());
+	const FText SharedNumNodesText = FText::AsNumber(SharedInstanceData.GetNumItems());
 
-	return FText::Format(LOCTEXT("RuntimeSize", "Runtime size: {0}, {1} nodes"), SizeText, NumNodesText);
+	return FText::Format(LOCTEXT("RuntimeSize", "Runtime size: {0}, {1} nodes\nShared size: {2}, {3} nodes"), RuntimeSizeText, RuntimeNumNodesText, SharedSizeText, SharedNumNodesText);
 }
 
 void FStateTreeEditor::HandleModelAssetChanged()
diff --git a/Engine/Plugins/Runtime/StateTree/Source/StateTreeModule/Private/Blueprint/StateTreeConditionBlueprintBase.cpp b/Engine/Plugins/Runtime/StateTree/Source/StateTreeModule/Private/Blueprint/StateTreeConditionBlueprintBase.cpp
--- a/Engine/Plugins/Runtime/StateTree/Source/StateTreeModule/Private/Blueprint/StateTreeConditionBlueprintBase.cpp
+++ b/Engine/Plugins/Runtime/StateTree/Source/StateTreeModule/Private/Blueprint/StateTreeConditionBlueprintBase.cpp
@@ -9,7 +9,7 @@
 //  UStateTreeConditionBlueprintBase
 //----------------------------------------------------------------------//
 
-bool UStateTreeConditionBlueprintBase::TestCondition(FStateTreeExecutionContext& Context)
+bool UStateTreeConditionBlueprintBase::TestCondition(FStateTreeExecutionContext& Context) const
 {
 	AActor* OwnerActor = GetOwnerActor(Context);
 	return ReceiveTestCondition(OwnerActor);
diff --git a/Engine/Plugins/Runtime/StateTree/Source/StateTreeModule/Private/StateTree.cpp b/Engine/Plugins/Runtime/StateTree/Source/StateTreeModule/Private/StateTree.cpp
--- a/Engine/Plugins/Runtime/StateTree/Source/StateTreeModule/Private/StateTree.cpp
+++ b/Engine/Plugins/Runtime/StateTree/Source/StateTreeModule/Private/StateTree.cpp
@@ -4,6 +4,9 @@
 #include "StateTreeLinker.h"
 #include "StateTreeNodeBase.h"
 
+const FGuid FStateTreeCustomVersion::GUID(0x28E21331, 0x501F4723, 0x8110FA64, 0xEA10DA1E);
+FCustomVersionRegistration GRegisterStateTreeCustomVersion(FStateTreeCustomVersion::GUID, FStateTreeCustomVersion::LatestVersion, TEXT("StateTreeAsset"));
+
 bool UStateTree::IsReadyToRun() const
 {
 	// Valid tree must have at least one state and valid instance data.
@@ -44,6 +47,15 @@ void UStateTree::PostLoad()
 {
 	Super::PostLoad();
 
+	const int32 CurrentVersion = GetLinkerCustomVersion(FStateTreeCustomVersion::GUID);
+
+	if (CurrentVersion < FStateTreeCustomVersion::SharedInstanceData)
+	{
+		ResetLinked();
+		UE_LOG(LogStateTree, Error, TEXT("%s: StateTree compiled data in older format. Please recompile the StateTree asset."), *GetName());
+		return;
+	}
+	
 	if (!Link())
 	{
 		UE_LOG(LogStateTree, Error, TEXT("%s failed to link. Asset will not be usable at runtime."), *GetName());	
@@ -54,6 +66,8 @@ void UStateTree::Serialize(FStructuredArchiveRecord Record)
 {
 	Super::Serialize(Record);
 
+	Record.GetUnderlyingArchive().UsingCustomVersion(FStateTreeCustomVersion::GUID);
+	
 	// We need to link and rebind property bindings each time a BP is compiled,
 	// because property bindings may get invalid, and instance data potentially needs refreshed.
 	if (Record.GetUnderlyingArchive().IsModifyingWeakAndStrongReferences())
@@ -65,15 +79,21 @@ void UStateTree::Serialize(FStructuredArchiveRecord Record)
 	}
 }
 
-bool UStateTree::Link()
+void UStateTree::ResetLinked()
 {
-	// Initialize the instance data default value.
-	// This data will be used to allocate runtime instance on all StateTree users.
 	InstanceDataDefaultValue.Reset();
+	SharedInstanceData.Reset();
 	
 	ExternalDataDescs.Reset();
 	NumDataViews = 0;
 	ExternalDataBaseIndex = 0;
+}
+
+bool UStateTree::Link()
+{
+	// Initialize the instance data default value.
+	// This data will be used to allocate runtime instance on all StateTree users.
+	ResetLinked();
 	
 	// Update property bag structs before resolving binding.
 	TArrayView<FStateTreeBindableStructDesc> SourceStructs = PropertyBindings.GetSourceStructs();
@@ -156,10 +176,16 @@ bool UStateTree::Link()
 	ExternalDataDescs = Linker.GetExternalDataDescs();
 	NumDataViews = ExternalDataBaseIndex + ExternalDataDescs.Num();
 
-	if (Instances.Num() > 0)
+	if (Instances.Num() > 0 || InstanceObjects.Num() > 0)
 	{
 		InstanceDataDefaultValue.Initialize(*this, Instances, InstanceObjects);
 	}
 
+	if (SharedInstances.Num())
+	{
+		SharedInstanceData.Initialize(*this, SharedInstances, SharedInstanceObjects);
+	}
+
+	
 	return true;
 }
diff --git a/Engine/Plugins/Runtime/StateTree/Source/StateTreeModule/Private/StateTreeExecutionContext.cpp b/Engine/Plugins/Runtime/StateTree/Source/StateTreeModule/Private/StateTreeExecutionContext.cpp
--- a/Engine/Plugins/Runtime/StateTree/Source/StateTreeModule/Private/StateTreeExecutionContext.cpp
+++ b/Engine/Plugins/Runtime/StateTree/Source/StateTreeModule/Private/StateTreeExecutionContext.cpp
@@ -761,7 +761,7 @@ EStateTreeRunStatus FStateTreeExecutionContext::TickTasks(FStateTreeInstanceData
 	return Result;
 }
 
-bool FStateTreeExecutionContext::TestAllConditions(FStateTreeInstanceData& InstanceData, const uint32 ConditionsOffset, const uint32 ConditionsNum)
+bool FStateTreeExecutionContext::TestAllConditions(const uint32 ConditionsOffset, const uint32 ConditionsNum)
 {
 	CSV_SCOPED_TIMING_STAT_EXCLUSIVE(StateTree_TestConditions);
 
@@ -769,6 +769,8 @@ bool FStateTreeExecutionContext::TestAllConditions(FStateTreeInstanceData& Insta
 	{
 		return true;
 	}
+
+	const FStateTreeInstanceData& SharedInstanceData = StateTree->GetSharedInstanceData();
 	
 	TStaticArray<EStateTreeConditionOperand, UE::StateTree::MaxConditionIndent + 1> Operands(InPlace, EStateTreeConditionOperand::Copy);
 	TStaticArray<bool, UE::StateTree::MaxConditionIndent + 1> Values(InPlace, false);
@@ -778,7 +780,7 @@ bool FStateTreeExecutionContext::TestAllConditions(FStateTreeInstanceData& Insta
 	for (uint32 Index = 0; Index < ConditionsNum; Index++)
 	{
 		const FStateTreeConditionBase& Cond = GetNode<FStateTreeConditionBase>(ConditionsOffset + Index);
-		const FStateTreeDataView CondInstanceView = GetInstanceData(InstanceData, Cond.bInstanceIsObject, Cond.InstanceIndex);
+		const FStateTreeDataView CondInstanceView = GetInstanceData(SharedInstanceData, Cond.bInstanceIsObject, Cond.InstanceIndex);
 		DataViews[Cond.DataViewIndex] = CondInstanceView;
 
 		// Copy bound properties.
@@ -858,7 +860,7 @@ bool FStateTreeExecutionContext::TriggerTransitions(FStateTreeInstanceData& Inst
 			// All transition conditions must pass
 			const int16 TransitionIndex = State.TransitionsBegin + i;
 			const FCompactStateTransition& Transition = StateTree->Transitions[TransitionIndex];
-			if (EnumHasAllFlags(Transition.Event, Event) && TestAllConditions(InstanceData, Transition.ConditionsBegin, Transition.ConditionsNum))
+			if (EnumHasAllFlags(Transition.Event, Event) && TestAllConditions(Transition.ConditionsBegin, Transition.ConditionsNum))
 			{
 				// If a transition has delay, we stop testing other transitions, but the transition will not pass the condition until the delay time passes.
 				if (Transition.GateDelay > 0)
@@ -1014,7 +1016,7 @@ bool FStateTreeExecutionContext::SelectStateInternal(FStateTreeInstanceData& Ins
 	TickEvaluatorsForSelect(InstanceData, NextState, EStateTreeEvaluationType::PreSelect, 0.0f);
 	
 	// Check that the state can be entered
-	if (TestAllConditions(InstanceData, State.EnterConditionsBegin, State.EnterConditionsNum))
+	if (TestAllConditions(State.EnterConditionsBegin, State.EnterConditionsNum))
 	{
 		if (!OutNewActiveState.Push(NextState))
 		{
diff --git a/Engine/Plugins/Runtime/StateTree/Source/StateTreeModule/Private/StateTreeInstanceData.cpp b/Engine/Plugins/Runtime/StateTree/Source/StateTreeModule/Private/StateTreeInstanceData.cpp
--- a/Engine/Plugins/Runtime/StateTree/Source/StateTreeModule/Private/StateTreeInstanceData.cpp
+++ b/Engine/Plugins/Runtime/StateTree/Source/StateTreeModule/Private/StateTreeInstanceData.cpp
@@ -74,6 +74,38 @@ int32 FStateTreeInstanceDataLayout::GetLayoutInstanceMinAlignment() const
 //  FStateTreeInstanceData
 //----------------------------------------------------------------//
 
+int32 FStateTreeInstanceData::GetEstimatedMemoryUsage() const
+{
+	int32 Size = sizeof(FStateTreeInstanceData);
+
+	if (Layout.IsValid())
+	{
+		Size = Layout->GetLayoutInstanceSize();
+	}
+
+	for (const UObject* InstanceObject : InstanceObjects)
+	{
+		if (InstanceObject)
+		{
+			Size += InstanceObject->GetClass()->GetStructureSize();
+		}
+	}
+
+	return Size;
+}
+
+int32 FStateTreeInstanceData::GetNumItems() const
+{
+	int32 Num = InstanceObjects.Num();
+
+	if (Layout.IsValid())
+	{
+		Num += Layout->Num();
+	}
+	
+	return Num;
+}
+
 void FStateTreeInstanceData::AddStructReferencedObjects(class FReferenceCollector& Collector) const
 {
 	if (!Layout.IsValid())
diff --git a/Engine/Plugins/Runtime/StateTree/Source/StateTreeModule/Public/Blueprint/StateTreeConditionBlueprintBase.h b/Engine/Plugins/Runtime/StateTree/Source/StateTreeModule/Public/Blueprint/StateTreeConditionBlueprintBase.h
--- a/Engine/Plugins/Runtime/StateTree/Source/StateTreeModule/Public/Blueprint/StateTreeConditionBlueprintBase.h
+++ b/Engine/Plugins/Runtime/StateTree/Source/StateTreeModule/Public/Blueprint/StateTreeConditionBlueprintBase.h
@@ -22,9 +22,9 @@ class STATETREEMODULE_API UStateTreeConditionBlueprintBase : public UStateTreeIt
 public:
 
 	UFUNCTION(BlueprintImplementableEvent)
-	bool ReceiveTestCondition(AActor* OwnerActor);
+	bool ReceiveTestCondition(AActor* OwnerActor) const;
 
-	virtual bool TestCondition(FStateTreeExecutionContext& Context);
+	virtual bool TestCondition(FStateTreeExecutionContext& Context) const;
 };
 
 /**
diff --git a/Engine/Plugins/Runtime/StateTree/Source/StateTreeModule/Public/Conditions/StateTreeCondition_Common.h b/Engine/Plugins/Runtime/StateTree/Source/StateTreeModule/Public/Conditions/StateTreeCondition_Common.h
--- a/Engine/Plugins/Runtime/StateTree/Source/StateTreeModule/Public/Conditions/StateTreeCondition_Common.h
+++ b/Engine/Plugins/Runtime/StateTree/Source/StateTreeModule/Public/Conditions/StateTreeCondition_Common.h
@@ -29,6 +29,7 @@ struct STATETREEMODULE_API FStateTreeCondition_CompareIntInstanceData
 	UPROPERTY(EditAnywhere, Category = Parameter)
 	int32 Right = 0;
 };
+STATETREE_POD_INSTANCEDATA(FStateTreeCondition_CompareIntInstanceData);
 
 USTRUCT(DisplayName="Integer Compare")
 struct STATETREEMODULE_API FStateTreeCondition_CompareInt : public FStateTreeConditionCommonBase
@@ -71,6 +72,7 @@ struct STATETREEMODULE_API FStateTreeCondition_CompareFloatInstanceData
 	UPROPERTY(EditAnywhere, Category = Parameter)
 	float Right = 0.0f;
 };
+STATETREE_POD_INSTANCEDATA(FStateTreeCondition_CompareFloatInstanceData);
 
 USTRUCT(DisplayName = "Float Compare")
 struct STATETREEMODULE_API FStateTreeCondition_CompareFloat : public FStateTreeConditionCommonBase
@@ -113,6 +115,7 @@ struct STATETREEMODULE_API FStateTreeCondition_CompareBoolInstanceData
 	UPROPERTY(EditAnywhere, Category = Parameter)
 	bool bRight = false;
 };
+STATETREE_POD_INSTANCEDATA(FStateTreeCondition_CompareBoolInstanceData);
 
 USTRUCT(DisplayName = "Bool Compare")
 struct STATETREEMODULE_API FStateTreeCondition_CompareBool : public FStateTreeConditionCommonBase
@@ -155,6 +158,7 @@ struct STATETREEMODULE_API FStateTreeCondition_CompareEnumInstanceData
 	UPROPERTY(EditAnywhere, Category = Parameter)
 	FStateTreeAnyEnum Right;
 };
+STATETREE_POD_INSTANCEDATA(FStateTreeCondition_CompareEnumInstanceData);
 
 USTRUCT(DisplayName = "Enum Compare")
 struct STATETREEMODULE_API FStateTreeCondition_CompareEnum : public FStateTreeConditionCommonBase
@@ -204,6 +208,7 @@ struct STATETREEMODULE_API FStateTreeCondition_CompareDistanceInstanceData
 	UPROPERTY(EditAnywhere, Category = Parameter)
 	float Distance = 0.0f;
 };
+STATETREE_POD_INSTANCEDATA(FStateTreeCondition_CompareDistanceInstanceData);
 
 USTRUCT(DisplayName = "Distance Compare")
 struct STATETREEMODULE_API FStateTreeCondition_CompareDistance : public FStateTreeConditionCommonBase
diff --git a/Engine/Plugins/Runtime/StateTree/Source/StateTreeModule/Public/StateTree.h b/Engine/Plugins/Runtime/StateTree/Source/StateTreeModule/Public/StateTree.h
--- a/Engine/Plugins/Runtime/StateTree/Source/StateTreeModule/Public/StateTree.h
+++ b/Engine/Plugins/Runtime/StateTree/Source/StateTreeModule/Public/StateTree.h
@@ -10,6 +10,29 @@
 #include "StateTreeInstanceData.h"
 #include "StateTree.generated.h"
 
+
+/** Custom serialization version for StateTree Asset */
+struct STATETREEMODULE_API FStateTreeCustomVersion
+{
+	enum Type
+	{
+		// Before any version changes were made in the plugin
+		BeforeCustomVersionWasAdded = 0,
+		// Separated conditions to shared instance data.
+		SharedInstanceData,
+
+		// -----<new versions can be added above this line>-------------------------------------------------
+		VersionPlusOne,
+		LatestVersion = VersionPlusOne - 1
+	};
+
+	/** The GUID for this custom version number */
+	const static FGuid GUID;
+
+private:
+	FStateTreeCustomVersion() {}
+};
+
 /**
  * StateTree asset. Contains the StateTree definition in both editor and runtime (baked) formats.
  */
@@ -22,9 +45,12 @@ public:
 
 	/** @return Default value for the instance data. */
 	const FStateTreeInstanceData& GetInstanceDataDefaultValue() const { return InstanceDataDefaultValue; }
-	
-	/** @return Number of runtime data (Evaluators, Tasks, Conditions) in the runtime storage. */
-	int32 GetNumInstances() const { return Instances.Num(); }
+
+	/**
+	 * @todo: This should return different data for each thread.
+	 * @return Shared instance data.
+	 */
+	const FStateTreeInstanceData& GetSharedInstanceData() const { return SharedInstanceData; }
 
 	/** @return Number of data views required for StateTree execution (Evaluators, Tasks, Conditions, External data). */
 	int32 GetNumDataViews() const { return NumDataViews; }
@@ -73,6 +99,12 @@ protected:
 
 private:
 
+	/**
+     * Reset the data generated by Link(), this in turn will cause IsReadyToRun() to return false.
+     * Used during linking, or to invalidate the linked data when data version is old (requires recompile). 
+	 */
+	void ResetLinked();
+	
 	// Properties
 
 	UPROPERTY(Instanced)
@@ -98,9 +130,20 @@ private:
 	UPROPERTY()
 	TArray<TObjectPtr<UObject>> InstanceObjects;
 
+	/** Conditions runtime data. */
+	UPROPERTY()
+	TArray<FInstancedStruct> SharedInstances;
+
+	/** Blueprint Conditions runtime data. */
+	UPROPERTY()
+	TArray<TObjectPtr<UObject>> SharedInstanceObjects;
+
 	UPROPERTY(Transient)
 	FStateTreeInstanceData InstanceDataDefaultValue;
 
+	UPROPERTY(Transient)
+	FStateTreeInstanceData SharedInstanceData;
+
 	/** List of external data required by the state tree, created during linking. */
 	UPROPERTY(Transient)
 	TArray<FStateTreeExternalDataDesc> ExternalDataDescs;
diff --git a/Engine/Plugins/Runtime/StateTree/Source/StateTreeModule/Public/StateTreeConditionBase.h b/Engine/Plugins/Runtime/StateTree/Source/StateTreeModule/Public/StateTreeConditionBase.h
--- a/Engine/Plugins/Runtime/StateTree/Source/StateTreeModule/Public/StateTreeConditionBase.h
+++ b/Engine/Plugins/Runtime/StateTree/Source/StateTreeModule/Public/StateTreeConditionBase.h
@@ -55,3 +55,15 @@ struct STATETREEMODULE_API FStateTreeConditionCommonBase : public FStateTreeCond
 	GENERATED_BODY()
 };
 
+/** Helper macro to define instance data as simple constructible. */
+#define STATETREE_POD_INSTANCEDATA(Type) \
+template <> struct TIsPODType<Type> { enum { Value = true }; }; \
+template<> \
+struct TStructOpsTypeTraits<Type> : public TStructOpsTypeTraitsBase2<Type> \
+{ \
+	enum \
+	{ \
+		WithZeroConstructor = true, \
+		WithNoDestructor = true, \
+	}; \
+};
diff --git a/Engine/Plugins/Runtime/StateTree/Source/StateTreeModule/Public/StateTreeExecutionContext.h b/Engine/Plugins/Runtime/StateTree/Source/StateTreeModule/Public/StateTreeExecutionContext.h
--- a/Engine/Plugins/Runtime/StateTree/Source/StateTreeModule/Public/StateTreeExecutionContext.h
+++ b/Engine/Plugins/Runtime/StateTree/Source/StateTreeModule/Public/StateTreeExecutionContext.h
@@ -319,7 +319,7 @@ protected:
 	 * Checks all conditions at given range
 	 * @return True if all conditions pass.
 	 */
-	bool TestAllConditions(FStateTreeInstanceData& InstanceData, const uint32 ConditionsOffset, const uint32 ConditionsNum);
+	bool TestAllConditions(const uint32 ConditionsOffset, const uint32 ConditionsNum);
 
 	/**
 	 * Triggers transitions based on current run status. CurrentStatus is used to select which transitions events are triggered.
@@ -362,7 +362,7 @@ protected:
 	}
 
 	/** @return View to an Evaluator, a Task, or a Condition instance data. */
-	FStateTreeDataView GetInstanceData(FStateTreeInstanceData& InstanceData, const bool bIsObject, const int32 Index) const
+	FStateTreeDataView GetInstanceData(const FStateTreeInstanceData& InstanceData, const bool bIsObject, const int32 Index) const
 	{
 		if (UNLIKELY(bIsObject == true))
 		{
diff --git a/Engine/Plugins/Runtime/StateTree/Source/StateTreeModule/Public/StateTreeInstanceData.h b/Engine/Plugins/Runtime/StateTree/Source/StateTreeModule/Public/StateTreeInstanceData.h
--- a/Engine/Plugins/Runtime/StateTree/Source/StateTreeModule/Public/StateTreeInstanceData.h
+++ b/Engine/Plugins/Runtime/StateTree/Source/StateTreeModule/Public/StateTreeInstanceData.h
@@ -156,6 +156,9 @@ struct STATETREEMODULE_API FStateTreeInstanceData
 	/** @return const pointer to an instance object   */
 	const UObject* GetObject(const int32 Index) const { return InstanceObjects[Index]; }
 
+	int32 GetEstimatedMemoryUsage() const;
+	int32 GetNumItems() const;
+	
 	/** Type traits */
 	void AddStructReferencedObjects(class FReferenceCollector& Collector) const;
 	bool Identical(const FStateTreeInstanceData* Other, uint32 PortFlags) const;
