[StateTree] updated test types to use InstanceDataType typedef and fixed one of them.
#rnx

diff --git a/Engine/Plugins/Runtime/StateTree/Source/StateTreeTestSuite/Private/StateTreeTestTypes.h b/Engine/Plugins/Runtime/StateTree/Source/StateTreeTestSuite/Private/StateTreeTestTypes.h
--- a/Engine/Plugins/Runtime/StateTree/Source/StateTreeTestSuite/Private/StateTreeTestTypes.h
+++ b/Engine/Plugins/Runtime/StateTree/Source/StateTreeTestSuite/Private/StateTreeTestTypes.h
@@ -114,13 +114,13 @@ struct FTestEval_A : public FStateTreeEvaluatorBase
 	FTestEval_A() = default;
 	virtual ~FTestEval_A() override {}
 
-	virtual const UStruct* GetInstanceDataType() const override { return FTestEval_AInstanceData::StaticStruct(); }
+	virtual const UStruct* GetInstanceDataType() const override { return InstanceDataType::StaticStruct(); }
 	
 	virtual bool Link(FStateTreeLinker& Linker) override
 	{
-		Linker.LinkInstanceDataProperty(FloatAHandle, STATETREE_INSTANCEDATA_PROPERTY(FTestEval_AInstanceData, FloatA));
-		Linker.LinkInstanceDataProperty(IntAHandle, STATETREE_INSTANCEDATA_PROPERTY(FTestEval_AInstanceData, IntA));
-		Linker.LinkInstanceDataProperty(BoolAHandle, STATETREE_INSTANCEDATA_PROPERTY(FTestEval_AInstanceData, bBoolA));
+		Linker.LinkInstanceDataProperty(FloatAHandle, STATETREE_INSTANCEDATA_PROPERTY(InstanceDataType, FloatA));
+		Linker.LinkInstanceDataProperty(IntAHandle, STATETREE_INSTANCEDATA_PROPERTY(InstanceDataType, IntA));
+		Linker.LinkInstanceDataProperty(BoolAHandle, STATETREE_INSTANCEDATA_PROPERTY(InstanceDataType, bBoolA));
 		return true;
 	}
 
@@ -154,13 +154,13 @@ struct FTestTask_B : public FStateTreeTaskBase
 	FTestTask_B() = default;
 	virtual ~FTestTask_B() override {}
 	
-	virtual const UStruct* GetInstanceDataType() const override { return FTestTask_BInstanceData::StaticStruct(); }
+	virtual const UStruct* GetInstanceDataType() const override { return InstanceDataType::StaticStruct(); }
 	
 	virtual bool Link(FStateTreeLinker& Linker) override
 	{
-		Linker.LinkInstanceDataProperty(FloatBHandle, STATETREE_INSTANCEDATA_PROPERTY(FTestTask_BInstanceData, FloatB));
-		Linker.LinkInstanceDataProperty(IntBHandle, STATETREE_INSTANCEDATA_PROPERTY(FTestTask_BInstanceData, IntB));
-		Linker.LinkInstanceDataProperty(BoolBHandle, STATETREE_INSTANCEDATA_PROPERTY(FTestTask_BInstanceData, bBoolB));
+		Linker.LinkInstanceDataProperty(FloatBHandle, STATETREE_INSTANCEDATA_PROPERTY(InstanceDataType, FloatB));
+		Linker.LinkInstanceDataProperty(IntBHandle, STATETREE_INSTANCEDATA_PROPERTY(InstanceDataType, IntB));
+		Linker.LinkInstanceDataProperty(BoolBHandle, STATETREE_INSTANCEDATA_PROPERTY(InstanceDataType, bBoolB));
 		return true;
 	}
 
@@ -268,12 +268,12 @@ struct FTestEval_Wander : public FStateTreeEvaluatorBase
 	FTestEval_Wander(const FName InName) { Name = InName; }
 	virtual ~FTestEval_Wander() override {}
 
-	virtual const UStruct* GetInstanceDataType() const override { return FTestEval_WanderInstanceData::StaticStruct(); }
+	virtual const UStruct* GetInstanceDataType() const override { return InstanceDataType::StaticStruct(); }
 	
 	virtual bool Link(FStateTreeLinker& Linker) override
 	{
-		Linker.LinkInstanceDataProperty(HasWanderLocationHandle, STATETREE_INSTANCEDATA_PROPERTY(FTestEval_WanderInstanceData, bHasWanderLocation));
-		Linker.LinkInstanceDataProperty(NextMoveLocationHandle, STATETREE_INSTANCEDATA_PROPERTY(FTestEval_WanderInstanceData, NextMoveLocation));
+		Linker.LinkInstanceDataProperty(HasWanderLocationHandle, STATETREE_INSTANCEDATA_PROPERTY(InstanceDataType, bHasWanderLocation));
+		Linker.LinkInstanceDataProperty(NextMoveLocationHandle, STATETREE_INSTANCEDATA_PROPERTY(InstanceDataType, NextMoveLocation));
 		return true;
 	}
 
@@ -334,13 +334,13 @@ struct FTestEval_SmartObjectSensor : public FStateTreeEvaluatorBase
 	FTestEval_SmartObjectSensor(const FName InName) { Name = InName; }
 	virtual ~FTestEval_SmartObjectSensor() {}
 
-	virtual const UStruct* GetInstanceDataType() const override { return FTestEval_SmartObjectSensorInstanceData::StaticStruct(); }
+	virtual const UStruct* GetInstanceDataType() const override { return InstanceDataType::StaticStruct(); }
 	
 	virtual bool Link(FStateTreeLinker& Linker) override
 	{
-		Linker.LinkInstanceDataProperty(PotentialSmartObjectsHandle, STATETREE_INSTANCEDATA_PROPERTY(FTestEval_SmartObjectSensorInstanceData, PotentialSmartObjects));
-		Linker.LinkInstanceDataProperty(HasSmartObjectsHandle, STATETREE_INSTANCEDATA_PROPERTY(FTestEval_SmartObjectSensorInstanceData, bHasSmartObjects));
-		Linker.LinkInstanceDataProperty(IsQueryingSmartObjectsHandle, STATETREE_INSTANCEDATA_PROPERTY(FTestEval_SmartObjectSensorInstanceData, bIsQueryingSmartObjects));
+		Linker.LinkInstanceDataProperty(PotentialSmartObjectsHandle, STATETREE_INSTANCEDATA_PROPERTY(InstanceDataType, PotentialSmartObjects));
+		Linker.LinkInstanceDataProperty(HasSmartObjectsHandle, STATETREE_INSTANCEDATA_PROPERTY(InstanceDataType, bHasSmartObjects));
+		Linker.LinkInstanceDataProperty(IsQueryingSmartObjectsHandle, STATETREE_INSTANCEDATA_PROPERTY(InstanceDataType, bIsQueryingSmartObjects));
 		return true;
 	}
 
@@ -438,13 +438,13 @@ struct FTestTask_ReserveSmartObject : public FStateTreeTaskBase
 	FTestTask_ReserveSmartObject(const FName InName) { Name = InName; }
 	virtual ~FTestTask_ReserveSmartObject() override {}
 
-	virtual const UStruct* GetInstanceDataType() const override { return FTestTask_ReserveSmartObjectInstanceData::StaticStruct(); }
+	virtual const UStruct* GetInstanceDataType() const override { return InstanceDataType::StaticStruct(); }
 	
 	virtual bool Link(FStateTreeLinker& Linker) override
 	{
-		Linker.LinkInstanceDataProperty(PotentialSmartObjectsHandle, STATETREE_INSTANCEDATA_PROPERTY(FTestTask_ReserveSmartObjectInstanceData, PotentialSmartObjects));
-		Linker.LinkInstanceDataProperty(ReservedSmartObjectLocationHandle, STATETREE_INSTANCEDATA_PROPERTY(FTestTask_ReserveSmartObjectInstanceData, ReservedSmartObjectLocation));
-		Linker.LinkInstanceDataProperty(ReservedSmartObjectHandle, STATETREE_INSTANCEDATA_PROPERTY(FTestTask_ReserveSmartObjectInstanceData, ReservedSmartObject));
+		Linker.LinkInstanceDataProperty(PotentialSmartObjectsHandle, STATETREE_INSTANCEDATA_PROPERTY(InstanceDataType, PotentialSmartObjects));
+		Linker.LinkInstanceDataProperty(ReservedSmartObjectLocationHandle, STATETREE_INSTANCEDATA_PROPERTY(InstanceDataType, ReservedSmartObjectLocation));
+		Linker.LinkInstanceDataProperty(ReservedSmartObjectHandle, STATETREE_INSTANCEDATA_PROPERTY(InstanceDataType, ReservedSmartObject));
 		return true;
 	}
 
@@ -583,12 +583,12 @@ struct FTestTask_MoveTo : public FStateTreeTaskBase
 	FTestTask_MoveTo(const FName InName) { Name = InName; }
 	virtual ~FTestTask_MoveTo() {}
 
-	virtual const UStruct* GetInstanceDataType() const override { return FTestTask_MoveToInstanceData::StaticStruct(); }
+	virtual const UStruct* GetInstanceDataType() const override { return InstanceDataType::StaticStruct(); }
 	
 	virtual bool Link(FStateTreeLinker& Linker) override
 	{
-		Linker.LinkInstanceDataProperty(MoveLocationHandle, STATETREE_INSTANCEDATA_PROPERTY(FTestTask_MoveToInstanceData, MoveLocation));
-		Linker.LinkInstanceDataProperty(CurrentTickHandle, STATETREE_INSTANCEDATA_PROPERTY(FTestTask_MoveToInstanceData, CurrentTick));
+		Linker.LinkInstanceDataProperty(MoveLocationHandle, STATETREE_INSTANCEDATA_PROPERTY(InstanceDataType, MoveLocation));
+		Linker.LinkInstanceDataProperty(CurrentTickHandle, STATETREE_INSTANCEDATA_PROPERTY(InstanceDataType, CurrentTick));
 		return true;
 	}
 
@@ -676,11 +676,11 @@ struct FTestTask_Stand : public FStateTreeTaskBase
 	FTestTask_Stand(const FName InName) { Name = InName; }
 	virtual ~FTestTask_Stand() {}
 
-	virtual const UStruct* GetInstanceDataType() const { return FTestTask_StandInstanceData::StaticStruct(); }
+	virtual const UStruct* GetInstanceDataType() const { return InstanceDataType::StaticStruct(); }
 
 	virtual bool Link(FStateTreeLinker& Linker) override
 	{
-		Linker.LinkInstanceDataProperty(CurrentTickHandle, STATETREE_INSTANCEDATA_PROPERTY(FTestTask_StandInstanceData, CurrentTick));
+		Linker.LinkInstanceDataProperty(CurrentTickHandle, STATETREE_INSTANCEDATA_PROPERTY(InstanceDataType, CurrentTick));
 		return true;
 	}
 	
@@ -752,18 +752,18 @@ struct FTestTask_UseSmartObject : public FStateTreeTaskBase
 {
 	GENERATED_BODY()
 
-	typedef FTestEval_WanderInstanceData InstanceDataType;
+	typedef FTestTask_UseSmartObjectInstanceData InstanceDataType;
 	
 	FTestTask_UseSmartObject() = default;
 	FTestTask_UseSmartObject(const FName InName) { Name = InName; }
 	virtual ~FTestTask_UseSmartObject() {}
 
-	virtual const UStruct* GetInstanceDataType()const override { return FTestTask_UseSmartObjectInstanceData::StaticStruct(); }
+	virtual const UStruct* GetInstanceDataType()const override { return InstanceDataType::StaticStruct(); }
 
 	virtual bool Link(FStateTreeLinker& Linker) override
 	{
-		Linker.LinkInstanceDataProperty(CurrentTickHandle, STATETREE_INSTANCEDATA_PROPERTY(FTestTask_UseSmartObjectInstanceData, CurrentTick));
-		Linker.LinkInstanceDataProperty(SmartObjectHandle, STATETREE_INSTANCEDATA_PROPERTY(FTestTask_UseSmartObjectInstanceData, SmartObject));
+		Linker.LinkInstanceDataProperty(CurrentTickHandle, STATETREE_INSTANCEDATA_PROPERTY(InstanceDataType, CurrentTick));
+		Linker.LinkInstanceDataProperty(SmartObjectHandle, STATETREE_INSTANCEDATA_PROPERTY(InstanceDataType, SmartObject));
 
 		return true;
 	}
