StateTree: Commandlet to compile all StateTree assets
- Added StateTreeCompileAllCommandlet
- Fixed StateTree editor compile button indication

diff --git a/Engine/Plugins/Runtime/StateTree/Source/StateTreeEditorModule/Private/Commandlets/StateTreeCompileAllCommandlet.cpp b/Engine/Plugins/Runtime/StateTree/Source/StateTreeEditorModule/Private/Commandlets/StateTreeCompileAllCommandlet.cpp
new file mode 100644
--- /dev/null
+++ b/Engine/Plugins/Runtime/StateTree/Source/StateTreeEditorModule/Private/Commandlets/StateTreeCompileAllCommandlet.cpp
@@ -0,0 +1,142 @@
+﻿// Copyright Epic Games, Inc. All Rights Reserved.
+
+#include "Commandlets/StateTreeCompileAllCommandlet.h"
+#include "StateTree.h"
+#include "PackageHelperFunctions.h"
+#include "AssetRegistry/AssetRegistryModule.h"
+#include "ISourceControlOperation.h"
+#include "SourceControlOperations.h"
+#include "SourceControlHelpers.h"
+#include "StateTreeEditor.h"
+#include "StateTreeCompiler.h"
+#include "StateTreeDelegates.h"
+
+DEFINE_LOG_CATEGORY_STATIC(LogStateTreeCompile, Log, Log);
+
+UStateTreeCompileAllCommandlet::UStateTreeCompileAllCommandlet(const FObjectInitializer& ObjectInitializer)
+	: Super(ObjectInitializer)
+{
+}
+
+int32 UStateTreeCompileAllCommandlet::Main(const FString& Params)
+{
+	// Parse command line.
+	TArray<FString> Tokens;
+	TArray<FString> Switches;
+
+	// want everything in upper case, it's a mess otherwise
+	const FString ParamsUpperCase = Params.ToUpper();
+	const TCHAR* Parms = *ParamsUpperCase;
+	ParseCommandLine(Parms, Tokens, Switches);
+
+	// Source control
+	bool bNoSourceControl = Switches.Contains(TEXT("nosourcecontrol"));
+	FScopedSourceControl SourceControl;
+	SourceControlProvider = bNoSourceControl ? nullptr : &SourceControl.GetProvider();
+
+	// Load assets
+	UE_LOG(LogStateTreeCompile, Display, TEXT("Loading Asset Registry..."));
+	const FAssetRegistryModule& AssetRegistryModule = FModuleManager::LoadModuleChecked<FAssetRegistryModule>(AssetRegistryConstants::ModuleName);
+	AssetRegistryModule.Get().SearchAllAssets(/*bSynchronousSearch =*/true);
+	UE_LOG(LogStateTreeCompile, Display, TEXT("Finished Loading Asset Registry."));
+	
+	UE_LOG(LogStateTreeCompile, Display, TEXT("Gathering All StateTrees From Asset Registry..."));
+	TArray<FAssetData> StateTreeAssetList;
+	AssetRegistryModule.Get().GetAssetsByClass(UStateTree::StaticClass()->GetFName(), StateTreeAssetList, /*bSearchSubClasses*/false);
+
+	int32 Counter = 0;
+	for (const FAssetData& Asset : StateTreeAssetList)
+	{
+		const FString ObjectPath = Asset.ObjectPath.ToString();
+		UE_LOG(LogStateTreeCompile, Display, TEXT("Loading and Compiling: '%s' [%d/%d]..."), *ObjectPath, Counter+1, StateTreeAssetList.Num());
+
+		UStateTree* StateTree = Cast<UStateTree>(StaticLoadObject(Asset.GetClass(), /*Outer*/nullptr, *ObjectPath, /*FileName*/nullptr, LOAD_NoWarn));
+		if (StateTree == nullptr)
+		{
+			UE_LOG(LogStateTreeCompile, Error, TEXT("Failed to Load: '%s'."), *ObjectPath);
+		}
+		else
+		{
+			CompileAndSaveStateTree(*StateTree);
+		}
+		Counter++;
+	}
+		
+	return 0;
+}
+
+bool UStateTreeCompileAllCommandlet::CompileAndSaveStateTree(UStateTree& StateTree) const
+{
+	UPackage* Package = StateTree.GetPackage();
+	const FString PackageFileName = SourceControlHelpers::PackageFilename(Package);
+
+	// Compile the StateTree asset.
+	UE::StateTree::Editor::ValidateAsset(StateTree);
+	const uint32 EditorDataHash = UE::StateTree::Editor::CalcAssetHash(StateTree);
+
+	FStateTreeCompilerLog Log;
+	FStateTreeCompiler Compiler(Log);
+
+	const bool bSuccess = Compiler.Compile(StateTree);
+
+	if (bSuccess)
+	{
+		// Success
+		StateTree.LastCompiledEditorDataHash = EditorDataHash;
+		UE::StateTree::Delegates::OnPostCompile.Broadcast(StateTree);
+		UE_LOG(LogStateTreeCompile, Log, TEXT("Compile StateTree %s succeeded."), *PackageFileName, EditorDataHash);
+	}
+	else
+	{
+		// Make sure not to leave stale data on failed compile.
+		StateTree.ResetCompiled();
+		StateTree.LastCompiledEditorDataHash = 0;
+
+		UE_LOG(LogStateTreeCompile, Error, TEXT("Failed to compile StateTree %s, errors follow."), *PackageFileName);
+		Log.DumpToLog(LogStateTreeCompile);
+
+		return false;
+	}
+
+	// Check out the StateTree asset
+	if (SourceControlProvider != nullptr)
+	{
+		const FSourceControlStatePtr SourceControlState = SourceControlProvider->GetState(PackageFileName, EStateCacheUsage::ForceUpdate);
+
+		if (SourceControlState.IsValid())
+		{
+			FString OtherCheckedOutUser;
+			if (SourceControlState->IsCheckedOutOther(&OtherCheckedOutUser))
+			{
+				UE_LOG(LogStateTreeCompile, Error, TEXT("Overwriting package %s already checked out by %s, will not submit"), *PackageFileName, *OtherCheckedOutUser);
+				return false;
+			}
+			else if (!SourceControlState->IsCurrent())
+			{
+				UE_LOG(LogStateTreeCompile, Error, TEXT("Overwriting package %s (not at head revision), will not submit"), *PackageFileName);
+				return false;
+			}
+			else if (SourceControlState->IsCheckedOut() || SourceControlState->IsAdded())
+			{
+				UE_LOG(LogStateTreeCompile, Log, TEXT("Package %s already checked out"), *PackageFileName);
+				return true;
+			}
+			else if (SourceControlState->IsSourceControlled())
+			{
+				UE_LOG(LogStateTreeCompile, Log, TEXT("Checking out package %s from source control"), *PackageFileName);
+				return SourceControlProvider->Execute(ISourceControlOperation::Create<FCheckOut>(), PackageFileName) == ECommandResult::Succeeded;
+			}
+		}
+	}
+
+	// Save StateTree asset.
+	if (!SavePackageHelper(Package, PackageFileName))
+	{
+		UE_LOG(LogStateTreeCompile, Error, TEXT("Failed to save %s."), *PackageFileName);
+		return false;
+	}
+
+	UE_LOG(LogStateTreeCompile, Log, TEXT("Compile and save %s succeeded."), *PackageFileName);
+
+	return true;
+}
diff --git a/Engine/Plugins/Runtime/StateTree/Source/StateTreeEditorModule/Private/Commandlets/StateTreeCompileAllCommandlet.h b/Engine/Plugins/Runtime/StateTree/Source/StateTreeEditorModule/Private/Commandlets/StateTreeCompileAllCommandlet.h
new file mode 100644
--- /dev/null
+++ b/Engine/Plugins/Runtime/StateTree/Source/StateTreeEditorModule/Private/Commandlets/StateTreeCompileAllCommandlet.h
@@ -0,0 +1,31 @@
+﻿// Copyright Epic Games, Inc. All Rights Reserved.
+
+#pragma once
+
+#include "Commandlets/Commandlet.h"
+#include "ISourceControlProvider.h"
+#include "StateTreeCompileAllCommandlet.generated.h"
+
+class UStateTree;
+
+/**
+ * Commandlet to recompile all StateTree assets in the project
+ */
+UCLASS()
+class UStateTreeCompileAllCommandlet : public UCommandlet
+{
+	GENERATED_BODY()
+public:
+	UStateTreeCompileAllCommandlet(const FObjectInitializer& ObjectInitializer);
+
+	//~ Begin UCommandlet Interface
+	virtual int32 Main(const FString& Params) override;
+	//~ End UCommandlet Interface
+
+protected:
+	bool CompileAndSaveStateTree(UStateTree& StateTree) const;
+
+	ISourceControlProvider* SourceControlProvider = nullptr;
+};
+
+
diff --git a/Engine/Plugins/Runtime/StateTree/Source/StateTreeEditorModule/Private/StateTreeCompilerLog.cpp b/Engine/Plugins/Runtime/StateTree/Source/StateTreeEditorModule/Private/StateTreeCompilerLog.cpp
--- a/Engine/Plugins/Runtime/StateTree/Source/StateTreeEditorModule/Private/StateTreeCompilerLog.cpp
+++ b/Engine/Plugins/Runtime/StateTree/Source/StateTreeEditorModule/Private/StateTreeCompilerLog.cpp
@@ -1,9 +1,9 @@
 ﻿// Copyright Epic Games, Inc. All Rights Reserved.
 
 #include "StateTreeCompilerLog.h"
-
 #include "StateTreeState.h"
 #include "Developer/MessageLog/Public/IMessageLogListing.h"
+#include "Logging/LogCategory.h"
 #include "Misc/UObjectToken.h"
 
 #define LOCTEXT_NAMESPACE "StateTreeEditor"
@@ -33,5 +33,41 @@ void FStateTreeCompilerLog::AppendToLog(IMessageLogListing* LogListing) const
 	}
 }
 
+void FStateTreeCompilerLog::DumpToLog(const FLogCategoryBase& Category) const
+{
+	for (const FStateTreeCompilerLogMessage& StateTreeMessage : Messages)
+	{
+		FString Message;
+		
+		if (StateTreeMessage.State != nullptr)
+		{
+			Message += FString::Printf(TEXT("State '%s': "), *StateTreeMessage.State->Name.ToString());
+		}
+
+		if (StateTreeMessage.Item.ID.IsValid())
+		{
+			Message += FString::Printf(TEXT("%s '%s': "), *UEnum::GetDisplayValueAsText(StateTreeMessage.Item.DataSource).ToString(), *StateTreeMessage.Item.Name.ToString());
+		}
+
+		Message += StateTreeMessage.Message;
+
+		switch (StateTreeMessage.Severity)
+		{
+		case EMessageSeverity::CriticalError:
+		case EMessageSeverity::Error:
+			UE_LOG_REF(Category, Error, TEXT("%s"), *Message);
+			break;
+		case EMessageSeverity::PerformanceWarning:
+		case EMessageSeverity::Warning:
+			UE_LOG_REF(Category, Warning, TEXT("%s"), *Message);
+			break;
+		case EMessageSeverity::Info:
+		default:
+			UE_LOG_REF(Category, Log, TEXT("%s"), *Message);
+			break;
+		};
+	}
+}
+
 
 #undef LOCTEXT_NAMESPACE
diff --git a/Engine/Plugins/Runtime/StateTree/Source/StateTreeEditorModule/Private/StateTreeEditor.cpp b/Engine/Plugins/Runtime/StateTree/Source/StateTreeEditorModule/Private/StateTreeEditor.cpp
--- a/Engine/Plugins/Runtime/StateTree/Source/StateTreeEditorModule/Private/StateTreeEditor.cpp
+++ b/Engine/Plugins/Runtime/StateTree/Source/StateTreeEditorModule/Private/StateTreeEditor.cpp
@@ -21,7 +21,7 @@
 #include "StateTreeEditorCommands.h"
 #include "ToolMenus.h"
 #include "StateTreeObjectHash.h"
-
+#include "StateTreeToolMenuContext.h"
 #include "Developer/MessageLog/Public/IMessageLogListing.h"
 #include "Developer/MessageLog/Public/MessageLogInitializationOptions.h"
 #include "Developer/MessageLog/Public/MessageLogModule.h"
@@ -106,8 +106,11 @@ void FStateTreeEditor::InitEditor( const EToolkitMode::Type Mode, const TSharedP
 		EditorData = NewObject<UStateTreeEditorData>(StateTree, FName(), RF_Transactional);
 		EditorData->AddRootState();
 		StateTree->EditorData = EditorData;
+		Compile();
 	}
 
+	EditorDataHash = UE::StateTree::Editor::CalcAssetHash(*StateTree);
+	
 	// @todo: Temporary fix
 	// Make sure all states are transactional
 	for (UStateTreeState* SubTree : EditorData->SubTrees)
@@ -239,6 +242,13 @@ FLinearColor FStateTreeEditor::GetWorldCentricTabColorScale() const
 	return FLinearColor( 0.0f, 0.0f, 0.2f, 0.5f );
 }
 
+void FStateTreeEditor::InitToolMenuContext(FToolMenuContext& MenuContext)
+{
+	UStateTreeToolMenuContext* Context = NewObject<UStateTreeToolMenuContext>();
+	Context->StateTreeEditor = SharedThis(this);
+	MenuContext.AddObject(Context);
+}
+
 void FStateTreeEditor::HandleMessageTokenClicked(const TSharedRef<IMessageToken>& InMessageToken)
 {
 	if (InMessageToken->GetType() == EMessageToken::Object)
@@ -652,6 +662,31 @@ namespace UE::StateTree::Editor::Internal
 
 }
 
+namespace UE::StateTree::Editor
+{
+	void ValidateAsset(UStateTree& StateTree)
+	{
+		UE::StateTree::Editor::Internal::UpdateParents(StateTree);
+		UE::StateTree::Editor::Internal::ApplySchema(StateTree);
+		UE::StateTree::Editor::Internal::RemoveUnusedBindings(StateTree);
+		UE::StateTree::Editor::Internal::ValidateLinkedStates(StateTree);
+		UE::StateTree::Editor::Internal::UpdateLinkedStateParameters(StateTree);
+	}
+
+	uint32 CalcAssetHash(const UStateTree& StateTree)
+	{
+		uint32 EditorDataHash = 0;
+		if (StateTree.EditorData != nullptr)
+		{
+			FStateTreeObjectCRC32 Archive;
+			EditorDataHash = Archive.Crc32(StateTree.EditorData, 0);
+		}
+
+		return EditorDataHash;
+	}
+
+};
+
 void FStateTreeEditor::BindCommands()
 {
 	const FStateTreeEditorCommands& Commands = FStateTreeEditorCommands::Get();
@@ -677,15 +712,31 @@ void FStateTreeEditor::RegisterToolbar()
 		ToolBar = UToolMenus::Get()->RegisterMenu(MenuName, ParentName, EMultiBoxType::ToolBar);
 	}
 
-	const FStateTreeEditorCommands& Commands = FStateTreeEditorCommands::Get();
 	const FToolMenuInsert InsertAfterAssetSection("Asset", EToolMenuInsertType::After);
 
 	FToolMenuSection& Section = ToolBar->AddSection("Compile", TAttribute<FText>(), InsertAfterAssetSection);
-	Section.AddEntry(FToolMenuEntry::InitToolBarButton(
-		Commands.Compile,
-		TAttribute<FText>(),
-		TAttribute<FText>(),
-		TAttribute<FSlateIcon>(this, &FStateTreeEditor::GetCompileStatusImage)));
+
+	Section.AddDynamicEntry("CompileCommands", FNewToolMenuSectionDelegate::CreateLambda([](FToolMenuSection& InSection)
+	{
+		const UStateTreeToolMenuContext* Context = InSection.FindContext<UStateTreeToolMenuContext>();
+		if (Context && Context->StateTreeEditor.IsValid())
+		{
+			TSharedPtr<FStateTreeEditor> StateTreeEditor = Context->StateTreeEditor.Pin();
+			if (StateTreeEditor.IsValid())
+			{
+				const FStateTreeEditorCommands& Commands = FStateTreeEditorCommands::Get();
+
+				FToolMenuEntry CompileButton = FToolMenuEntry::InitToolBarButton(
+					Commands.Compile,
+					TAttribute<FText>(),
+					TAttribute<FText>(),
+					TAttribute<FSlateIcon>(StateTreeEditor.ToSharedRef(), &FStateTreeEditor::GetCompileStatusImage)
+					);
+
+				InSection.AddEntry(CompileButton);
+			}
+		}
+	}));
 }
 
 void FStateTreeEditor::Compile()
@@ -695,6 +746,8 @@ void FStateTreeEditor::Compile()
 		return;
 	}
 
+	// Note: If the compilation process changes, also update UStateTreeCompileAllCommandlet.
+	
 	UpdateAsset();
 	
 	if (CompilerResultsListing.IsValid())
@@ -705,14 +758,14 @@ void FStateTreeEditor::Compile()
 	FStateTreeCompilerLog Log;
 	FStateTreeCompiler Compiler(Log);
 
-	const bool bSuccess = Compiler.Compile(*StateTree);
+	bLastCompileSucceeded = Compiler.Compile(*StateTree);
 
 	if (CompilerResultsListing.IsValid())
 	{
 		Log.AppendToLog(CompilerResultsListing.Get());
 	}
 
-	if (bSuccess)
+	if (bLastCompileSucceeded)
 	{
 		// Success
 		StateTree->LastCompiledEditorDataHash = EditorDataHash;
@@ -754,12 +807,14 @@ FSlateIcon FStateTreeEditor::GetCompileStatusImage() const
 	static const FName CompileStatusGood("Blueprint.CompileStatus.Overlay.Good");
 	static const FName CompileStatusWarning("Blueprint.CompileStatus.Overlay.Warning");
 
+	const bool bCompiledDataResetDuringLoad = StateTree->LastCompiledEditorDataHash == EditorDataHash && !StateTree->IsReadyToRun(); 
+	
 	if (StateTree == nullptr)
 	{
 		return FSlateIcon(FAppStyle::GetAppStyleSetName(), CompileStatusBackground, NAME_None, CompileStatusUnknown);
 	}
 	
-	if (!StateTree->IsReadyToRun())
+	if (!bLastCompileSucceeded || bCompiledDataResetDuringLoad)
 	{
 		return FSlateIcon(FAppStyle::GetAppStyleSetName(), CompileStatusBackground, NAME_None, CompileStatusError);
 	}
@@ -779,18 +834,8 @@ void FStateTreeEditor::UpdateAsset()
 		return;
 	}
 
-	UE::StateTree::Editor::Internal::UpdateParents(*StateTree);
-	UE::StateTree::Editor::Internal::ApplySchema(*StateTree);
-	UE::StateTree::Editor::Internal::RemoveUnusedBindings(*StateTree);
-	UE::StateTree::Editor::Internal::ValidateLinkedStates(*StateTree);
-	UE::StateTree::Editor::Internal::UpdateLinkedStateParameters(*StateTree);
-
-	EditorDataHash = 0;
-	if (StateTree->EditorData != nullptr)
-	{
-		FStateTreeObjectCRC32 Archive;
-		EditorDataHash = Archive.Crc32(StateTree->EditorData, 0);
-	}
+	UE::StateTree::Editor::ValidateAsset(*StateTree);
+	EditorDataHash = UE::StateTree::Editor::CalcAssetHash(*StateTree);
 }
 
 
diff --git a/Engine/Plugins/Runtime/StateTree/Source/StateTreeEditorModule/Private/StateTreeEditor.h b/Engine/Plugins/Runtime/StateTree/Source/StateTreeEditorModule/Private/StateTreeEditor.h
--- a/Engine/Plugins/Runtime/StateTree/Source/StateTreeEditorModule/Private/StateTreeEditor.h
+++ b/Engine/Plugins/Runtime/StateTree/Source/StateTreeEditorModule/Private/StateTreeEditor.h
@@ -13,6 +13,17 @@ class IDetailsView;
 class UStateTree;
 class FStateTreeViewModel;
 
+namespace UE::StateTree::Editor
+{
+	// Validates asset state
+	void ValidateAsset(UStateTree& StateTree);
+
+	// Calculates editor data hash of the asset. 
+	uint32 CalcAssetHash(const UStateTree& StateTree);
+ 
+} // UE::StateTree::Editor
+
+
 class FStateTreeEditor : public IStateTreeEditor, public FSelfRegisteringEditorUndoClient, public FGCObject
 {
 public:
@@ -27,6 +38,7 @@ public:
 	virtual FText GetBaseToolkitName() const override;
 	virtual FString GetWorldCentricTabPrefix() const override;
 	virtual FLinearColor GetWorldCentricTabColorScale() const override;
+	virtual void InitToolMenuContext(FToolMenuContext& MenuContext) override;
 	//~ End IToolkit Interface
 
 	void OnAssetFinishedChangingProperties(const FPropertyChangedEvent& PropertyChangedEvent);
@@ -80,6 +92,7 @@ private:
 	UStateTree* StateTree = nullptr;
 
 	uint32 EditorDataHash = 0;
+	bool bLastCompileSucceeded = true;
 	
 	/** Selection Property View */
 	TSharedPtr<class IDetailsView> SelectionDetailsView;
diff --git a/Engine/Plugins/Runtime/StateTree/Source/StateTreeEditorModule/Public/StateTreeCompilerLog.h b/Engine/Plugins/Runtime/StateTree/Source/StateTreeEditorModule/Public/StateTreeCompilerLog.h
--- a/Engine/Plugins/Runtime/StateTree/Source/StateTreeEditorModule/Public/StateTreeCompilerLog.h
+++ b/Engine/Plugins/Runtime/StateTree/Source/StateTreeEditorModule/Public/StateTreeCompilerLog.h
@@ -91,6 +91,9 @@ struct STATETREEEDITORMODULE_API FStateTreeCompilerLog
 
 	/** Appends StateTree log to log listing. */
 	void AppendToLog(class IMessageLogListing* LogListing) const;
+
+	/** Dumps StateTree log to log */
+	void DumpToLog(const FLogCategoryBase& Category) const;
 	
 protected:
 	UPROPERTY()
diff --git a/Engine/Plugins/Runtime/StateTree/Source/StateTreeEditorModule/Public/StateTreeToolMenuContext.h b/Engine/Plugins/Runtime/StateTree/Source/StateTreeEditorModule/Public/StateTreeToolMenuContext.h
new file mode 100644
--- /dev/null
+++ b/Engine/Plugins/Runtime/StateTree/Source/StateTreeEditorModule/Public/StateTreeToolMenuContext.h
@@ -0,0 +1,18 @@
+﻿// Copyright Epic Games, Inc. All Rights Reserved.
+
+#pragma once
+
+#include "UObject/Object.h"
+#include "Templates/SharedPointer.h"
+#include "StateTreeToolMenuContext.generated.h"
+
+class FStateTreeEditor;
+
+UCLASS()
+class STATETREEEDITORMODULE_API UStateTreeToolMenuContext : public UObject
+{
+	GENERATED_BODY()
+
+public:
+	TWeakPtr<FStateTreeEditor> StateTreeEditor;
+};
diff --git a/Engine/Plugins/Runtime/StateTree/Source/StateTreeEditorModule/StateTreeEditorModule.Build.cs b/Engine/Plugins/Runtime/StateTree/Source/StateTreeEditorModule/StateTreeEditorModule.Build.cs
--- a/Engine/Plugins/Runtime/StateTree/Source/StateTreeEditorModule/StateTreeEditorModule.Build.cs
+++ b/Engine/Plugins/Runtime/StateTree/Source/StateTreeEditorModule/StateTreeEditorModule.Build.cs
@@ -25,6 +25,7 @@ namespace UnrealBuildTool.Rules
 				"EditorStyle",
 				"PropertyEditor",
 				"StateTreeModule",
+				"SourceControl",
 				"Projects",
 				"BlueprintGraph",
 				"PropertyAccessEditor",
diff --git a/Engine/Plugins/Runtime/StateTree/Source/StateTreeModule/Private/StateTree.cpp b/Engine/Plugins/Runtime/StateTree/Source/StateTreeModule/Private/StateTree.cpp
--- a/Engine/Plugins/Runtime/StateTree/Source/StateTreeModule/Private/StateTree.cpp
+++ b/Engine/Plugins/Runtime/StateTree/Source/StateTreeModule/Private/StateTree.cpp
@@ -30,6 +30,7 @@ void UStateTree::ResetCompiled()
 
 	NumDataViews = 0;
 	ExternalDataBaseIndex = 0;
+	DefaultParametersDataViewIndex = INDEX_NONE;
 
 	InstanceDataDefaultValue.Reset();
 }
@@ -99,7 +100,7 @@ bool UStateTree::Link()
 	TArrayView<FStateTreeBindableStructDesc> SourceStructs = PropertyBindings.GetSourceStructs();
 	TArrayView<FStateTreePropCopyBatch> CopyBatches = PropertyBindings.GetCopyBatches();
 
-	if (DefaultParametersDataViewIndex != INDEX_NONE)
+	if (SourceStructs.IsValidIndex(DefaultParametersDataViewIndex))
 	{
 		SourceStructs[DefaultParametersDataViewIndex].Struct = Parameters.GetPropertyBagStruct();
 	}
