StateTree: moved evaluators global to tree
- Deprecated per state evaluators and moved them to global to the tree
- Tick evals once per Tick()
- Updated editor node customizations to work on UStateTreeEditorData
- Added separate detail customization for UStateTreeEditorData

diff --git a/Engine/Plugins/AI/MassAI/Source/MassAIBehavior/Private/Evaluators/MassComponentHitEvaluator.cpp b/Engine/Plugins/AI/MassAI/Source/MassAIBehavior/Private/Evaluators/MassComponentHitEvaluator.cpp
--- a/Engine/Plugins/AI/MassAI/Source/MassAIBehavior/Private/Evaluators/MassComponentHitEvaluator.cpp
+++ b/Engine/Plugins/AI/MassAI/Source/MassAIBehavior/Private/Evaluators/MassComponentHitEvaluator.cpp
@@ -18,7 +18,7 @@ bool FMassComponentHitEvaluator::Link(FStateTreeLinker& Linker)
 	return true;
 }
 
-void FMassComponentHitEvaluator::Evaluate(FStateTreeExecutionContext &Context, const EStateTreeEvaluationType EvalType, const float DeltaTime) const
+void FMassComponentHitEvaluator::Tick(FStateTreeExecutionContext &Context, const float DeltaTime) const
 {
 	// Look for recent hits
 	UMassComponentHitSubsystem& HitSubsystem = Context.GetExternalData(ComponentHitSubsystemHandle);
diff --git a/Engine/Plugins/AI/MassAI/Source/MassAIBehavior/Private/Evaluators/MassZoneGraphAnnotationEvaluator.cpp b/Engine/Plugins/AI/MassAI/Source/MassAIBehavior/Private/Evaluators/MassZoneGraphAnnotationEvaluator.cpp
--- a/Engine/Plugins/AI/MassAI/Source/MassAIBehavior/Private/Evaluators/MassZoneGraphAnnotationEvaluator.cpp
+++ b/Engine/Plugins/AI/MassAI/Source/MassAIBehavior/Private/Evaluators/MassZoneGraphAnnotationEvaluator.cpp
@@ -19,7 +19,7 @@ bool FMassZoneGraphAnnotationEvaluator::Link(FStateTreeLinker& Linker)
 	return true;
 }
 
-void FMassZoneGraphAnnotationEvaluator::Evaluate(FStateTreeExecutionContext& Context, const EStateTreeEvaluationType EvalType, const float DeltaTime) const
+void FMassZoneGraphAnnotationEvaluator::Tick(FStateTreeExecutionContext& Context, const float DeltaTime) const
 {
 	const FMassZoneGraphAnnotationFragment& AnnotationTagsFragment = Context.GetExternalData(AnnotationTagsFragmentHandle);
 	FZoneGraphTagMask& AnnotationTags = Context.GetInstanceData(AnnotationTagsHandle);
diff --git a/Engine/Plugins/AI/MassAI/Source/MassAIBehavior/Public/Evaluators/MassComponentHitEvaluator.h b/Engine/Plugins/AI/MassAI/Source/MassAIBehavior/Public/Evaluators/MassComponentHitEvaluator.h
--- a/Engine/Plugins/AI/MassAI/Source/MassAIBehavior/Public/Evaluators/MassComponentHitEvaluator.h
+++ b/Engine/Plugins/AI/MassAI/Source/MassAIBehavior/Public/Evaluators/MassComponentHitEvaluator.h
@@ -32,7 +32,7 @@ struct MASSAIBEHAVIOR_API FMassComponentHitEvaluator : public FMassStateTreeEval
 protected:
 	virtual bool Link(FStateTreeLinker& Linker) override;
 	virtual const UStruct* GetInstanceDataType() const override { return FMassComponentHitEvaluatorInstanceData::StaticStruct(); }
-	virtual void Evaluate(FStateTreeExecutionContext& Context, const EStateTreeEvaluationType EvalType, const float DeltaTime) const override;
+	virtual void Tick(FStateTreeExecutionContext& Context, const float DeltaTime) const override;
 
 	TStateTreeExternalDataHandle<UMassComponentHitSubsystem> ComponentHitSubsystemHandle;
 
diff --git a/Engine/Plugins/AI/MassAI/Source/MassAIBehavior/Public/Evaluators/MassZoneGraphAnnotationEvaluator.h b/Engine/Plugins/AI/MassAI/Source/MassAIBehavior/Public/Evaluators/MassZoneGraphAnnotationEvaluator.h
--- a/Engine/Plugins/AI/MassAI/Source/MassAIBehavior/Public/Evaluators/MassZoneGraphAnnotationEvaluator.h
+++ b/Engine/Plugins/AI/MassAI/Source/MassAIBehavior/Public/Evaluators/MassZoneGraphAnnotationEvaluator.h
@@ -30,7 +30,7 @@ struct MASSAIBEHAVIOR_API FMassZoneGraphAnnotationEvaluator : public FMassStateT
 protected:
 	virtual bool Link(FStateTreeLinker& Linker) override;
 	virtual const UStruct* GetInstanceDataType() const override { return FMassZoneGraphAnnotationEvaluatorInstanceData::StaticStruct(); }
-	virtual void Evaluate(FStateTreeExecutionContext& Context, const EStateTreeEvaluationType EvalType, const float DeltaTime) const override;
+	virtual void Tick(FStateTreeExecutionContext& Context, const float DeltaTime) const override;
 
 	TStateTreeExternalDataHandle<FMassZoneGraphAnnotationFragment> AnnotationTagsFragmentHandle;
 
diff --git a/Engine/Plugins/Runtime/StateTree/Source/StateTreeEditorModule/Private/Customizations/StateTreeEditorDataDetails.cpp b/Engine/Plugins/Runtime/StateTree/Source/StateTreeEditorModule/Private/Customizations/StateTreeEditorDataDetails.cpp
new file mode 100644
--- /dev/null
+++ b/Engine/Plugins/Runtime/StateTree/Source/StateTreeEditorModule/Private/Customizations/StateTreeEditorDataDetails.cpp
@@ -0,0 +1,79 @@
+﻿// Copyright Epic Games, Inc. All Rights Reserved.
+
+#include "StateTreeEditorDataDetails.h"
+#include "Widgets/DeclarativeSyntaxSupport.h"
+#include "IDetailChildrenBuilder.h"
+#include "DetailWidgetRow.h"
+#include "DetailLayoutBuilder.h"
+#include "DetailCategoryBuilder.h"
+#include "PropertyCustomizationHelpers.h"
+#include "IPropertyUtilities.h"
+#include "StateTreeEditorData.h"
+#include "StateTreeViewModel.h"
+#include "StateTree.h"
+
+
+#define LOCTEXT_NAMESPACE "StateTreeEditor"
+
+
+TSharedRef<IDetailCustomization> FStateTreeEditorDataDetails::MakeInstance()
+{
+	return MakeShareable(new FStateTreeEditorDataDetails);
+}
+
+void FStateTreeEditorDataDetails::CustomizeDetails(IDetailLayoutBuilder& DetailBuilder)
+{
+	// Find StateTreeEditorData associated with this panel.
+	const UStateTreeEditorData* EditorData = nullptr;
+	TArray<TWeakObjectPtr<UObject>> Objects;
+	DetailBuilder.GetObjectsBeingCustomized(Objects);
+	for (TWeakObjectPtr<UObject>& WeakObject : Objects)
+	{
+		if (UStateTreeEditorData* Object = Cast<UStateTreeEditorData>(WeakObject.Get()))
+		{
+			EditorData = Object;
+			break;
+		}
+	}
+	const UStateTreeSchema* Schema = EditorData ? EditorData->Schema : nullptr;
+
+	TSharedPtr<IPropertyHandle> EvaluatorsProperty = DetailBuilder.GetProperty(TEXT("Evaluators"));
+
+	IDetailCategoryBuilder& CommonCategory = DetailBuilder.EditCategory(TEXT("Common"), LOCTEXT("EditorDataDetailsCommon", "Common"));
+	CommonCategory.SetSortOrder(0);
+
+	const FName EvalCategoryName(TEXT("Evaluators"));
+	if (Schema && Schema->AllowEvaluators())
+	{
+		MakeArrayCategory(DetailBuilder, EvalCategoryName, LOCTEXT("EditorDataDetailsEvaluators", "Evaluators"), 1, EvaluatorsProperty);
+	}
+	else
+	{
+		DetailBuilder.EditCategory(EvalCategoryName).SetCategoryVisibility(false);
+	}
+}
+
+void FStateTreeEditorDataDetails::MakeArrayCategory(IDetailLayoutBuilder& DetailBuilder, FName CategoryName, const FText& DisplayName, int32 SortOrder, TSharedPtr<IPropertyHandle> PropertyHandle)
+{
+	IDetailCategoryBuilder& Category = DetailBuilder.EditCategory(CategoryName, DisplayName);
+	Category.SetSortOrder(SortOrder);
+
+	TSharedRef<SHorizontalBox> HeaderContentWidget = SNew(SHorizontalBox);
+	HeaderContentWidget->AddSlot()
+	.HAlign(HAlign_Right)
+	.VAlign(VAlign_Center)
+	[
+		PropertyHandle->CreateDefaultPropertyButtonWidgets()
+	];
+	Category.HeaderContent(HeaderContentWidget);
+
+	// Add items inline
+	TSharedRef<FDetailArrayBuilder> Builder = MakeShareable(new FDetailArrayBuilder(PropertyHandle.ToSharedRef(), /*InGenerateHeader*/ false, /*InDisplayResetToDefault*/ true, /*InDisplayElementNum*/ false));
+	Builder->OnGenerateArrayElementWidget(FOnGenerateArrayElementWidget::CreateLambda([](TSharedRef<IPropertyHandle> PropertyHandle, int32 ArrayIndex, IDetailChildrenBuilder& ChildrenBuilder)
+		{
+			ChildrenBuilder.AddProperty(PropertyHandle);
+		}));
+	Category.AddCustomBuilder(Builder, /*bForAdvanced*/ false);
+}
+
+#undef LOCTEXT_NAMESPACE
diff --git a/Engine/Plugins/Runtime/StateTree/Source/StateTreeEditorModule/Private/Customizations/StateTreeEditorDataDetails.h b/Engine/Plugins/Runtime/StateTree/Source/StateTreeEditorModule/Private/Customizations/StateTreeEditorDataDetails.h
new file mode 100644
--- /dev/null
+++ b/Engine/Plugins/Runtime/StateTree/Source/StateTreeEditorModule/Private/Customizations/StateTreeEditorDataDetails.h
@@ -0,0 +1,22 @@
+﻿// Copyright Epic Games, Inc. All Rights Reserved.
+
+#pragma once
+
+#include "IDetailCustomization.h"
+
+class IDetailLayoutBuilder;
+class IDetailChildrenBuilder;
+class IPropertyHandle;
+
+class FStateTreeEditorDataDetails : public IDetailCustomization
+{
+public:
+	/** Makes a new instance of this detail layout class for a specific detail view requesting it */
+	static TSharedRef<IDetailCustomization> MakeInstance();
+
+	/** IDetailCustomization interface */
+	virtual void CustomizeDetails(IDetailLayoutBuilder& DetailBuilder) override;
+
+protected:
+	void MakeArrayCategory(IDetailLayoutBuilder& DetailBuilder, FName CategoryName, const FText& DisplayName, int32 SortOrder, TSharedPtr<IPropertyHandle> PropertyHandle);
+};
diff --git a/Engine/Plugins/Runtime/StateTree/Source/StateTreeEditorModule/Private/Customizations/StateTreeEditorNodeDetails.cpp b/Engine/Plugins/Runtime/StateTree/Source/StateTreeEditorModule/Private/Customizations/StateTreeEditorNodeDetails.cpp
--- a/Engine/Plugins/Runtime/StateTree/Source/StateTreeEditorModule/Private/Customizations/StateTreeEditorNodeDetails.cpp
+++ b/Engine/Plugins/Runtime/StateTree/Source/StateTreeEditorModule/Private/Customizations/StateTreeEditorNodeDetails.cpp
@@ -500,7 +500,12 @@ void FStateTreeEditorNodeDetails::FindOuterObjects()
 	StructProperty->GetOuterObjects(OuterObjects);
 	for (UObject* Outer : OuterObjects)
 	{
-		UStateTreeEditorData* OuterEditorData = Outer->GetTypedOuter<UStateTreeEditorData>();
+		UStateTreeEditorData* OuterEditorData = Cast<UStateTreeEditorData>(Outer);
+		if (OuterEditorData == nullptr)
+		{
+			OuterEditorData = Outer->GetTypedOuter<UStateTreeEditorData>();
+		}
+		
 		UStateTree* OuterStateTree = OuterEditorData ? OuterEditorData->GetTypedOuter<UStateTree>() : nullptr;
 		if (OuterEditorData && OuterStateTree)
 		{
diff --git a/Engine/Plugins/Runtime/StateTree/Source/StateTreeEditorModule/Private/Customizations/StateTreeStateDetails.cpp b/Engine/Plugins/Runtime/StateTree/Source/StateTreeEditorModule/Private/Customizations/StateTreeStateDetails.cpp
--- a/Engine/Plugins/Runtime/StateTree/Source/StateTreeEditorModule/Private/Customizations/StateTreeStateDetails.cpp
+++ b/Engine/Plugins/Runtime/StateTree/Source/StateTreeEditorModule/Private/Customizations/StateTreeStateDetails.cpp
@@ -43,7 +43,6 @@ void FStateTreeStateDetails::CustomizeDetails(IDetailLayoutBuilder& DetailBuilde
 	TSharedPtr<IPropertyHandle> TasksProperty = DetailBuilder.GetProperty(TEXT("Tasks"));
 	TSharedPtr<IPropertyHandle> SingleTaskProperty = DetailBuilder.GetProperty(TEXT("SingleTask"));
 	TSharedPtr<IPropertyHandle> EnterConditionsProperty = DetailBuilder.GetProperty(TEXT("EnterConditions"));
-	TSharedPtr<IPropertyHandle> EvaluatorsProperty = DetailBuilder.GetProperty(TEXT("Evaluators"));
 	TSharedPtr<IPropertyHandle> TransitionsProperty = DetailBuilder.GetProperty(TEXT("Transitions"));
 	TSharedPtr<IPropertyHandle> TypeProperty = DetailBuilder.GetProperty(TEXT("Type"));
 	TSharedPtr<IPropertyHandle> LinkedStateProperty = DetailBuilder.GetProperty(TEXT("LinkedState"));
@@ -66,16 +65,6 @@ void FStateTreeStateDetails::CustomizeDetails(IDetailLayoutBuilder& DetailBuilde
 		ParametersProperty->MarkHiddenByCustomization();
 	}
 	
-	const FName EvalCategoryName(TEXT("Evaluators"));
-	if ((StateType == EStateTreeStateType::State || StateType == EStateTreeStateType::Subtree) && Schema && Schema->AllowEvaluators())
-	{
-		MakeArrayCategory(DetailBuilder, EvalCategoryName, LOCTEXT("StateDetailsEvaluators", "Evaluators"), 1, EvaluatorsProperty);
-	}
-	else
-	{
-		DetailBuilder.EditCategory(EvalCategoryName).SetCategoryVisibility(false);
-	}
-
 	const FName EnterConditionsCategoryName(TEXT("Enter Conditions"));
 	if (Schema && Schema->AllowEnterConditions())
 	{
@@ -141,13 +130,11 @@ void FStateTreeStateDetails::MakeArrayCategory(IDetailLayoutBuilder& DetailBuild
 
 	// Add items inline
 	TSharedRef<FDetailArrayBuilder> Builder = MakeShareable(new FDetailArrayBuilder(PropertyHandle.ToSharedRef(), /*InGenerateHeader*/ false, /*InDisplayResetToDefault*/ true, /*InDisplayElementNum*/ false));
-	Builder->OnGenerateArrayElementWidget(FOnGenerateArrayElementWidget::CreateSP(this, &FStateTreeStateDetails::GenerateArrayElementWidget));
+	Builder->OnGenerateArrayElementWidget(FOnGenerateArrayElementWidget::CreateLambda([](TSharedRef<IPropertyHandle> PropertyHandle, int32 ArrayIndex, IDetailChildrenBuilder& ChildrenBuilder)
+	{
+		ChildrenBuilder.AddProperty(PropertyHandle);
+	}));
 	Category.AddCustomBuilder(Builder, /*bForAdvanced*/ false);
 }
 
-void FStateTreeStateDetails::GenerateArrayElementWidget(TSharedRef<IPropertyHandle> PropertyHandle, int32 ArrayIndex, IDetailChildrenBuilder& ChildrenBuilder)
-{
-	ChildrenBuilder.AddProperty(PropertyHandle);
-};
-
 #undef LOCTEXT_NAMESPACE
diff --git a/Engine/Plugins/Runtime/StateTree/Source/StateTreeEditorModule/Private/SStateTreeView.cpp b/Engine/Plugins/Runtime/StateTree/Source/StateTreeEditorModule/Private/SStateTreeView.cpp
--- a/Engine/Plugins/Runtime/StateTree/Source/StateTreeEditorModule/Private/SStateTreeView.cpp
+++ b/Engine/Plugins/Runtime/StateTree/Source/StateTreeEditorModule/Private/SStateTreeView.cpp
@@ -259,12 +259,8 @@ void SStateTreeView::HandleModelStatesChanged(const TSet<UStateTreeState*>& Affe
 {
 	bool bArraysChanged = false;
 
-	// The purpose of the rebuild below is to update the evaluator and task visualization (number of widgets change).
+	// The purpose of the rebuild below is to update the task visualization (number of widgets change).
 	// This method is called when anything in a state changes, make sure to only rebuild when needed.
-	if (PropertyChangedEvent.MemberProperty->GetFName() == GET_MEMBER_NAME_CHECKED(UStateTreeState, Evaluators))
-	{
-		bArraysChanged = true;
-	}
 	if (PropertyChangedEvent.MemberProperty->GetFName() == GET_MEMBER_NAME_CHECKED(UStateTreeState, Tasks))
 	{
 		bArraysChanged = true;
diff --git a/Engine/Plugins/Runtime/StateTree/Source/StateTreeEditorModule/Private/SStateTreeViewRow.cpp b/Engine/Plugins/Runtime/StateTree/Source/StateTreeEditorModule/Private/SStateTreeViewRow.cpp
--- a/Engine/Plugins/Runtime/StateTree/Source/StateTreeEditorModule/Private/SStateTreeViewRow.cpp
+++ b/Engine/Plugins/Runtime/StateTree/Source/StateTreeEditorModule/Private/SStateTreeViewRow.cpp
@@ -24,7 +24,6 @@
 #include "StateTree.h"
 #include "StateTreeState.h"
 #include "StateTreeConditionBase.h"
-#include "StateTreeEvaluatorBase.h"
 #include "StateTreeTaskBase.h"
 #include "StateTreeViewModel.h"
 #include "Algo/ForEach.h"
@@ -57,7 +56,6 @@ void SStateTreeViewRow::Construct(const FArguments& InArgs, const TSharedRef<STa
     , InOwnerTableView);
 
 	static const FLinearColor TasksBackground = FLinearColor(FColor(17, 117, 131));
-	static const FLinearColor EvaluatorsBackground = FLinearColor(FColor(48, 48, 48));
 	static const FLinearColor LinkBackground = FLinearColor(FColor(84, 84, 84));
 	
 	this->ChildSlot
@@ -137,47 +135,6 @@ void SStateTreeViewRow::Construct(const FArguments& InArgs, const TSharedRef<STa
 			]
 		]
 
-		// Evaluators Box
-		+ SHorizontalBox::Slot()
-		.VAlign(VAlign_Center)
-		.Padding(FMargin(0.0f, 4.0f))
-		.AutoWidth()
-		[
-			SNew(SBox)
-			.HeightOverride(28.0f)
-			.VAlign(VAlign_Fill)
-			.Visibility(this, &SStateTreeViewRow::GetEvaluatorsVisibility)
-			[
-				SNew(SBorder)
-				.BorderImage(FAppStyle::GetBrush("WhiteBrush"))
-				.BorderBackgroundColor(EvaluatorsBackground)
-				.Padding(FMargin(12.0f, 0.0f, 16.0f, 0.0f))
-				[
-					// Evaluator icon
-					SNew(SHorizontalBox)
-					+ SHorizontalBox::Slot()
-					.VAlign(VAlign_Center)
-					.Padding(FMargin(0.0f, 0.0f, 4.0f, 0.0f))
-					.AutoWidth()
-					[
-						SNew(STextBlock)
-						.Text(FEditorFontGlyphs::Crosshairs)
-						.TextStyle(FStateTreeEditorStyle::Get(), "StateTree.DetailsIcon")
-					]
-
-					// Evalutors
-					+ SHorizontalBox::Slot()
-					.VAlign(VAlign_Center)
-					.AutoWidth()
-					[
-						SNew(STextBlock)
-						.Text(this, &SStateTreeViewRow::GetEvaluatorsDesc)
-						.TextStyle(FStateTreeEditorStyle::Get(), "StateTree.Details")
-					]
-				]
-			]
-		]
-
 		// Tasks Box
 		+ SHorizontalBox::Slot()
 		.VAlign(VAlign_Center)
@@ -491,44 +448,6 @@ FText SStateTreeViewRow::GetSelectorDesc() const
 	return FText::GetEmpty(); 
 }
 
-EVisibility SStateTreeViewRow::GetEvaluatorsVisibility() const
-{
-	if (const UStateTreeState* State = WeakState.Get())
-	{
-		int32 ValidCount = 0;
-		for (int32 i = 0; i < State->Evaluators.Num(); i++)
-		{
-			if (const FStateTreeEvaluatorBase* Eval = State->Evaluators[i].Node.GetPtr<FStateTreeEvaluatorBase>())
-			{
-				ValidCount++;
-			}
-		}
-		return ValidCount > 0 ? EVisibility::Visible : EVisibility::Collapsed;
-	}
-	return EVisibility::Collapsed;
-}
-
-FText SStateTreeViewRow::GetEvaluatorsDesc() const
-{
-	const UStateTreeState* State = WeakState.Get();
-	if (!State)
-	{
-		return FText::GetEmpty();
-	}
-
-	TArray<FText> Names;
-
-	for (int32 i = 0; i < State->Evaluators.Num(); i++)
-	{
-		if (const FStateTreeEvaluatorBase* Eval = State->Evaluators[i].Node.GetPtr<FStateTreeEvaluatorBase>())
-		{
-			Names.Add(FText::FromName(Eval->Name));
-		}
-	}
-
-	return FText::Join((FText::FromString(TEXT(", "))), Names);
-}
-
 EVisibility SStateTreeViewRow::GetTasksVisibility() const
 {
 	if (const UStateTreeState* State = WeakState.Get())
diff --git a/Engine/Plugins/Runtime/StateTree/Source/StateTreeEditorModule/Private/SStateTreeViewRow.h b/Engine/Plugins/Runtime/StateTree/Source/StateTreeEditorModule/Private/SStateTreeViewRow.h
--- a/Engine/Plugins/Runtime/StateTree/Source/StateTreeEditorModule/Private/SStateTreeViewRow.h
+++ b/Engine/Plugins/Runtime/StateTree/Source/StateTreeEditorModule/Private/SStateTreeViewRow.h
@@ -32,9 +32,6 @@ private:
 	EVisibility GetSelectorVisibility() const;
 	FText GetSelectorDesc() const;
 
-	EVisibility GetEvaluatorsVisibility() const;
-	FText GetEvaluatorsDesc() const;
-
 	EVisibility GetTasksVisibility() const;
 	FText GetTasksDesc() const;
 
diff --git a/Engine/Plugins/Runtime/StateTree/Source/StateTreeEditorModule/Private/StateTreeCompiler.cpp b/Engine/Plugins/Runtime/StateTree/Source/StateTreeEditorModule/Private/StateTreeCompiler.cpp
--- a/Engine/Plugins/Runtime/StateTree/Source/StateTreeEditorModule/Private/StateTreeCompiler.cpp
+++ b/Engine/Plugins/Runtime/StateTree/Source/StateTreeEditorModule/Private/StateTreeCompiler.cpp
@@ -56,7 +56,7 @@ bool FStateTreeCompiler::Compile(UStateTree& InStateTree)
 		return false;
 	}
 
-	if (!CreateStateEvaluators())
+	if (!CreateEvaluators())
 	{
 		StateTree->ResetCompiled();
 		return false;
@@ -208,6 +208,26 @@ bool FStateTreeCompiler::CreateConditions(UStateTreeState& State, TConstArrayVie
 	return true;
 }
 
+bool FStateTreeCompiler::CreateEvaluators()
+{
+	check(StateTree->Nodes.Num() <= int32(MAX_uint16));
+	StateTree->EvaluatorsBegin = uint16(StateTree->Nodes.Num());
+
+	for (FStateTreeEditorNode& EvalNode : TreeData->Evaluators)
+	{
+		if (!CreateEvaluator(EvalNode))
+		{
+			return false;
+		}
+	}
+	
+	const int32 EvaluatorsNum = StateTree->Nodes.Num() - int32(StateTree->EvaluatorsBegin);
+	check(EvaluatorsNum <= int32(MAX_uint16));
+	StateTree->EvaluatorsNum = uint16(EvaluatorsNum);
+
+	return true;
+}
+
 bool FStateTreeCompiler::CreateStateTasksAndParameters()
 {
 	for (int32 i = 0; i < StateTree->States.Num(); i++)
@@ -250,7 +270,7 @@ bool FStateTreeCompiler::CreateStateTasksAndParameters()
 
 				// Check that the bindings for this struct are still all valid.
 				TArray<FStateTreeEditorPropertyBinding> Bindings;
-				if (!GetAndValidateBindings(*SourceState, StructDesc, Bindings))
+				if (!GetAndValidateBindings(StructDesc, Bindings))
 				{
 					return false;
 				}
@@ -290,36 +310,6 @@ bool FStateTreeCompiler::CreateStateTasksAndParameters()
 	return true;
 }
 
-bool FStateTreeCompiler::CreateStateEvaluators()
-{
-	for (int32 i = 0; i < StateTree->States.Num(); i++)
-	{
-		FCompactStateTreeState& CompactState = StateTree->States[i];
-		UStateTreeState* SourceState = SourceStates[i];
-		check(SourceState != nullptr);
-
-		FStateTreeCompilerLogStateScope LogStateScope(SourceState, Log);
-		
-		// Collect evaluators
-		check(StateTree->Nodes.Num() <= int32(MAX_uint16));
-		CompactState.EvaluatorsBegin = uint16(StateTree->Nodes.Num());
-
-		for (FStateTreeEditorNode& EvalNode : SourceState->Evaluators)
-		{
-			if (!CreateEvaluator(*SourceState, EvalNode))
-			{
-				return false;
-			}
-		}
-		
-		const int32 EvaluatorsNum = StateTree->Nodes.Num() - int32(CompactState.EvaluatorsBegin);
-		check(EvaluatorsNum <= int32(MAX_uint8));
-		CompactState.EvaluatorsNum = uint8(EvaluatorsNum);
-	}
-
-	return true;
-}
-
 bool FStateTreeCompiler::CreateStateTransitions()
 {
 	for (int32 i = 0; i < StateTree->States.Num(); i++)
@@ -512,7 +502,7 @@ bool FStateTreeCompiler::CreateCondition(UStateTreeState& State, const FStateTre
 
 	// Check that the bindings for this struct are still all valid.
 	TArray<FStateTreeEditorPropertyBinding> Bindings;
-	if (!GetAndValidateBindings(State, StructDesc, Bindings))
+	if (!GetAndValidateBindings(StructDesc, Bindings))
 	{
 		return false;
 	}
@@ -600,7 +590,7 @@ bool FStateTreeCompiler::CreateTask(UStateTreeState& State, const FStateTreeEdit
 
 	// Check that the bindings for this struct are still all valid.
 	TArray<FStateTreeEditorPropertyBinding> Bindings;
-	if (!GetAndValidateBindings(State, StructDesc, Bindings))
+	if (!GetAndValidateBindings(StructDesc, Bindings))
 	{
 		return false;
 	}
@@ -621,7 +611,7 @@ bool FStateTreeCompiler::CreateTask(UStateTreeState& State, const FStateTreeEdit
 	return true;
 }
 
-bool FStateTreeCompiler::CreateEvaluator(UStateTreeState& State, const FStateTreeEditorNode& EvalNode)
+bool FStateTreeCompiler::CreateEvaluator(const FStateTreeEditorNode& EvalNode)
 {
 	// Silently ignore empty items.
 	if (!EvalNode.Node.IsValid())
@@ -689,7 +679,7 @@ bool FStateTreeCompiler::CreateEvaluator(UStateTreeState& State, const FStateTre
 
 	// Check that the bindings for this struct are still all valid.
 	TArray<FStateTreeEditorPropertyBinding> Bindings;
-	if (!GetAndValidateBindings(State, StructDesc, Bindings))
+	if (!GetAndValidateBindings(StructDesc, Bindings))
 	{
 		return false;
 	}
@@ -730,7 +720,7 @@ bool FStateTreeCompiler::IsPropertyAnyEnum(const FStateTreeBindableStructDesc& S
 	return bIsAnyEnum;
 }
 
-bool FStateTreeCompiler::GetAndValidateBindings(UStateTreeState& State, const FStateTreeBindableStructDesc& TargetStruct, TArray<FStateTreeEditorPropertyBinding>& OutBindings) const
+bool FStateTreeCompiler::GetAndValidateBindings(const FStateTreeBindableStructDesc& TargetStruct, TArray<FStateTreeEditorPropertyBinding>& OutBindings) const
 {
 	OutBindings.Reset();
 	
diff --git a/Engine/Plugins/Runtime/StateTree/Source/StateTreeEditorModule/Private/StateTreeEditor.cpp b/Engine/Plugins/Runtime/StateTree/Source/StateTreeEditorModule/Private/StateTreeEditor.cpp
--- a/Engine/Plugins/Runtime/StateTree/Source/StateTreeEditorModule/Private/StateTreeEditor.cpp
+++ b/Engine/Plugins/Runtime/StateTree/Source/StateTreeEditorModule/Private/StateTreeEditor.cpp
@@ -585,6 +585,14 @@ namespace UE::StateTree::Editor::Internal
 			return;
 		}
 
+		// Clear evaluators if not allowed.
+		if (Schema->AllowEvaluators() == false && TreeData->Evaluators.Num() > 0)
+		{
+			UE_LOG(LogStateTree, Warning, TEXT("%s: Resetting Evaluators due to current schema restrictions."), *GetNameSafe(&StateTree));
+			TreeData->Evaluators.Reset();
+		}
+
+
 		TreeData->VisitHierarchy([&StateTree, Schema](UStateTreeState& State, UStateTreeState* /*ParentState*/)
 		{
 			// Clear enter conditions if not allowed.
@@ -594,13 +602,6 @@ namespace UE::StateTree::Editor::Internal
 				State.EnterConditions.Reset();
 			}
 
-			// Clear evaluators if not allowed.
-			if (Schema->AllowEvaluators() == false && State.Evaluators.Num() > 0)
-			{
-				UE_LOG(LogStateTree, Warning, TEXT("%s: Resetting Evaluators in state %s due to current schema restrictions."), *GetNameSafe(&StateTree), *GetNameSafe(&State));
-				State.Evaluators.Reset();
-			}
-
 			// Keep single and many tasks based on what is allowed.
 			if (Schema->AllowMultipleTasks() == false)
 			{
diff --git a/Engine/Plugins/Runtime/StateTree/Source/StateTreeEditorModule/Private/StateTreeEditorData.cpp b/Engine/Plugins/Runtime/StateTree/Source/StateTreeEditorModule/Private/StateTreeEditorData.cpp
--- a/Engine/Plugins/Runtime/StateTree/Source/StateTreeEditorModule/Private/StateTreeEditorData.cpp
+++ b/Engine/Plugins/Runtime/StateTree/Source/StateTreeEditorModule/Private/StateTreeEditorData.cpp
@@ -14,16 +14,23 @@ void UStateTreeEditorData::PostInitProperties()
 }
 
 #if WITH_EDITOR
-void UStateTreeEditorData::PostEditChangeProperty(struct FPropertyChangedEvent& PropertyChangedEvent)
+void UStateTreeEditorData::PostEditChangeChainProperty(FPropertyChangedChainEvent& PropertyChangedEvent)
 {
-	Super::PostEditChangeProperty(PropertyChangedEvent);
-	
-	if (PropertyChangedEvent.MemberProperty && PropertyChangedEvent.Property)
+	Super::PostEditChangeChainProperty(PropertyChangedEvent);
+
+	FProperty* Property = PropertyChangedEvent.Property;
+	FProperty* MemberProperty = nullptr;
+	if (PropertyChangedEvent.PropertyChain.GetActiveMemberNode())
+	{
+		MemberProperty = PropertyChangedEvent.PropertyChain.GetActiveMemberNode()->GetValue();
+	}
+
+	if (MemberProperty && Property)
 	{
 		const UStateTree* StateTree = GetTypedOuter<UStateTree>();
 		checkf(StateTree, TEXT("UStateTreeEditorData should only be allocated within a UStateTree"));
 		
-		const FName MemberName = PropertyChangedEvent.MemberProperty->GetFName();
+		const FName MemberName = MemberProperty->GetFName();
 		if (MemberName == GET_MEMBER_NAME_CHECKED(UStateTreeEditorData, Schema))
 		{			
 			UE::StateTree::Delegates::OnSchemaChanged.Broadcast(*StateTree);
@@ -32,6 +39,23 @@ void UStateTreeEditorData::PostEditChangeProperty(struct FPropertyChangedEvent&
 		{
 			UE::StateTree::Delegates::OnParametersChanged.Broadcast(*StateTree);
 		}
+
+		// Ensure unique ID on duplicated items.
+		if (PropertyChangedEvent.ChangeType == EPropertyChangeType::Duplicate)
+		{
+			if (MemberName == GET_MEMBER_NAME_CHECKED(UStateTreeEditorData, Evaluators))
+			{
+				const int32 ArrayIndex = PropertyChangedEvent.GetArrayIndex(MemberProperty->GetFName().ToString());
+				if (Evaluators.IsValidIndex(ArrayIndex))
+				{
+					if (FStateTreeEvaluatorBase* Eval = Evaluators[ArrayIndex].Node.GetMutablePtr<FStateTreeEvaluatorBase>())
+					{
+						Eval->Name = FName(Eval->Name.ToString() + TEXT(" Duplicate"));
+					}
+					Evaluators[ArrayIndex].ID = FGuid::NewGuid();
+				}
+			}
+		}
 	}
 }
 #endif // WITH_EDITOR
@@ -87,58 +111,66 @@ void UStateTreeEditorData::GetAccessibleStructs(const TConstArrayView<const USta
 		}
 	}
 
-	TArray<FStateTreeBindableStructDesc> EvalDescs;
-	TArray<FStateTreeBindableStructDesc> TaskDescs;
+	bool bFoundTarget = false;
 
-	for (const UStateTreeState* State : Path)
+	// Evaluators
+	// Evaluators can access other evaluators that come before them.
+	for (const FStateTreeEditorNode& Node : Evaluators)
 	{
-		if (State == nullptr)
+		if (const FStateTreeEvaluatorBase* Evaluator = Node.Node.GetPtr<FStateTreeEvaluatorBase>())
 		{
-			continue;
+			// Stop iterating as soon as we find the target node.
+			if (Node.ID == TargetStructID)
+			{
+				bFoundTarget = true;
+				break;
+			}
+
+			// Collect evaluators accessible so far.
+			FStateTreeBindableStructDesc& Desc = OutStructDescs.AddDefaulted_GetRef();
+			Desc.Struct = Evaluator->GetInstanceDataType();
+			Desc.Name = Evaluator->Name;
+			Desc.ID = Node.ID;
+			Desc.DataSource = EStateTreeBindableStructSource::Evaluator;
 		}
-		
-		VisitStateNodes(*State, [&OutStructDescs, &EvalDescs, &TaskDescs, TargetStructID]
-			(const UStateTreeState& State, const FGuid& ID, const FName& Name, const EStateTreeNodeType NodeType, const UScriptStruct* NodeStruct, const UStruct* InstanceStruct)
+	}
+
+	// Conditions and Tasks
+	// Visit the tree in execution order. Conditions and tasks can access tasks that are executed before them.
+	if (!bFoundTarget)
+	{
+		TArray<FStateTreeBindableStructDesc> TaskDescs;
+
+		for (const UStateTreeState* State : Path)
+		{
+			if (State == nullptr)
 			{
-				if (ID == TargetStructID)
+				continue;
+			}
+			
+			VisitStateNodes(*State, [&OutStructDescs, &TaskDescs, TargetStructID]
+				(const UStateTreeState* State, const FGuid& ID, const FName& Name, const EStateTreeNodeType NodeType, const UScriptStruct* NodeStruct, const UStruct* InstanceStruct)
 				{
-					if (NodeType == EStateTreeNodeType::EnterCondition
-						|| NodeType == EStateTreeNodeType::Task
-						|| NodeType == EStateTreeNodeType::TransitionCondition
-						|| NodeType == EStateTreeNodeType::StateParameters)
+					// Stop iterating as soon as we find the target node.
+					if (ID == TargetStructID)
 					{
-						// These can only see both evaluators and tasks
-						OutStructDescs.Append(EvalDescs);
 						OutStructDescs.Append(TaskDescs);
+						return EStateTreeVisitor::Break;
 					}
-					else if (NodeType == EStateTreeNodeType::Evaluator)
+
+					// Not at target yet, collect all tasks accessible so far.
+					if (NodeStruct->IsChildOf(FStateTreeTaskBase::StaticStruct()))
 					{
-						// Evaluators can only see other evaluators
-						OutStructDescs.Append(EvalDescs);
+						FStateTreeBindableStructDesc& Desc = TaskDescs.AddDefaulted_GetRef();
+						Desc.Struct = InstanceStruct;
+						Desc.Name = Name;
+						Desc.ID = ID;
+						Desc.DataSource = EStateTreeBindableStructSource::Task;
 					}
-					
-					return EStateTreeVisitor::Break;
-				}
-
-				// Not at target yet, collect all evaluators and tasks visible so far.
-				if (NodeStruct->IsChildOf(FStateTreeEvaluatorBase::StaticStruct()))
-				{
-					FStateTreeBindableStructDesc& Desc = EvalDescs.AddDefaulted_GetRef();
-					Desc.Struct = InstanceStruct;
-					Desc.Name = Name;
-					Desc.ID = ID;
-					Desc.DataSource = EStateTreeBindableStructSource::Evaluator;
-				}
-				else if (NodeStruct->IsChildOf(FStateTreeTaskBase::StaticStruct()))
-				{
-					FStateTreeBindableStructDesc& Desc = TaskDescs.AddDefaulted_GetRef();
-					Desc.Struct = InstanceStruct;
-					Desc.Name = Name;
-					Desc.ID = ID;
-					Desc.DataSource = EStateTreeBindableStructSource::Task;
-				}
-				return EStateTreeVisitor::Continue;
-			});
+				
+					return EStateTreeVisitor::Continue;
+				});
+		}
 	}
 
 	OutStructDescs.StableSort([](const FStateTreeBindableStructDesc& A, const FStateTreeBindableStructDesc& B)
@@ -151,7 +183,7 @@ bool UStateTreeEditorData::GetStructByID(const FGuid StructID, FStateTreeBindabl
 {
 	bool bResult = false;
 
-	VisitHierarchyNodes([&bResult, &OutStructDesc, StructID](const UStateTreeState& State, const FGuid& ID, const FName& Name, const EStateTreeNodeType NodeType, const UScriptStruct* NodeStruct, const UStruct* InstanceStruct)
+	VisitHierarchyNodes([&bResult, &OutStructDesc, StructID](const UStateTreeState* State, const FGuid& ID, const FName& Name, const EStateTreeNodeType NodeType, const UScriptStruct* NodeStruct, const UStruct* InstanceStruct)
 		{
 			if (ID == StructID)
 			{
@@ -171,11 +203,11 @@ const UStateTreeState* UStateTreeEditorData::GetStateByStructID(const FGuid Targ
 {
 	const UStateTreeState* Result = nullptr;
 
-	VisitHierarchyNodes([&Result, TargetStructID](const UStateTreeState& State, const FGuid& ID, const FName& Name, const EStateTreeNodeType NodeType, const UScriptStruct* NodeStruct, const UStruct* InstanceStruct)
+	VisitHierarchyNodes([&Result, TargetStructID](const UStateTreeState* State, const FGuid& ID, const FName& Name, const EStateTreeNodeType NodeType, const UScriptStruct* NodeStruct, const UStruct* InstanceStruct)
 		{
 			if (ID == TargetStructID)
 			{
-				Result = &State;
+				Result = State;
 				return EStateTreeVisitor::Break;
 			}
 			return EStateTreeVisitor::Continue;
@@ -221,29 +253,17 @@ void UStateTreeEditorData::GetAllStructIDs(TMap<FGuid, const UStruct*>& AllStruc
 		}	
 	}
 
-	VisitHierarchyNodes([&AllStructs](const UStateTreeState& State, const FGuid& ID, const FName& Name, const EStateTreeNodeType NodeType, const UScriptStruct* NodeStruct, const UStruct* InstanceStruct)
+	VisitHierarchyNodes([&AllStructs](const UStateTreeState* State, const FGuid& ID, const FName& Name, const EStateTreeNodeType NodeType, const UScriptStruct* NodeStruct, const UStruct* InstanceStruct)
 		{
 			AllStructs.Add(ID, InstanceStruct);
 			return EStateTreeVisitor::Continue;
 		});
 }
 
-EStateTreeVisitor UStateTreeEditorData::VisitStateNodes(const UStateTreeState& State, TFunctionRef<EStateTreeVisitor(const UStateTreeState& State, const FGuid& ID, const FName& Name, const EStateTreeNodeType NodeType, const UScriptStruct* NodeStruct, const UStruct* InstanceStruct)> InFunc) const
+EStateTreeVisitor UStateTreeEditorData::VisitStateNodes(const UStateTreeState& State, TFunctionRef<EStateTreeVisitor(const UStateTreeState* State, const FGuid& ID, const FName& Name, const EStateTreeNodeType NodeType, const UScriptStruct* NodeStruct, const UStruct* InstanceStruct)> InFunc) const
 {
 	bool bContinue = true;
 
-	// Evaluators
-	for (const FStateTreeEditorNode& Node : State.Evaluators)
-	{
-		if (const FStateTreeEvaluatorBase* Evaluator = Node.Node.GetPtr<FStateTreeEvaluatorBase>())
-		{
-			if (InFunc(State, Node.ID, Evaluator->Name, EStateTreeNodeType::Evaluator, Node.Node.GetScriptStruct(), Evaluator->GetInstanceDataType()) == EStateTreeVisitor::Break)
-			{
-				bContinue = false;
-				break;
-			}
-		}
-	}
 	if (bContinue)
 	{
 		// Enter conditions
@@ -251,7 +271,7 @@ EStateTreeVisitor UStateTreeEditorData::VisitStateNodes(const UStateTreeState& S
 		{
 			if (const FStateTreeConditionBase* Cond = Node.Node.GetPtr<FStateTreeConditionBase>())
 			{
-				if (InFunc(State, Node.ID, Node.Node.GetScriptStruct()->GetFName(), EStateTreeNodeType::EnterCondition, Node.Node.GetScriptStruct(), Cond->GetInstanceDataType()) == EStateTreeVisitor::Break)
+				if (InFunc(&State, Node.ID, Node.Node.GetScriptStruct()->GetFName(), EStateTreeNodeType::EnterCondition, Node.Node.GetScriptStruct(), Cond->GetInstanceDataType()) == EStateTreeVisitor::Break)
 				{
 					bContinue = false;
 					break;
@@ -266,7 +286,7 @@ EStateTreeVisitor UStateTreeEditorData::VisitStateNodes(const UStateTreeState& S
 		{
 			if (const FStateTreeTaskBase* Task = Node.Node.GetPtr<FStateTreeTaskBase>())
 			{
-				if (InFunc(State, Node.ID, Task->Name, EStateTreeNodeType::Task, Node.Node.GetScriptStruct(), Task->GetInstanceDataType()) == EStateTreeVisitor::Break)
+				if (InFunc(&State, Node.ID, Task->Name, EStateTreeNodeType::Task, Node.Node.GetScriptStruct(), Task->GetInstanceDataType()) == EStateTreeVisitor::Break)
 				{
 					bContinue = false;
 					break;
@@ -278,7 +298,7 @@ EStateTreeVisitor UStateTreeEditorData::VisitStateNodes(const UStateTreeState& S
 	{
 		if (const FStateTreeTaskBase* Task = State.SingleTask.Node.GetPtr<FStateTreeTaskBase>())
 		{
-			if (InFunc(State, State.SingleTask.ID, Task->Name, EStateTreeNodeType::Task, State.SingleTask.Node.GetScriptStruct(), Task->GetInstanceDataType()) == EStateTreeVisitor::Break)
+			if (InFunc(&State, State.SingleTask.ID, Task->Name, EStateTreeNodeType::Task, State.SingleTask.Node.GetScriptStruct(), Task->GetInstanceDataType()) == EStateTreeVisitor::Break)
 			{
 				bContinue = false;
 			}
@@ -294,7 +314,7 @@ EStateTreeVisitor UStateTreeEditorData::VisitStateNodes(const UStateTreeState& S
 			{
 				if (const FStateTreeConditionBase* Cond = Node.Node.GetPtr<FStateTreeConditionBase>())
 				{
-					if (InFunc(State, Node.ID, Node.Node.GetScriptStruct()->GetFName(), EStateTreeNodeType::TransitionCondition, Node.Node.GetScriptStruct(), Cond->GetInstanceDataType()) == EStateTreeVisitor::Break)
+					if (InFunc(&State, Node.ID, Node.Node.GetScriptStruct()->GetFName(), EStateTreeNodeType::TransitionCondition, Node.Node.GetScriptStruct(), Cond->GetInstanceDataType()) == EStateTreeVisitor::Break)
 					{
 						bContinue = false;
 						break;
@@ -308,7 +328,7 @@ EStateTreeVisitor UStateTreeEditorData::VisitStateNodes(const UStateTreeState& S
 		// Bindable state parameters
 		if (State.Type == EStateTreeStateType::Linked && State.Parameters.Parameters.IsValid())
 		{
-			if (InFunc(State, State.Parameters.ID, State.Name, EStateTreeNodeType::StateParameters, nullptr, State.Parameters.Parameters.GetPropertyBagStruct()) == EStateTreeVisitor::Break)
+			if (InFunc(&State, State.Parameters.ID, State.Name, EStateTreeNodeType::StateParameters, nullptr, State.Parameters.Parameters.GetPropertyBagStruct()) == EStateTreeVisitor::Break)
 			{
 				bContinue = false;
 			}
@@ -361,11 +381,28 @@ void UStateTreeEditorData::VisitHierarchy(TFunctionRef<EStateTreeVisitor(UStateT
 	}
 }
 
-void UStateTreeEditorData::VisitHierarchyNodes(TFunctionRef<EStateTreeVisitor(const UStateTreeState& State, const FGuid& ID, const FName& Name, const EStateTreeNodeType NodeType, const UScriptStruct* NodeStruct, const UStruct* InstanceStruct)> InFunc) const
+void UStateTreeEditorData::VisitHierarchyNodes(TFunctionRef<EStateTreeVisitor(const UStateTreeState* State, const FGuid& ID, const FName& Name, const EStateTreeNodeType NodeType, const UScriptStruct* NodeStruct, const UStruct* InstanceStruct)> InFunc) const
 {
-	VisitHierarchy([this, &InFunc](const UStateTreeState& State, UStateTreeState* /*ParentState*/)
+	// Evaluators
+	bool bContinue = true;
+	for (const FStateTreeEditorNode& Node : Evaluators)
 	{
-		return VisitStateNodes(State, InFunc);
-	});
+		if (const FStateTreeEvaluatorBase* Evaluator = Node.Node.GetPtr<FStateTreeEvaluatorBase>())
+		{
+			if (InFunc(nullptr, Node.ID, Evaluator->Name, EStateTreeNodeType::Task, Node.Node.GetScriptStruct(), Evaluator->GetInstanceDataType()) == EStateTreeVisitor::Break)
+			{
+				bContinue = false;
+				break;
+			}
+		}
+	}
+
+	if (bContinue)
+	{
+		VisitHierarchy([this, &InFunc](const UStateTreeState& State, UStateTreeState* /*ParentState*/)
+		{
+			return VisitStateNodes(State, InFunc);
+		});
+	}
 }
 
diff --git a/Engine/Plugins/Runtime/StateTree/Source/StateTreeEditorModule/Private/StateTreeEditorModule.cpp b/Engine/Plugins/Runtime/StateTree/Source/StateTreeEditorModule/Private/StateTreeEditorModule.cpp
--- a/Engine/Plugins/Runtime/StateTree/Source/StateTreeEditorModule/Private/StateTreeEditorModule.cpp
+++ b/Engine/Plugins/Runtime/StateTree/Source/StateTreeEditorModule/Private/StateTreeEditorModule.cpp
@@ -8,6 +8,7 @@
 #include "Customizations/StateTreeTransitionDetails.h"
 #include "Customizations/StateTreeStateLinkDetails.h"
 #include "Customizations/StateTreeStateDetails.h"
+#include "Customizations/StateTreeEditorDataDetails.h"
 #include "Customizations/StateTreeEditorNodeDetails.h"
 #include "Customizations/StateTreeStateParametersDetails.h"
 #include "Customizations/StateTreeAnyEnumDetails.h"
@@ -56,6 +57,7 @@ void FStateTreeEditorModule::StartupModule()
 	PropertyModule.RegisterCustomPropertyTypeLayout("StateTreeAnyEnum", FOnGetPropertyTypeCustomizationInstance::CreateStatic(&FStateTreeAnyEnumDetails::MakeInstance));
 	PropertyModule.RegisterCustomPropertyTypeLayout("StateTreeReference", FOnGetPropertyTypeCustomizationInstance::CreateStatic(&FStateTreeReferenceDetails::MakeInstance));
 	PropertyModule.RegisterCustomClassLayout("StateTreeState", FOnGetDetailCustomizationInstance::CreateStatic(&FStateTreeStateDetails::MakeInstance));
+	PropertyModule.RegisterCustomClassLayout("StateTreeEditorData", FOnGetDetailCustomizationInstance::CreateStatic(&FStateTreeEditorDataDetails::MakeInstance));
 
 	PropertyModule.NotifyCustomizationModuleChanged();
 }
diff --git a/Engine/Plugins/Runtime/StateTree/Source/StateTreeEditorModule/Private/StateTreeState.cpp b/Engine/Plugins/Runtime/StateTree/Source/StateTreeEditorModule/Private/StateTreeState.cpp
--- a/Engine/Plugins/Runtime/StateTree/Source/StateTreeEditorModule/Private/StateTreeState.cpp
+++ b/Engine/Plugins/Runtime/StateTree/Source/StateTreeEditorModule/Private/StateTreeState.cpp
@@ -76,7 +76,6 @@ void UStateTreeState::PostEditChangeChainProperty(FPropertyChangedChainEvent& Pr
 			if (Type == EStateTreeStateType::Group || Type == EStateTreeStateType::Linked)
 			{
 				Tasks.Reset();
-				Evaluators.Reset();
 			}
 
 			// If transitioning from linked state, reset the linked state.
@@ -131,18 +130,6 @@ void UStateTreeState::PostEditChangeChainProperty(FPropertyChangedChainEvent& Pr
 			// Ensure unique ID on duplicated items.
 			if (PropertyChangedEvent.ChangeType == EPropertyChangeType::Duplicate)
 			{
-				if (MemberProperty->GetFName() == GET_MEMBER_NAME_CHECKED(UStateTreeState, Evaluators))
-				{
-					const int32 ArrayIndex = PropertyChangedEvent.GetArrayIndex(MemberProperty->GetFName().ToString());
-					if (Evaluators.IsValidIndex(ArrayIndex))
-					{
-						if (FStateTreeEvaluatorBase* Eval = Evaluators[ArrayIndex].Node.GetMutablePtr<FStateTreeEvaluatorBase>())
-						{
-							Eval->Name = FName(Eval->Name.ToString() + TEXT(" Duplicate"));
-						}
-						Evaluators[ArrayIndex].ID = FGuid::NewGuid();
-					}
-				}
 				if (MemberProperty->GetFName() == GET_MEMBER_NAME_CHECKED(UStateTreeState, Tasks))
 				{
 					const int32 ArrayIndex = PropertyChangedEvent.GetArrayIndex(MemberProperty->GetFName().ToString());
@@ -179,6 +166,18 @@ void UStateTreeState::PostLoad()
 	{
 		SetFlags(RF_Transactional);
 	}
+
+	// Move deprecated evaluators to editor data.
+	PRAGMA_DISABLE_DEPRECATION_WARNINGS
+	if (Evaluators_DEPRECATED.Num() > 0)
+	{
+		if (UStateTreeEditorData* TreeData = GetTypedOuter<UStateTreeEditorData>())
+		{
+			TreeData->Evaluators.Append(Evaluators_DEPRECATED);
+			Evaluators_DEPRECATED.Reset();
+		}		
+	}
+	PRAGMA_ENABLE_DEPRECATION_WARNINGS
 }
 
 void UStateTreeState::UpdateParametersFromLinkedState()
diff --git a/Engine/Plugins/Runtime/StateTree/Source/StateTreeEditorModule/Public/StateTreeCompiler.h b/Engine/Plugins/Runtime/StateTree/Source/StateTreeEditorModule/Public/StateTreeCompiler.h
--- a/Engine/Plugins/Runtime/StateTree/Source/StateTreeEditorModule/Public/StateTreeCompiler.h
+++ b/Engine/Plugins/Runtime/StateTree/Source/StateTreeEditorModule/Public/StateTreeCompiler.h
@@ -34,15 +34,15 @@ private:
 	bool CreateStates();
 	bool CreateStateRecursive(UStateTreeState& State, const FStateTreeHandle Parent);
 	
+	bool CreateEvaluators();
 	bool CreateStateTasksAndParameters();
-	bool CreateStateEvaluators();
 	bool CreateStateTransitions();
 	
 	bool CreateConditions(UStateTreeState& State, TConstArrayView<FStateTreeEditorNode> Conditions);
 	bool CreateCondition(UStateTreeState& State, const FStateTreeEditorNode& CondNode, const EStateTreeConditionOperand Operand, const int8 DeltaIndent);
 	bool CreateTask(UStateTreeState& State, const FStateTreeEditorNode& TaskNode);
-	bool CreateEvaluator(UStateTreeState& State, const FStateTreeEditorNode& EvalNode);
-	bool GetAndValidateBindings(UStateTreeState& State, const FStateTreeBindableStructDesc& TargetStruct, TArray<FStateTreeEditorPropertyBinding>& OutBindings) const;
+	bool CreateEvaluator(const FStateTreeEditorNode& EvalNode);
+	bool GetAndValidateBindings(const FStateTreeBindableStructDesc& TargetStruct, TArray<FStateTreeEditorPropertyBinding>& OutBindings) const;
 	bool IsPropertyAnyEnum(const FStateTreeBindableStructDesc& Struct, FStateTreeEditorPropertyPath Path) const;
 
 	FStateTreeCompilerLog& Log;
diff --git a/Engine/Plugins/Runtime/StateTree/Source/StateTreeEditorModule/Public/StateTreeEditorData.h b/Engine/Plugins/Runtime/StateTree/Source/StateTreeEditorModule/Public/StateTreeEditorData.h
--- a/Engine/Plugins/Runtime/StateTree/Source/StateTreeEditorModule/Public/StateTreeEditorData.h
+++ b/Engine/Plugins/Runtime/StateTree/Source/StateTreeEditorModule/Public/StateTreeEditorData.h
@@ -8,16 +8,6 @@
 
 class UStateTreeSchema;
 
-UENUM()
-enum class EStateTreeNodeType : uint8
-{
-	EnterCondition,
-	Evaluator,
-	Task,
-	TransitionCondition,
-	StateParameters,
-};
-
 UENUM()
 enum class EStateTreeVisitor : uint8
 {
@@ -25,7 +15,6 @@ enum class EStateTreeVisitor : uint8
 	Break,
 };
 
-
 /**
  * Edit time data for StateTree asset. This data gets baked into runtime format before being used by the StateTreeInstance.
  */
@@ -44,7 +33,7 @@ public:
 	// ~IStateTreeEditorPropertyBindingsOwner
 
 #if WITH_EDITOR
-	virtual void PostEditChangeProperty(struct FPropertyChangedEvent& PropertyChangedEvent) override;
+	virtual void PostEditChangeChainProperty(FPropertyChangedChainEvent& PropertyChangedEvent) override;
 #endif
 
 	/** Returns parent state of a struct, or nullptr if not found. */
@@ -66,13 +55,13 @@ public:
 	 * Iterates over all structs that are related to binding
 	 * @param InFunc function called at each node, should return true if visiting is continued or false to stop.
 	 */
-	void VisitHierarchyNodes(TFunctionRef<EStateTreeVisitor(const UStateTreeState& State, const FGuid& ID, const FName& Name, const EStateTreeNodeType NodeType, const UScriptStruct* NodeStruct, const UStruct* InstanceStruct)> InFunc) const;
+	void VisitHierarchyNodes(TFunctionRef<EStateTreeVisitor(const UStateTreeState* State, const FGuid& ID, const FName& Name, const EStateTreeNodeType NodeType, const UScriptStruct* NodeStruct, const UStruct* InstanceStruct)> InFunc) const;
 
 	/**
 	 * Iterates over all nodes in a given state.
 	 * @param InFunc function called at each node, should return true if visiting is continued or false to stop.
 	 */
-	EStateTreeVisitor VisitStateNodes(const UStateTreeState& State, TFunctionRef<EStateTreeVisitor(const UStateTreeState& State, const FGuid& ID, const FName& Name, const EStateTreeNodeType NodeType, const UScriptStruct* NodeStruct, const UStruct* InstanceStruct)> InFunc) const;
+	EStateTreeVisitor VisitStateNodes(const UStateTreeState& State, TFunctionRef<EStateTreeVisitor(const UStateTreeState* State, const FGuid& ID, const FName& Name, const EStateTreeNodeType NodeType, const UScriptStruct* NodeStruct, const UStruct* InstanceStruct)> InFunc) const;
 
 	/**
 	 * Returns array of nodes along the execution path, up to the TargetStruct.
@@ -105,7 +94,25 @@ public:
 	{
 		return AddSubTree(FName(TEXT("Root")));
 	}
-	
+
+	/**
+	 * Adds Evaluator of specified type.
+	 * @return reference to the new Evaluator. 
+	 */
+	template<typename T, typename... TArgs>
+	TStateTreeEditorNode<T>& AddEvaluator(TArgs&&... InArgs)
+	{
+		FStateTreeEditorNode& EvalItem = Evaluators.AddDefaulted_GetRef();
+		EvalItem.ID = FGuid::NewGuid();
+		EvalItem.Node.InitializeAs<T>(Forward<TArgs>(InArgs)...);
+		T& Eval = EvalItem.Node.GetMutable<T>();
+		if (const UScriptStruct* InstanceType = Cast<const UScriptStruct>(Eval.GetInstanceDataType()))
+		{
+			EvalItem.Instance.InitializeAs(InstanceType);
+		}
+		return static_cast<TStateTreeEditorNode<T>&>(EvalItem);
+	}
+
 	/**
 	 * Adds property binding between two structs.
 	 */
@@ -122,7 +129,10 @@ public:
 	/** Public parameters that could be used for bindings within the Tree. */
 	UPROPERTY(EditDefaultsOnly, Category = Common)
 	FStateTreeStateParameters RootParameters;
-	
+
+	UPROPERTY(EditDefaultsOnly, Category = "Evaluators", meta = (BaseStruct = "StateTreeEvaluatorBase", BaseClass = "StateTreeEvaluatorBlueprintBase"))
+	TArray<FStateTreeEditorNode> Evaluators;
+
 	UPROPERTY(meta = (ExcludeFromHash))
 	FStateTreeEditorPropertyBindings EditorBindings;
 
diff --git a/Engine/Plugins/Runtime/StateTree/Source/StateTreeEditorModule/Public/StateTreeEditorNode.h b/Engine/Plugins/Runtime/StateTree/Source/StateTreeEditorModule/Public/StateTreeEditorNode.h
new file mode 100644
--- /dev/null
+++ b/Engine/Plugins/Runtime/StateTree/Source/StateTreeEditorModule/Public/StateTreeEditorNode.h
@@ -0,0 +1,71 @@
+﻿// Copyright Epic Games, Inc. All Rights Reserved.
+
+#pragma once
+
+#include "StateTreeTypes.h"
+#include "StateTreeNodeBase.h"
+#include "Misc/Guid.h"
+#include "InstancedStruct.h"
+#include "StateTreeEditorNode.generated.h"
+
+UENUM()
+enum class EStateTreeNodeType : uint8
+{
+	EnterCondition,
+	Evaluator,
+	Task,
+	TransitionCondition,
+	StateParameters,
+};
+
+/**
+ * Base for Evaluator, Task and Condition nodes.
+ */
+USTRUCT()
+struct STATETREEEDITORMODULE_API FStateTreeEditorNode
+{
+	GENERATED_BODY()
+
+	void Reset()
+	{
+		Node.Reset();
+		Instance.Reset();
+		InstanceObject = nullptr;
+		ID = FGuid();
+	}
+
+	FName GetName() const
+	{
+		if (const FStateTreeNodeBase* NodePtr = Node.GetPtr<FStateTreeNodeBase>())
+		{
+			return NodePtr->Name;
+		}
+		return FName();
+	}
+
+	UPROPERTY(EditDefaultsOnly, Category = Node)
+	FInstancedStruct Node;
+
+	UPROPERTY(EditDefaultsOnly, Category = Node)
+	FInstancedStruct Instance;
+
+	UPROPERTY(EditDefaultsOnly, Instanced, Category = Node)
+	UObject* InstanceObject = nullptr;
+	
+	UPROPERTY(EditDefaultsOnly, Category = Node)
+	FGuid ID;
+
+	UPROPERTY(EditDefaultsOnly, Category = Node)
+	uint8 ConditionIndent = 0;
+
+	UPROPERTY(EditDefaultsOnly, Category = Node)
+	EStateTreeConditionOperand ConditionOperand = EStateTreeConditionOperand::And; 
+};
+
+template <typename T>
+struct TStateTreeEditorNode : public FStateTreeEditorNode
+{
+	typedef T NodeType;
+	FORCEINLINE T& GetItem() { return Node.template GetMutable<T>(); }
+	FORCEINLINE typename T::InstanceDataType& GetInstance() { return Instance.template GetMutable<typename T::InstanceDataType>(); }
+};
diff --git a/Engine/Plugins/Runtime/StateTree/Source/StateTreeEditorModule/Public/StateTreeState.h b/Engine/Plugins/Runtime/StateTree/Source/StateTreeEditorModule/Public/StateTreeState.h
--- a/Engine/Plugins/Runtime/StateTree/Source/StateTreeEditorModule/Public/StateTreeState.h
+++ b/Engine/Plugins/Runtime/StateTree/Source/StateTreeEditorModule/Public/StateTreeState.h
@@ -7,6 +7,7 @@
 #include "Misc/Guid.h"
 #include "InstancedStruct.h"
 #include "PropertyBag.h"
+#include "StateTreeEditorNode.h"
 #include "StateTreeState.generated.h"
 
 class UStateTreeState;
@@ -37,59 +38,6 @@ struct STATETREEEDITORMODULE_API FStateTreeStateLink
 	EStateTreeTransitionType Type = EStateTreeTransitionType::GotoState;
 };
 
-
-/**
- * Base for Evaluator, Task and Condition nodes.
- */
-USTRUCT()
-struct STATETREEEDITORMODULE_API FStateTreeEditorNode
-{
-	GENERATED_BODY()
-
-	void Reset()
-	{
-		Node.Reset();
-		Instance.Reset();
-		InstanceObject = nullptr;
-		ID = FGuid();
-	}
-
-	FName GetName() const
-	{
-		if (const FStateTreeNodeBase* NodePtr = Node.GetPtr<FStateTreeNodeBase>())
-		{
-			return NodePtr->Name;
-		}
-		return FName();
-	}
-
-	UPROPERTY(EditDefaultsOnly, Category = Node)
-	FInstancedStruct Node;
-
-	UPROPERTY(EditDefaultsOnly, Category = Node)
-	FInstancedStruct Instance;
-
-	UPROPERTY(EditDefaultsOnly, Instanced, Category = Node)
-	UObject* InstanceObject = nullptr;
-	
-	UPROPERTY(EditDefaultsOnly, Category = Node)
-	FGuid ID;
-
-	UPROPERTY(EditDefaultsOnly, Category = Node)
-	uint8 ConditionIndent = 0;
-
-	UPROPERTY(EditDefaultsOnly, Category = Node)
-	EStateTreeConditionOperand ConditionOperand = EStateTreeConditionOperand::And; 
-};
-
-template <typename T>
-struct TStateTreeEditorNode : public FStateTreeEditorNode
-{
-	typedef T NodeType;
-	FORCEINLINE T& GetItem() { return Node.template GetMutable<T>(); }
-	FORCEINLINE typename T::InstanceDataType& GetInstance() { return Instance.template GetMutable<typename T::InstanceDataType>(); }
-};
-
 /**
  * Editor representation of a transition in StateTree
  */
@@ -219,24 +167,6 @@ public:
 		return static_cast<TStateTreeEditorNode<T>&>(TaskItem);
 	}
 
-	/**
-	 * Adds Evaluator of specified type.
-	 * @return reference to the new Evaluator. 
-	 */
-	template<typename T, typename... TArgs>
-    TStateTreeEditorNode<T>& AddEvaluator(TArgs&&... InArgs)
-	{
-		FStateTreeEditorNode& EvalItem = Evaluators.AddDefaulted_GetRef();
-		EvalItem.ID = FGuid::NewGuid();
-		EvalItem.Node.InitializeAs<T>(Forward<TArgs>(InArgs)...);
-		T& Eval = EvalItem.Node.GetMutable<T>();
-		if (const UScriptStruct* InstanceType = Cast<const UScriptStruct>(Eval.GetInstanceDataType()))
-		{
-			EvalItem.Instance.InitializeAs(InstanceType);
-		}
-		return static_cast<TStateTreeEditorNode<T>&>(EvalItem);
-	}
-
 	/**
 	 * Adds Transition.
 	 * @return reference to the new Transition. 
@@ -266,8 +196,9 @@ public:
 	UPROPERTY(EditDefaultsOnly, Category = "Enter Conditions", meta = (BaseStruct = "StateTreeConditionBase", BaseClass = "StateTreeConditionBlueprintBase"))
 	TArray<FStateTreeEditorNode> EnterConditions;
 
-	UPROPERTY(EditDefaultsOnly, Category = "Evaluators", meta = (BaseStruct = "StateTreeEvaluatorBase", BaseClass = "StateTreeEvaluatorBlueprintBase"))
-	TArray<FStateTreeEditorNode> Evaluators;
+	UE_DEPRECATED(5.1, "Evaluators are moved into UStateTreeEditorData. This property will be removed for 5.1.")
+	UPROPERTY(meta = (DeprecatedProperty, BaseStruct = "StateTreeEvaluatorBase", BaseClass = "StateTreeEvaluatorBlueprintBase"))
+	TArray<FStateTreeEditorNode> Evaluators_DEPRECATED;
 
 	UPROPERTY(EditDefaultsOnly, Category = "Tasks", meta = (BaseStruct = "StateTreeTaskBase", BaseClass = "StateTreeTaskBlueprintBase"))
 	TArray<FStateTreeEditorNode> Tasks;
diff --git a/Engine/Plugins/Runtime/StateTree/Source/StateTreeModule/Private/Blueprint/StateTreeEvaluatorBlueprintBase.cpp b/Engine/Plugins/Runtime/StateTree/Source/StateTreeModule/Private/Blueprint/StateTreeEvaluatorBlueprintBase.cpp
--- a/Engine/Plugins/Runtime/StateTree/Source/StateTreeModule/Private/Blueprint/StateTreeEvaluatorBlueprintBase.cpp
+++ b/Engine/Plugins/Runtime/StateTree/Source/StateTreeModule/Private/Blueprint/StateTreeEvaluatorBlueprintBase.cpp
@@ -11,25 +11,46 @@
 void UStateTreeEvaluatorBlueprintBase::EnterState(FStateTreeExecutionContext& Context, const EStateTreeStateChangeType ChangeType, const FStateTreeTransitionResult& Transition)
 {
 	AActor* OwnerActor = GetOwnerActor(Context);
+	PRAGMA_DISABLE_DEPRECATION_WARNINGS
 	ReceiveEnterState(OwnerActor, ChangeType, Transition);
+	PRAGMA_ENABLE_DEPRECATION_WARNINGS
 }
 
 void UStateTreeEvaluatorBlueprintBase::ExitState(FStateTreeExecutionContext& Context, const EStateTreeStateChangeType ChangeType, const FStateTreeTransitionResult& Transition)
 {
 	AActor* OwnerActor = GetOwnerActor(Context);
+	PRAGMA_DISABLE_DEPRECATION_WARNINGS
 	ReceiveExitState(OwnerActor, ChangeType, Transition);
+	PRAGMA_ENABLE_DEPRECATION_WARNINGS
 }
 
 void UStateTreeEvaluatorBlueprintBase::StateCompleted(FStateTreeExecutionContext& Context, const EStateTreeRunStatus CompletionStatus, const FStateTreeActiveStates& CompletedActiveStates)
 {
 	AActor* OwnerActor = GetOwnerActor(Context);
+	PRAGMA_DISABLE_DEPRECATION_WARNINGS
 	ReceiveStateCompleted(OwnerActor, CompletionStatus, CompletedActiveStates);
+	PRAGMA_ENABLE_DEPRECATION_WARNINGS
 }
 
-void UStateTreeEvaluatorBlueprintBase::Evaluate(FStateTreeExecutionContext& Context, const EStateTreeEvaluationType EvalType, const float DeltaTime)
+void UStateTreeEvaluatorBlueprintBase::TreeStart(FStateTreeExecutionContext& Context)
 {
 	AActor* OwnerActor = GetOwnerActor(Context);
-	ReceiveEvaluate(OwnerActor, EvalType, DeltaTime);
+	ReceiveTreeStart(OwnerActor);
+}
+
+void UStateTreeEvaluatorBlueprintBase::TreeStop(FStateTreeExecutionContext& Context)
+{
+	AActor* OwnerActor = GetOwnerActor(Context);
+	ReceiveTreeStop(OwnerActor);
+}
+
+void UStateTreeEvaluatorBlueprintBase::Tick(FStateTreeExecutionContext& Context, const float DeltaTime)
+{
+	AActor* OwnerActor = GetOwnerActor(Context);
+	PRAGMA_DISABLE_DEPRECATION_WARNINGS
+	ReceiveEvaluate(OwnerActor, EStateTreeEvaluationType::Tick, DeltaTime);
+	PRAGMA_ENABLE_DEPRECATION_WARNINGS
+	ReceiveTick(OwnerActor, DeltaTime);
 }
 
 //----------------------------------------------------------------------//
@@ -53,7 +74,9 @@ void FStateTreeBlueprintEvaluatorWrapper::EnterState(FStateTreeExecutionContext&
 	check(Instance);
 	
 	Instance->CopyExternalData(Context, ExternalDataHandles);
+	PRAGMA_DISABLE_DEPRECATION_WARNINGS
 	Instance->EnterState(Context, ChangeType, Transition);
+	PRAGMA_ENABLE_DEPRECATION_WARNINGS
 }
 
 void FStateTreeBlueprintEvaluatorWrapper::ExitState(FStateTreeExecutionContext& Context, const EStateTreeStateChangeType ChangeType, const FStateTreeTransitionResult& Transition) const
@@ -62,7 +85,9 @@ void FStateTreeBlueprintEvaluatorWrapper::ExitState(FStateTreeExecutionContext&
 	check(Instance);
 	
 	Instance->CopyExternalData(Context, ExternalDataHandles);
+	PRAGMA_DISABLE_DEPRECATION_WARNINGS
 	Instance->ExitState(Context, ChangeType, Transition);
+	PRAGMA_ENABLE_DEPRECATION_WARNINGS
 }
 
 void FStateTreeBlueprintEvaluatorWrapper::StateCompleted(FStateTreeExecutionContext& Context, const EStateTreeRunStatus CompletionStatus, const FStateTreeActiveStates& CompletedActiveStates) const
@@ -71,14 +96,34 @@ void FStateTreeBlueprintEvaluatorWrapper::StateCompleted(FStateTreeExecutionCont
 	check(Instance);
 	
 	Instance->CopyExternalData(Context, ExternalDataHandles);
+	PRAGMA_DISABLE_DEPRECATION_WARNINGS
 	Instance->StateCompleted(Context, CompletionStatus, CompletedActiveStates);
+	PRAGMA_ENABLE_DEPRECATION_WARNINGS
+}
+
+void FStateTreeBlueprintEvaluatorWrapper::TreeStart(FStateTreeExecutionContext& Context) const
+{
+	UStateTreeEvaluatorBlueprintBase* Instance = Context.GetInstanceObjectInternal<UStateTreeEvaluatorBlueprintBase>(DataViewIndex);
+	check(Instance);
+	
+	Instance->CopyExternalData(Context, ExternalDataHandles);
+	Instance->TreeStart(Context);
+}
+
+void FStateTreeBlueprintEvaluatorWrapper::TreeStop(FStateTreeExecutionContext& Context) const
+{
+	UStateTreeEvaluatorBlueprintBase* Instance = Context.GetInstanceObjectInternal<UStateTreeEvaluatorBlueprintBase>(DataViewIndex);
+	check(Instance);
+	
+	Instance->CopyExternalData(Context, ExternalDataHandles);
+	Instance->TreeStop(Context);
 }
 
-void FStateTreeBlueprintEvaluatorWrapper::Evaluate(FStateTreeExecutionContext& Context, const EStateTreeEvaluationType EvalType, const float DeltaTime) const
+void FStateTreeBlueprintEvaluatorWrapper::Tick(FStateTreeExecutionContext& Context, const float DeltaTime) const
 {
 	UStateTreeEvaluatorBlueprintBase* Instance = Context.GetInstanceObjectInternal<UStateTreeEvaluatorBlueprintBase>(DataViewIndex);
 	check(Instance);
 	
 	Instance->CopyExternalData(Context, ExternalDataHandles);
-	Instance->Evaluate(Context, EvalType, DeltaTime);
+	Instance->Tick(Context, DeltaTime);
 }
diff --git a/Engine/Plugins/Runtime/StateTree/Source/StateTreeModule/Private/StateTree.cpp b/Engine/Plugins/Runtime/StateTree/Source/StateTreeModule/Private/StateTree.cpp
--- a/Engine/Plugins/Runtime/StateTree/Source/StateTreeModule/Private/StateTree.cpp
+++ b/Engine/Plugins/Runtime/StateTree/Source/StateTreeModule/Private/StateTree.cpp
@@ -28,6 +28,9 @@ void UStateTree::ResetCompiled()
 	NamedExternalDataDescs.Reset();
 	PropertyBindings.Reset();
 
+	EvaluatorsBegin = 0;
+	EvaluatorsNum = 0;
+
 	NumDataViews = 0;
 	ExternalDataBaseIndex = 0;
 	DefaultParametersDataViewIndex = INDEX_NONE;
diff --git a/Engine/Plugins/Runtime/StateTree/Source/StateTreeModule/Private/StateTreeExecutionContext.cpp b/Engine/Plugins/Runtime/StateTree/Source/StateTreeModule/Private/StateTreeExecutionContext.cpp
--- a/Engine/Plugins/Runtime/StateTree/Source/StateTreeModule/Private/StateTreeExecutionContext.cpp
+++ b/Engine/Plugins/Runtime/StateTree/Source/StateTreeModule/Private/StateTreeExecutionContext.cpp
@@ -53,9 +53,6 @@ bool FStateTreeExecutionContext::Init(UObject& InOwner, const UStateTree& InStat
 	// Set data views associated to the parameters using the default values
 	SetDefaultParameters();
 
-	// Initialize array to keep track of which states have been updated.
-	VisitedStates.Init(false, StateTree->States.Num());
-
 	return true;
 }
 
@@ -84,7 +81,6 @@ void FStateTreeExecutionContext::Reset()
 	InternalInstanceData.Reset();
 	DataViews.Reset();
 	StorageType = EStateTreeStorage::Internal;
-	VisitedStates.Reset();
 	StateTree = nullptr;
 	Owner = nullptr;
 }
@@ -140,6 +136,29 @@ EStateTreeRunStatus FStateTreeExecutionContext::Start(FStateTreeInstanceData* Ex
 	checkSlow(InstanceData.GetLayout() == StateTree->GetInstanceDataDefaultValue().GetLayout());
 	FStateTreeExecutionState& Exec = GetExecState(InstanceData);
 
+	// Call TreeStart on evaluators.
+	STATETREE_CLOG(StateTree->EvaluatorsNum > 0, Verbose, TEXT("Start Evaluators"));
+	for (int32 EvalIndex = 0; EvalIndex < int32(StateTree->EvaluatorsNum); EvalIndex++)
+	{
+		const FStateTreeEvaluatorBase& Eval = GetNode<FStateTreeEvaluatorBase>(int32(StateTree->EvaluatorsBegin) + EvalIndex);
+		const FStateTreeDataView EvalInstanceView = GetInstanceData(InstanceData, Eval.bInstanceIsObject, Eval.InstanceIndex);
+		DataViews[Eval.DataViewIndex] = EvalInstanceView;
+
+		// Copy bound properties.
+		if (Eval.BindingsBatch.IsValid())
+		{
+			StateTree->PropertyBindings.CopyTo(DataViews, Eval.BindingsBatch.Index, EvalInstanceView);
+		}
+		STATETREE_LOG(Verbose, TEXT("  TreeStart: '%s'"), *Eval.Name.ToString());
+		{
+			QUICK_SCOPE_CYCLE_COUNTER(StateTree_Eval_TreeStart);
+			Eval.TreeStart(*this);
+		}
+	}
+
+	// First tick
+	TickEvaluators(InstanceData, 0.0f);
+
 	// Stop if still running previous state.
 	if (Exec.TreeRunStatus == EStateTreeRunStatus::Running)
 	{
@@ -157,9 +176,6 @@ EStateTreeRunStatus FStateTreeExecutionContext::Start(FStateTreeInstanceData* Ex
 	Exec.ActiveStates.Reset();
 	Exec.LastTickStatus = EStateTreeRunStatus::Unset;
 
-	// Select new Starting from root.
-	VisitedStates.Init(false, StateTree->States.Num());
-	
 	static const FStateTreeHandle RootState = FStateTreeHandle(0);
 
 	FStateTreeActiveStates NextActiveStates;
@@ -214,6 +230,8 @@ EStateTreeRunStatus FStateTreeExecutionContext::Stop(FStateTreeInstanceData* Ext
 	checkSlow(InstanceData.GetLayout() == StateTree->GetInstanceDataDefaultValue().GetLayout());
 	FStateTreeExecutionState& Exec = GetExecState(InstanceData);
 
+	TickEvaluators(InstanceData, 0.0f);
+
 	// Exit states if still in some valid state.
 	if (!Exec.ActiveStates.IsEmpty() && (Exec.ActiveStates.Last() != FStateTreeHandle::Succeeded || Exec.ActiveStates.Last() != FStateTreeHandle::Failed))
 	{
@@ -237,6 +255,26 @@ EStateTreeRunStatus FStateTreeExecutionContext::Stop(FStateTreeInstanceData* Ext
 
 	Exec.LastTickStatus = EStateTreeRunStatus::Unset;
 
+	// Call TreeStop on evaluators.
+	STATETREE_CLOG(StateTree->EvaluatorsNum > 0, Verbose, TEXT("Stop Evaluators"));
+	for (int32 EvalIndex = 0; EvalIndex < int32(StateTree->EvaluatorsNum); EvalIndex++)
+	{
+		const FStateTreeEvaluatorBase& Eval = GetNode<FStateTreeEvaluatorBase>(int32(StateTree->EvaluatorsBegin) + EvalIndex);
+		const FStateTreeDataView EvalInstanceView = GetInstanceData(InstanceData, Eval.bInstanceIsObject, Eval.InstanceIndex);
+		DataViews[Eval.DataViewIndex] = EvalInstanceView;
+
+		// Copy bound properties.
+		if (Eval.BindingsBatch.IsValid())
+		{
+			StateTree->PropertyBindings.CopyTo(DataViews, Eval.BindingsBatch.Index, EvalInstanceView);
+		}
+		STATETREE_LOG(Verbose, TEXT("  TreeStop: '%s'"), *Eval.Name.ToString());
+		{
+			QUICK_SCOPE_CYCLE_COUNTER(StateTree_Eval_TreeStart);
+			Eval.TreeStop(*this);
+		}
+	}
+
 	return Exec.TreeRunStatus;
 }
 
@@ -264,15 +302,11 @@ EStateTreeRunStatus FStateTreeExecutionContext::Tick(const float DeltaTime, FSta
 		Exec.GatedTransitionTime -= DeltaTime;
 	}
 	
-	// Reset visited state, used to keep track which state's evaluators have been ticked.
-	VisitedStates.Init(false, StateTree->States.Num());
+	// Tick global evaluators.
+	TickEvaluators(InstanceData, DeltaTime);
 
 	if (Exec.LastTickStatus == EStateTreeRunStatus::Running)
 	{
-		// Tick evaluators on active states. Further evaluator may be ticked during selection as non-active states are visited.
-		// Ticked states are kept track of and evaluators are ticked just once per frame.
-		TickEvaluators(InstanceData, Exec.ActiveStates, EStateTreeEvaluationType::Tick, DeltaTime);
-
 		// Tick tasks on active states.
 		Exec.LastTickStatus = TickTasks(InstanceData, Exec, DeltaTime);
 		
@@ -316,9 +350,6 @@ EStateTreeRunStatus FStateTreeExecutionContext::Tick(const float DeltaTime, FSta
 		{
 			break;
 		}
-		
-		// Reset visited states to allow evaluators to update during the next transition selection.
-		VisitedStates.Init(false, StateTree->States.Num());
 	}
 	
 	if (Exec.ActiveStates.IsEmpty())
@@ -390,27 +421,6 @@ EStateTreeRunStatus FStateTreeExecutionContext::EnterState(FStateTreeInstanceDat
 
 		STATETREE_CLOG(bIsEnteringState, Log, TEXT("%*sEnter state '%s' %s"), Index*UE::StateTree::DebugIndentSize, TEXT(""), *DebugGetStatePath(Transition.NextActiveStates, Index), *UEnum::GetValueAsString(ChangeType));
 		
-		for (int32 EvalIndex = State.EvaluatorsBegin; EvalIndex < (int32)(State.EvaluatorsBegin + State.EvaluatorsNum); EvalIndex++)
-		{
-			const FStateTreeEvaluatorBase& Eval = GetNode<FStateTreeEvaluatorBase>(EvalIndex);
-			const FStateTreeDataView EvalInstanceView = GetInstanceData(InstanceData, Eval.bInstanceIsObject, Eval.InstanceIndex);
-			DataViews[Eval.DataViewIndex] = EvalInstanceView;
-
-			// Copy bound properties.
-			if (Eval.BindingsBatch.IsValid())
-			{
-				StateTree->PropertyBindings.CopyTo(DataViews, Eval.BindingsBatch.Index, EvalInstanceView);
-			}
-
-			if (bIsEnteringState)
-			{
-				STATETREE_LOG(Verbose, TEXT("%*s  Notify Evaluator '%s'."), Index*UE::StateTree::DebugIndentSize, TEXT(""), *Eval.Name.ToString());
-				QUICK_SCOPE_CYCLE_COUNTER(StateTree_Eval_EnterState);
-				CSV_SCOPED_TIMING_STAT_EXCLUSIVE(StateTree_Eval_EnterState);
-				Eval.EnterState(*this, ChangeType, CurrentTransition);
-			}
-		}
-		
 		// Activate tasks on current state.
 		for (int32 TaskIndex = State.TasksBegin; TaskIndex < (int32)(State.TasksBegin + State.TasksNum); TaskIndex++)
 		{
@@ -501,19 +511,6 @@ void FStateTreeExecutionContext::ExitState(FStateTreeInstanceData& InstanceData,
 		}
 
 		// Do property copies, ExitState() is called below.
-		for (int32 EvalIndex = State.EvaluatorsBegin; EvalIndex < (int32)(State.EvaluatorsBegin + State.EvaluatorsNum); EvalIndex++)
-		{
-			const FStateTreeEvaluatorBase& Eval = GetNode<FStateTreeEvaluatorBase>(EvalIndex);
-			const FStateTreeDataView EvalInstanceView = GetInstanceData(InstanceData, Eval.bInstanceIsObject, Eval.InstanceIndex);
-			DataViews[Eval.DataViewIndex] = EvalInstanceView;
-
-			// Copy bound properties.
-			if (Eval.BindingsBatch.IsValid())
-			{
-				StateTree->PropertyBindings.CopyTo(DataViews, Eval.BindingsBatch.Index, EvalInstanceView);
-			}
-		}
-
 		for (int32 TaskIndex = State.TasksBegin; TaskIndex < (int32)(State.TasksBegin + State.TasksNum); TaskIndex++)
 		{
 			const FStateTreeTaskBase& Task = GetNode<FStateTreeTaskBase>(TaskIndex);
@@ -558,19 +555,6 @@ void FStateTreeExecutionContext::ExitState(FStateTreeInstanceData& InstanceData,
 				}
 			}
 		}
-
-		// Evaluators
-		for (int32 EvalIndex = (int32)(State.EvaluatorsBegin + State.EvaluatorsNum) - 1; EvalIndex >= State.EvaluatorsBegin; EvalIndex--)
-		{
-			const FStateTreeEvaluatorBase& Eval = GetNode<FStateTreeEvaluatorBase>(EvalIndex);
-
-			STATETREE_LOG(Verbose, TEXT("%*s  Notify Evaluator '%s'."), Index*UE::StateTree::DebugIndentSize, TEXT(""), *Eval.Name.ToString());
-			{
-				QUICK_SCOPE_CYCLE_COUNTER(StateTree_Eval_ExitState);
-				CSV_SCOPED_TIMING_STAT_EXCLUSIVE(StateTree_Eval_ExitState);
-				Eval.ExitState(*this, ChangeType, CurrentTransition);
-			}
-		}
 	}
 }
 
@@ -607,86 +591,32 @@ void FStateTreeExecutionContext::StateCompleted(FStateTreeInstanceData& Instance
 				Task.StateCompleted(*this, Exec.LastTickStatus, Exec.ActiveStates);
 			}
 		}
-
-		// Notify evaluators
-		for (int32 EvalIndex = (int32)(State.EvaluatorsBegin + State.EvaluatorsNum) - 1; EvalIndex >= State.EvaluatorsBegin; EvalIndex--)
-		{
-			const FStateTreeEvaluatorBase& Eval = GetNode<FStateTreeEvaluatorBase>(EvalIndex);
-
-			STATETREE_LOG(Verbose, TEXT("%*s  Notify Evaluator '%s'"), Index*UE::StateTree::DebugIndentSize, TEXT(""), *Eval.Name.ToString());
-			Eval.StateCompleted(*this, Exec.LastTickStatus, Exec.ActiveStates);
-		}
-	}
-}
-
-void FStateTreeExecutionContext::TickEvaluatorsForSelect(FStateTreeInstanceData& InstanceData, const FStateTreeHandle CurrentState, const EStateTreeEvaluationType EvalType, const float DeltaTime)
-{
-	if (CurrentState == FStateTreeHandle::Succeeded || CurrentState == FStateTreeHandle::Failed || CurrentState == FStateTreeHandle::Invalid)
-	{
-		return;
-	}
-
-	FStateTreeActiveStates Branch;
-	FStateTreeHandle StateHandle = CurrentState;
-	while (StateHandle.IsValid())
-	{
-		if (!Branch.PushFront(StateHandle))
-		{
-			break;
-		}
-		StateHandle = StateTree->States[StateHandle.Index].Parent;
 	}
-
-	TickEvaluators(InstanceData, Branch, EvalType, DeltaTime);
 }
 
-void FStateTreeExecutionContext::TickEvaluators(FStateTreeInstanceData& InstanceData, const FStateTreeActiveStates& ActiveStates, const EStateTreeEvaluationType EvalType, const float DeltaTime)
+void FStateTreeExecutionContext::TickEvaluators(FStateTreeInstanceData& InstanceData, const float DeltaTime)
 {
 	CSV_SCOPED_TIMING_STAT_EXCLUSIVE(StateTree_TickEvaluators);
 
-	for (int32 Index = 0; Index < ActiveStates.Num(); Index++)
-	{
-		const FStateTreeHandle StateHandle = ActiveStates[Index];
+	STATETREE_CLOG(StateTree->EvaluatorsNum > 0, Verbose, TEXT("Ticking Evaluators"));
 
-		if (VisitedStates[StateHandle.Index])
-		{
-			// Already ticked this frame.
-			continue;
-		}
-
-		const FCompactStateTreeState& State = StateTree->States[StateHandle.Index];
+	// Tick evaluators
+	for (int32 EvalIndex = 0; EvalIndex < int32(StateTree->EvaluatorsNum); EvalIndex++)
+	{
+		const FStateTreeEvaluatorBase& Eval = GetNode<FStateTreeEvaluatorBase>(int32(StateTree->EvaluatorsBegin) + EvalIndex);
+		const FStateTreeDataView EvalInstanceView = GetInstanceData(InstanceData, Eval.bInstanceIsObject, Eval.InstanceIndex);
+		DataViews[Eval.DataViewIndex] = EvalInstanceView;
 
-		if (State.Type == EStateTreeStateType::Linked)
-		{
-			UpdateLinkedStateParameters(InstanceData, State);
-		}
-		else if (State.Type == EStateTreeStateType::Subtree)
+		// Copy bound properties.
+		if (Eval.BindingsBatch.IsValid())
 		{
-			UpdateSubtreeStateParameters(InstanceData, State);
+			StateTree->PropertyBindings.CopyTo(DataViews, Eval.BindingsBatch.Index, EvalInstanceView);
 		}
-		
-		STATETREE_CLOG(State.EvaluatorsNum > 0, Verbose, TEXT("%*sTicking Evaluators of state '%s' %s"), Index*UE::StateTree::DebugIndentSize, TEXT(""), *DebugGetStatePath(ActiveStates, Index), *UEnum::GetValueAsString(EvalType));
-
-		// Tick evaluators
-		for (int32 EvalIndex = 0; EvalIndex < int32(State.EvaluatorsNum); EvalIndex++)
+		STATETREE_LOG(Verbose, TEXT("  Tick: '%s'"), *Eval.Name.ToString());
 		{
-			const FStateTreeEvaluatorBase& Eval = GetNode<FStateTreeEvaluatorBase>(int32(State.EvaluatorsBegin) + EvalIndex);
-			const FStateTreeDataView EvalInstanceView = GetInstanceData(InstanceData, Eval.bInstanceIsObject, Eval.InstanceIndex);
-			DataViews[Eval.DataViewIndex] = EvalInstanceView;
-
-			// Copy bound properties.
-			if (Eval.BindingsBatch.IsValid())
-			{
-				StateTree->PropertyBindings.CopyTo(DataViews, Eval.BindingsBatch.Index, EvalInstanceView);
-			}
-			STATETREE_LOG(Verbose, TEXT("%*s  Evaluate: '%s'"), Index*UE::StateTree::DebugIndentSize, TEXT(""), *Eval.Name.ToString());
-			{
-				QUICK_SCOPE_CYCLE_COUNTER(StateTree_Eval_Tick);
-				Eval.Evaluate(*this, EvalType, DeltaTime);
-			}
+			QUICK_SCOPE_CYCLE_COUNTER(StateTree_Eval_Tick);
+			Eval.Tick(*this, DeltaTime);
 		}
-
-		VisitedStates[StateHandle.Index] = true;
 	}
 }
 
@@ -1012,9 +942,6 @@ bool FStateTreeExecutionContext::SelectStateInternal(FStateTreeInstanceData& Ins
 
 	const FCompactStateTreeState& State = StateTree->States[NextState.Index];
 
-	// Make sure all the evaluators for the target state are up to date.
-	TickEvaluatorsForSelect(InstanceData, NextState, EStateTreeEvaluationType::PreSelect, 0.0f);
-	
 	// Check that the state can be entered
 	if (TestAllConditions(State.EnterConditionsBegin, State.EnterConditionsNum))
 	{
@@ -1150,6 +1077,16 @@ FString FStateTreeExecutionContext::GetDebugInfoString(const FStateTreeInstanceD
 		DebugString += TEXT("--\n");
 	}
 
+	if (StateTree->EvaluatorsNum > 0)
+	{
+		DebugString += TEXT("\nEvaluators:\n");
+		for (int32 EvalIndex = 0; EvalIndex < int32(StateTree->EvaluatorsNum); EvalIndex++)
+		{
+			const FStateTreeEvaluatorBase& Eval = GetNode<FStateTreeEvaluatorBase>(int32(StateTree->EvaluatorsBegin) + EvalIndex);
+			Eval.AppendDebugInfoString(DebugString, *this);
+		}
+	}
+
 	// Active States
 	DebugString += TEXT("Current State:\n");
 	for (int32 Index = 0; Index < Exec.ActiveStates.Num(); Index++)
@@ -1169,15 +1106,6 @@ FString FStateTreeExecutionContext::GetDebugInfoString(const FStateTreeInstanceD
 					Task.AppendDebugInfoString(DebugString, *this);
 				}
 			}
-			if (State.EvaluatorsNum > 0)
-			{
-				DebugString += TEXT("\nEvaluators:\n");
-				for (int32 EvalIndex = 0; EvalIndex < int32(State.EvaluatorsNum); EvalIndex++)
-				{
-					const FStateTreeEvaluatorBase& Eval = GetNode<FStateTreeEvaluatorBase>(int32(State.EvaluatorsBegin) + EvalIndex);
-					Eval.AppendDebugInfoString(DebugString, *this);
-				}
-			}
 		}
 	}
 
@@ -1205,14 +1133,14 @@ void FStateTreeExecutionContext::DebugPrintInternalLayout(const FStateTreeInstan
 		DebugString += FString::Printf(TEXT("  %s\n"), Node.IsValid() ? *Node.GetScriptStruct()->GetName() : TEXT("null"));
 	}
 
-	// Instance InstanceData data (e.g. tasks, evaluators, conditions)
+	// Instance InstanceData data (e.g. tasks)
 	DebugString += FString::Printf(TEXT("\nInstance Data(%d)\n"), StateTree->Instances.Num());
 	for (const FInstancedStruct& Data : StateTree->Instances)
 	{
 		DebugString += FString::Printf(TEXT("  %s\n"), Data.IsValid() ? *Data.GetScriptStruct()->GetName() : TEXT("null"));
 	}
 
-	// Instance InstanceData offsets (e.g. tasks, evaluators, conditions)
+	// Instance InstanceData offsets (e.g. tasks)
 
 	if (const FStateTreeInstanceDataLayout* InstanceLayout = StateTree->GetInstanceDataDefaultValue().GetLayout().Get())
 	{
@@ -1264,28 +1192,29 @@ void FStateTreeExecutionContext::DebugPrintInternalLayout(const FStateTreeInstan
 		);
 	for (const FCompactStateTreeState& State : StateTree->States)
 	{
-		DebugString += FString::Printf(TEXT("  | %-30s | %15s | %5s [%3d:%-3d[ | %9s   %4d %4d %4d %4d | %3s   %4d %4d %4d %4d\n"),
+		DebugString += FString::Printf(TEXT("  | %-30s | %15s | %5s [%3d:%-3d[ | %9s   %4d %4d %4d | %3s   %4d %4d %4d\n"),
 									*State.Name.ToString(), *State.Parent.Describe(),
 									TEXT(""), State.ChildrenBegin, State.ChildrenEnd,
-									TEXT(""), State.EnterConditionsBegin, State.TransitionsBegin, State.TasksBegin, State.EvaluatorsBegin,
-									TEXT(""), State.EnterConditionsNum, State.TransitionsNum, State.TasksNum, State.EvaluatorsNum);
+									TEXT(""), State.EnterConditionsBegin, State.TransitionsBegin, State.TasksBegin,
+									TEXT(""), State.EnterConditionsNum, State.TransitionsNum, State.TasksNum);
 	}
 
-	DebugString += FString::Printf(TEXT("\nStates\n  [ %-30s | %-12s | %-30s | %15s | %10s ]\n"),
-		TEXT("State"), TEXT("Type"), TEXT("Name"), TEXT("Bindings"), TEXT("Struct Idx"));
-	for (const FCompactStateTreeState& State : StateTree->States)
+	// Evaluators
+	if (StateTree->EvaluatorsNum)
 	{
-		// Evaluators
-		if (State.EvaluatorsNum)
+		for (int32 EvalIndex = 0; EvalIndex < StateTree->EvaluatorsNum; EvalIndex++)
 		{
-			for (int32 EvalIndex = 0; EvalIndex < State.EvaluatorsNum; EvalIndex++)
-			{
-				const FStateTreeEvaluatorBase& Eval = GetNode<FStateTreeEvaluatorBase>(State.EvaluatorsBegin + EvalIndex);
-				DebugString += FString::Printf(TEXT("  | %-30s | %-12s | %-30s | %15s | %10d |\n"), *State.Name.ToString(),
-					TEXT("  Evaluator"), *Eval.Name.ToString(), *Eval.BindingsBatch.Describe(), Eval.DataViewIndex);
-			}
+			const FStateTreeEvaluatorBase& Eval = GetNode<FStateTreeEvaluatorBase>(StateTree->EvaluatorsBegin + EvalIndex);
+			DebugString += FString::Printf(TEXT("| %-12s | %-30s | %15s | %10d |\n"),
+				TEXT("Evaluator"), *Eval.Name.ToString(), *Eval.BindingsBatch.Describe(), Eval.DataViewIndex);
 		}
+	}
+
 
+	DebugString += FString::Printf(TEXT("\nStates\n  [ %-30s | %-12s | %-30s | %15s | %10s ]\n"),
+		TEXT("State"), TEXT("Type"), TEXT("Name"), TEXT("Bindings"), TEXT("Struct Idx"));
+	for (const FCompactStateTreeState& State : StateTree->States)
+	{
 		// Tasks
 		if (State.TasksNum)
 		{
diff --git a/Engine/Plugins/Runtime/StateTree/Source/StateTreeModule/Public/Blueprint/StateTreeEvaluatorBlueprintBase.h b/Engine/Plugins/Runtime/StateTree/Source/StateTreeModule/Public/Blueprint/StateTreeEvaluatorBlueprintBase.h
--- a/Engine/Plugins/Runtime/StateTree/Source/StateTreeModule/Public/Blueprint/StateTreeEvaluatorBlueprintBase.h
+++ b/Engine/Plugins/Runtime/StateTree/Source/StateTreeModule/Public/Blueprint/StateTreeEvaluatorBlueprintBase.h
@@ -21,26 +21,41 @@ class STATETREEMODULE_API UStateTreeEvaluatorBlueprintBase : public UStateTreeIt
 	GENERATED_BODY()
 public:
 
-	UFUNCTION(BlueprintImplementableEvent, meta = (DisplayName = "EnterState"))
+	UE_DEPRECATED(5.1, "This function will be removed for 5.1.")
+	UFUNCTION(BlueprintImplementableEvent, meta = (DisplayName = "EnterState", DeprecatedFunction, DeprecationMessage = "Evaluators are not ticked at state level anymore, Use ReceiveTreeStart instead."))
 	void ReceiveEnterState(AActor* OwnerActor, const EStateTreeStateChangeType ChangeType, const FStateTreeTransitionResult& Transition);
 
-	UFUNCTION(BlueprintImplementableEvent, meta = (DisplayName = "ExitState"))
+	UE_DEPRECATED(5.1, "This function will be removed for 5.1.")
+	UFUNCTION(BlueprintImplementableEvent, meta = (DisplayName = "ExitState", DeprecatedFunction, DeprecationMessage = "Evaluators are not ticked at state level anymore, Use ReceiveTreeStop."))
 	void ReceiveExitState(AActor* OwnerActor, const EStateTreeStateChangeType ChangeType, const FStateTreeTransitionResult& Transition);
 
-	UFUNCTION(BlueprintImplementableEvent, meta = (DisplayName = "StateCompleted"))
+	UE_DEPRECATED(5.1, "This function will be removed for 5.1.")
+	UFUNCTION(BlueprintImplementableEvent, meta = (DisplayName = "StateCompleted", DeprecatedFunction, DeprecationMessage = "Evaluators are not ticked at state level anymore, Use ReceiveTreeStop instead."))
 	void ReceiveStateCompleted(AActor* OwnerActor, const EStateTreeRunStatus CompletionStatus, const FStateTreeActiveStates CompletedState);
 
-	UFUNCTION(BlueprintImplementableEvent, meta = (DisplayName = "Evaluate"))
+	UE_DEPRECATED(5.1, "This function will be removed for 5.1.")
+	UFUNCTION(BlueprintImplementableEvent, meta = (DisplayName = "Evaluate", DeprecatedFunction, DeprecationMessage = "Use ReceiveTick instead."))
 	void ReceiveEvaluate(AActor* OwnerActor, const EStateTreeEvaluationType EvalType, const float DeltaTime);
 
+	UFUNCTION(BlueprintImplementableEvent, meta = (DisplayName = "TreeStart"))
+	void ReceiveTreeStart(AActor* OwnerActor);
+
+	UFUNCTION(BlueprintImplementableEvent, meta = (DisplayName = "TreeStop"))
+	void ReceiveTreeStop(AActor* OwnerActor);
+
+	UFUNCTION(BlueprintImplementableEvent, meta = (DisplayName = "Tick"))
+	void ReceiveTick(AActor* OwnerActor, const float DeltaTime);
 	
+	UE_DEPRECATED(5.1, "This function will be removed for 5.1.")
 	virtual void EnterState(FStateTreeExecutionContext& Context, const EStateTreeStateChangeType ChangeType, const FStateTreeTransitionResult& Transition);
-
+	UE_DEPRECATED(5.1, "This function will be removed for 5.1.")
 	virtual void ExitState(FStateTreeExecutionContext& Context, const EStateTreeStateChangeType ChangeType, const FStateTreeTransitionResult& Transition);
-
+	UE_DEPRECATED(5.1, "This function will be removed for 5.1.")
 	virtual void StateCompleted(FStateTreeExecutionContext& Context, const EStateTreeRunStatus CompletionStatus, const FStateTreeActiveStates& CompletedActiveStates);
 
-	virtual void Evaluate(FStateTreeExecutionContext& Context, const EStateTreeEvaluationType EvalType, const float DeltaTime);
+	virtual void TreeStart(FStateTreeExecutionContext& Context);
+	virtual void TreeStop(FStateTreeExecutionContext& Context);
+	virtual void Tick(FStateTreeExecutionContext& Context, const float DeltaTime);
 };
 
 /**
@@ -53,10 +68,14 @@ struct STATETREEMODULE_API FStateTreeBlueprintEvaluatorWrapper : public FStateTr
 
 	virtual const UStruct* GetInstanceDataType() const override { return EvaluatorClass; };
 	virtual bool Link(FStateTreeLinker& Linker) override;
+	
 	virtual void EnterState(FStateTreeExecutionContext& Context, const EStateTreeStateChangeType ChangeType, const FStateTreeTransitionResult& Transition) const override;
 	virtual void ExitState(FStateTreeExecutionContext& Context, const EStateTreeStateChangeType ChangeType, const FStateTreeTransitionResult& Transition) const override;
 	virtual void StateCompleted(FStateTreeExecutionContext& Context, const EStateTreeRunStatus CompletionStatus, const FStateTreeActiveStates& CompletedActiveStates) const override;
-	virtual void Evaluate(FStateTreeExecutionContext& Context, const EStateTreeEvaluationType EvalType, const float DeltaTime) const override;
+
+	virtual void TreeStart(FStateTreeExecutionContext& Context) const override;
+	virtual void TreeStop(FStateTreeExecutionContext& Context) const override;
+	virtual void Tick(FStateTreeExecutionContext& Context, const float DeltaTime) const override;
 
 	UPROPERTY()
 	TSubclassOf<UStateTreeEvaluatorBlueprintBase> EvaluatorClass = nullptr;
diff --git a/Engine/Plugins/Runtime/StateTree/Source/StateTreeModule/Public/StateTree.h b/Engine/Plugins/Runtime/StateTree/Source/StateTreeModule/Public/StateTree.h
--- a/Engine/Plugins/Runtime/StateTree/Source/StateTreeModule/Public/StateTree.h
+++ b/Engine/Plugins/Runtime/StateTree/Source/StateTreeModule/Public/StateTree.h
@@ -20,6 +20,8 @@ struct STATETREEMODULE_API FStateTreeCustomVersion
 		BeforeCustomVersionWasAdded = 0,
 		// Separated conditions to shared instance data.
 		SharedInstanceData,
+		// Moved evaluators to be global.
+		GlobalEvaluators,
 
 		// -----<new versions can be added above this line>-------------------------------------------------
 		VersionPlusOne,
@@ -151,7 +153,13 @@ private:
 	/** List of names external data enforced by the schema, created at compilation. */
 	UPROPERTY()
 	TArray<FStateTreeExternalDataDesc> NamedExternalDataDescs;
+
+	UPROPERTY()
+	uint16 EvaluatorsBegin = 0;
 	
+	UPROPERTY()
+	uint16 EvaluatorsNum = 0;
+
 	UPROPERTY(Transient)
 	int32 NumDataViews = 0;
 
diff --git a/Engine/Plugins/Runtime/StateTree/Source/StateTreeModule/Public/StateTreeEvaluatorBase.h b/Engine/Plugins/Runtime/StateTree/Source/StateTreeModule/Public/StateTreeEvaluatorBase.h
--- a/Engine/Plugins/Runtime/StateTree/Source/StateTreeModule/Public/StateTreeEvaluatorBase.h
+++ b/Engine/Plugins/Runtime/StateTree/Source/StateTreeModule/Public/StateTreeEvaluatorBase.h
@@ -24,6 +24,7 @@ struct STATETREEMODULE_API FStateTreeEvaluatorBase : public FStateTreeNodeBase
 	 * @param ChangeType Describes the change type (Changed/Sustained).
 	 * @param Transition Describes the states involved in the transition
 	 */
+	UE_DEPRECATED(5.1, "This function will be removed for 5.1.")
 	virtual void EnterState(FStateTreeExecutionContext& Context, const EStateTreeStateChangeType ChangeType, const FStateTreeTransitionResult& Transition) const {}
 
 	/**
@@ -33,6 +34,7 @@ struct STATETREEMODULE_API FStateTreeEvaluatorBase : public FStateTreeNodeBase
 	 * @param ChangeType Describes the change type (Changed/Sustained).
 	 * @param Transition Describes the states involved in the transition
 	 */
+	UE_DEPRECATED(5.1, "This function will be removed for 5.1.")
 	virtual void ExitState(FStateTreeExecutionContext& Context, const EStateTreeStateChangeType ChangeType, const FStateTreeTransitionResult& Transition) const {}
 
 	/**
@@ -42,17 +44,36 @@ struct STATETREEMODULE_API FStateTreeEvaluatorBase : public FStateTreeNodeBase
 	 * @param CompletionStatus Describes the running status of the completed state (Succeeded/Failed).
 	 * @param CompletedActiveStates Active states at the time of completion.
 	 */
+	UE_DEPRECATED(5.1, "This function will be removed for 5.1.")
 	virtual void StateCompleted(FStateTreeExecutionContext& Context, const EStateTreeRunStatus CompletionStatus, const FStateTreeActiveStates& CompletedActiveStates) const {}
 	
+	UE_DEPRECATED(5.1, "This function will be removed for 5.1, use Tick() instead.")
+	virtual void Evaluate(FStateTreeExecutionContext& Context, const EStateTreeEvaluationType EvalType, const float DeltaTime) const { }
+
+	
+	/**
+	 * Called when StateTree is started.
+	 * @param Context Reference to current execution context.
+	 */
+	virtual void TreeStart(FStateTreeExecutionContext& Context) const {}
+
+	/**
+	 * Called when StateTree is stopped.
+	 * @param Context Reference to current execution context.
+	 */
+	virtual void TreeStop(FStateTreeExecutionContext& Context) const {}
+
 	/**
-	 * Called when evaluator needs to be updated. EvalType describes if the tick happens during state tree tick when the evaluator is on active state (Tick),
-	 * or during state selection process when the evaluator's state is visited while it's inactive (PreSelection).
-	 * That is, type "Tick" means that the call happens between EnterState()/ExitState() pair, "PreSelection" is used otherwise.
+	 * Called each frame to update the evaluator.
 	 * @param Context Reference to current execution context.
-	 * @param EvalType Describes tick type.
 	 * @param DeltaTime Time since last StateTree tick, or 0 if called during preselection.
 	 */
-	virtual void Evaluate(FStateTreeExecutionContext& Context, const EStateTreeEvaluationType EvalType, const float DeltaTime) const {}
+	virtual void Tick(FStateTreeExecutionContext& Context, const float DeltaTime) const
+	{
+		PRAGMA_DISABLE_DEPRECATION_WARNINGS
+		Evaluate(Context, EStateTreeEvaluationType::Tick, DeltaTime);
+		PRAGMA_ENABLE_DEPRECATION_WARNINGS
+	}
 
 #if WITH_GAMEPLAY_DEBUGGER
 	virtual void AppendDebugInfoString(FString& DebugString, const FStateTreeExecutionContext& Context) const;
diff --git a/Engine/Plugins/Runtime/StateTree/Source/StateTreeModule/Public/StateTreeExecutionContext.h b/Engine/Plugins/Runtime/StateTree/Source/StateTreeModule/Public/StateTreeExecutionContext.h
--- a/Engine/Plugins/Runtime/StateTree/Source/StateTreeModule/Public/StateTreeExecutionContext.h
+++ b/Engine/Plugins/Runtime/StateTree/Source/StateTreeModule/Public/StateTreeExecutionContext.h
@@ -304,10 +304,9 @@ protected:
 	void TickEvaluatorsForSelect(FStateTreeInstanceData& InstanceData, const FStateTreeHandle CurrentState, const EStateTreeEvaluationType EvalType, const float DeltaTime);
 
 	/**
-	 * Ticks evaluators of all active states by delta time.
-	 * If TickEvaluators() is called multiple times per frame (i.e. during selection when visiting new states), each state and evaluator is ticked only once.
+	 * Ticks global evaluators by delta time.
 	 */
-	void TickEvaluators(FStateTreeInstanceData& InstanceData, const FStateTreeActiveStates& ActiveStates, const EStateTreeEvaluationType EvalType, const float DeltaTime);
+	void TickEvaluators(FStateTreeInstanceData& InstanceData, const float DeltaTime);
 
 	/**
 	 * Ticks tasks of all active states starting from current state by delta time.
@@ -417,9 +416,6 @@ protected:
 	UPROPERTY(Transient)
 	FStateTreeInstanceData InternalInstanceData;
 
-	/** States visited during a tick while updating evaluators. Initialized to match the number of states in the asset. */ 
-	TArray<bool> VisitedStates;
-
 	/** Array of data pointers (external data, tasks, evaluators, conditions), used during evaluation. Initialized to match the number of items in the asset. */
 	TArray<FStateTreeDataView> DataViews;
 
diff --git a/Engine/Plugins/Runtime/StateTree/Source/StateTreeModule/Public/StateTreeTypes.h b/Engine/Plugins/Runtime/StateTree/Source/StateTreeModule/Public/StateTreeTypes.h
--- a/Engine/Plugins/Runtime/StateTree/Source/StateTreeModule/Public/StateTreeTypes.h
+++ b/Engine/Plugins/Runtime/StateTree/Source/StateTreeModule/Public/StateTreeTypes.h
@@ -372,8 +372,6 @@ struct STATETREEMODULE_API FCompactStateTreeState
 	uint16 TransitionsBegin = 0;						// Index to first transition
 	UPROPERTY()
 	uint16 TasksBegin = 0;								// Index to first task
-	UPROPERTY()
-	uint16 EvaluatorsBegin = 0;							// Index to first evaluator
 
 	UPROPERTY()
 	uint16 ParameterInstanceIndex = MAX_uint16;			// Index to state instance data
@@ -387,8 +385,6 @@ struct STATETREEMODULE_API FCompactStateTreeState
 	UPROPERTY()
 	uint8 TasksNum = 0;									// Number of tasks
 	UPROPERTY()
-	uint8 EvaluatorsNum = 0;							// Number of evaluators
-	UPROPERTY()
 	EStateTreeStateType Type = EStateTreeStateType::State; 
 };
 
diff --git a/Engine/Plugins/Runtime/StateTree/Source/StateTreeTestSuite/Private/StateTreeTest.cpp b/Engine/Plugins/Runtime/StateTree/Source/StateTreeTestSuite/Private/StateTreeTest.cpp
--- a/Engine/Plugins/Runtime/StateTree/Source/StateTreeTestSuite/Private/StateTreeTest.cpp
+++ b/Engine/Plugins/Runtime/StateTree/Source/StateTreeTestSuite/Private/StateTreeTest.cpp
@@ -42,7 +42,7 @@ struct FStateTreeTest_MakeAndBakeStateTree : FAITestBase
 		UStateTreeState& StateB = Root.AddChildState(FName(TEXT("B")));
 
 		// Root
-		auto& EvalA = Root.AddEvaluator<FTestEval_A>();
+		auto& EvalA = EditorData.AddEvaluator<FTestEval_A>();
 		
 		// State A
 		auto& TaskB1 = StateA.AddTask<FTestTask_B>();
