Static analysis fix for FStateTreeReference
#rnx

diff --git a/Engine/Plugins/Runtime/StateTree/Source/StateTreeModule/Private/StateTreeReference.cpp b/Engine/Plugins/Runtime/StateTree/Source/StateTreeModule/Private/StateTreeReference.cpp
--- a/Engine/Plugins/Runtime/StateTree/Source/StateTreeModule/Private/StateTreeReference.cpp
+++ b/Engine/Plugins/Runtime/StateTree/Source/StateTreeModule/Private/StateTreeReference.cpp
@@ -33,7 +33,7 @@ FStateTreeReference::~FStateTreeReference()
 #if WITH_EDITOR
 void FStateTreeReference::PostSerialize(const FArchive& Ar)
 {
-	if (Ar.IsLoading() & Ar.IsPersistent())
+	if (Ar.IsLoading() && Ar.IsPersistent())
 	{
 		// This might modify the object but we don't want to dirty on load
 		SyncParameters();
