StateTree: Fix for abstract classes showing up in class selector
- added condition to prevent abstract classes being selected
- made some base node types hidden

diff --git a/Engine/Plugins/Runtime/StateTree/Source/StateTreeEditorModule/Private/Customizations/StateTreeEditorNodeDetails.cpp b/Engine/Plugins/Runtime/StateTree/Source/StateTreeEditorModule/Private/Customizations/StateTreeEditorNodeDetails.cpp
--- a/Engine/Plugins/Runtime/StateTree/Source/StateTreeEditorModule/Private/Customizations/StateTreeEditorNodeDetails.cpp
+++ b/Engine/Plugins/Runtime/StateTree/Source/StateTreeEditorModule/Private/Customizations/StateTreeEditorNodeDetails.cpp
@@ -1021,17 +1021,17 @@ TSharedRef<SWidget> FStateTreeEditorNodeDetails::GeneratePicker()
 			{
 				continue;
 			}
-			if (Class->HasAnyClassFlags(CLASS_Hidden | CLASS_HideDropDown))
+			if (Class->HasAnyClassFlags(CLASS_Abstract | CLASS_Hidden | CLASS_HideDropDown))
 			{
 				continue;				
 			}
 			if (Class->HasMetaData(TEXT("Hidden")))
 			{
-				continue;				
+				continue;
 			}
 			if (Schema && !Schema->IsClassAllowed(Class))
 			{
-				continue;				
+				continue;
 			}
 
 			FUIAction ItemAction(FExecuteAction::CreateSP(this, &FStateTreeEditorNodeDetails::OnClassPicked, Class));
diff --git a/Engine/Plugins/Runtime/StateTree/Source/StateTreeModule/Public/StateTreeConditionBase.h b/Engine/Plugins/Runtime/StateTree/Source/StateTreeModule/Public/StateTreeConditionBase.h
--- a/Engine/Plugins/Runtime/StateTree/Source/StateTreeModule/Public/StateTreeConditionBase.h
+++ b/Engine/Plugins/Runtime/StateTree/Source/StateTreeModule/Public/StateTreeConditionBase.h
@@ -18,7 +18,7 @@ struct FStateTreeEditorPropertyPath;
 /**
  * Base struct for all conditions.
  */
-USTRUCT()
+USTRUCT(meta = (Hidden))
 struct STATETREEMODULE_API FStateTreeConditionBase : public FStateTreeNodeBase
 {
 	GENERATED_BODY()
diff --git a/Engine/Plugins/Runtime/StateTree/Source/StateTreeModule/Public/StateTreeEvaluatorBase.h b/Engine/Plugins/Runtime/StateTree/Source/StateTreeModule/Public/StateTreeEvaluatorBase.h
--- a/Engine/Plugins/Runtime/StateTree/Source/StateTreeModule/Public/StateTreeEvaluatorBase.h
+++ b/Engine/Plugins/Runtime/StateTree/Source/StateTreeModule/Public/StateTreeEvaluatorBase.h
@@ -12,7 +12,7 @@ struct FStateTreeExecutionContext;
  * Base struct of StateTree Evaluators.
  * Evaluators calculate and expose data to be used for decision making in a StateTree.
  */
-USTRUCT()
+USTRUCT(meta = (Hidden))
 struct STATETREEMODULE_API FStateTreeEvaluatorBase : public FStateTreeNodeBase
 {
 	GENERATED_BODY()
diff --git a/Engine/Plugins/Runtime/StateTree/Source/StateTreeModule/Public/StateTreeTaskBase.h b/Engine/Plugins/Runtime/StateTree/Source/StateTreeModule/Public/StateTreeTaskBase.h
--- a/Engine/Plugins/Runtime/StateTree/Source/StateTreeModule/Public/StateTreeTaskBase.h
+++ b/Engine/Plugins/Runtime/StateTree/Source/StateTreeModule/Public/StateTreeTaskBase.h
@@ -12,7 +12,7 @@ struct FStateTreeExecutionContext;
  * Base struct for StateTree Tasks.
  * Tasks are logic executed in an active state.
  */
-USTRUCT()
+USTRUCT(meta = (Hidden))
 struct STATETREEMODULE_API FStateTreeTaskBase : public FStateTreeNodeBase
 {
 	GENERATED_BODY()
