StateTree: Fixed handling of missing instance data types
- Fixed alignment of null structs from 0 to 1
- 0 causes the offset to reset, which in turn could could cause 0 size malloc

diff --git a/Engine/Plugins/Runtime/StateTree/Source/StateTreeModule/Private/StateTreeInstanceData.cpp b/Engine/Plugins/Runtime/StateTree/Source/StateTreeModule/Private/StateTreeInstanceData.cpp
--- a/Engine/Plugins/Runtime/StateTree/Source/StateTreeModule/Private/StateTreeInstanceData.cpp
+++ b/Engine/Plugins/Runtime/StateTree/Source/StateTreeModule/Private/StateTreeInstanceData.cpp
@@ -28,7 +28,7 @@ TSharedPtr<FStateTreeInstanceDataLayout> FStateTreeInstanceDataLayout::Create(TC
 	{
 		FLayoutItem& Item = Items[ItemIndex];
 		const UScriptStruct* ScriptStruct = Structs[ItemIndex];
-		const int32 Alignment = ScriptStruct != nullptr ? ScriptStruct->GetMinAlignment() : 0;
+		const int32 Alignment = ScriptStruct != nullptr ? ScriptStruct->GetMinAlignment() : 1;
 		const int32 Size = ScriptStruct != nullptr ? ScriptStruct->GetStructureSize() : 0;
 
 		Offset = Align(Offset, Alignment);
@@ -63,7 +63,7 @@ int32 FStateTreeInstanceDataLayout::GetLayoutInstanceMinAlignment() const
 	for (int32 Index = 0; Index < NumItems; Index++)
 	{
 		const FLayoutItem& Item = Items[Index];
-		const int32 Alignment = Item.ScriptStruct != nullptr ? Item.ScriptStruct->GetMinAlignment() : 0;
+		const int32 Alignment = Item.ScriptStruct != nullptr ? Item.ScriptStruct->GetMinAlignment() : 1;
 		CombinedAlignment = FMath::Max(CombinedAlignment, Alignment);
 	}
 		
