Fixing bad merge

diff --git a/Engine/Plugins/Runtime/StateTree/Source/StateTreeEditorModule/Public/StateTreeState.h b/Engine/Plugins/Runtime/StateTree/Source/StateTreeEditorModule/Public/StateTreeState.h
--- a/Engine/Plugins/Runtime/StateTree/Source/StateTreeEditorModule/Public/StateTreeState.h
+++ b/Engine/Plugins/Runtime/StateTree/Source/StateTreeEditorModule/Public/StateTreeState.h
@@ -73,8 +73,8 @@ struct STATETREEEDITORMODULE_API FStateTreeTransition
 	UPROPERTY(EditDefaultsOnly, Category = Transition, meta = (UIMin = "0", ClampMin = "0", UIMax = "25", ClampMax = "25"))
 	float GateDelay = 0.0f;
 	
-	UPROPERTY(EditDefaultsOnly, Category = Transition, meta = (BaseStruct = "StateTreeConditionBase", BaseClass = "StateTreeConditionBlueprintBase"))
-	TArray<FStateTreeConditionItem> Conditions;
+	UPROPERTY(EditDefaultsOnly, Category = Transition, meta = (BaseStruct = "/Script/StateTreeModule.StateTreeConditionBase", BaseClass = "/Script/StateTreeModule.StateTreeConditionBlueprintBase"))
+	TArray<FStateTreeEditorNode> Conditions;
 };
 
 USTRUCT()
