StateTree: refactored index types
- Changed FStateTreeHandle to FStateTreeStateHandle, used only for indexing states
- Added uint16 and uin8 index types, which can represent invalid index (aka INDEX_NONE)
- Changed indices that can be optional to the index types above
- Added validation and error loggic when index types overflow during compile
- Removed 2 indirections (in common case) and halved the memory usage of property copies

diff --git a/Engine/Plugins/Runtime/StateTree/Config/BaseStateTree.ini b/Engine/Plugins/Runtime/StateTree/Config/BaseStateTree.ini
--- a/Engine/Plugins/Runtime/StateTree/Config/BaseStateTree.ini
+++ b/Engine/Plugins/Runtime/StateTree/Config/BaseStateTree.ini
@@ -14,6 +14,7 @@
 
 +StructRedirects=(OldName="BakedStateTransition",NewName="CompactStateTransition")
 +StructRedirects=(OldName="BakedStateTreeState",NewName="CompactStateTreeState")
++StructRedirects=(OldName="StateTreeHandle",NewName="StateTreeStateHandle")
 
 +PropertyRedirects=(OldName="StateTreeState.Transitions2",NewName="StateTreeState.Transitions")
 +PropertyRedirects=(OldName="StateTreeState.EnterConditions2",NewName="StateTreeState.EnterConditions")
diff --git a/Engine/Plugins/Runtime/StateTree/Source/StateTreeEditorModule/Private/StateTreeCompiler.cpp b/Engine/Plugins/Runtime/StateTree/Source/StateTreeEditorModule/Private/StateTreeCompiler.cpp
--- a/Engine/Plugins/Runtime/StateTree/Source/StateTreeEditorModule/Private/StateTreeCompiler.cpp
+++ b/Engine/Plugins/Runtime/StateTree/Source/StateTreeEditorModule/Private/StateTreeCompiler.cpp
@@ -12,6 +12,15 @@
 #include "StateTreeExecutionContext.h"
 #include "StateTreePropertyBindingCompiler.h"
 
+namespace UE::StateTree::Compiler
+{
+	void FValidationResult::Log(FStateTreeCompilerLog& Log, const TCHAR* ContextText, const FStateTreeBindableStructDesc& ContextStruct) const
+	{
+		Log.Reportf(EMessageSeverity::Error, ContextStruct, TEXT("The StateTree is too complex. Compact index %s out of range %d/%d."), ContextText, Value, MaxValue);
+	}
+
+}; // UE::StateTree::Compiler
+
 bool FStateTreeCompiler::Compile(UStateTree& InStateTree)
 {
 	StateTree = &InStateTree;
@@ -37,8 +46,20 @@
 	StateTree->Parameters = TreeData->RootParameters.Parameters;
 	
 	// Mark parameters as binding source
-	StateTree->ParametersDataViewIndex = BindingsCompiler.AddSourceStruct(
-		{ TEXT("Parameters"), StateTree->Parameters.GetPropertyBagStruct(), EStateTreeBindableStructSource::TreeParameter, TreeData->RootParameters.ID });
+	const FStateTreeBindableStructDesc ParametersDesc = {
+			TEXT("Parameters"),
+			StateTree->Parameters.GetPropertyBagStruct(),
+			EStateTreeBindableStructSource::TreeParameter,
+			TreeData->RootParameters.ID
+		};
+	const int32 ParametersDataViewIndex = BindingsCompiler.AddSourceStruct(ParametersDesc);
+
+	if (const auto Validation = UE::StateTree::Compiler::IsValidIndex8(ParametersDataViewIndex); Validation.DidFail())
+	{
+		Validation.Log(Log, TEXT("ParametersDataViewIndex"), ParametersDesc);
+		return false;
+	}
+	StateTree->ParametersDataViewIndex = FStateTreeIndex8(ParametersDataViewIndex); 
 	
 	// Mark all named external values as binding source
 	if (StateTree->Schema)
@@ -46,7 +67,19 @@
 		StateTree->NamedExternalDataDescs = StateTree->Schema->GetNamedExternalDataDescs();
 		for (FStateTreeExternalDataDesc& Desc : StateTree->NamedExternalDataDescs)
 		{
-			Desc.Handle.DataViewIndex = BindingsCompiler.AddSourceStruct({Desc.Name, Desc.Struct, EStateTreeBindableStructSource::TreeData, Desc.ID});
+			const FStateTreeBindableStructDesc ExtDataDesc = {
+					Desc.Name,
+					Desc.Struct,
+					EStateTreeBindableStructSource::TreeData,
+					Desc.ID
+				};
+			const int32 ExternalStructIndex = BindingsCompiler.AddSourceStruct(ExtDataDesc);
+			if (const auto Validation = UE::StateTree::Compiler::IsValidIndex8(ExternalStructIndex); Validation.DidFail())
+			{
+				Validation.Log(Log, TEXT("ExternalStructIndex"), ParametersDesc);
+				return false;
+			}
+			Desc.Handle.DataViewIndex = FStateTreeIndex8(ExternalStructIndex); 
 		} 
 	}
 	
@@ -89,15 +122,15 @@
 	return true;
 }
 
-FStateTreeHandle FStateTreeCompiler::GetStateHandle(const FGuid& StateID) const
+FStateTreeStateHandle FStateTreeCompiler::GetStateHandle(const FGuid& StateID) const
 {
 	const int32* Idx = IDToState.Find(StateID);
 	if (Idx == nullptr)
 	{
-		return FStateTreeHandle::Invalid;
+		return FStateTreeStateHandle::Invalid;
 	}
 
-	return FStateTreeHandle(uint16(*Idx));
+	return FStateTreeStateHandle(uint16(*Idx));
 }
 
 UStateTreeState* FStateTreeCompiler::GetState(const FGuid& StateID)
@@ -121,7 +154,7 @@
 	{
 		if (SubTree != nullptr)
 		{
-			if (!CreateStateRecursive(*SubTree, FStateTreeHandle::Invalid))
+			if (!CreateStateRecursive(*SubTree, FStateTreeStateHandle::Invalid))
 			{
 				return false;
 			}
@@ -139,7 +172,7 @@
 			{
 				if (State->Type == EStateTreeStateType::Subtree)
 				{
-					if (!CreateStateRecursive(*State, FStateTreeHandle::Invalid))
+					if (!CreateStateRecursive(*State, FStateTreeStateHandle::Invalid))
 					{
 						return false;
 					}
@@ -152,7 +185,7 @@
 	return true;
 }
 
-bool FStateTreeCompiler::CreateStateRecursive(UStateTreeState& State, const FStateTreeHandle Parent)
+bool FStateTreeCompiler::CreateStateRecursive(UStateTreeState& State, const FStateTreeStateHandle Parent)
 {
 	FStateTreeCompilerLogStateScope LogStateScope(&State, Log);
 
@@ -167,20 +200,32 @@
 	IDToState.Add(State.ID, StateIdx);
 
 	// Child states
-	check(StateTree->States.Num() <= int32(MAX_uint16));
-	CompactState.ChildrenBegin = uint16(StateTree->States.Num());
+	const int32 ChildrenBegin = StateTree->States.Num();
+	if (const auto Validation = UE::StateTree::Compiler::IsValidCount16(ChildrenBegin); Validation.DidFail())
+	{
+		Validation.Log(Log, TEXT("ChildrenBegin"));
+		return false;
+	}
+	CompactState.ChildrenBegin = uint16(ChildrenBegin);
+	
 	for (UStateTreeState* Child : State.Children)
 	{
 		if (Child != nullptr && Child->Type != EStateTreeStateType::Subtree)
 		{
-			if (!CreateStateRecursive(*Child, FStateTreeHandle((uint16)StateIdx)))
+			if (!CreateStateRecursive(*Child, FStateTreeStateHandle((uint16)StateIdx)))
 			{
 				return false;
 			}
 		}
 	}
-	check(StateTree->States.Num() <= int32(MAX_uint16));
-	StateTree->States[StateIdx].ChildrenEnd = uint16(StateTree->States.Num()); // Cannot use CompactState here, it may be invalid due to array resize.
+	
+	const int32 ChildrenEnd = StateTree->States.Num();
+	if (const auto Validation = UE::StateTree::Compiler::IsValidCount16(ChildrenEnd); Validation.DidFail())
+	{
+		Validation.Log(Log, TEXT("ChildrenEnd"));
+		return false;
+	}
+	StateTree->States[StateIdx].ChildrenEnd = uint16(ChildrenEnd);
 
 	return true;
 }
@@ -194,12 +239,10 @@
 		// First operand should be copy as we dont have a previous item to operate on.
 		const EStateTreeConditionOperand Operand = bIsFirst ? EStateTreeConditionOperand::Copy : CondNode.ConditionOperand;
 		// First indent must be 0 to make the parentheses calculation match.
-		const int32 CurrIndent = bIsFirst ? 0 : CondNode.ConditionIndent;
+		const int32 CurrIndent = bIsFirst ? 0 : FMath::Clamp((int32)CondNode.ConditionIndent, 0, UE::StateTree::MaxConditionIndent);
 		// Next indent, or terminate at zero.
-		const int32 NextIndent = Conditions.IsValidIndex(Index + 1) ? Conditions[Index].ConditionIndent : 0;
+		const int32 NextIndent = Conditions.IsValidIndex(Index + 1) ? FMath::Clamp((int32)Conditions[Index].ConditionIndent, 0, UE::StateTree::MaxConditionIndent) : 0;
 		const int32 DeltaIndent = NextIndent - CurrIndent;
-		
-		check(DeltaIndent >= MIN_int8 && DeltaIndent <= MAX_int8);
 
 		if (!CreateCondition(State, CondNode, Operand, (int8)DeltaIndent))
 		{
@@ -212,8 +255,13 @@
 
 bool FStateTreeCompiler::CreateEvaluators()
 {
-	check(Nodes.Num() <= int32(MAX_uint16));
-	StateTree->EvaluatorsBegin = uint16(Nodes.Num());
+	const int32 EvaluatorsBegin = Nodes.Num();
+	if (const auto Validation = UE::StateTree::Compiler::IsValidCount16(EvaluatorsBegin); Validation.DidFail())
+	{
+		Validation.Log(Log, TEXT("EvaluatorsBegin"));
+		return false;
+	}
+	StateTree->EvaluatorsBegin = uint16(EvaluatorsBegin);
 
 	for (FStateTreeEditorNode& EvalNode : TreeData->Evaluators)
 	{
@@ -223,8 +271,12 @@
 		}
 	}
 	
-	const int32 EvaluatorsNum = Nodes.Num() - int32(StateTree->EvaluatorsBegin);
-	check(EvaluatorsNum <= int32(MAX_uint16));
+	const int32 EvaluatorsNum = Nodes.Num() - EvaluatorsBegin;
+	if (const auto Validation = UE::StateTree::Compiler::IsValidCount16(EvaluatorsNum); Validation.DidFail())
+	{
+		Validation.Log(Log, TEXT("EvaluatorsNum"));
+		return false;
+	}
 	StateTree->EvaluatorsNum = uint16(EvaluatorsNum);
 
 	return true;
@@ -247,9 +299,12 @@
 			// This allows to resolve the binding paths and lets us have bindable parameters when transitioned into a parameterized subtree directly.
 			FInstancedStruct& Instance = InstanceStructs.AddDefaulted_GetRef();
 			const int32 InstanceIndex = InstanceStructs.Num() - 1;
-
-			check(InstanceIndex <= int32(MAX_uint16));
-			CompactState.ParameterInstanceIndex = InstanceIndex;
+			if (const auto Validation = UE::StateTree::Compiler::IsValidIndex16(InstanceIndex); Validation.DidFail())
+			{
+				Validation.Log(Log, TEXT("InstanceIndex"));
+				return false;
+			}
+			CompactState.ParameterInstanceIndex = FStateTreeIndex16(InstanceIndex);
 		
 			Instance.InitializeAs<FCompactStateTreeParameters>();
 			FCompactStateTreeParameters& CompactParams = Instance.GetMutable<FCompactStateTreeParameters>();
@@ -259,37 +314,60 @@
 			if (SourceState->Type == EStateTreeStateType::Subtree)
 			{
 				// Register a binding source
-				CompactState.ParameterDataViewIndex = BindingsCompiler.AddSourceStruct({ SourceState->Name, SourceState->Parameters.Parameters.GetPropertyBagStruct(), EStateTreeBindableStructSource::StateParameter, SourceState->Parameters.ID });
+				const FStateTreeBindableStructDesc SubtreeParamsDesc = {
+						SourceState->Name,
+						SourceState->Parameters.Parameters.GetPropertyBagStruct(),
+						EStateTreeBindableStructSource::StateParameter,
+						SourceState->Parameters.ID
+					};
+				const int32 SourceStructIndex = BindingsCompiler.AddSourceStruct(SubtreeParamsDesc);
+				if (const auto Validation = UE::StateTree::Compiler::IsValidIndex16(SourceStructIndex); Validation.DidFail())
+				{
+					Validation.Log(Log, TEXT("SourceStructIndex"), SubtreeParamsDesc);
+					return false;
+				}
+				CompactState.ParameterDataViewIndex = FStateTreeIndex16(SourceStructIndex);
 			}
 			else if (SourceState->Type == EStateTreeStateType::Linked)
 			{
 				// Binding target
-				FStateTreeBindableStructDesc StructDesc;
-				StructDesc.ID = SourceState->Parameters.ID;
-				StructDesc.Name = SourceState->Name;
-				StructDesc.DataSource = EStateTreeBindableStructSource::StateParameter;
-				StructDesc.Struct = SourceState->Parameters.Parameters.GetPropertyBagStruct();
+				FStateTreeBindableStructDesc LinkedParamsDesc = {
+						SourceState->Name,
+						SourceState->Parameters.Parameters.GetPropertyBagStruct(),
+						EStateTreeBindableStructSource::StateParameter,
+						SourceState->Parameters.ID
+					};
 
 				// Check that the bindings for this struct are still all valid.
 				TArray<FStateTreeEditorPropertyBinding> Bindings;
-				if (!GetAndValidateBindings(StructDesc, Bindings))
+				if (!GetAndValidateBindings(LinkedParamsDesc, Bindings))
 				{
 					return false;
 				}
 
 				int32 BatchIndex = INDEX_NONE;
-				if (!BindingsCompiler.CompileBatch(StructDesc, Bindings, BatchIndex))
+				if (!BindingsCompiler.CompileBatch(LinkedParamsDesc, Bindings, BatchIndex))
 				{
 					return false;
 				}
-				check(BatchIndex < int32(MAX_uint16));
-				CompactParams.BindingsBatch = BatchIndex == INDEX_NONE ? FStateTreeHandle::Invalid : FStateTreeHandle(uint16(BatchIndex));
+
+				if (const auto Validation = UE::StateTree::Compiler::IsValidIndex16(BatchIndex); Validation.DidFail())
+				{
+					Validation.Log(Log, TEXT("BatchIndex"), LinkedParamsDesc);
+					return false;
+				}
+				CompactParams.BindingsBatch = FStateTreeIndex16(BatchIndex);
 			}
 		}
 		
 		// Create tasks
-		check(Nodes.Num() <= int32(MAX_uint16));
-		CompactState.TasksBegin = uint16(Nodes.Num());
+		const int32 TasksBegin = Nodes.Num();
+		if (const auto Validation = UE::StateTree::Compiler::IsValidCount16(TasksBegin); Validation.DidFail())
+		{
+			Validation.Log(Log, TEXT("TasksBegin"));
+			return false;
+		}
+		CompactState.TasksBegin = uint16(TasksBegin);
 
 		for (FStateTreeEditorNode& TaskNode : SourceState->Tasks)
 		{
@@ -304,8 +382,12 @@
 			return false;
 		}
 	
-		const int32 TasksNum = Nodes.Num() - int32(CompactState.TasksBegin);
-		check(TasksNum <= int32(MAX_uint8));
+		const int32 TasksNum = Nodes.Num() - TasksBegin;
+		if (const auto Validation = UE::StateTree::Compiler::IsValidCount8(TasksNum); Validation.DidFail())
+		{
+			Validation.Log(Log, TEXT("TasksNum"));
+			return false;
+		}
 		CompactState.TasksNum = uint8(TasksNum);
 	}
 	
@@ -323,14 +405,28 @@
 		FStateTreeCompilerLogStateScope LogStateScope(SourceState, Log);
 		
 		// Enter conditions.
-		CompactState.EnterConditionsBegin = uint16(Nodes.Num());
+		const int32 EnterConditionsBegin = Nodes.Num();
+		if (const auto Validation = UE::StateTree::Compiler::IsValidCount16(EnterConditionsBegin); Validation.DidFail())
+		{
+			Validation.Log(Log, TEXT("EnterConditionsBegin"));
+			return false;
+		}
+		CompactState.EnterConditionsBegin = uint16(EnterConditionsBegin);
+		
 		if (!CreateConditions(*SourceState, SourceState->EnterConditions))
 		{
 			Log.Reportf(EMessageSeverity::Error,
 				TEXT("Failed to create state enter condition."));
 			return false;
 		}
-		CompactState.EnterConditionsNum = uint8(uint16(Nodes.Num()) - CompactState.EnterConditionsBegin);
+		
+		const int32 EnterConditionsNum = Nodes.Num() - EnterConditionsBegin;
+		if (const auto Validation = UE::StateTree::Compiler::IsValidCount8(EnterConditionsNum); Validation.DidFail())
+		{
+			Validation.Log(Log, TEXT("EnterConditionsNum"));
+			return false;
+		}
+		CompactState.EnterConditionsNum = uint8(EnterConditionsNum);
 
 		// Linked state
 		if (SourceState->Type == EStateTreeStateType::Linked)
@@ -366,21 +462,34 @@
 		}
 		
 		// Transitions
-		CompactState.TransitionsBegin = uint16(StateTree->Transitions.Num());
+		const int32 TransitionsBegin = StateTree->Transitions.Num();
+		if (const auto Validation = UE::StateTree::Compiler::IsValidCount16(TransitionsBegin); Validation.DidFail())
+		{
+			Validation.Log(Log, TEXT("TransitionsBegin"));
+			return false;
+		}
+		CompactState.TransitionsBegin = uint16(TransitionsBegin);
+		
 		for (FStateTreeTransition& Transition : SourceState->Transitions)
 		{
 			FCompactStateTransition& CompactTransition = StateTree->Transitions.AddDefaulted_GetRef();
 			CompactTransition.Event = Transition.Event;
 			CompactTransition.Type = Transition.State.Type;
 			CompactTransition.GateDelay = (uint8)FMath::Clamp(FMath::CeilToInt(Transition.GateDelay * 10.0f), 0, 255);
-			CompactTransition.State = FStateTreeHandle::Invalid;
+			CompactTransition.State = FStateTreeStateHandle::Invalid;
 			if (!ResolveTransitionState(*SourceState, Transition.State, CompactTransition.State))
 			{
 				return false;
 			}
 			// Note: Unset transition is allowed here. It can be used to mask a transition at parent.
-
-			CompactTransition.ConditionsBegin = uint16(Nodes.Num());
+			const int32 ConditionsBegin = Nodes.Num();
+			if (const auto Validation = UE::StateTree::Compiler::IsValidCount16(ConditionsBegin); Validation.DidFail())
+			{
+				Validation.Log(Log, TEXT("ConditionsBegin"));
+				return false;
+			}
+			CompactTransition.ConditionsBegin = uint16(ConditionsBegin);
+			
 			if (!CreateConditions(*SourceState, Transition.Conditions))
 			{
 				Log.Reportf(EMessageSeverity::Error,
@@ -388,9 +497,23 @@
 					*Transition.State.Name.ToString());
 				return false;
 			}
-			CompactTransition.ConditionsNum = uint8(uint16(Nodes.Num()) - CompactTransition.ConditionsBegin);
+
+			const int32 ConditionsNum = Nodes.Num() - ConditionsBegin;
+			if (const auto Validation = UE::StateTree::Compiler::IsValidCount8(ConditionsNum); Validation.DidFail())
+			{
+				Validation.Log(Log, TEXT("ConditionsNum"));
+				return false;
+			}
+			CompactTransition.ConditionsNum = uint8(ConditionsNum);
 		}
-		CompactState.TransitionsNum = uint8(uint16(StateTree->Transitions.Num()) - CompactState.TransitionsBegin);
+		
+		const int32 TransitionsNum = StateTree->Transitions.Num() - TransitionsBegin;
+		if (const auto Validation = UE::StateTree::Compiler::IsValidCount8(TransitionsNum); Validation.DidFail())
+		{
+			Validation.Log(Log, TEXT("TransitionsNum"));
+			return false;
+		}
+		CompactState.TransitionsNum = uint8(TransitionsNum);
 	}
 
 	// @todo: Add test to check that all success/failure transition is possible (see editor).
@@ -398,7 +521,7 @@
 	return true;
 }
 
-bool FStateTreeCompiler::ResolveTransitionState(const UStateTreeState& SourceState, const FStateTreeStateLink& Link, FStateTreeHandle& OutTransitionHandle) const 
+bool FStateTreeCompiler::ResolveTransitionState(const UStateTreeState& SourceState, const FStateTreeStateLink& Link, FStateTreeStateHandle& OutTransitionHandle) const 
 {
 	if (Link.Type == EStateTreeTransitionType::GotoState)
 	{
@@ -471,8 +594,12 @@
 		StructDesc.Struct = CondNode.Instance.GetScriptStruct();
 		StructDesc.Name = Cond.Name;
 
-		check(InstanceIndex <= int32(MAX_uint16));
-		Cond.InstanceIndex = uint16(InstanceIndex);
+		if (const auto Validation = UE::StateTree::Compiler::IsValidIndex16(InstanceIndex); Validation.DidFail())
+		{
+			Validation.Log(Log, TEXT("InstanceIndex"), StructDesc);
+			return false;
+		}
+		Cond.InstanceIndex = FStateTreeIndex16(InstanceIndex);
 		Cond.bInstanceIsObject = false;
 	}
 	else
@@ -487,8 +614,12 @@
 		StructDesc.Struct = Instance->GetClass();
 		StructDesc.Name = Cond.Name;
 
-		check(InstanceIndex <= int32(MAX_uint16));
-		Cond.InstanceIndex = uint16(InstanceIndex);
+		if (const auto Validation = UE::StateTree::Compiler::IsValidIndex16(InstanceIndex); Validation.DidFail())
+		{
+			Validation.Log(Log, TEXT("InstanceIndex"), StructDesc);
+			return false;
+		}
+		Cond.InstanceIndex = FStateTreeIndex16(InstanceIndex);
 		Cond.bInstanceIsObject = true;
 	}
 
@@ -508,11 +639,20 @@
 	{
 		return false;
 	}
-	check(BatchIndex < int32(MAX_uint16));
-	Cond.BindingsBatch = BatchIndex == INDEX_NONE ? FStateTreeHandle::Invalid : FStateTreeHandle(uint16(BatchIndex));
 
-	check(SourceStructIndex <= int32(MAX_uint16));
-	Cond.DataViewIndex = uint16(SourceStructIndex);
+	if (const auto Validation = UE::StateTree::Compiler::IsValidIndex16(BatchIndex); Validation.DidFail())
+	{
+		Validation.Log(Log, TEXT("BatchIndex"), StructDesc);
+		return false;
+	}
+	Cond.BindingsBatch = FStateTreeIndex16(BatchIndex);
+
+	if (const auto Validation = UE::StateTree::Compiler::IsValidIndex16(SourceStructIndex); Validation.DidFail())
+	{
+		Validation.Log(Log, TEXT("SourceStructIndex"), StructDesc);
+		return false;
+	}
+	Cond.DataViewIndex = FStateTreeIndex16(SourceStructIndex);
 	
 	return true;
 }
@@ -552,8 +692,12 @@
 		StructDesc.Struct = TaskNode.Instance.GetScriptStruct();
 		StructDesc.Name = Task.Name;
 
-		check(InstanceIndex <= int32(MAX_uint16));
-		Task.InstanceIndex = uint16(InstanceIndex);
+		if (const auto Validation = UE::StateTree::Compiler::IsValidIndex16(InstanceIndex); Validation.DidFail())
+		{
+			Validation.Log(Log, TEXT("InstanceIndex"), StructDesc);
+			return false;
+		}
+		Task.InstanceIndex = FStateTreeIndex16(InstanceIndex);
 		Task.bInstanceIsObject = false;
 	}
 	else
@@ -568,8 +712,12 @@
 		StructDesc.Struct = Instance->GetClass();
 		StructDesc.Name = Task.Name;
 
-		check(InstanceIndex <= int32(MAX_uint16));
-		Task.InstanceIndex = uint16(InstanceIndex);
+		if (const auto Validation = UE::StateTree::Compiler::IsValidIndex16(InstanceIndex); Validation.DidFail())
+		{
+			Validation.Log(Log, TEXT("InstanceIndex"), StructDesc);
+			return false;
+		}
+		Task.InstanceIndex = FStateTreeIndex16(InstanceIndex);
 		Task.bInstanceIsObject = true;
 	}
 
@@ -589,12 +737,20 @@
 	{
 		return false;
 	}
-	
-	check(BatchIndex < int32(MAX_uint16));
-	Task.BindingsBatch = BatchIndex == INDEX_NONE ? FStateTreeHandle::Invalid : FStateTreeHandle(uint16(BatchIndex));
-	
-	check(SourceStructIndex <= int32(MAX_uint16));
-	Task.DataViewIndex = uint16(SourceStructIndex);
+
+	if (const auto Validation = UE::StateTree::Compiler::IsValidIndex16(BatchIndex); Validation.DidFail())
+	{
+		Validation.Log(Log, TEXT("BatchIndex"), StructDesc);
+		return false;
+	}
+	Task.BindingsBatch = FStateTreeIndex16(BatchIndex);
+
+	if (const auto Validation = UE::StateTree::Compiler::IsValidIndex16(SourceStructIndex); Validation.DidFail())
+	{
+		Validation.Log(Log, TEXT("SourceStructIndex"), StructDesc);
+		return false;
+	}
+	Task.DataViewIndex = FStateTreeIndex16(SourceStructIndex);
 
 	return true;
 }
@@ -634,8 +790,12 @@
 		StructDesc.Struct = EvalNode.Instance.GetScriptStruct();
 		StructDesc.Name = Eval.Name;
 
-		check(InstanceIndex <= int32(MAX_uint16));
-		Eval.InstanceIndex = uint16(InstanceIndex);
+		if (const auto Validation = UE::StateTree::Compiler::IsValidIndex16(InstanceIndex); Validation.DidFail())
+		{
+			Validation.Log(Log, TEXT("InstanceIndex"), StructDesc);
+			return false;
+		}
+		Eval.InstanceIndex = FStateTreeIndex16(InstanceIndex);
 		Eval.bInstanceIsObject = false;
 	}
 	else
@@ -650,8 +810,12 @@
 		StructDesc.Struct = Instance->GetClass();
 		StructDesc.Name = Eval.Name;
 
-		check(InstanceIndex <= int32(MAX_uint16));
-		Eval.InstanceIndex = uint16(InstanceIndex);
+		if (const auto Validation = UE::StateTree::Compiler::IsValidIndex16(InstanceIndex); Validation.DidFail())
+		{
+			Validation.Log(Log, TEXT("InstanceIndex"), StructDesc);
+			return false;
+		}
+		Eval.InstanceIndex = FStateTreeIndex16(InstanceIndex);
 		Eval.bInstanceIsObject = true;
 	}
 		
@@ -671,12 +835,20 @@
 	{
 		return false;
 	}
-	
-	check(BatchIndex < int32(MAX_uint16));
-	Eval.BindingsBatch = BatchIndex == INDEX_NONE ? FStateTreeHandle::Invalid : FStateTreeHandle(uint16(BatchIndex));
-	
-	check(SourceStructIndex <= int32(MAX_uint16));
-	Eval.DataViewIndex = uint16(SourceStructIndex);
+
+	if (const auto Validation = UE::StateTree::Compiler::IsValidIndex16(BatchIndex); Validation.DidFail())
+	{
+		Validation.Log(Log, TEXT("BatchIndex"), StructDesc);
+		return false;
+	}
+	Eval.BindingsBatch = FStateTreeIndex16(BatchIndex);
+
+	if (const auto Validation = UE::StateTree::Compiler::IsValidIndex16(SourceStructIndex); Validation.DidFail())
+	{
+		Validation.Log(Log, TEXT("SourceStructIndex"), StructDesc);
+		return false;
+	}
+	Eval.DataViewIndex = FStateTreeIndex16(SourceStructIndex);
 
 	return true;
 }
diff --git a/Engine/Plugins/Runtime/StateTree/Source/StateTreeEditorModule/Private/StateTreePropertyBindingCompiler.cpp b/Engine/Plugins/Runtime/StateTree/Source/StateTreeEditorModule/Private/StateTreePropertyBindingCompiler.cpp
--- a/Engine/Plugins/Runtime/StateTree/Source/StateTreeEditorModule/Private/StateTreePropertyBindingCompiler.cpp
+++ b/Engine/Plugins/Runtime/StateTree/Source/StateTreeEditorModule/Private/StateTreePropertyBindingCompiler.cpp
@@ -4,6 +4,7 @@
 #include "PropertyPathHelpers.h"
 #include "StateTreeTypes.h"
 #include "StateTreePropertyBindings.h"
+#include "StateTreeCompiler.h"
 
 bool FStateTreePropertyBindingCompiler::Init(FStateTreePropertyBindings& InPropertyBindings, FStateTreeCompilerLog& InLog)
 {
@@ -21,8 +22,11 @@ bool FStateTreePropertyBindingCompiler::CompileBatch(const FStateTreeBindableStr
 	OutBatchIndex = INDEX_NONE;
 
 	StoreSourceStructs();
+
+	TArray<FStateTreePropertySegment> SourceSegments;
+	TArray<FStateTreePropertySegment> TargetSegments;
 	
-	int32 BindingsBegin = PropertyBindings->PropertyBindings.Num();
+	const int32 BindingsBegin = PropertyBindings->PropertyBindings.Num();
 
 	for (const FStateTreeEditorPropertyBinding& EditorBinding : EditorPropertyBindings)
 	{
@@ -46,9 +50,13 @@ bool FStateTreePropertyBindingCompiler::CompileBatch(const FStateTreeBindableStr
 
 		FStateTreePropertyBinding& NewBinding = PropertyBindings->PropertyBindings.AddDefaulted_GetRef();
 
+		SourceSegments.Reset();
+		TargetSegments.Reset();
+
 		// Resolve paths
-		FResolvedPathResult SourceResult;
-		if (!ResolvePropertyPath(TargetStruct, SourceStruct, EditorBinding.SourcePath, SourceResult))
+		const FProperty* SourceLeafProperty = nullptr;
+		int32 SourceLeafArrayIndex = INDEX_NONE;
+		if (!ResolvePropertyPath(SourceStruct, EditorBinding.SourcePath, SourceSegments, SourceLeafProperty, SourceLeafArrayIndex, Log, &TargetStruct))
 		{
 			Log->Reportf(EMessageSeverity::Error, TargetStruct,
 				TEXT("Could not resolve path to '%s:%s'."),
@@ -57,8 +65,9 @@ bool FStateTreePropertyBindingCompiler::CompileBatch(const FStateTreeBindableStr
 		}
 
 		// Destination container is set to 0, it is assumed to be passed in when doing the batch copy.
-		FResolvedPathResult TargetResult;
-		if (!ResolvePropertyPath(TargetStruct, TargetStruct, EditorBinding.TargetPath, TargetResult))
+		const FProperty* TargetLeafProperty = nullptr;
+		int32 TargetLeafArrayIndex = INDEX_NONE;
+		if (!ResolvePropertyPath(TargetStruct, EditorBinding.TargetPath, TargetSegments, TargetLeafProperty, TargetLeafArrayIndex, Log, &TargetStruct))
 		{
 			Log->Reportf(EMessageSeverity::Error, TargetStruct,
 				TEXT("Could not resolve path to '%s:%s'."),
@@ -66,11 +75,45 @@ bool FStateTreePropertyBindingCompiler::CompileBatch(const FStateTreeBindableStr
 			return false;
 		}
 
-		NewBinding.SourcePathIndex = SourceResult.PathIndex;
-		NewBinding.TargetPathIndex = TargetResult.PathIndex;
-		NewBinding.SourceStructIndex = SourceStructIdx;
-		NewBinding.CopyType = GetCopyType(SourceStruct.Struct, SourceResult.LeafProperty, SourceResult.LeafArrayIndex,
-			TargetStruct.Struct, TargetResult.LeafProperty, TargetResult.LeafArrayIndex);
+		auto StorePropertyPath = [this, &TargetStruct](FStateTreePropertySegment& FirstPathSegment, TArray<FStateTreePropertySegment>& Segments) -> bool
+		{
+			check(Segments.Num() > 0);
+			FirstPathSegment = Segments[0];
+			FStateTreePropertySegment* PrevSegment = &FirstPathSegment; 
+			for (int32 Index = 1; Index < Segments.Num(); Index++)
+			{
+				const int32 SegmentIndex = PropertyBindings->PropertySegments.Num();
+				FStateTreePropertySegment& NewSegment = PropertyBindings->PropertySegments.Add_GetRef(Segments[Index]);
+
+				if (const auto Validation = UE::StateTree::Compiler::IsValidIndex16(SegmentIndex); Validation.DidFail())
+				{
+					Validation.Log(*Log, TEXT("SegmentIndex"), TargetStruct);
+					return false;
+				}
+				PrevSegment->NextIndex = FStateTreeIndex16(SegmentIndex);
+				PrevSegment = &NewSegment;
+			}
+			return true;
+		};
+
+		if (!StorePropertyPath(NewBinding.SourcePath, SourceSegments))
+		{
+			return false;
+		}
+		if (!StorePropertyPath(NewBinding.TargetPath, TargetSegments))
+		{
+			return false;
+		}
+
+		if (const auto Validation = UE::StateTree::Compiler::IsValidIndex16(SourceStructIdx); Validation.DidFail())
+		{
+			Validation.Log(*Log, TEXT("SourceStructIdx"), TargetStruct);
+			return false;
+		}
+		NewBinding.SourceStructIndex = FStateTreeIndex16(SourceStructIdx);
+		
+		NewBinding.CopyType = GetCopyType(SourceStruct.Struct, SourceLeafProperty, SourceLeafArrayIndex,
+			TargetStruct.Struct, TargetLeafProperty, TargetLeafArrayIndex);
 
 		if (NewBinding.CopyType == EStateTreePropertyCopyType::None)
 		{
@@ -112,7 +155,7 @@ int32 FStateTreePropertyBindingCompiler::GetSourceStructIndexByID(const FGuid& I
 }
 
 EStateTreePropertyCopyType FStateTreePropertyBindingCompiler::GetCopyType(const UStruct* SourceStruct, const FProperty* SourceProperty, const int32 SourceArrayIndex,
-																		  const UStruct* TargetStruct, const FProperty* TargetProperty, const int32 TargetArrayIndex)
+																		  const UStruct* TargetStruct, const FProperty* TargetProperty, const int32 TargetArrayIndex) const
 {
 	if (SourceProperty == nullptr)
 	{
@@ -297,20 +340,6 @@ EStateTreePropertyCopyType FStateTreePropertyBindingCompiler::GetCopyType(const
 	return EStateTreePropertyCopyType::None;
 }
 
-bool FStateTreePropertyBindingCompiler::ResolvePropertyPath(const FStateTreeBindableStructDesc& InOwnerStructDesc, const FStateTreeBindableStructDesc& InStructDesc, const FStateTreeEditorPropertyPath& InPath, FResolvedPathResult& OutResult)
-{
-	const int32 PathIndex = PropertyBindings->PropertyPaths.Num();
-	FStateTreePropertyPath& Path = PropertyBindings->PropertyPaths.AddDefaulted_GetRef();
-	Path.SegmentsBegin = PropertyBindings->PropertySegments.Num();
-
-	const bool bResult = ResolvePropertyPath(InStructDesc, InPath, PropertyBindings->PropertySegments, OutResult.LeafProperty, OutResult.LeafArrayIndex, Log, &InOwnerStructDesc);
-
-	Path.SegmentsEnd = PropertyBindings->PropertySegments.Num();
-	OutResult.PathIndex = PathIndex;
-
-	return bResult;
-}
-
 bool FStateTreePropertyBindingCompiler::ResolvePropertyPath(const FStateTreeBindableStructDesc& InStructDesc, const FStateTreeEditorPropertyPath& InPath,
 															TArray<FStateTreePropertySegment>& OutSegments, const FProperty*& OutLeafProperty, int32& OutLeafArrayIndex,
 															FStateTreeCompilerLog* InLog, const FStateTreeBindableStructDesc* InLogContextStruct)
@@ -379,9 +408,18 @@ bool FStateTreePropertyBindingCompiler::ResolvePropertyPath(const FStateTreeBind
 			break;
 		}
 
+		if (const auto Validation = UE::StateTree::Compiler::IsValidIndex16(ArrayIndex); Validation.DidFail())
+		{
+			if (InLog != nullptr)
+			{
+				Validation.Log(*InLog, TEXT("ArrayIndex"), InStructDesc);
+			}
+			return false;
+		}
+
 		FStateTreePropertySegment& Segment = OutSegments.AddDefaulted_GetRef();
 		Segment.Name = PropertyName;
-		Segment.ArrayIndex = ArrayIndex;
+		Segment.ArrayIndex = FStateTreeIndex16(ArrayIndex);
 
 		// Check to see if it is an array access first
 		const FArrayProperty* ArrayProperty = CastField<FArrayProperty>(Property);
@@ -403,7 +441,7 @@ bool FStateTreePropertyBindingCompiler::ResolvePropertyPath(const FStateTreeBind
 				{
 					// Object arrays need an object dereference adding if non-leaf
 					FStateTreePropertySegment& ExtraSegment = OutSegments.AddDefaulted_GetRef();
-					ExtraSegment.ArrayIndex = 0;
+					ExtraSegment.ArrayIndex = FStateTreeIndex16(0);
 					ExtraSegment.Type = EStateTreePropertyAccessType::Object;
 					const FProperty* InnerProperty = ArrayProperty->Inner;
 					if (const FObjectProperty* ObjectProperty = CastField<FObjectProperty>(InnerProperty))
@@ -423,7 +461,7 @@ bool FStateTreePropertyBindingCompiler::ResolvePropertyPath(const FStateTreeBind
 			else
 			{
 				Segment.Type = EStateTreePropertyAccessType::IndexArray;
-				Segment.ArrayIndex = ArrayIndex;
+				Segment.ArrayIndex = FStateTreeIndex16(ArrayIndex);
 				CurrentStruct = nullptr;
 			}
 		}
diff --git a/Engine/Plugins/Runtime/StateTree/Source/StateTreeEditorModule/Public/StateTreeCompiler.h b/Engine/Plugins/Runtime/StateTree/Source/StateTreeEditorModule/Public/StateTreeCompiler.h
--- a/Engine/Plugins/Runtime/StateTree/Source/StateTreeEditorModule/Public/StateTreeCompiler.h
+++ b/Engine/Plugins/Runtime/StateTree/Source/StateTreeEditorModule/Public/StateTreeCompiler.h
@@ -30,12 +30,12 @@
 	
 private:
 
-	bool ResolveTransitionState(const UStateTreeState& SourceState, const FStateTreeStateLink& Link, FStateTreeHandle& OutTransitionHandle) const;
-	FStateTreeHandle GetStateHandle(const FGuid& StateID) const;
+	bool ResolveTransitionState(const UStateTreeState& SourceState, const FStateTreeStateLink& Link, FStateTreeStateHandle& OutTransitionHandle) const;
+	FStateTreeStateHandle GetStateHandle(const FGuid& StateID) const;
 	UStateTreeState* GetState(const FGuid& StateID);
 
 	bool CreateStates();
-	bool CreateStateRecursive(UStateTreeState& State, const FStateTreeHandle Parent);
+	bool CreateStateRecursive(UStateTreeState& State, const FStateTreeStateHandle Parent);
 	
 	bool CreateEvaluators();
 	bool CreateStateTasksAndParameters();
@@ -48,8 +48,6 @@
 	bool GetAndValidateBindings(const FStateTreeBindableStructDesc& TargetStruct, TArray<FStateTreeEditorPropertyBinding>& OutBindings) const;
 	bool IsPropertyAnyEnum(const FStateTreeBindableStructDesc& Struct, FStateTreeEditorPropertyPath Path) const;
 
-	void CalculateEstimatedMemoryUsage();
-
 	FStateTreeCompilerLog& Log;
 	UStateTree* StateTree = nullptr;
 	UStateTreeEditorData* TreeData = nullptr;
@@ -64,3 +62,72 @@
 	
 	FStateTreePropertyBindingCompiler BindingsCompiler;
 };
+
+
+namespace UE::StateTree::Compiler
+{
+	struct FValidationResult
+	{
+		FValidationResult() = default;
+		FValidationResult(const bool bInResult, const int32 InValue, const int32 InMaxValue) : bResult(bInResult), Value(InValue), MaxValue(InMaxValue) {}
+
+		/** Validation succeeded */
+		bool DidSucceed() const { return bResult == true; }
+
+		/** Validation failed */
+		bool DidFail() const { return bResult == false; }
+
+		/**
+		 * Logs common validation for IsValidIndex16(), IsValidIndex8(), IsValidCount16(), IsValidCount8().
+		 * @param Log reference to the compiler log.
+		 * @param ContextText Text identifier for the context where the test is done.
+		 * @param ContextStruct Struct identifier for the context where the test is done.
+		 */
+		void Log(FStateTreeCompilerLog& Log, const TCHAR* ContextText, const FStateTreeBindableStructDesc& ContextStruct = FStateTreeBindableStructDesc()) const;
+		
+		bool bResult = true;
+		int32 Value = 0;
+		int32 MaxValue = 0;
+	};
+
+	/**
+	 * Checks if given index can be represented as uint16, including MAX_uint16 as INDEX_NONE.
+	 * @param Index Index to test
+	 * @return validation result.
+	 */
+	inline FValidationResult IsValidIndex16(const int32 Index)
+	{
+		return FValidationResult(Index == INDEX_NONE || (Index >= 0 && Index < MAX_uint16), Index, MAX_uint16 - 1);
+	}
+
+	/**
+	 * Checks if given index can be represented as uint8, including MAX_uint8 as INDEX_NONE. 
+	 * @param Index Index to test
+	 * @return true if the index is valid.
+	 */
+	inline FValidationResult IsValidIndex8(const int32 Index)
+	{
+		return FValidationResult(Index == INDEX_NONE || (Index >= 0 && Index < MAX_uint8), Index, MAX_uint8 - 1);
+	}
+
+	/**
+	 * Checks if given count can be represented as uint16. 
+	 * @param Count Count to test
+	 * @return true if the count is valid.
+	 */
+	inline FValidationResult IsValidCount16(const int32 Count)
+	{
+		return FValidationResult(Count >= 0 && Count <= MAX_uint16, Count, MAX_uint16);
+	}
+
+	/**
+	 * Checks if given count can be represented as uint8. 
+	 * @param Count Count to test
+	 * @return true if the count is valid.
+	 */
+	inline FValidationResult IsValidCount8(const int32 Count)
+	{
+		return FValidationResult(Count >= 0 && Count <= MAX_uint8, Count, MAX_uint8);
+	}
+
+}; // UE::StateTree::Compiler
diff --git a/Engine/Plugins/Runtime/StateTree/Source/StateTreeEditorModule/Public/StateTreePropertyBindingCompiler.h b/Engine/Plugins/Runtime/StateTree/Source/StateTreeEditorModule/Public/StateTreePropertyBindingCompiler.h
--- a/Engine/Plugins/Runtime/StateTree/Source/StateTreeEditorModule/Public/StateTreePropertyBindingCompiler.h
+++ b/Engine/Plugins/Runtime/StateTree/Source/StateTreeEditorModule/Public/StateTreePropertyBindingCompiler.h
@@ -85,20 +85,12 @@ protected:
 
 	void StoreSourceStructs();
 
+	EStateTreePropertyCopyType GetCopyType(const UStruct* SourceStruct, const FProperty* SourceProperty, const int32 SourceArrayIndex, const UStruct* TargetStruct, const FProperty* TargetProperty, const int32 TargetArrayIndex) const;
+
 	UPROPERTY()
 	TArray<FStateTreeBindableStructDesc> SourceStructs;
 
 	FStateTreePropertyBindings* PropertyBindings = nullptr;
 
 	FStateTreeCompilerLog* Log = nullptr;
-	
-	struct FResolvedPathResult
-	{
-		int32 PathIndex = INDEX_NONE;
-		const FProperty* LeafProperty = nullptr;
-		int32 LeafArrayIndex = INDEX_NONE;
-	};
-
-	EStateTreePropertyCopyType GetCopyType(const UStruct* SourceStruct, const FProperty* SourceProperty, const int32 SourceArrayIndex, const UStruct* TargetStruct, const FProperty* TargetProperty, const int32 TargetArrayIndex);
-	[[nodiscard]] bool ResolvePropertyPath(const FStateTreeBindableStructDesc& InOwnerStructDesc, const FStateTreeBindableStructDesc& InStructDesc, const FStateTreeEditorPropertyPath& InPath, FResolvedPathResult& OutResult);
 };
diff --git a/Engine/Plugins/Runtime/StateTree/Source/StateTreeModule/Private/StateTree.cpp b/Engine/Plugins/Runtime/StateTree/Source/StateTreeModule/Private/StateTree.cpp
--- a/Engine/Plugins/Runtime/StateTree/Source/StateTreeModule/Private/StateTree.cpp
+++ b/Engine/Plugins/Runtime/StateTree/Source/StateTreeModule/Private/StateTree.cpp
@@ -28,7 +28,7 @@
 	PropertyBindings.Reset();
 	Parameters.Reset();
 
-	ParametersDataViewIndex = INDEX_NONE;
+	ParametersDataViewIndex = FStateTreeIndex8::Invalid;
 	
 	EvaluatorsBegin = 0;
 	EvaluatorsNum = 0;
@@ -51,9 +51,11 @@
 
 	const int32 CurrentVersion = GetLinkerCustomVersion(FStateTreeCustomVersion::GUID);
 
-	if (CurrentVersion < FStateTreeCustomVersion::InstanceDataArrays)
+	if (CurrentVersion < FStateTreeCustomVersion::LatestVersion)
 	{
-		ResetLinked();
+#if WITH_EDITOR
+		ResetCompiled();
+#endif
 		UE_LOG(LogStateTree, Error, TEXT("%s: StateTree compiled data in older format. Please recompile the StateTree asset."), *GetName());
 		return;
 	}
@@ -105,42 +107,42 @@
 	const TArrayView<FStateTreeBindableStructDesc> SourceStructs = PropertyBindings.GetSourceStructs();
 	const TArrayView<FStateTreePropCopyBatch> CopyBatches = PropertyBindings.GetCopyBatches();
 
-	if (SourceStructs.IsValidIndex(ParametersDataViewIndex))
+	if (ParametersDataViewIndex.IsValid() && SourceStructs.IsValidIndex(ParametersDataViewIndex.Get()))
 	{
-		SourceStructs[ParametersDataViewIndex].Struct = Parameters.GetPropertyBagStruct();
+		SourceStructs[ParametersDataViewIndex.Get()].Struct = Parameters.GetPropertyBagStruct();
 	}
 	
 	for (const FCompactStateTreeState& State : States)
 	{
 		if (State.Type == EStateTreeStateType::Subtree)
 		{
-			if (State.ParameterInstanceIndex == MAX_uint16
-				|| State.ParameterDataViewIndex == MAX_uint16)
+			if (State.ParameterInstanceIndex.IsValid() == false
+				|| State.ParameterDataViewIndex.IsValid() == false)
 			{
 				UE_LOG(LogStateTree, Error, TEXT("%s: Data for state '%s' is malformed. Please recompile the StateTree asset."), *GetName(), *State.Name.ToString());
 				return false;
 			}
 
 			// Subtree is a bind source, update bag struct.
-			const FCompactStateTreeParameters& Params = DefaultInstanceData.GetMutableStruct(State.ParameterInstanceIndex).GetMutable<FCompactStateTreeParameters>();
-			FStateTreeBindableStructDesc& Desc = SourceStructs[State.ParameterDataViewIndex];
+			const FCompactStateTreeParameters& Params = DefaultInstanceData.GetMutableStruct(State.ParameterInstanceIndex.Get()).GetMutable<FCompactStateTreeParameters>();
+			FStateTreeBindableStructDesc& Desc = SourceStructs[State.ParameterDataViewIndex.Get()];
 			Desc.Struct = Params.Parameters.GetPropertyBagStruct();
 		}
 		else if (State.Type == EStateTreeStateType::Linked && State.LinkedState.IsValid())
 		{
 			const FCompactStateTreeState& LinkedState = States[State.LinkedState.Index];
 
-			if (State.ParameterInstanceIndex == MAX_uint16
-				|| LinkedState.ParameterInstanceIndex == MAX_uint16)
+			if (State.ParameterInstanceIndex.IsValid() == false
+				|| LinkedState.ParameterInstanceIndex.IsValid() == false)
 			{
 				UE_LOG(LogStateTree, Error, TEXT("%s: Data for state '%s' is malformed. Please recompile the StateTree asset."), *GetName(), *State.Name.ToString());
 				return false;
 			}
 
-			const FCompactStateTreeParameters& Params = DefaultInstanceData.GetMutableStruct(State.ParameterInstanceIndex).GetMutable<FCompactStateTreeParameters>();
+			const FCompactStateTreeParameters& Params = DefaultInstanceData.GetMutableStruct(State.ParameterInstanceIndex.Get()).GetMutable<FCompactStateTreeParameters>();
 
 			// Check that the bag in linked state matches.
-			const FCompactStateTreeParameters& LinkedStateParams = DefaultInstanceData.GetMutableStruct(LinkedState.ParameterInstanceIndex).GetMutable<FCompactStateTreeParameters>();
+			const FCompactStateTreeParameters& LinkedStateParams = DefaultInstanceData.GetMutableStruct(LinkedState.ParameterInstanceIndex.Get()).GetMutable<FCompactStateTreeParameters>();
 
 			if (LinkedStateParams.Parameters.GetPropertyBagStruct() != Params.Parameters.GetPropertyBagStruct())
 			{
@@ -148,7 +150,7 @@
 				return false;
 			}
 
-			FStateTreePropCopyBatch& Batch = CopyBatches[Params.BindingsBatch.Index];
+			FStateTreePropCopyBatch& Batch = CopyBatches[Params.BindingsBatch.Get()];
 			Batch.TargetStruct.Struct = Params.Parameters.GetPropertyBagStruct();
 		}
 	}
@@ -168,7 +170,7 @@
 		const FStructView Node = Nodes[Index];
 		if (FStateTreeNodeBase* NodePtr = Node.GetMutablePtr<FStateTreeNodeBase>())
 		{
-			Linker.SetCurrentInstanceDataType(NodePtr->GetInstanceDataType(), NodePtr->DataViewIndex);
+			Linker.SetCurrentInstanceDataType(NodePtr->GetInstanceDataType(), NodePtr->DataViewIndex.Get());
 			const bool bLinkSucceeded = NodePtr->Link(Linker);
 			if (!bLinkSucceeded || Linker.GetStatus() == EStateTreeLinkerStatus::Failed)
 			{
@@ -222,9 +224,9 @@
 	const int32 EvalMemUsageIndex = MemoryUsages.Emplace(TEXT("Evaluators"));
 	const int32 SharedMemUsageIndex = MemoryUsages.Emplace(TEXT("Shared Data"));
 
-	auto GetRootStateHandle = [this](const FStateTreeHandle InState) -> FStateTreeHandle
+	auto GetRootStateHandle = [this](const FStateTreeStateHandle InState) -> FStateTreeStateHandle
 	{
-		FStateTreeHandle Result = InState;
+		FStateTreeStateHandle Result = InState;
 		while (Result.IsValid() && States[Result.Index].Parent.IsValid())
 		{
 			Result = States[Result.Index].Parent;
@@ -232,7 +234,7 @@
 		return Result;		
 	};
 
-	auto GetUsageIndexForState = [&MemoryUsages, this](const FStateTreeHandle InStateHandle) -> int32
+	auto GetUsageIndexForState = [&MemoryUsages, this](const FStateTreeStateHandle InStateHandle) -> int32
 	{
 		check(InStateHandle.IsValid());
 		
@@ -249,9 +251,9 @@
 
 	for (int32 Index = 0; Index < States.Num(); Index++)
 	{
-		const FStateTreeHandle StateHandle((uint16)Index);
+		const FStateTreeStateHandle StateHandle((uint16)Index);
 		const FCompactStateTreeState& CompactState = States[Index];
-		const FStateTreeHandle ParentHandle = GetRootStateHandle(StateHandle);
+		const FStateTreeStateHandle ParentHandle = GetRootStateHandle(StateHandle);
 		const int32 ParentUsageIndex = GetUsageIndexForState(ParentHandle);
 		FStateTreeMemoryUsage& MemUsage = MemoryUsages[ParentUsageIndex];
 
@@ -263,7 +265,7 @@
 			StateLinks.Emplace(ParentUsageIndex, LinkedUsageIndex);
 
 			MemUsage.NodeCount++;
-			MemUsage.AddUsage(DefaultInstanceData.GetStruct(CompactState.ParameterInstanceIndex));
+			MemUsage.AddUsage(DefaultInstanceData.GetStruct(CompactState.ParameterInstanceIndex.Get()));
 		}
 		
 		for (int32 TaskIndex = CompactState.TasksBegin; TaskIndex < (CompactState.TasksBegin + CompactState.TasksNum); TaskIndex++)
@@ -272,12 +274,12 @@
 			if (Task.bInstanceIsObject == false)
 			{
 				MemUsage.NodeCount++;
-				MemUsage.AddUsage(DefaultInstanceData.GetStruct(Task.InstanceIndex));
+				MemUsage.AddUsage(DefaultInstanceData.GetStruct(Task.InstanceIndex.Get()));
 			}
 			else
 			{
 				MemUsage.NodeCount++;
-				MemUsage.AddUsage(DefaultInstanceData.GetMutableObject(Task.InstanceIndex));
+				MemUsage.AddUsage(DefaultInstanceData.GetMutableObject(Task.InstanceIndex.Get()));
 			}
 		}
 	}
@@ -302,11 +304,11 @@
 		const FStateTreeEvaluatorBase& Eval = Nodes[EvalIndex].Get<FStateTreeEvaluatorBase>();
 		if (Eval.bInstanceIsObject == false)
 		{
-			EvalMemUsage.AddUsage(DefaultInstanceData.GetMutableStruct(Eval.InstanceIndex));
+			EvalMemUsage.AddUsage(DefaultInstanceData.GetMutableStruct(Eval.InstanceIndex.Get()));
 		}
 		else
 		{
-			EvalMemUsage.AddUsage(DefaultInstanceData.GetMutableObject(Eval.InstanceIndex));
+			EvalMemUsage.AddUsage(DefaultInstanceData.GetMutableObject(Eval.InstanceIndex.Get()));
 		}
 	}
 
diff --git a/Engine/Plugins/Runtime/StateTree/Source/StateTreeModule/Private/StateTreeExecutionContext.cpp b/Engine/Plugins/Runtime/StateTree/Source/StateTreeModule/Private/StateTreeExecutionContext.cpp
--- a/Engine/Plugins/Runtime/StateTree/Source/StateTreeModule/Private/StateTreeExecutionContext.cpp
+++ b/Engine/Plugins/Runtime/StateTree/Source/StateTreeModule/Private/StateTreeExecutionContext.cpp
@@ -58,9 +58,9 @@
 void FStateTreeExecutionContext::SetDefaultParameters()
 {
 	if (ensureMsgf(StateTree != nullptr, TEXT("Execution context must be initialized before calling %s"), ANSI_TO_TCHAR(__FUNCTION__))
-		&& DataViews.IsValidIndex(StateTree->ParametersDataViewIndex))
+		&& DataViews.IsValidIndex(StateTree->ParametersDataViewIndex.Get()))
 	{
-		DataViews[StateTree->ParametersDataViewIndex] = FStateTreeDataView(StateTree->GetDefaultParameters().GetMutableValue());	
+		DataViews[StateTree->ParametersDataViewIndex.Get()] = FStateTreeDataView(StateTree->GetDefaultParameters().GetMutableValue());	
 	}
 }
 
@@ -69,9 +69,9 @@
 	if (ensureMsgf(StateTree != nullptr, TEXT("Execution context must be initialized before calling %s"), ANSI_TO_TCHAR(__FUNCTION__))
 		&& ensureMsgf(StateTree->GetDefaultParameters().GetPropertyBagStruct() == Parameters.GetPropertyBagStruct(),
 			TEXT("Parameters must be of the same struct type. Make sure to migrate the provided parameters to the same type as the StateTree default parameters."))
-		&& DataViews.IsValidIndex(StateTree->ParametersDataViewIndex))
+		&& DataViews.IsValidIndex(StateTree->ParametersDataViewIndex.Get()))
 	{
-		DataViews[StateTree->ParametersDataViewIndex] = FStateTreeDataView(Parameters.GetMutableValue());
+		DataViews[StateTree->ParametersDataViewIndex.Get()] = FStateTreeDataView(Parameters.GetMutableValue());
 	}
 }
 
@@ -93,31 +93,31 @@
 	const FStateTreeDataView ParametersView(StateParams.Parameters.GetMutableValue());
 	if (StateParams.BindingsBatch.IsValid())
 	{
-		StateTree->PropertyBindings.CopyTo(DataViews, StateParams.BindingsBatch.Index, ParametersView);
+		StateTree->PropertyBindings.CopyTo(DataViews, StateParams.BindingsBatch, ParametersView);
 	}
 
 	// Set the parameters as the input parameters for the linked state.
 	check(State.LinkedState.IsValid());
 	const FCompactStateTreeState& LinkedState = StateTree->States[State.LinkedState.Index];
-	check(LinkedState.ParameterDataViewIndex != MAX_uint16);
-	DataViews[LinkedState.ParameterDataViewIndex] = ParametersView;
+	check(LinkedState.ParameterDataViewIndex.IsValid());
+	DataViews[LinkedState.ParameterDataViewIndex.Get()] = ParametersView;
 }
 
 void FStateTreeExecutionContext::UpdateSubtreeStateParameters(const FStateTreeInstanceData& InstanceData, const FCompactStateTreeState& State)
 {
-	check(State.ParameterDataViewIndex != MAX_uint16);
-	check(State.ParameterInstanceIndex != MAX_uint16);
+	check(State.ParameterDataViewIndex.IsValid());
+	check(State.ParameterInstanceIndex.IsValid());
 	
 	// Usually the subtree parameter view is set by the linked state. If it's not (i.e. transitioned into a parametrized subtree), we'll set the view default params.
-	if (DataViews[State.ParameterDataViewIndex].IsValid())
+	if (DataViews[State.ParameterDataViewIndex.Get()].IsValid())
 	{
 		return;
 	}
 
 	// Set view to default parameters.
-	const FStateTreeDataView ParamInstanceView = StateTree->DefaultInstanceData.GetMutableStruct(State.ParameterInstanceIndex); // These are used as const, so get them from the tree initial values.
+	const FStateTreeDataView ParamInstanceView = StateTree->DefaultInstanceData.GetMutableStruct(State.ParameterInstanceIndex.Get()); // These are used as const, so get them from the tree initial values.
 	const FCompactStateTreeParameters& Params = ParamInstanceView.GetMutable<FCompactStateTreeParameters>();
-	DataViews[State.ParameterDataViewIndex] = FStateTreeDataView(Params.Parameters.GetMutableValue());
+	DataViews[State.ParameterDataViewIndex.Get()] = FStateTreeDataView(Params.Parameters.GetMutableValue());
 }
 
 EStateTreeRunStatus FStateTreeExecutionContext::Start(FStateTreeInstanceData* ExternalInstanceData)
@@ -155,10 +155,10 @@
 	if (Exec->TreeRunStatus == EStateTreeRunStatus::Running)
 	{
 		FStateTreeTransitionResult Transition;
-		Transition.TargetState = FStateTreeHandle::Succeeded;
+		Transition.TargetState = FStateTreeStateHandle::Succeeded;
 		Transition.CurrentActiveStates = Exec->ActiveStates;
 		Transition.CurrentRunStatus = Exec->LastTickStatus;
-		Transition.NextActiveStates = FStateTreeActiveStates(FStateTreeHandle::Succeeded);
+		Transition.NextActiveStates = FStateTreeActiveStates(FStateTreeStateHandle::Succeeded);
 
 		ExitState(InstanceData, Transition);
 	}
@@ -168,17 +168,17 @@
 	Exec->ActiveStates.Reset();
 	Exec->LastTickStatus = EStateTreeRunStatus::Unset;
 
-	static const FStateTreeHandle RootState = FStateTreeHandle(0);
+	static const FStateTreeStateHandle RootState = FStateTreeStateHandle(0);
 
 	FStateTreeActiveStates NextActiveStates;
 	if (SelectState(InstanceData, RootState, NextActiveStates))
 	{
-		if (NextActiveStates.Last() == FStateTreeHandle::Succeeded || NextActiveStates.Last() == FStateTreeHandle::Failed)
+		if (NextActiveStates.Last() == FStateTreeStateHandle::Succeeded || NextActiveStates.Last() == FStateTreeStateHandle::Failed)
 		{
 			// Transition to a terminal state (succeeded/failed), or default transition failed.
 			STATETREE_LOG(Warning, TEXT("%s: Tree %s at StateTree start on '%s' using StateTree '%s'."),
-				ANSI_TO_TCHAR(__FUNCTION__), NextActiveStates.Last() == FStateTreeHandle::Succeeded ? TEXT("succeeded") : TEXT("failed"), *GetNameSafe(Owner), *GetNameSafe(StateTree));
-			Exec->TreeRunStatus = NextActiveStates.Last() == FStateTreeHandle::Succeeded ? EStateTreeRunStatus::Succeeded : EStateTreeRunStatus::Failed;
+				ANSI_TO_TCHAR(__FUNCTION__), NextActiveStates.Last() == FStateTreeStateHandle::Succeeded ? TEXT("succeeded") : TEXT("failed"), *GetNameSafe(Owner), *GetNameSafe(StateTree));
+			Exec->TreeRunStatus = NextActiveStates.Last() == FStateTreeStateHandle::Succeeded ? EStateTreeRunStatus::Succeeded : EStateTreeRunStatus::Failed;
 		}
 		else
 		{
@@ -233,14 +233,14 @@
 	TickEvaluators(InstanceData, 0.0f);
 
 	// Exit states if still in some valid state.
-	if (!Exec.ActiveStates.IsEmpty() && (Exec.ActiveStates.Last() != FStateTreeHandle::Succeeded || Exec.ActiveStates.Last() != FStateTreeHandle::Failed))
+	if (!Exec.ActiveStates.IsEmpty() && (Exec.ActiveStates.Last() != FStateTreeStateHandle::Succeeded || Exec.ActiveStates.Last() != FStateTreeStateHandle::Failed))
 	{
 		// Transition to Succeeded state.
 		FStateTreeTransitionResult Transition;
-		Transition.TargetState = FStateTreeHandle::Succeeded;
+		Transition.TargetState = FStateTreeStateHandle::Succeeded;
 		Transition.CurrentActiveStates = Exec.ActiveStates;
 		Transition.CurrentRunStatus = Exec.LastTickStatus;
-		Transition.NextActiveStates = FStateTreeActiveStates(FStateTreeHandle::Succeeded);
+		Transition.NextActiveStates = FStateTreeActiveStates(FStateTreeStateHandle::Succeeded);
 
 		ExitState(InstanceData, Transition);
 
@@ -248,7 +248,7 @@
 	}
 	else
 	{
-		Exec.TreeRunStatus = Exec.ActiveStates.Last() == FStateTreeHandle::Succeeded ? EStateTreeRunStatus::Succeeded : EStateTreeRunStatus::Failed;
+		Exec.TreeRunStatus = Exec.ActiveStates.Last() == FStateTreeStateHandle::Succeeded ? EStateTreeRunStatus::Succeeded : EStateTreeRunStatus::Failed;
 	}
 
 	// Call TreeStop on evaluators.
@@ -256,10 +256,10 @@
 
 	Exec.ActiveStates.Reset();
 	Exec.LastTickStatus = EStateTreeRunStatus::Unset;
-	Exec.FirstTaskStructIndex = 0;
-	Exec.FirstTaskObjectIndex = 0;
+	Exec.FirstTaskStructIndex = FStateTreeIndex16::Invalid;
+	Exec.FirstTaskObjectIndex = FStateTreeIndex16::Invalid;
 
-	EStateTreeRunStatus Result = Exec.TreeRunStatus; 
+	const EStateTreeRunStatus Result = Exec.TreeRunStatus; 
 
 	// Destruct all allocated instance data (does not shrink the buffer). This will invalidate Exec too.
 	InstanceData.Reset();
@@ -292,7 +292,7 @@
 	}
 
 	// Update the gated transition time.
-	if (Exec->GatedTransitionIndex != INDEX_NONE)
+	if (Exec->GatedTransitionIndex.IsValid())
 	{
 		Exec->GatedTransitionTime -= DeltaTime;
 	}
@@ -323,10 +323,10 @@
 		{
 			ExitState(InstanceData, Transition);
 			
-			if (Transition.NextActiveStates.Last() == FStateTreeHandle::Succeeded || Transition.NextActiveStates.Last() == FStateTreeHandle::Failed)
+			if (Transition.NextActiveStates.Last() == FStateTreeStateHandle::Succeeded || Transition.NextActiveStates.Last() == FStateTreeStateHandle::Failed)
 			{
 				// Transition to a terminal state (succeeded/failed), or default transition failed.
-				Exec->TreeRunStatus = Transition.NextActiveStates.Last() == FStateTreeHandle::Succeeded ? EStateTreeRunStatus::Succeeded : EStateTreeRunStatus::Failed;
+				Exec->TreeRunStatus = Transition.NextActiveStates.Last() == FStateTreeStateHandle::Succeeded ? EStateTreeRunStatus::Succeeded : EStateTreeRunStatus::Failed;
 				return Exec->TreeRunStatus;
 			}
 			
@@ -394,11 +394,11 @@
 		const FStateTreeEvaluatorBase& Eval =  StateTree->Nodes[EvalIndex].Get<FStateTreeEvaluatorBase>();
 		if (Eval.bInstanceIsObject)
 		{
-			InstanceObjects.Add(StateTree->DefaultInstanceData.GetMutableObject(Eval.InstanceIndex));
+			InstanceObjects.Add(StateTree->DefaultInstanceData.GetMutableObject(Eval.InstanceIndex.Get()));
 		}
 		else
 		{
-			InstanceStructs.Add(StateTree->DefaultInstanceData.GetMutableStruct(Eval.InstanceIndex));
+			InstanceStructs.Add(StateTree->DefaultInstanceData.GetMutableStruct(Eval.InstanceIndex.Get()));
 		}
 	}
 
@@ -415,13 +415,13 @@
 	
 	for (int32 Index = 0; Index < NextActiveStates.Num(); Index++)
 	{
-		const FStateTreeHandle CurrentHandle = NextActiveStates[Index];
+		const FStateTreeStateHandle CurrentHandle = NextActiveStates[Index];
 		const FCompactStateTreeState& State = StateTree->States[CurrentHandle.Index];
 
 		if (State.Type == EStateTreeStateType::Linked)
 		{
-			check(State.ParameterInstanceIndex != MAX_uint16);
-			InstanceStructs.Add(StateTree->DefaultInstanceData.GetMutableStruct(State.ParameterInstanceIndex));
+			check(State.ParameterInstanceIndex.IsValid());
+			InstanceStructs.Add(StateTree->DefaultInstanceData.GetMutableStruct(State.ParameterInstanceIndex.Get()));
 		}
 		
 		for (int32 TaskIndex = State.TasksBegin; TaskIndex < (State.TasksBegin + State.TasksNum); TaskIndex++)
@@ -429,11 +429,11 @@
 			const FStateTreeTaskBase& Task = StateTree->Nodes[TaskIndex].Get<FStateTreeTaskBase>();
 			if (Task.bInstanceIsObject)
 			{
-				InstanceObjects.Add(StateTree->DefaultInstanceData.GetMutableObject(Task.InstanceIndex));
+				InstanceObjects.Add(StateTree->DefaultInstanceData.GetMutableObject(Task.InstanceIndex.Get()));
 			}
 			else
 			{
-				InstanceStructs.Add(StateTree->DefaultInstanceData.GetMutableStruct(Task.InstanceIndex));
+				InstanceStructs.Add(StateTree->DefaultInstanceData.GetMutableStruct(Task.InstanceIndex.Get()));
 			}
 		}
 		
@@ -467,8 +467,8 @@
 		MakeArrayView(InstanceObjects.GetData() + NumCommonInstanceObjects, InstanceObjects.Num() - NumCommonInstanceObjects));
 
 	FStateTreeExecutionState& Exec = GetExecState(InstanceData);
-	Exec.FirstTaskStructIndex = FirstTaskStructIndex;
-	Exec.FirstTaskObjectIndex = FirstTaskObjectIndex;
+	Exec.FirstTaskStructIndex = FStateTreeIndex16(FirstTaskStructIndex);
+	Exec.FirstTaskObjectIndex = FStateTreeIndex16(FirstTaskObjectIndex);
 }
 
 EStateTreeRunStatus FStateTreeExecutionContext::EnterState(FStateTreeInstanceData& InstanceData, const FStateTreeTransitionResult& Transition)
@@ -495,17 +495,18 @@
 	
 	EStateTreeRunStatus Result = EStateTreeRunStatus::Running;
 
-	Exec.EnterStateFailedTaskIndex = MAX_uint16; // This will make all tasks to be accepted.
+	Exec.EnterStateFailedTaskIndex = FStateTreeIndex16::Invalid; // This will make all tasks to be accepted.
 	Exec.ActiveStates.Reset();
 	
 	// Do property copy on all states, propagating the results from last tick.
-	int32 InstanceStructIndex = Exec.FirstTaskStructIndex;
-	int32 InstanceObjectIndex = Exec.FirstTaskObjectIndex;
+	check(Exec.FirstTaskStructIndex.IsValid() && Exec.FirstTaskObjectIndex.IsValid()); 
+	int32 InstanceStructIndex = Exec.FirstTaskStructIndex.Get();
+	int32 InstanceObjectIndex = Exec.FirstTaskObjectIndex.Get();
 	
 	for (int32 Index = 0; Index < Transition.NextActiveStates.Num() && Result != EStateTreeRunStatus::Failed; Index++)
 	{
-		const FStateTreeHandle CurrentHandle = Transition.NextActiveStates[Index];
-		const FStateTreeHandle PreviousHandle = Transition.CurrentActiveStates.GetStateSafe(Index);
+		const FStateTreeStateHandle CurrentHandle = Transition.NextActiveStates[Index];
+		const FStateTreeStateHandle PreviousHandle = Transition.CurrentActiveStates.GetStateSafe(Index);
 		const FCompactStateTreeState& State = StateTree->States[CurrentHandle.Index];
 
 		if (!Exec.ActiveStates.Push(CurrentHandle))
@@ -541,19 +542,19 @@
 			const FStateTreeTaskBase& Task = StateTree->Nodes[TaskIndex].Get<FStateTreeTaskBase>();
 			if (Task.bInstanceIsObject)
 			{
-				DataViews[Task.DataViewIndex] = InstanceData.GetMutableObject(InstanceObjectIndex);
+				DataViews[Task.DataViewIndex.Get()] = InstanceData.GetMutableObject(InstanceObjectIndex);
 				InstanceObjectIndex++;
 			}
 			else
 			{
-				DataViews[Task.DataViewIndex] = InstanceData.GetMutableStruct(InstanceStructIndex);
+				DataViews[Task.DataViewIndex.Get()] = InstanceData.GetMutableStruct(InstanceStructIndex);
 				InstanceStructIndex++;
 			}
 
 			// Copy bound properties.
 			if (Task.BindingsBatch.IsValid())
 			{
-				StateTree->PropertyBindings.CopyTo(DataViews, Task.BindingsBatch.Index, DataViews[Task.DataViewIndex]);
+				StateTree->PropertyBindings.CopyTo(DataViews, Task.BindingsBatch, DataViews[Task.DataViewIndex.Get()]);
 			}
 
 			if (bIsEnteringState)
@@ -566,7 +567,7 @@
 				if (Status == EStateTreeRunStatus::Failed)
 				{
 					// Store how far in the enter state we got. This will be used to match the ExitState() calls.
-					Exec.EnterStateFailedTaskIndex = TaskIndex;
+					Exec.EnterStateFailedTaskIndex = FStateTreeIndex16(TaskIndex);
 					Result = Status;
 					break;
 				}
@@ -588,7 +589,7 @@
 
 	// Reset transition delay
 	FStateTreeExecutionState& Exec = GetExecState(InstanceData);
-	Exec.GatedTransitionIndex = INDEX_NONE;
+	Exec.GatedTransitionIndex = FStateTreeIndex16::Invalid;
 	Exec.GatedTransitionTime = 0.0f;
 
 	// On target branch means that the state is the target of current transition or child of it.
@@ -596,21 +597,21 @@
 	// EnterState called. That is, a transition is handled as "replan from this state".
 	bool bOnTargetBranch = false;
 
-	FStateTreeHandle ExitedStates[FStateTreeActiveStates::MaxStates];
+	FStateTreeStateHandle ExitedStates[FStateTreeActiveStates::MaxStates];
 	EStateTreeStateChangeType ExitedStateChangeType[FStateTreeActiveStates::MaxStates];
 	int32 ExitedStateActiveIndex[FStateTreeActiveStates::MaxStates];
 	int32 NumExitedStates = 0;
 	
 	// Do property copy on all states, propagating the results from last tick.
 	// Collect the states that need to be called, the actual call is done below in reverse order.
+	check(Exec.FirstTaskStructIndex.IsValid() && Exec.FirstTaskObjectIndex.IsValid()); 
+	int32 InstanceStructIndex = Exec.FirstTaskStructIndex.Get();
+	int32 InstanceObjectIndex = Exec.FirstTaskObjectIndex.Get();
 
-	int32 InstanceStructIndex = Exec.FirstTaskStructIndex;
-	int32 InstanceObjectIndex = Exec.FirstTaskObjectIndex;
-
 	for (int32 Index = 0; Index < Transition.CurrentActiveStates.Num(); Index++)
 	{
-		const FStateTreeHandle CurrentHandle = Transition.CurrentActiveStates[Index];
-		const FStateTreeHandle NextHandle = Transition.NextActiveStates.GetStateSafe(Index);
+		const FStateTreeStateHandle CurrentHandle = Transition.CurrentActiveStates[Index];
+		const FStateTreeStateHandle NextHandle = Transition.NextActiveStates.GetStateSafe(Index);
 		const FCompactStateTreeState& State = StateTree->States[CurrentHandle.Index];
 
 		if (State.Type == EStateTreeStateType::Linked)
@@ -643,19 +644,19 @@
 			const FStateTreeTaskBase& Task = StateTree->Nodes[TaskIndex].Get<FStateTreeTaskBase>();
 			if (Task.bInstanceIsObject)
 			{
-				DataViews[Task.DataViewIndex] = InstanceData.GetMutableObject(InstanceObjectIndex);
+				DataViews[Task.DataViewIndex.Get()] = InstanceData.GetMutableObject(InstanceObjectIndex);
 				InstanceObjectIndex++;
 			}
 			else
 			{
-				DataViews[Task.DataViewIndex] = InstanceData.GetMutableStruct(InstanceStructIndex);
+				DataViews[Task.DataViewIndex.Get()] = InstanceData.GetMutableStruct(InstanceStructIndex);
 				InstanceStructIndex++;
 			}
 
 			// Copy bound properties.
 			if (Task.BindingsBatch.IsValid())
 			{
-				StateTree->PropertyBindings.CopyTo(DataViews, Task.BindingsBatch.Index, DataViews[Task.DataViewIndex]);
+				StateTree->PropertyBindings.CopyTo(DataViews, Task.BindingsBatch, DataViews[Task.DataViewIndex.Get()]);
 			}
 		}
 	}
@@ -665,7 +666,7 @@
 
 	for (int32 Index = NumExitedStates - 1; Index >= 0; Index--)
 	{
-		const FStateTreeHandle CurrentHandle = ExitedStates[Index];
+		const FStateTreeStateHandle CurrentHandle = ExitedStates[Index];
 		const FCompactStateTreeState& State = StateTree->States[CurrentHandle.Index];
 		const EStateTreeStateChangeType ChangeType = ExitedStateChangeType[Index];
 
@@ -678,7 +679,8 @@
 		{
 			// Call task completed only if EnterState() was called.
 			// The task order in the tree (BF) allows us to use the comparison.
-			if (TaskIndex <= (int32)Exec.EnterStateFailedTaskIndex)
+			// Relying here that invalid value of Exec.EnterStateFailedTaskIndex == MAX_uint16.
+			if (TaskIndex <= Exec.EnterStateFailedTaskIndex.Get())
 			{
 				const FStateTreeTaskBase& Task = StateTree->Nodes[TaskIndex].Get<FStateTreeTaskBase>();
 
@@ -708,7 +710,7 @@
 	// Note: Completed is assumed to be called immediately after tick or enter state, so there's no property copying.
 	for (int32 Index = Exec.ActiveStates.Num() - 1; Index >= 0; Index--)
 	{
-		const FStateTreeHandle CurrentHandle = Exec.ActiveStates[Index];
+		const FStateTreeStateHandle CurrentHandle = Exec.ActiveStates[Index];
 		const FCompactStateTreeState& State = StateTree->States[CurrentHandle.Index];
 
 		STATETREE_LOG(Verbose, TEXT("%*sState Completed '%s' %s"), Index*UE::StateTree::DebugIndentSize, TEXT(""), *DebugGetStatePath(Exec.ActiveStates, Index), *UEnum::GetValueAsString(Exec.LastTickStatus));
@@ -718,7 +720,8 @@
 		{
 			// Call task completed only if EnterState() was called.
 			// The task order in the tree (BF) allows us to use the comparison.
-			if (TaskIndex <= (int32)Exec.EnterStateFailedTaskIndex)
+			// Relying here that invalid value of Exec.EnterStateFailedTaskIndex == MAX_uint16.
+			if (TaskIndex <= Exec.EnterStateFailedTaskIndex.Get())
 			{
 				const FStateTreeTaskBase& Task = StateTree->Nodes[TaskIndex].Get<FStateTreeTaskBase>();
 
@@ -744,19 +747,19 @@
 		const FStateTreeEvaluatorBase& Eval = StateTree->Nodes[EvalIndex].Get<FStateTreeEvaluatorBase>();
 		if (Eval.bInstanceIsObject)
 		{
-			DataViews[Eval.DataViewIndex] = InstanceData.GetMutableObject(InstanceObjectIndex);
+			DataViews[Eval.DataViewIndex.Get()] = InstanceData.GetMutableObject(InstanceObjectIndex);
 			InstanceObjectIndex++;
 		}
 		else
 		{
-			DataViews[Eval.DataViewIndex] = InstanceData.GetMutableStruct(InstanceStructIndex);
+			DataViews[Eval.DataViewIndex.Get()] = InstanceData.GetMutableStruct(InstanceStructIndex);
 			InstanceStructIndex++;
 		}
 
 		// Copy bound properties.
 		if (Eval.BindingsBatch.IsValid())
 		{
-			StateTree->PropertyBindings.CopyTo(DataViews, Eval.BindingsBatch.Index, DataViews[Eval.DataViewIndex]);
+			StateTree->PropertyBindings.CopyTo(DataViews, Eval.BindingsBatch, DataViews[Eval.DataViewIndex.Get()]);
 		}
 		STATETREE_LOG(Verbose, TEXT("  Tick: '%s'"), *Eval.Name.ToString());
 		{
@@ -781,19 +784,19 @@
 		const FStateTreeEvaluatorBase& Eval = StateTree->Nodes[EvalIndex].Get<FStateTreeEvaluatorBase>();
 		if (Eval.bInstanceIsObject)
 		{
-			DataViews[Eval.DataViewIndex] = InstanceData.GetMutableObject(InstanceObjectIndex);
+			DataViews[Eval.DataViewIndex.Get()] = InstanceData.GetMutableObject(InstanceObjectIndex);
 			InstanceObjectIndex++;
 		}
 		else
 		{
-			DataViews[Eval.DataViewIndex] = InstanceData.GetMutableStruct(InstanceStructIndex);
+			DataViews[Eval.DataViewIndex.Get()] = InstanceData.GetMutableStruct(InstanceStructIndex);
 			InstanceStructIndex++;
 		}
 
 		// Copy bound properties.
 		if (Eval.BindingsBatch.IsValid())
 		{
-			StateTree->PropertyBindings.CopyTo(DataViews, Eval.BindingsBatch.Index, DataViews[Eval.DataViewIndex]);
+			StateTree->PropertyBindings.CopyTo(DataViews, Eval.BindingsBatch, DataViews[Eval.DataViewIndex.Get()]);
 		}
 		STATETREE_LOG(Verbose, TEXT("  Start: '%s'"), *Eval.Name.ToString());
 		{
@@ -818,19 +821,19 @@
 		const FStateTreeEvaluatorBase& Eval = StateTree->Nodes[EvalIndex].Get<FStateTreeEvaluatorBase>();
 		if (Eval.bInstanceIsObject)
 		{
-			DataViews[Eval.DataViewIndex] = InstanceData.GetMutableObject(InstanceObjectIndex);
+			DataViews[Eval.DataViewIndex.Get()] = InstanceData.GetMutableObject(InstanceObjectIndex);
 			InstanceObjectIndex++;
 		}
 		else
 		{
-			DataViews[Eval.DataViewIndex] = InstanceData.GetMutableStruct(InstanceStructIndex);
+			DataViews[Eval.DataViewIndex.Get()] = InstanceData.GetMutableStruct(InstanceStructIndex);
 			InstanceStructIndex++;
 		}
 
 		// Copy bound properties.
 		if (Eval.BindingsBatch.IsValid())
 		{
-			StateTree->PropertyBindings.CopyTo(DataViews, Eval.BindingsBatch.Index, DataViews[Eval.DataViewIndex]);
+			StateTree->PropertyBindings.CopyTo(DataViews, Eval.BindingsBatch, DataViews[Eval.DataViewIndex.Get()]);
 		}
 		STATETREE_LOG(Verbose, TEXT("  Stop: '%s'"), *Eval.Name.ToString());
 		{
@@ -854,12 +857,13 @@
 	EStateTreeRunStatus Result = EStateTreeRunStatus::Running;
 	int32 NumTotalTasks = 0;
 
-	int32 InstanceStructIndex = Exec.FirstTaskStructIndex;
-	int32 InstanceObjectIndex = Exec.FirstTaskObjectIndex;
+	check(Exec.FirstTaskStructIndex.IsValid() && Exec.FirstTaskObjectIndex.IsValid()); 
+	int32 InstanceStructIndex = Exec.FirstTaskStructIndex.Get();
+	int32 InstanceObjectIndex = Exec.FirstTaskObjectIndex.Get();
 
 	for (int32 Index = 0; Index < Exec.ActiveStates.Num() && Result != EStateTreeRunStatus::Failed; Index++)
 	{
-		const FStateTreeHandle CurrentHandle = Exec.ActiveStates[Index];
+		const FStateTreeStateHandle CurrentHandle = Exec.ActiveStates[Index];
 		const FCompactStateTreeState& State = StateTree->States[CurrentHandle.Index];
 
 		STATETREE_CLOG(State.TasksNum > 0, Verbose, TEXT("%*sTicking Tasks of state '%s'"), Index*UE::StateTree::DebugIndentSize, TEXT(""), *DebugGetStatePath(Exec.ActiveStates, Index));
@@ -880,19 +884,19 @@
 			const FStateTreeTaskBase& Task = StateTree->Nodes[TaskIndex].Get<FStateTreeTaskBase>();
 			if (Task.bInstanceIsObject)
 			{
-				DataViews[Task.DataViewIndex] = InstanceData.GetMutableObject(InstanceObjectIndex);
+				DataViews[Task.DataViewIndex.Get()] = InstanceData.GetMutableObject(InstanceObjectIndex);
 				InstanceObjectIndex++;
 			}
 			else
 			{
-				DataViews[Task.DataViewIndex] = InstanceData.GetMutableStruct(InstanceStructIndex);
+				DataViews[Task.DataViewIndex.Get()] = InstanceData.GetMutableStruct(InstanceStructIndex);
 				InstanceStructIndex++;
 			}
 
 			// Copy bound properties.
 			if (Task.BindingsBatch.IsValid())
 			{
-				StateTree->PropertyBindings.CopyTo(DataViews, Task.BindingsBatch.Index, DataViews[Task.DataViewIndex]);
+				StateTree->PropertyBindings.CopyTo(DataViews, Task.BindingsBatch, DataViews[Task.DataViewIndex.Get()]);
 			}
 			STATETREE_LOG(Verbose, TEXT("%*s  Tick: '%s'"), Index*UE::StateTree::DebugIndentSize, TEXT(""), *Task.Name.ToString());
 			
@@ -946,17 +950,17 @@
 		const FStateTreeConditionBase& Cond = StateTree->Nodes[ConditionsOffset + Index].Get<FStateTreeConditionBase>();
 		if (Cond.bInstanceIsObject)
 		{
-			DataViews[Cond.DataViewIndex] = SharedInstanceData.GetMutableObject(Cond.InstanceIndex);
+			DataViews[Cond.DataViewIndex.Get()] = SharedInstanceData.GetMutableObject(Cond.InstanceIndex.Get());
 		}
 		else
 		{
-			DataViews[Cond.DataViewIndex] = SharedInstanceData.GetMutableStruct(Cond.InstanceIndex);
+			DataViews[Cond.DataViewIndex.Get()] = SharedInstanceData.GetMutableStruct(Cond.InstanceIndex.Get());
 		}
 
 		// Copy bound properties.
 		if (Cond.BindingsBatch.IsValid())
 		{
-			if (!StateTree->PropertyBindings.CopyTo(DataViews, Cond.BindingsBatch.Index, DataViews[Cond.DataViewIndex]))
+			if (!StateTree->PropertyBindings.CopyTo(DataViews, Cond.BindingsBatch, DataViews[Cond.DataViewIndex.Get()]))
 			{
 				// If the source data cannot be accessed, the whole expression evaluates to false.
 				Values[0] = false;
@@ -1035,9 +1039,9 @@
 				// If a transition has delay, we stop testing other transitions, but the transition will not pass the condition until the delay time passes.
 				if (Transition.GateDelay > 0)
 				{
-					if (Exec.GatedTransitionIndex != TransitionIndex)
+					if ((int32)Exec.GatedTransitionIndex.Get() != TransitionIndex)
 					{
-						Exec.GatedTransitionIndex = TransitionIndex;
+						Exec.GatedTransitionIndex = FStateTreeIndex16(TransitionIndex);
 						Exec.GatedTransitionTime = FMath::RandRange(0.0f, Transition.GateDelay * 0.1f); // TODO: we need variance too.
 						BeginGatedTransition(Exec);
 						STATETREE_LOG(Verbose, TEXT("Gated transition triggered from '%s' (%s) -> '%s' %.1fs"), *GetSafeStateName(Exec.ActiveStates.Last()), *State.Name.ToString(), *GetSafeStateName(Transition.State), Exec.GatedTransitionTime);
@@ -1073,23 +1077,23 @@
 				{
 					STATETREE_LOG(Verbose, TEXT("Stop tree execution from state '%s' (%s): Succeeded"), *GetSafeStateName(Exec.ActiveStates.Last()), *State.Name.ToString());
 					OutTransition.CurrentActiveStates = Exec.ActiveStates;
-					OutTransition.TargetState = FStateTreeHandle::Succeeded;
-					OutTransition.NextActiveStates = FStateTreeActiveStates(FStateTreeHandle::Succeeded);
+					OutTransition.TargetState = FStateTreeStateHandle::Succeeded;
+					OutTransition.NextActiveStates = FStateTreeActiveStates(FStateTreeStateHandle::Succeeded);
 					return true;
 				}
 				else
 				{
 					STATETREE_LOG(Verbose, TEXT("Stop tree execution from state '%s' (%s): Failed"), *GetSafeStateName(Exec.ActiveStates.Last()), *State.Name.ToString());
 					OutTransition.CurrentActiveStates = Exec.ActiveStates;
-					OutTransition.TargetState = FStateTreeHandle::Failed;
-					OutTransition.NextActiveStates = FStateTreeActiveStates(FStateTreeHandle::Failed);
+					OutTransition.TargetState = FStateTreeStateHandle::Failed;
+					OutTransition.NextActiveStates = FStateTreeActiveStates(FStateTreeStateHandle::Failed);
 					return true;
 				}
 			}
-			else if (Exec.GatedTransitionIndex == TransitionIndex)
+			else if ((int32)Exec.GatedTransitionIndex.Get() == TransitionIndex)
 			{
 				// If the current transition was gated transition, reset it if the condition failed.
-				Exec.GatedTransitionIndex = INDEX_NONE;
+				Exec.GatedTransitionIndex = FStateTreeIndex16::Invalid;
 				Exec.GatedTransitionTime = 0.0f;
 			}
 		}
@@ -1098,7 +1102,7 @@
 	if (Exec.LastTickStatus != EStateTreeRunStatus::Running)
 	{
 		// Could not trigger completion transition, jump back to start.
-		static const FStateTreeHandle RootState = FStateTreeHandle(0);
+		static const FStateTreeStateHandle RootState = FStateTreeStateHandle(0);
 		OutTransition.TargetState = RootState;
 		return SelectState(InstanceData, RootState, OutTransition.NextActiveStates);
 	}
@@ -1107,7 +1111,7 @@
 	return false;
 }
 
-bool FStateTreeExecutionContext::SelectState(FStateTreeInstanceData& InstanceData, const FStateTreeHandle NextState, FStateTreeActiveStates& OutNewActiveState)
+bool FStateTreeExecutionContext::SelectState(FStateTreeInstanceData& InstanceData, const FStateTreeStateHandle NextState, FStateTreeActiveStates& OutNewActiveState)
 {
 	const FStateTreeExecutionState& Exec = GetExecState(InstanceData);
 
@@ -1120,12 +1124,12 @@
 	// This allows transitions within a subtree.
 	OutNewActiveState = Exec.ActiveStates;
 	
-	TStaticArray<FStateTreeHandle, FStateTreeActiveStates::MaxStates> InBetweenStates;
+	TStaticArray<FStateTreeStateHandle, FStateTreeActiveStates::MaxStates> InBetweenStates;
 	int32 NumInBetweenStates = 0;
 	int32 CommonActiveAncestorIndex = INDEX_NONE;
 
 	// Walk towards the root from current state.
-	FStateTreeHandle CurrState = NextState;
+	FStateTreeStateHandle CurrState = NextState;
 	while (CurrState.IsValid())
 	{
 		// Store the states that are in between the 'NextState' and common ancestor. 
@@ -1166,7 +1170,7 @@
 	return SelectStateInternal(InstanceData, NextState, OutNewActiveState);
 }
 
-bool FStateTreeExecutionContext::SelectStateInternal(FStateTreeInstanceData& InstanceData, const FStateTreeHandle NextState, FStateTreeActiveStates& OutNewActiveState)
+bool FStateTreeExecutionContext::SelectStateInternal(FStateTreeInstanceData& InstanceData, const FStateTreeStateHandle NextState, FStateTreeActiveStates& OutNewActiveState)
 {
 	CSV_SCOPED_TIMING_STAT_EXCLUSIVE(StateTree_SelectState);
 
@@ -1206,7 +1210,7 @@
 			// If the state has children, proceed to select children.
 			for (uint16 ChildState = State.ChildrenBegin; ChildState < State.ChildrenEnd; ChildState = StateTree->States[ChildState].GetNextSibling())
 			{
-				if (SelectStateInternal(InstanceData, FStateTreeHandle(ChildState), OutNewActiveState))
+				if (SelectStateInternal(InstanceData, FStateTreeStateHandle(ChildState), OutNewActiveState))
 				{
 					// Selection succeeded
 					return true;
@@ -1226,18 +1230,18 @@
 	return false;
 }
 
-FString FStateTreeExecutionContext::GetSafeStateName(const FStateTreeHandle State) const
+FString FStateTreeExecutionContext::GetSafeStateName(const FStateTreeStateHandle State) const
 {
 	check(StateTree);
-	if (State == FStateTreeHandle::Invalid)
+	if (State == FStateTreeStateHandle::Invalid)
 	{
 		return TEXT("(State Invalid)");
 	}
-	else if (State == FStateTreeHandle::Succeeded)
+	else if (State == FStateTreeStateHandle::Succeeded)
 	{
 		return TEXT("(State Succeeded)");
 	}
-	else if (State == FStateTreeHandle::Failed)
+	else if (State == FStateTreeStateHandle::Failed)
 	{
 		return TEXT("(State Failed)");
 	}
@@ -1336,7 +1340,7 @@
 	DebugString += TEXT("Current State:\n");
 	for (int32 Index = 0; Index < Exec.ActiveStates.Num(); Index++)
 	{
-		FStateTreeHandle Handle = Exec.ActiveStates[Index];
+		FStateTreeStateHandle Handle = Exec.ActiveStates[Index];
 		if (Handle.IsValid())
 		{
 			const FCompactStateTreeState& State = StateTree->States[Handle.Index];
@@ -1397,7 +1401,7 @@
 	DebugString += FString::Printf(TEXT("\nExternal Data(%d)\n  [ %-40s | %-8s | %5s ]\n"), StateTree->ExternalDataDescs.Num(), TEXT("Name"), TEXT("Optional"), TEXT("Index"));
 	for (const FStateTreeExternalDataDesc& Desc : StateTree->ExternalDataDescs)
 	{
-		DebugString += FString::Printf(TEXT("  | %-40s | %8s | %5d |\n"), Desc.Struct ? *Desc.Struct->GetName() : TEXT("null"), *UEnum::GetValueAsString(Desc.Requirement), Desc.Handle.DataViewIndex);
+		DebugString += FString::Printf(TEXT("  | %-40s | %8s | %5d |\n"), Desc.Struct ? *Desc.Struct->GetName() : TEXT("null"), *UEnum::GetValueAsString(Desc.Requirement), Desc.Handle.DataViewIndex.Get());
 	}
 
 	// Bindings
@@ -1442,17 +1446,19 @@
 	// Evaluators
 	if (StateTree->EvaluatorsNum)
 	{
+		DebugString += FString::Printf(TEXT("\nEvaluators\n  [ %-30s | %8s | %10s ]\n"),
+			TEXT("Name"), TEXT("Bindings"), TEXT("Struct Idx"));
 		for (int32 EvalIndex = StateTree->EvaluatorsBegin; EvalIndex < (StateTree->EvaluatorsBegin + StateTree->EvaluatorsNum); EvalIndex++)
 		{
 			const FStateTreeEvaluatorBase& Eval = StateTree->Nodes[EvalIndex].Get<FStateTreeEvaluatorBase>();
-			DebugString += FString::Printf(TEXT("| %-12s | %-30s | %15s | %10d |\n"),
-				TEXT("Evaluator"), *Eval.Name.ToString(), *Eval.BindingsBatch.Describe(), Eval.DataViewIndex);
+			DebugString += FString::Printf(TEXT("| %-30s | %8d | %10d |\n"),
+				*Eval.Name.ToString(), Eval.BindingsBatch.Get(), Eval.DataViewIndex.Get());
 		}
 	}
 
 
-	DebugString += FString::Printf(TEXT("\nStates\n  [ %-30s | %-12s | %-30s | %15s | %10s ]\n"),
-		TEXT("State"), TEXT("Type"), TEXT("Name"), TEXT("Bindings"), TEXT("Struct Idx"));
+	DebugString += FString::Printf(TEXT("\nTasks\n  [ %-30s | %-30s | %8s | %10s ]\n"),
+		TEXT("State"), TEXT("Name"), TEXT("Bindings"), TEXT("Struct Idx"));
 	for (const FCompactStateTreeState& State : StateTree->States)
 	{
 		// Tasks
@@ -1461,8 +1467,8 @@
 			for (int32 TaskIndex = State.TasksBegin; TaskIndex < (State.TasksBegin + State.TasksNum); TaskIndex++)
 			{
 				const FStateTreeTaskBase& Task = StateTree->Nodes[TaskIndex].Get<FStateTreeTaskBase>();
-				DebugString += FString::Printf(TEXT("  | %-30s | %-12s | %-30s | %15s | %10d |\n"), *State.Name.ToString(),
-					TEXT("  Task"), *Task.Name.ToString(), *Task.BindingsBatch.Describe(), Task.DataViewIndex);
+				DebugString += FString::Printf(TEXT("  | %-30s | %-30s | %8d | %10d |\n"), *State.Name.ToString(),
+					*Task.Name.ToString(), Task.BindingsBatch.Get(), Task.DataViewIndex.Get());
 			}
 		}
 	}
@@ -1503,7 +1509,7 @@
 	// Active States
 	for (int32 Index = 0; Index < Exec.ActiveStates.Num(); Index++)
 	{
-		const FStateTreeHandle Handle = Exec.ActiveStates[Index];
+		const FStateTreeStateHandle Handle = Exec.ActiveStates[Index];
 		if (Handle.IsValid())
 		{
 			const FCompactStateTreeState& State = StateTree->States[Handle.Index];
@@ -1560,7 +1566,7 @@
 	// Active States
 	for (int32 Index = 0; Index < Exec.ActiveStates.Num(); Index++)
 	{
-		const FStateTreeHandle Handle = Exec.ActiveStates[Index];
+		const FStateTreeStateHandle Handle = Exec.ActiveStates[Index];
 		if (Handle.IsValid())
 		{
 			const FCompactStateTreeState& State = StateTree->States[Handle.Index];
diff --git a/Engine/Plugins/Runtime/StateTree/Source/StateTreeModule/Private/StateTreePropertyBindings.cpp b/Engine/Plugins/Runtime/StateTree/Source/StateTreeModule/Private/StateTreePropertyBindings.cpp
--- a/Engine/Plugins/Runtime/StateTree/Source/StateTreeModule/Private/StateTreePropertyBindings.cpp
+++ b/Engine/Plugins/Runtime/StateTree/Source/StateTreeModule/Private/StateTreePropertyBindings.cpp
@@ -11,10 +11,8 @@ void FStateTreePropertyBindings::Reset()
 	SourceStructs.Reset();
 	CopyBatches.Reset();
 	PropertyBindings.Reset();
-	PropertyPaths.Reset();
 	PropertySegments.Reset();
 	PropertyCopies.Reset();
-	PropertyAccesses.Reset();
 	PropertyIndirections.Reset();
 	bBindingsResolved = false;
 }
@@ -23,7 +21,6 @@ bool FStateTreePropertyBindings::ResolvePaths()
 {
 	PropertyIndirections.Reset();
 	PropertyCopies.SetNum(PropertyBindings.Num());
-	PropertyAccesses.SetNum(PropertyPaths.Num());
 
 	bBindingsResolved = true;
 
@@ -36,12 +33,10 @@ bool FStateTreePropertyBindings::ResolvePaths()
 			const FStateTreePropertyBinding& Binding = PropertyBindings[i];
 			FStateTreePropCopy& Copy = PropertyCopies[i];
 
-			Copy.SourceAccessIndex = Binding.SourcePathIndex;
-			Copy.TargetAccessIndex = Binding.TargetPathIndex;
 			Copy.SourceStructIndex = Binding.SourceStructIndex;
 			Copy.Type = Binding.CopyType;
 
-			const UStruct* SourceStruct = SourceStructs[Binding.SourceStructIndex].Struct;
+			const UStruct* SourceStruct = SourceStructs[Binding.SourceStructIndex.Get()].Struct;
 			const UStruct* TargetStruct = Batch.TargetStruct.Struct;
 			if (!SourceStruct || !TargetStruct)
 			{
@@ -53,8 +48,8 @@ bool FStateTreePropertyBindings::ResolvePaths()
 
 			// Resolve paths and validate the copy. Stops on first failure.
 			bool bSuccess = true;
-			bSuccess = bSuccess && ResolvePath(SourceStruct, PropertyPaths[Binding.SourcePathIndex], PropertyAccesses[Copy.SourceAccessIndex]);
-			bSuccess = bSuccess && ResolvePath(TargetStruct, PropertyPaths[Binding.TargetPathIndex], PropertyAccesses[Copy.TargetAccessIndex]);
+			bSuccess = bSuccess && ResolvePath(SourceStruct, Binding.SourcePath, Copy.SourceIndirection, Copy.SourceLeafProperty);
+			bSuccess = bSuccess && ResolvePath(TargetStruct, Binding.TargetPath, Copy.TargetIndirection, Copy.TargetLeafProperty);
 			bSuccess = bSuccess && ValidateCopy(Copy);
 			if (!bSuccess)
 			{
@@ -68,51 +63,51 @@ bool FStateTreePropertyBindings::ResolvePaths()
 	return bResult;
 }
 
-bool FStateTreePropertyBindings::ResolvePath(const UStruct* Struct, const FStateTreePropertyPath& Path, FStateTreePropertyAccess& OutAccess)
+bool FStateTreePropertyBindings::ResolvePath(const UStruct* Struct, const FStateTreePropertySegment& FirstPathSegment, FStateTreePropertyIndirection& OutFirstIndirection, const FProperty*& OutLeafProperty)
 {
 	if (!Struct)
 	{
-		UE_LOG(LogStateTree, Error, TEXT("%s: '%s' Invalid source struct."), ANSI_TO_TCHAR(__FUNCTION__), *GetPathAsString(Path));
+		UE_LOG(LogStateTree, Error, TEXT("%s: '%s' Invalid source struct."), ANSI_TO_TCHAR(__FUNCTION__), *GetPathAsString(FirstPathSegment));
 		return false;
 	}
 
 	const UStruct* CurrentStruct = Struct;
-	FProperty* LeafProperty = nullptr;
-	const int32 IndirectionsBegin = PropertyIndirections.Num();
+	const FProperty* LeafProperty = nullptr;
+	TArray<FStateTreePropertyIndirection, TInlineAllocator<16>> TempIndirections;
+	const FStateTreePropertySegment* Segment = &FirstPathSegment;
 
-	for (int32 i = Path.SegmentsBegin; i < Path.SegmentsEnd; i++)
+	while (Segment != nullptr)
 	{
-		const FStateTreePropertySegment& Segment = PropertySegments[i];
-		FStateTreePropertyIndirection& Indirection = PropertyIndirections.AddDefaulted_GetRef();
+		FStateTreePropertyIndirection& Indirection = TempIndirections.AddDefaulted_GetRef();
 
-		const bool bFinalSegment = i == (Path.SegmentsEnd - 1);
+		const bool bFinalSegment = Segment->NextIndex.IsValid() == false;;
 
 		if (!ensure(CurrentStruct))
 		{
-			UE_LOG(LogStateTree, Error, TEXT("%s: '%s' Invalid struct."), ANSI_TO_TCHAR(__FUNCTION__), *GetPathAsString(Path, i, TEXT("<"), TEXT(">")));
+			UE_LOG(LogStateTree, Error, TEXT("%s: '%s' Invalid struct."), ANSI_TO_TCHAR(__FUNCTION__), *GetPathAsString(FirstPathSegment, Segment, TEXT("<"), TEXT(">")));
 			return false;
 		}
-		FProperty* Property = CurrentStruct->FindPropertyByName(Segment.Name);
+		FProperty* Property = CurrentStruct->FindPropertyByName(Segment->Name);
 		if (!Property)
 		{
 			// TODO: use core redirects to fix up the name.
 			UE_LOG(LogStateTree, Error, TEXT("%s: Malformed path '%s', could not to find property '%s%s.%s'."),
-				ANSI_TO_TCHAR(__FUNCTION__), *GetPathAsString(Path, i, TEXT("<"), TEXT(">")),
-				CurrentStruct->GetPrefixCPP(), *CurrentStruct->GetName(), *Segment.Name.ToString());
+				ANSI_TO_TCHAR(__FUNCTION__), *GetPathAsString(FirstPathSegment, Segment, TEXT("<"), TEXT(">")),
+				CurrentStruct->GetPrefixCPP(), *CurrentStruct->GetName(), *Segment->Name.ToString());
 			return false;
 		}
-		Indirection.ArrayIndex = Segment.ArrayIndex == INDEX_NONE ? 0 : Segment.ArrayIndex;
-		Indirection.Type = Segment.Type;
-		Indirection.Offset = Property->GetOffset_ForInternal() + Property->ElementSize * Indirection.ArrayIndex;
+		Indirection.ArrayIndex = FStateTreeIndex16(Segment->ArrayIndex.IsValid() ? Segment->ArrayIndex.Get() : 0);
+		Indirection.Type = Segment->Type;
+		Indirection.Offset = Property->GetOffset_ForInternal() + Property->ElementSize * Indirection.ArrayIndex.Get();
 
 		// Check to see if it is an array access first.
-		if (FArrayProperty* ArrayProperty = CastField<FArrayProperty>(Property))
+		if (const FArrayProperty* ArrayProperty = CastField<FArrayProperty>(Property))
 		{
-			if (FStructProperty* ArrayOfStructsProperty = CastField<FStructProperty>(ArrayProperty->Inner))
+			if (const FStructProperty* ArrayOfStructsProperty = CastField<FStructProperty>(ArrayProperty->Inner))
 			{
 				CurrentStruct = ArrayOfStructsProperty->Struct;
 			}
-			else if (FObjectPropertyBase* ArrayOfObjectsProperty = CastField<FObjectPropertyBase>(ArrayProperty->Inner))
+			else if (const FObjectPropertyBase* ArrayOfObjectsProperty = CastField<FObjectPropertyBase>(ArrayProperty->Inner))
 			{
 				CurrentStruct = ArrayOfObjectsProperty->PropertyClass;
 			}
@@ -124,67 +119,77 @@ bool FStateTreePropertyBindings::ResolvePath(const UStruct* Struct, const FState
 			CurrentStruct = nullptr;
 		}
 		// Check to see if this is a simple structure (eg. not an array of structures)
-		else if (FStructProperty* StructProperty = CastField<FStructProperty>(Property))
+		else if (const FStructProperty* StructProperty = CastField<FStructProperty>(Property))
 		{
 			CurrentStruct = StructProperty->Struct;
 		}
 		// Check to see if this is a simple object (eg. not an array of objects)
-		else if (FObjectProperty* ObjectProperty = CastField<FObjectProperty>(Property))
+		else if (const FObjectProperty* ObjectProperty = CastField<FObjectProperty>(Property))
 		{
 			CurrentStruct = ObjectProperty->PropertyClass;
 		}
 		// Check to see if this is a simple weak object property (eg. not an array of weak objects).
-		else if (FWeakObjectProperty* WeakObjectProperty = CastField<FWeakObjectProperty>(Property))
+		else if (const FWeakObjectProperty* WeakObjectProperty = CastField<FWeakObjectProperty>(Property))
 		{
 			CurrentStruct = WeakObjectProperty->PropertyClass;
 		}
 		// Check to see if this is a simple soft object property (eg. not an array of soft objects).
-		else if (FSoftObjectProperty* SoftObjectProperty = CastField<FSoftObjectProperty>(Property))
+		else if (const FSoftObjectProperty* SoftObjectProperty = CastField<FSoftObjectProperty>(Property))
 		{
 			CurrentStruct = SoftObjectProperty->PropertyClass;
 		}
 		else
 		{
 			UE_LOG(LogStateTree, Error, TEXT("%s: Unsupported segment %s in path '%s'."),
-			    ANSI_TO_TCHAR(__FUNCTION__), *Property->GetCPPType(), *GetPathAsString(Path, i, TEXT("<"), TEXT(">")));
+			    ANSI_TO_TCHAR(__FUNCTION__), *Property->GetCPPType(), *GetPathAsString(FirstPathSegment, Segment, TEXT("<"), TEXT(">")));
 		}
 
 		if (bFinalSegment)
 		{
 			LeafProperty = Property;
 		}
-	}
 
-	int32 IndirectionsEnd = PropertyIndirections.Num();
+		Segment = Segment->NextIndex.IsValid() ? &PropertySegments[Segment->NextIndex.Get()] : nullptr;
+	}
 
 	// Collapse adjacent offset indirections
-	for (int32 i = IndirectionsBegin; i < IndirectionsEnd; i++)
+	for (int32 Index = 0; Index < TempIndirections.Num(); Index++)
 	{
-		FStateTreePropertyIndirection& Indirection = PropertyIndirections[i];
-		if (Indirection.Type == EStateTreePropertyAccessType::Offset && (i + 1) < IndirectionsEnd)
+		FStateTreePropertyIndirection& Indirection = TempIndirections[Index];
+		if (Indirection.Type == EStateTreePropertyAccessType::Offset && (Index + 1) < TempIndirections.Num())
 		{
-			FStateTreePropertyIndirection& NextIndirection = PropertyIndirections[i + 1];
+			const FStateTreePropertyIndirection& NextIndirection = TempIndirections[Index + 1];
 			if (NextIndirection.Type == EStateTreePropertyAccessType::Offset)
 			{
 				Indirection.Offset += NextIndirection.Offset;
-				PropertyIndirections.RemoveAt(i + 1);
-				i--;
-				IndirectionsEnd--;
+				TempIndirections.RemoveAt(Index + 1);
+				Index--;
 			}
 		}
 	}
 
-	OutAccess.IndirectionsBegin = IndirectionsBegin;
-	OutAccess.IndirectionsEnd = IndirectionsEnd;
-	OutAccess.LeafProperty = LeafProperty;
+	check(TempIndirections.Num() > 0);
+
+	// Store indirections
+	OutFirstIndirection = TempIndirections[0];
+	FStateTreePropertyIndirection* PrevIndirection = &OutFirstIndirection;
+	for (int32 Index = 1; Index < TempIndirections.Num(); Index++)
+	{
+		const int32 IndirectionIndex = PropertyIndirections.Num();
+		FStateTreePropertyIndirection& NewIndirection = PropertyIndirections.Add_GetRef(TempIndirections[Index]);
+		PrevIndirection->NextIndex = FStateTreeIndex16(IndirectionIndex);
+		PrevIndirection = &NewIndirection;
+	}
+	
+	OutLeafProperty = LeafProperty;
 
 	return true;
 }
 
 bool FStateTreePropertyBindings::ValidateCopy(FStateTreePropCopy& Copy) const
 {
-	const FProperty* SourceProperty = PropertyAccesses[Copy.SourceAccessIndex].LeafProperty;
-	const FProperty* TargetProperty = PropertyAccesses[Copy.TargetAccessIndex].LeafProperty;
+	const FProperty* SourceProperty = Copy.SourceLeafProperty;
+	const FProperty* TargetProperty = Copy.TargetLeafProperty;
 
 	if (!TargetProperty)
 	{
@@ -212,7 +217,7 @@ bool FStateTreePropertyBindings::ValidateCopy(FStateTreePropCopy& Copy) const
 	switch (Copy.Type)
 	{
 	case EStateTreePropertyCopyType::CopyPlain:
-		Copy.CopySize = PropertyAccesses[Copy.SourceAccessIndex].LeafProperty->ElementSize * PropertyAccesses[Copy.SourceAccessIndex].LeafProperty->ArrayDim;
+		Copy.CopySize = SourceProperty->ElementSize * SourceProperty->ArrayDim;
 		bResult = (SourceProperty->PropertyFlags & CPF_IsPlainOldData) != 0 && (TargetProperty->PropertyFlags & CPF_IsPlainOldData) != 0;
 		break;
 	case EStateTreePropertyCopyType::CopyComplex:
@@ -300,51 +305,52 @@ bool FStateTreePropertyBindings::ValidateCopy(FStateTreePropCopy& Copy) const
 	return bResult;
 }
 
-uint8* FStateTreePropertyBindings::GetAddress(FStateTreeDataView InStructView, const FStateTreePropertyAccess& InAccess) const
+uint8* FStateTreePropertyBindings::GetAddress(FStateTreeDataView InStructView, const FStateTreePropertyIndirection& FirstIndirection, const FProperty* LeafProperty) const
 {
 	uint8* Address = InStructView.GetMutableMemory();
 	if (Address == nullptr)
 	{
 		return nullptr;
 	}
-	
-	for (int32 i = InAccess.IndirectionsBegin; Address != nullptr && i < InAccess.IndirectionsEnd; i++)
+
+	const FStateTreePropertyIndirection* Indirection = &FirstIndirection;
+
+	while (Indirection != nullptr && Address != nullptr)
 	{
-		const FStateTreePropertyIndirection& Indirection = PropertyIndirections[i];
-		switch (Indirection.Type)
+		switch (Indirection->Type)
 		{
 		case EStateTreePropertyAccessType::Offset:
 		{
-			Address = Address + Indirection.Offset;
+			Address = Address + Indirection->Offset;
 			break;
 		}
 		case EStateTreePropertyAccessType::Object:
 		{
-			UObject* Object = *reinterpret_cast<UObject**>(Address + Indirection.Offset);
+			UObject* Object = *reinterpret_cast<UObject**>(Address + Indirection->Offset);
 			Address = reinterpret_cast<uint8*>(Object);
 			break;
 		}
 		case EStateTreePropertyAccessType::WeakObject:
 		{
-			TWeakObjectPtr<UObject>& WeakObjectPtr = *reinterpret_cast<TWeakObjectPtr<UObject>*>(Address + Indirection.Offset);
+			TWeakObjectPtr<UObject>& WeakObjectPtr = *reinterpret_cast<TWeakObjectPtr<UObject>*>(Address + Indirection->Offset);
 			UObject* Object = WeakObjectPtr.Get();
 			Address = reinterpret_cast<uint8*>(Object);
 			break;
 		}
 		case EStateTreePropertyAccessType::SoftObject:
 		{
-			FSoftObjectPtr& SoftObjectPtr = *reinterpret_cast<FSoftObjectPtr*>(Address + Indirection.Offset);
+			FSoftObjectPtr& SoftObjectPtr = *reinterpret_cast<FSoftObjectPtr*>(Address + Indirection->Offset);
 			UObject* Object = SoftObjectPtr.Get();
 			Address = reinterpret_cast<uint8*>(Object);
 			break;
 		}
 		case EStateTreePropertyAccessType::IndexArray:
 		{
-			check(Indirection.ArrayProperty);
-			FScriptArrayHelper Helper(Indirection.ArrayProperty, Address + Indirection.Offset);
-			if (Helper.IsValidIndex(Indirection.ArrayIndex))
+			check(Indirection->ArrayProperty);
+			FScriptArrayHelper Helper(Indirection->ArrayProperty, Address + Indirection->Offset);
+			if (Helper.IsValidIndex(Indirection->ArrayIndex.Get()))
 			{
-				Address = reinterpret_cast<uint8*>(Helper.GetRawPtr(Indirection.ArrayIndex));
+				Address = reinterpret_cast<uint8*>(Helper.GetRawPtr(Indirection->ArrayIndex.Get()));
 			}
 			else
 			{
@@ -354,18 +360,20 @@ uint8* FStateTreePropertyBindings::GetAddress(FStateTreeDataView InStructView, c
 		}
 		default:
 			ensureMsgf(false, TEXT("FStateTreePropertyBindings::GetAddress: Unhandled indirection type %s for '%s'"),
-				*StaticEnum<EStateTreePropertyAccessType>()->GetValueAsString(Indirection.Type), *InAccess.LeafProperty->GetNameCPP());
+				*StaticEnum<EStateTreePropertyAccessType>()->GetValueAsString(Indirection->Type), *LeafProperty->GetNameCPP());
 		}
+
+		Indirection = Indirection->NextIndex.IsValid() ? &PropertyIndirections[Indirection->NextIndex.Get()] : nullptr;
 	}
 
 	return Address;
 }
 
-void FStateTreePropertyBindings::PerformCopy(const FStateTreePropCopy& Copy, const FProperty* SourceProperty, const uint8* SourceAddress, const FProperty* TargetProperty, uint8* TargetAddress) const
+void FStateTreePropertyBindings::PerformCopy(const FStateTreePropCopy& Copy, const uint8* SourceAddress, uint8* TargetAddress) const
 {
 	// Source property can be null
 	check(SourceAddress);
-	check(TargetProperty);
+	check(Copy.TargetLeafProperty);
 	check(TargetAddress);
 
 	switch (Copy.Type)
@@ -374,34 +382,34 @@ void FStateTreePropertyBindings::PerformCopy(const FStateTreePropCopy& Copy, con
 		FMemory::Memcpy(TargetAddress, SourceAddress, Copy.CopySize);
 		break;
 	case EStateTreePropertyCopyType::CopyComplex:
-		TargetProperty->CopyCompleteValue(TargetAddress, SourceAddress);
+		Copy.TargetLeafProperty->CopyCompleteValue(TargetAddress, SourceAddress);
 		break;
 	case EStateTreePropertyCopyType::CopyBool:
-		static_cast<const FBoolProperty*>(TargetProperty)->SetPropertyValue(TargetAddress, static_cast<const FBoolProperty*>(SourceProperty)->GetPropertyValue(SourceAddress));
+		static_cast<const FBoolProperty*>(Copy.TargetLeafProperty)->SetPropertyValue(TargetAddress, static_cast<const FBoolProperty*>(Copy.SourceLeafProperty)->GetPropertyValue(SourceAddress));
 		break;
 	case EStateTreePropertyCopyType::CopyStruct:
 		// If SourceProperty == nullptr (pointing to the struct source directly), the GetAddress() did the right thing and is pointing the the beginning of the struct. 
-		static_cast<const FStructProperty*>(TargetProperty)->Struct->CopyScriptStruct(TargetAddress, SourceAddress);
+		static_cast<const FStructProperty*>(Copy.TargetLeafProperty)->Struct->CopyScriptStruct(TargetAddress, SourceAddress);
 		break;
 	case EStateTreePropertyCopyType::CopyObject:
-		if (SourceProperty == nullptr)
+		if (Copy.SourceLeafProperty == nullptr)
 		{
 			// Source is pointing at object directly.
-			static_cast<const FObjectPropertyBase*>(TargetProperty)->SetObjectPropertyValue(TargetAddress, (UObject*)SourceAddress);
+			static_cast<const FObjectPropertyBase*>(Copy.TargetLeafProperty)->SetObjectPropertyValue(TargetAddress, (UObject*)SourceAddress);
 		}
 		else
 		{
-			static_cast<const FObjectPropertyBase*>(TargetProperty)->SetObjectPropertyValue(TargetAddress, static_cast<const FObjectPropertyBase*>(SourceProperty)->GetObjectPropertyValue(SourceAddress));
+			static_cast<const FObjectPropertyBase*>(Copy.TargetLeafProperty)->SetObjectPropertyValue(TargetAddress, static_cast<const FObjectPropertyBase*>(Copy.SourceLeafProperty)->GetObjectPropertyValue(SourceAddress));
 		}
 		break;
 	case EStateTreePropertyCopyType::CopyName:
-		static_cast<const FNameProperty*>(TargetProperty)->SetPropertyValue(TargetAddress, static_cast<const FNameProperty*>(SourceProperty)->GetPropertyValue(SourceAddress));
+		static_cast<const FNameProperty*>(Copy.TargetLeafProperty)->SetPropertyValue(TargetAddress, static_cast<const FNameProperty*>(Copy.SourceLeafProperty)->GetPropertyValue(SourceAddress));
 		break;
 	case EStateTreePropertyCopyType::CopyFixedArray:
 	{
 		// Copy into fixed sized array (EditFixedSize). Resizable arrays are copied as Complex, and regular fixed sizes arrays via the regular copies (dim specifies array size).
-		const FArrayProperty* SourceArrayProperty = static_cast<const FArrayProperty*>(SourceProperty);
-		const FArrayProperty* TargetArrayProperty = static_cast<const FArrayProperty*>(TargetProperty);
+		const FArrayProperty* SourceArrayProperty = static_cast<const FArrayProperty*>(Copy.SourceLeafProperty);
+		const FArrayProperty* TargetArrayProperty = static_cast<const FArrayProperty*>(Copy.TargetLeafProperty);
 		FScriptArrayHelper SourceArrayHelper(SourceArrayProperty, SourceAddress);
 		FScriptArrayHelper TargetArrayHelper(TargetArrayProperty, TargetAddress);
 			
@@ -413,22 +421,22 @@ void FStateTreePropertyBindings::PerformCopy(const FStateTreePropCopy& Copy, con
 		break;
 	}
 	case EStateTreePropertyCopyType::PromoteBoolToByte:
-		*reinterpret_cast<uint8*>(TargetAddress) = (uint8)static_cast<const FBoolProperty*>(SourceProperty)->GetPropertyValue(SourceAddress);
+		*reinterpret_cast<uint8*>(TargetAddress) = (uint8)static_cast<const FBoolProperty*>(Copy.SourceLeafProperty)->GetPropertyValue(SourceAddress);
 		break;
 	case EStateTreePropertyCopyType::PromoteBoolToInt32:
-		*reinterpret_cast<int32*>(TargetAddress) = (int32)static_cast<const FBoolProperty*>(SourceProperty)->GetPropertyValue(SourceAddress);
+		*reinterpret_cast<int32*>(TargetAddress) = (int32)static_cast<const FBoolProperty*>(Copy.SourceLeafProperty)->GetPropertyValue(SourceAddress);
 		break;
 	case EStateTreePropertyCopyType::PromoteBoolToUInt32:
-		*reinterpret_cast<uint32*>(TargetAddress) = (uint32)static_cast<const FBoolProperty*>(SourceProperty)->GetPropertyValue(SourceAddress);
+		*reinterpret_cast<uint32*>(TargetAddress) = (uint32)static_cast<const FBoolProperty*>(Copy.SourceLeafProperty)->GetPropertyValue(SourceAddress);
 		break;
 	case EStateTreePropertyCopyType::PromoteBoolToInt64:
-		*reinterpret_cast<int64*>(TargetAddress) = (int64)static_cast<const FBoolProperty*>(SourceProperty)->GetPropertyValue(SourceAddress);
+		*reinterpret_cast<int64*>(TargetAddress) = (int64)static_cast<const FBoolProperty*>(Copy.SourceLeafProperty)->GetPropertyValue(SourceAddress);
 		break;
 	case EStateTreePropertyCopyType::PromoteBoolToFloat:
-		*reinterpret_cast<float*>(TargetAddress) = (float)static_cast<const FBoolProperty*>(SourceProperty)->GetPropertyValue(SourceAddress);
+		*reinterpret_cast<float*>(TargetAddress) = (float)static_cast<const FBoolProperty*>(Copy.SourceLeafProperty)->GetPropertyValue(SourceAddress);
 		break;
 	case EStateTreePropertyCopyType::PromoteBoolToDouble:
-		*reinterpret_cast<double*>(TargetAddress) = (double)static_cast<const FBoolProperty*>(SourceProperty)->GetPropertyValue(SourceAddress);
+		*reinterpret_cast<double*>(TargetAddress) = (double)static_cast<const FBoolProperty*>(Copy.SourceLeafProperty)->GetPropertyValue(SourceAddress);
 		break;
 	case EStateTreePropertyCopyType::PromoteByteToInt32:
 		*reinterpret_cast<int32*>(TargetAddress) = (int32)*reinterpret_cast<const uint8*>(SourceAddress);
@@ -468,12 +476,12 @@ void FStateTreePropertyBindings::PerformCopy(const FStateTreePropCopy& Copy, con
 		break;
 	default:
 		ensureMsgf(false, TEXT("FStateTreePropertyBindings::PerformCopy: Unhandled copy type %s between '%s' and '%s'"),
-			*StaticEnum<EStateTreePropertyCopyType>()->GetValueAsString(Copy.Type), *SourceProperty->GetNameCPP(), *TargetProperty->GetNameCPP());
+			*StaticEnum<EStateTreePropertyCopyType>()->GetValueAsString(Copy.Type), *Copy.SourceLeafProperty->GetNameCPP(), *Copy.TargetLeafProperty->GetNameCPP());
 		break;
 	}
 }
 
-bool FStateTreePropertyBindings::CopyTo(TConstArrayView<FStateTreeDataView> SourceStructViews, const int32 TargetBatchIndex, FStateTreeDataView TargetStructView) const
+bool FStateTreePropertyBindings::CopyTo(TConstArrayView<FStateTreeDataView> SourceStructViews, const FStateTreeIndex16 TargetBatchIndex, FStateTreeDataView TargetStructView) const
 {
 	// This is made ensure so that the programmers have the change to catch it (it's usually programming error not to call ResolvePaths(), and it wont spam log for others.
 	if (!ensureMsgf(bBindingsResolved, TEXT("Bindings must be resolved successfully before copying. See ResolvePaths()")))
@@ -481,13 +489,13 @@ bool FStateTreePropertyBindings::CopyTo(TConstArrayView<FStateTreeDataView> Sour
 		return false;
 	}
 
-	if (TargetBatchIndex == INDEX_NONE)
+	if (TargetBatchIndex.IsValid() == false)
 	{
 		return false;
 	}
 
-	check(CopyBatches.IsValidIndex(TargetBatchIndex));
-	const FStateTreePropCopyBatch& Batch = CopyBatches[TargetBatchIndex];
+	check(CopyBatches.IsValidIndex(TargetBatchIndex.Get()));
+	const FStateTreePropCopyBatch& Batch = CopyBatches[TargetBatchIndex.Get()];
 
 	check(TargetStructView.IsValid());
 	check(TargetStructView.GetStruct() == Batch.TargetStruct.Struct);
@@ -503,19 +511,17 @@ bool FStateTreePropertyBindings::CopyTo(TConstArrayView<FStateTreeDataView> Sour
 			continue;
 		}
 		
-		const FStateTreeDataView SourceStructView = SourceStructViews[Copy.SourceStructIndex];
+		const FStateTreeDataView SourceStructView = SourceStructViews[Copy.SourceStructIndex.Get()];
 		if (SourceStructView.IsValid())
 		{
-			check(SourceStructView.GetStruct() == SourceStructs[Copy.SourceStructIndex].Struct || (SourceStructView.GetStruct() && SourceStructView.GetStruct()->IsChildOf(SourceStructs[Copy.SourceStructIndex].Struct)));
+			check(SourceStructView.GetStruct() == SourceStructs[Copy.SourceStructIndex.Get()].Struct
+				|| (SourceStructView.GetStruct() && SourceStructView.GetStruct()->IsChildOf(SourceStructs[Copy.SourceStructIndex.Get()].Struct)));
 				
-			const FStateTreePropertyAccess& SourceAccess = PropertyAccesses[Copy.SourceAccessIndex];
-			const FStateTreePropertyAccess& TargetAccess = PropertyAccesses[Copy.TargetAccessIndex];
-			
-			const uint8* SourceAddress = GetAddress(SourceStructView, SourceAccess);
-			uint8* TargetAddress = GetAddress(TargetStructView, TargetAccess);
+			const uint8* SourceAddress = GetAddress(SourceStructView, Copy.SourceIndirection, Copy.SourceLeafProperty);
+			uint8* TargetAddress = GetAddress(TargetStructView, Copy.TargetIndirection, Copy.TargetLeafProperty);
 			
 			check (SourceAddress != nullptr && TargetAddress != nullptr);
-			PerformCopy(Copy, SourceAccess.LeafProperty, SourceAddress, TargetAccess.LeafProperty, TargetAddress);
+			PerformCopy(Copy, SourceAddress, TargetAddress);
 		}
 		else
 		{
@@ -549,84 +555,96 @@ void FStateTreePropertyBindings::DebugPrintInternalLayout(FString& OutString) co
 	}
 
 	/** Array of property bindings, resolved into arrays of copies before use. */
-	OutString += FString::Printf(TEXT("\nPropertyBindings (%d)\n  [ %-8s | %-8s | %-10s | %-50s ]\n"), PropertyBindings.Num(),
-	TEXT("Src Path"), TEXT("Tgt Path"), TEXT("Src Struct"), TEXT("Copy Type"));
+	OutString += FString::Printf(TEXT("\nPropertyBindings (%d)\n  [ %-20s | %-4s | %-4s | %-10s | %-20s | %-4s | %-7s | %-20s | %-7s | %-20s ]\n"),
+		PropertyBindings.Num(),
+		TEXT("Source Name"), TEXT("Arr#"), TEXT("Next"), TEXT("Access"),
+		TEXT("Target Name"), TEXT("Arr#"), TEXT("Next"), TEXT("Access"),
+		TEXT("Struct#"), TEXT("Copy Type"));
 	for (const FStateTreePropertyBinding& PropertyBinding : PropertyBindings)
 	{
-		OutString += FString::Printf(TEXT("  | %8d | %8d | %10d | %-50s |\n"),
-									 PropertyBinding.SourcePathIndex, PropertyBinding.TargetPathIndex, PropertyBinding.SourceStructIndex,
-									 *StaticEnum<EStateTreePropertyCopyType>()->GetValueAsString(PropertyBinding.CopyType));
-	}
-
-	/** Array of property bindings, indexed by property paths. */
-	OutString += FString::Printf(TEXT("\nPropertyPaths (%d)\n  [ %-4s [%3s:%3s[ |\n"), PropertyPaths.Num(),
-		TEXT("Seg."), TEXT("Beg"), TEXT("End"));
-	for (const FStateTreePropertyPath& PropertyPath : PropertyPaths)
-	{
-		OutString += FString::Printf(TEXT("  | %4s [%3d:%-3d[ |\n"), TEXT(""), PropertyPath.SegmentsBegin, PropertyPath.SegmentsEnd);
+		OutString += FString::Printf(TEXT("  | %-20s | %4d | %4d | %-10s | %-20s | %4d | %4d | %-10s | %7d | %-20s | \n"),
+									 *PropertyBinding.SourcePath.Name.ToString(),
+									 PropertyBinding.SourcePath.ArrayIndex.Get(),
+									 PropertyBinding.SourcePath.NextIndex.Get(),
+									 *UEnum::GetDisplayValueAsText(PropertyBinding.TargetPath.Type).ToString(),
+									 *PropertyBinding.TargetPath.Name.ToString(),
+									 PropertyBinding.TargetPath.ArrayIndex.Get(),
+									 PropertyBinding.TargetPath.NextIndex.Get(),
+									 *UEnum::GetDisplayValueAsText(PropertyBinding.TargetPath.Type).ToString(),
+									 PropertyBinding.SourceStructIndex.Get(),
+									 *UEnum::GetValueAsString(PropertyBinding.CopyType));
 	}
 
 	/** Array of property segments, indexed by property paths. */
-	OutString += FString::Printf(TEXT("\nPropertySegments (%d)\n  [ %-20s | %3s | %-40s ]\n"), PropertySegments.Num(),
-		TEXT("Name"), TEXT("Idx"), TEXT("Access Type"));
+	OutString += FString::Printf(TEXT("\nPropertySegments (%d)\n  [ %-20s | %4s | %4s | %-10s ]\n"), PropertySegments.Num(),
+		TEXT("Name"), TEXT("Arr#"), TEXT("Next"), TEXT("Access"));
 	for (const FStateTreePropertySegment& PropertySegment : PropertySegments)
 	{
-		OutString += FString::Printf(TEXT("  | %-20s | %3d | %-40s |\n"),
+		OutString += FString::Printf(TEXT("  | %-20s | %4d | %4d | %-10s |\n"),
 									 *PropertySegment.Name.ToString(),
-									 PropertySegment.ArrayIndex, *StaticEnum<EStateTreePropertyAccessType>()->GetValueAsString(PropertySegment.Type));
+									 PropertySegment.ArrayIndex.Get(),
+									 PropertySegment.NextIndex.Get(),
+									 *UEnum::GetDisplayValueAsText(PropertySegment.Type).ToString());
 	}
 
 	/** Array of property copies */
-	OutString += FString::Printf(TEXT("\nPropertyCopies (%d)\n  [ %-7s | %-7s | %-7s | %-50s | %-4s ]\n"), PropertyCopies.Num(),
-		TEXT("Src Idx"), TEXT("Tgt Idx"), TEXT("Struct"), TEXT("Copy Type"), TEXT("Size"));
+	OutString += FString::Printf(TEXT("\nPropertyCopies (%d)\n  [ %-7s | %-4s | %-4s | %-10s | %-7s | %-4s | %-4s | %-10s | %-7s | %-20s | %-4s ]\n"), PropertyCopies.Num(),
+		TEXT("Src Idx"), TEXT("Off."), TEXT("Next"), TEXT("Type"),
+		TEXT("Tgt Idx"), TEXT("Off."), TEXT("Next"), TEXT("Type"),
+		TEXT("Struct"), TEXT("Copy Type"), TEXT("Size"));
 	for (const FStateTreePropCopy& PropertyCopy : PropertyCopies)
 	{
-		OutString += FString::Printf(TEXT("  | %7d | %7d | %7d | %-50s | %4d |\n"),
-					PropertyCopy.SourceAccessIndex, PropertyCopy.TargetAccessIndex, PropertyCopy.SourceStructIndex,
-					*StaticEnum<EStateTreePropertyCopyType>()->GetValueAsString(PropertyCopy.Type),
+		OutString += FString::Printf(TEXT("  | %7d | %4d | %4d | %-10s | %7d | %4d | %4d | %-10s | %7d | %-20s | %4d |\n"),
+					PropertyCopy.SourceIndirection.ArrayIndex.Get(),
+					PropertyCopy.SourceIndirection.Offset,
+					PropertyCopy.SourceIndirection.NextIndex.Get(),
+					*UEnum::GetDisplayValueAsText(PropertyCopy.SourceIndirection.Type).ToString(),
+					PropertyCopy.TargetIndirection.ArrayIndex.Get(),
+					PropertyCopy.TargetIndirection.Offset,
+					PropertyCopy.TargetIndirection.NextIndex.Get(),
+					*UEnum::GetDisplayValueAsText(PropertyCopy.TargetIndirection.Type).ToString(),
+					PropertyCopy.SourceStructIndex.Get(),
+					*UEnum::GetDisplayValueAsText(PropertyCopy.Type).ToString(),
 					PropertyCopy.CopySize);
 	}
 
-	/** Array of property accesses, indexed by copies*/
-	OutString += FString::Printf(TEXT("\nPropertyAccesses (%d)\n  [ %-8s [%-4s:%-4s [ |\n"), PropertyAccesses.Num(),
-		TEXT("Indirect"), TEXT("Beg"), TEXT("End"));
-	for (const FStateTreePropertyAccess& PropertyAccess : PropertyAccesses)
-	{
-		OutString += FString::Printf(TEXT("  | %8s [%4d:%-4d [ |\n"), TEXT(""), PropertyAccess.IndirectionsBegin, PropertyAccess.IndirectionsEnd);
-	}
-
 	/** Array of property indirections, indexed by accesses*/
-	OutString += FString::Printf(TEXT("\nPropertyIndirections (%d)\n  [ %-4s | %-4s | %-40s ] \n"), PropertyIndirections.Num(),
-		TEXT("Idx"), TEXT("Off."), TEXT("Access Type"));
+	OutString += FString::Printf(TEXT("\nPropertyIndirections (%d)\n  [ %-4s | %-4s | %-4s | %-10s ] \n"), PropertyIndirections.Num(),
+		TEXT("Idx"), TEXT("Off."), TEXT("Next"), TEXT("Access Type"));
 	for (const FStateTreePropertyIndirection& PropertyIndirection : PropertyIndirections)
 	{
-		OutString += FString::Printf(TEXT("  | %4d | %4d | %-40s |\n"),
-					PropertyIndirection.ArrayIndex, PropertyIndirection.Offset,*StaticEnum<EStateTreePropertyAccessType>()->GetValueAsString(PropertyIndirection.Type));
+		OutString += FString::Printf(TEXT("  | %4d | %4d | %4d | %-10s |\n"),
+					PropertyIndirection.ArrayIndex.Get(),
+					PropertyIndirection.Offset,
+					PropertyIndirection.NextIndex.Get(),
+					*UEnum::GetDisplayValueAsText(PropertyIndirection.Type).ToString());
 	}
 }
 
-FString FStateTreePropertyBindings::GetPathAsString(const FStateTreePropertyPath& Path, const int32 HighlightedSegment, const TCHAR* HighlightPrefix, const TCHAR* HighlightPostfix)
+FString FStateTreePropertyBindings::GetPathAsString(const FStateTreePropertySegment& FirstPathSegment, const FStateTreePropertySegment* HighlightedSegment, const TCHAR* HighlightPrefix, const TCHAR* HighlightPostfix)
 {
 	FString Result;
-	for (int32 i = Path.SegmentsBegin; i < Path.SegmentsEnd; i++)
+	const FStateTreePropertySegment* Segment = &FirstPathSegment;
+	while (Segment != nullptr)
 	{
-		const FStateTreePropertySegment& Segment = PropertySegments[i];
-		if (i > 0)
+		if (!Result.IsEmpty())
 		{
 			Result += TEXT(".");
 		}
 		
-		if (i == HighlightedSegment && HighlightPrefix)
+		if (Segment == HighlightedSegment && HighlightPrefix)
 		{
 			Result += HighlightPrefix;
 		}
 
-		Result += Segment.Name.ToString();
+		Result += Segment->Name.ToString();
 
-		if (i == HighlightedSegment && HighlightPostfix)
+		if (Segment == HighlightedSegment && HighlightPostfix)
 		{
 			Result += HighlightPostfix;
 		}
+
+		Segment = Segment->NextIndex.IsValid() ? &PropertySegments[Segment->NextIndex.Get()] : nullptr;
 	}
 	return Result;
 }
diff --git a/Engine/Plugins/Runtime/StateTree/Source/StateTreeModule/Private/StateTreeTypes.cpp b/Engine/Plugins/Runtime/StateTree/Source/StateTreeModule/Private/StateTreeTypes.cpp
--- a/Engine/Plugins/Runtime/StateTree/Source/StateTreeModule/Private/StateTreeTypes.cpp
+++ b/Engine/Plugins/Runtime/StateTree/Source/StateTreeModule/Private/StateTreeTypes.cpp
@@ -4,8 +4,11 @@
 
 DEFINE_LOG_CATEGORY(LogStateTree);
 
-const FStateTreeHandle FStateTreeHandle::Invalid = FStateTreeHandle();
-const FStateTreeHandle FStateTreeHandle::Succeeded = FStateTreeHandle(FStateTreeHandle::SucceededIndex);
-const FStateTreeHandle FStateTreeHandle::Failed = FStateTreeHandle(FStateTreeHandle::FailedIndex);
+const FStateTreeStateHandle FStateTreeStateHandle::Invalid = FStateTreeStateHandle();
+const FStateTreeStateHandle FStateTreeStateHandle::Succeeded = FStateTreeStateHandle(FStateTreeStateHandle::SucceededIndex);
+const FStateTreeStateHandle FStateTreeStateHandle::Failed = FStateTreeStateHandle(FStateTreeStateHandle::FailedIndex);
 
-const FStateTreeExternalDataHandle FStateTreeExternalDataHandle::Invalid = FStateTreeExternalDataHandle();
+const FStateTreeIndex16 FStateTreeIndex16::Invalid = FStateTreeIndex16();
+const FStateTreeIndex8 FStateTreeIndex8::Invalid = FStateTreeIndex8();
+
+const FStateTreeExternalDataHandle FStateTreeExternalDataHandle::Invalid = FStateTreeExternalDataHandle();
\ No newline at end of file
diff --git a/Engine/Plugins/Runtime/StateTree/Source/StateTreeModule/Public/StateTree.h b/Engine/Plugins/Runtime/StateTree/Source/StateTreeModule/Public/StateTree.h
--- a/Engine/Plugins/Runtime/StateTree/Source/StateTreeModule/Public/StateTree.h
+++ b/Engine/Plugins/Runtime/StateTree/Source/StateTreeModule/Public/StateTree.h
@@ -25,6 +25,8 @@
 		GlobalEvaluators,
 		// Moved instance data to arrays.
 		InstanceDataArrays,
+		// Added index types.
+		IndexTypes,
 
 		// -----<new versions can be added above this line>-------------------------------------------------
 		VersionPlusOne,
@@ -44,7 +46,7 @@
 struct FStateTreeMemoryUsage
 {
 	FStateTreeMemoryUsage() = default;
-	FStateTreeMemoryUsage(const FString InName, const FStateTreeHandle InHandle = FStateTreeHandle::Invalid)
+	FStateTreeMemoryUsage(const FString InName, const FStateTreeStateHandle InHandle = FStateTreeStateHandle::Invalid)
 		: Name(InName)
 		, Handle(InHandle)
 	{
@@ -54,7 +56,7 @@
 	void AddUsage(const UObject* Object);
 
 	FString Name;
-	FStateTreeHandle Handle;
+	FStateTreeStateHandle Handle;
 	int32 NodeCount = 0;
 	int32 EstimatedMemoryUsage = 0;
 	int32 ChildNodeCount = 0;
@@ -178,7 +180,7 @@
 
 	/** Data view index of the tree Parameters */
 	UPROPERTY()
-	int32 ParametersDataViewIndex = INDEX_NONE;
+	FStateTreeIndex8 ParametersDataViewIndex = FStateTreeIndex8::Invalid;
 
 	/** Index of first evaluator in Nodes. */
 	UPROPERTY()
diff --git a/Engine/Plugins/Runtime/StateTree/Source/StateTreeModule/Public/StateTreeExecutionContext.h b/Engine/Plugins/Runtime/StateTree/Source/StateTreeModule/Public/StateTreeExecutionContext.h
--- a/Engine/Plugins/Runtime/StateTree/Source/StateTreeModule/Public/StateTreeExecutionContext.h
+++ b/Engine/Plugins/Runtime/StateTree/Source/StateTreeModule/Public/StateTreeExecutionContext.h
@@ -20,13 +20,13 @@
 	FStateTreeActiveStates ActiveStates;
 
 	/** Index of the first task struct in the currently initialized instance data. */
-	uint16 FirstTaskStructIndex = 0;
+	FStateTreeIndex16 FirstTaskStructIndex = FStateTreeIndex16::Invalid;
 	
 	/** Index of the first task object in the currently initialized instance data. */
-	uint16 FirstTaskObjectIndex = 0;
+	FStateTreeIndex16 FirstTaskObjectIndex = FStateTreeIndex16::Invalid;
 
 	/** The index of the task that failed during enter state. Exit state uses it to call ExitState() symmetrically. */
-	uint16 EnterStateFailedTaskIndex = MAX_uint16;
+	FStateTreeIndex16 EnterStateFailedTaskIndex = FStateTreeIndex16::Invalid;
 
 	/** Result of last tick */
 	EStateTreeRunStatus LastTickStatus = EStateTreeRunStatus::Failed;
@@ -35,7 +35,7 @@
 	EStateTreeRunStatus TreeRunStatus = EStateTreeRunStatus::Unset;
 
 	/** Delayed transition handle, if exists */
-	int16 GatedTransitionIndex = INDEX_NONE;
+	FStateTreeIndex16 GatedTransitionIndex = FStateTreeIndex16::Invalid;
 
 	/** Number of times a new state has been changed. */
 	uint16 StateChangeCount = 0;
@@ -101,7 +101,7 @@
 	EStateTreeRunStatus Tick(const float DeltaTime, FStateTreeInstanceData* ExternalInstanceData = nullptr);
 
 	/** @return Pointer to a State or null if state not found */ 
-	const FCompactStateTreeState* GetStateFromHandle(const FStateTreeHandle StateHandle) const
+	const FCompactStateTreeState* GetStateFromHandle(const FStateTreeStateHandle StateHandle) const
 	{
 		return (StateTree && StateTree->States.IsValidIndex(StateHandle.Index)) ? &StateTree->States[StateHandle.Index] : nullptr;
 	}
@@ -127,7 +127,7 @@
 		bool bResult = true;
 		for (const FStateTreeExternalDataDesc& DataDesc : StateTree->ExternalDataDescs)
 		{
-			const FStateTreeDataView& DataView = DataViews[DataDesc.Handle.DataViewIndex];
+			const FStateTreeDataView& DataView = DataViews[DataDesc.Handle.DataViewIndex.Get()];
 			
 			if (DataDesc.Requirement == EStateTreeExternalDataRequirement::Required)
 			{
@@ -151,7 +151,7 @@
 
 		for (const FStateTreeExternalDataDesc& DataDesc : StateTree->GetNamedExternalDataDescs())
 		{
-			const FStateTreeDataView& DataView = DataViews[DataDesc.Handle.DataViewIndex];
+			const FStateTreeDataView& DataView = DataViews[DataDesc.Handle.DataViewIndex.Get()];
 
 			// Items must have valid pointer of the expected type.  
 			if (!DataView.IsValid() || !DataView.GetStruct()->IsChildOf(DataDesc.Struct))
@@ -176,7 +176,7 @@
 	{
 		check(StateTree);
 		check(Handle.IsValid());
-		DataViews[Handle.DataViewIndex] = DataView;
+		DataViews[Handle.DataViewIndex.Get()] = DataView;
 	}
 
 	/**
@@ -189,8 +189,8 @@
 	{
 		check(StateTree);
 		check(Handle.IsValid());
-		checkSlow(StateTree->ExternalDataDescs[Handle.DataViewIndex - StateTree->ExternalDataBaseIndex].Requirement != EStateTreeExternalDataRequirement::Optional); // Optionals should query pointer instead.
-		return DataViews[Handle.DataViewIndex].template GetMutable<typename T::DataType>();
+		checkSlow(StateTree->ExternalDataDescs[Handle.DataViewIndex.Get() - StateTree->ExternalDataBaseIndex].Requirement != EStateTreeExternalDataRequirement::Optional); // Optionals should query pointer instead.
+		return DataViews[Handle.DataViewIndex.Get()].template GetMutable<typename T::DataType>();
 	}
 
 	/**
@@ -202,7 +202,7 @@
 	typename T::DataType* GetExternalDataPtr(const T Handle) const
 	{
 		check(StateTree);
-		return Handle.IsValid() ? DataViews[Handle.DataViewIndex].template GetMutablePtr<typename T::DataType>() : nullptr;
+		return Handle.IsValid() ? DataViews[Handle.DataViewIndex.Get()].template GetMutablePtr<typename T::DataType>() : nullptr;
 	}
 
 	FStateTreeDataView GetExternalDataView(const FStateTreeExternalDataHandle Handle)
@@ -210,7 +210,7 @@
 		check(StateTree);
 		if (Handle.IsValid())
 		{
-			return DataViews[Handle.DataViewIndex];
+			return DataViews[Handle.DataViewIndex.Get()];
 		}
 		return FStateTreeDataView();
 	}
@@ -225,7 +225,7 @@
 	{
 		check(StateTree);
 		check(Handle.IsValid());
-		return *(typename T::DataType*)(DataViews[Handle.DataViewIndex].GetMemory() + Handle.PropertyOffset);
+		return *(typename T::DataType*)(DataViews[Handle.DataViewIndex.Get()].GetMemory() + Handle.PropertyOffset);
 	}
 
 	/**
@@ -237,7 +237,7 @@
 	typename T::DataType* GetInstanceDataPtr(const T Handle) const
 	{
 		check(StateTree);
-		return Handle.IsValid() ? (typename T::DataType*)(DataViews[Handle.DataViewIndex].GetMemory() + Handle.PropertyOffset) : nullptr;
+		return Handle.IsValid() ? (typename T::DataType*)(DataViews[Handle.DataViewIndex.Get()].GetMemory() + Handle.PropertyOffset) : nullptr;
 	}
 
 	/**
@@ -246,12 +246,12 @@
 	 * @return Pointer to an instance object based.
 	 */
 	template <typename T>
-	T* GetInstanceObjectInternal(const int32 DataViewIndex) const
+	T* GetInstanceObjectInternal(const FStateTreeIndex16 DataViewIndex) const
 	{
-		const UStruct* Struct = DataViews[DataViewIndex].GetStruct();
+		const UStruct* Struct = DataViews[DataViewIndex.Get()].GetStruct();
 		if (Struct != nullptr && Struct->IsChildOf<T>())
 		{
-			return DataViews[DataViewIndex].template GetMutablePtr<T>();
+			return DataViews[DataViewIndex.Get()].template GetMutablePtr<T>();
 		}
 		return nullptr;
 	}
@@ -344,12 +344,12 @@
 	 * @param OutNewActiveStates Active states that got selected.
 	 * @return True if succeeded to select new active states.
 	 */
-	bool SelectState(FStateTreeInstanceData& InstanceData, const FStateTreeHandle NextState, FStateTreeActiveStates& OutNewActiveStates);
+	bool SelectState(FStateTreeInstanceData& InstanceData, const FStateTreeStateHandle NextState, FStateTreeActiveStates& OutNewActiveStates);
 
 	/**
 	 * Used internally to do the recursive part of the SelectState().
 	 */
-	bool SelectStateInternal(FStateTreeInstanceData& InstanceData, const FStateTreeHandle NextState, FStateTreeActiveStates& OutNewActiveStates);
+	bool SelectStateInternal(FStateTreeInstanceData& InstanceData, const FStateTreeStateHandle NextState, FStateTreeActiveStates& OutNewActiveStates);
 
 	/** @return Mutable storage based on storage settings. */
 	FStateTreeInstanceData& SelectMutableInstanceData(FStateTreeInstanceData* ExternalInstanceData)
@@ -387,7 +387,7 @@
 	FString GetStateStatusString(const FStateTreeExecutionState& ExecState) const;
 
 	/** @return String describing state name for logging and debug. */
-	FString GetSafeStateName(const FStateTreeHandle State) const;
+	FString GetSafeStateName(const FStateTreeStateHandle State) const;
 
 	/** @return String describing full path of an activate state for logging and debug. */
 	FString DebugGetStatePath(const FStateTreeActiveStates& ActiveStates, int32 ActiveStateIndex) const;
diff --git a/Engine/Plugins/Runtime/StateTree/Source/StateTreeModule/Public/StateTreeLinker.h b/Engine/Plugins/Runtime/StateTree/Source/StateTreeModule/Public/StateTreeLinker.h
--- a/Engine/Plugins/Runtime/StateTree/Source/StateTreeModule/Public/StateTreeLinker.h
+++ b/Engine/Plugins/Runtime/StateTree/Source/StateTreeModule/Public/StateTreeLinker.h
@@ -78,9 +78,9 @@ struct FStateTreeLinker
 		{
 			Index = ExternalDataDescs.Add(Desc);
 			check(FStateTreeExternalDataHandle::IsValidIndex(Index + ExternalDataBaseIndex));
-			ExternalDataDescs[Index].Handle.DataViewIndex = (uint8)(Index + ExternalDataBaseIndex);
+			ExternalDataDescs[Index].Handle.DataViewIndex = FStateTreeIndex8(Index + ExternalDataBaseIndex);
 		}
-		Handle.DataViewIndex = (uint8)(Index + ExternalDataBaseIndex);
+		Handle.DataViewIndex = FStateTreeIndex8(Index + ExternalDataBaseIndex);
 	}
 	/**
 	 * Links reference to a property in instance data.
@@ -117,10 +117,10 @@ protected:
 			return;
 		}
 
-		check(CurrentInstanceIndex < MAX_uint8);
+		check(FStateTreeIndex8::IsValidIndex(CurrentInstanceIndex));
 		check(Property->GetOffset_ForInternal() < MAX_uint16);
 		
-		Handle.DataViewIndex = (uint8)CurrentInstanceIndex;
+		Handle.DataViewIndex = FStateTreeIndex8(CurrentInstanceIndex);
 		Handle.Type = EStateTreePropertyIndirection::Offset;
 		Handle.PropertyOffset = (uint16)Property->GetOffset_ForInternal();
 	}
diff --git a/Engine/Plugins/Runtime/StateTree/Source/StateTreeModule/Public/StateTreeNodeBase.h b/Engine/Plugins/Runtime/StateTree/Source/StateTreeModule/Public/StateTreeNodeBase.h
--- a/Engine/Plugins/Runtime/StateTree/Source/StateTreeModule/Public/StateTreeNodeBase.h
+++ b/Engine/Plugins/Runtime/StateTree/Source/StateTreeModule/Public/StateTreeNodeBase.h
@@ -37,15 +37,15 @@ struct STATETREEMODULE_API FStateTreeNodeBase
 
 	/** Property binding copy batch handle. */
 	UPROPERTY()
-	FStateTreeHandle BindingsBatch = FStateTreeHandle::Invalid;
+	FStateTreeIndex16 BindingsBatch = FStateTreeIndex16::Invalid;
 
 	/** The runtime data's data view index in the StateTreeExecutionContext, and source struct index in property binding. */
 	UPROPERTY()
-	uint16 DataViewIndex = 0;
+	FStateTreeIndex16 DataViewIndex = FStateTreeIndex16::Invalid;
 
 	/** Index in runtime instance storage. */
 	UPROPERTY()
-	uint16 InstanceIndex = 0;
+	FStateTreeIndex16 InstanceIndex = FStateTreeIndex16::Invalid;
 
 	/** True if the instance is an UObject. */
 	UPROPERTY()
diff --git a/Engine/Plugins/Runtime/StateTree/Source/StateTreeModule/Public/StateTreePropertyBindings.h b/Engine/Plugins/Runtime/StateTree/Source/StateTreeModule/Public/StateTreePropertyBindings.h
--- a/Engine/Plugins/Runtime/StateTree/Source/StateTreeModule/Public/StateTreePropertyBindings.h
+++ b/Engine/Plugins/Runtime/StateTree/Source/StateTreeModule/Public/StateTreePropertyBindings.h
@@ -3,6 +3,7 @@
 #pragma once
 
 #include "CoreMinimal.h"
+#include "StateTreeTypes.h"
 #include "Misc/Guid.h"
 #include "StructView.h"
 #include "StateTreePropertyBindings.generated.h"
@@ -273,29 +274,17 @@ struct STATETREEMODULE_API FStateTreePropertySegment
 
 	/** Index in the array the property points at. */
 	UPROPERTY()
-	int32 ArrayIndex = 0;
+	FStateTreeIndex16 ArrayIndex = FStateTreeIndex16::Invalid;
+
+	/** Index to next segment. */
+	UPROPERTY()
+	FStateTreeIndex16 NextIndex = FStateTreeIndex16::Invalid;
 
 	/** Type of access/indirection. */
 	UPROPERTY()
 	EStateTreePropertyAccessType Type = EStateTreePropertyAccessType::Offset;
 };
 
-/**
- * Describes a property path made of segments (i.e. Foo.Bar consists of two segments). Used for storage only.
- */
-USTRUCT()
-struct STATETREEMODULE_API FStateTreePropertyPath
-{
-	GENERATED_BODY()
-
-	/** Index to first segment of the property path. */
-	UPROPERTY()
-	int32 SegmentsBegin = 0;
-
-	/** Index to one past the last segment in the property path. */
-	UPROPERTY()
-	int32 SegmentsEnd = 0;
-};
 
 /**
  * Describes property binding, the property from source path is copied into the property at the target path.
@@ -305,18 +294,18 @@ struct STATETREEMODULE_API FStateTreePropertyBinding
 {
 	GENERATED_BODY()
 
-	/** Index to source property path. */
+	/** Source property path. */
 	UPROPERTY()
-	int32 SourcePathIndex = 0;
+	FStateTreePropertySegment SourcePath;
 
-	/** Index to target property path. */
+	/** Target property path. */
 	UPROPERTY()
-	int32 TargetPathIndex = 0;
-
+	FStateTreePropertySegment TargetPath;
+	
 	/** Index to the source struct the source path refers to, sources are stored in FStateTreePropertyBindings. */
 	UPROPERTY()
-	int32 SourceStructIndex = 0;
-
+	FStateTreeIndex16 SourceStructIndex = FStateTreeIndex16::Invalid;
+	
 	/** The type of copy to use between the properties. */
 	UPROPERTY()
 	EStateTreePropertyCopyType CopyType = EStateTreePropertyCopyType::None;
@@ -332,11 +321,15 @@ struct STATETREEMODULE_API FStateTreePropertyIndirection
 
 	/** Index in the array the property points at. */
 	UPROPERTY()
-	int32 ArrayIndex = 0;
+	FStateTreeIndex16 ArrayIndex = FStateTreeIndex16::Invalid;
 
 	/** Cached offset of the property */
 	UPROPERTY()
-	int32 Offset = 0;
+	uint16 Offset = 0;
+
+	/** Cached offset of the property */
+	UPROPERTY()
+	FStateTreeIndex16 NextIndex = FStateTreeIndex16::Invalid;
 
 	/** Type of access/indirection. */
 	UPROPERTY()
@@ -347,57 +340,41 @@ struct STATETREEMODULE_API FStateTreePropertyIndirection
 };
 
 /**
- * Property access is a resolved property path, used for accessing properties in structs via one or more indirections.
+ * Describes property copy, the property from source is copied into the property at the target.
+ * Copy target struct is described in the property copy batch.
  */
 USTRUCT()
-struct STATETREEMODULE_API FStateTreePropertyAccess
+struct STATETREEMODULE_API FStateTreePropCopy
 {
 	GENERATED_BODY()
 
-	/** Index to first indirection of the property access. */
+	/** Source property access. */
 	UPROPERTY()
-	int32 IndirectionsBegin = 0;
+	FStateTreePropertyIndirection SourceIndirection;
 
-	/** Index to one past the last indirection in the property access. */
+	/** Target property access. */
 	UPROPERTY()
-	int32 IndirectionsEnd = 0;
+	FStateTreePropertyIndirection TargetIndirection;
 
 	/** Cached pointer to the leaf property of the access. */
-	const FProperty* LeafProperty = nullptr;
-};
+	const FProperty* SourceLeafProperty = nullptr;
 
+	/** Cached pointer to the leaf property of the access. */
+	const FProperty* TargetLeafProperty = nullptr;
 
-/**
- * Describes property copy, the property from source is copied into the property at the target.
- * Copy target struct is described in the property copy batch.
- */
-USTRUCT()
-struct STATETREEMODULE_API FStateTreePropCopy
-{
-	GENERATED_BODY()
-
-	/** Index to source property access. */
-	UPROPERTY()
-	int32 SourceAccessIndex = 0;
-
-	/** Index to target property access. */
+	/** Cached property element size * dim. */
 	UPROPERTY()
-	int32 TargetAccessIndex = 0;
+	int32 CopySize = 0;
 
 	/** Index to the struct the source path refers to, sources are stored in FStateTreePropertyBindings. */
 	UPROPERTY()
-	int32 SourceStructIndex = 0;
+	FStateTreeIndex16 SourceStructIndex = FStateTreeIndex16::Invalid;
 
 	/** Type of the copy */
 	UPROPERTY()
 	EStateTreePropertyCopyType Type = EStateTreePropertyCopyType::None;
-
-	/** Cached property element size * dim. */
-	UPROPERTY()
-	int32 CopySize = 0;
 };
 
-
 /**
  * Describes a batch of property copies from many sources to one target struct.
  * Note: The batch is used to reference both bindings and copies (a binding turns into copy when resolved).
@@ -413,11 +390,11 @@ struct STATETREEMODULE_API FStateTreePropCopyBatch
 
 	/** Index to first binding/copy. */
 	UPROPERTY()
-	int32 BindingsBegin = 0;
+	uint16 BindingsBegin = 0;
 
 	/** Index to one past the last binding/copy. */
 	UPROPERTY()
-	int32 BindingsEnd = 0;
+	uint16 BindingsEnd = 0;
 };
 
 /**
@@ -461,16 +438,16 @@ struct STATETREEMODULE_API FStateTreePropertyBindings
 	 * @param TargetStructView View to struct where properties are copied to.
 	 * @return true if all copies succeeded (a copy can fail e.g. if source or destination struct view is invalid).
 	 */
-	bool CopyTo(TConstArrayView<FStateTreeDataView> SourceStructViews, const int32 TargetBatchIndex, FStateTreeDataView TargetStructView) const;
+	bool CopyTo(TConstArrayView<FStateTreeDataView> SourceStructViews, const FStateTreeIndex16 TargetBatchIndex, FStateTreeDataView TargetStructView) const;
 
 	void DebugPrintInternalLayout(FString& OutString) const;
 
 protected:
-	[[nodiscard]] bool ResolvePath(const UStruct* Struct, const FStateTreePropertyPath& Path, FStateTreePropertyAccess& OutAccess);
+	[[nodiscard]] bool ResolvePath(const UStruct* Struct, const FStateTreePropertySegment& FirstPathSegment, FStateTreePropertyIndirection& OutFirstIndirection, const FProperty*& OutLeafProperty);
 	[[nodiscard]] bool ValidateCopy(FStateTreePropCopy& Copy) const;
-	void PerformCopy(const FStateTreePropCopy& Copy, const FProperty* SourceProperty, const uint8* SourceAddress, const FProperty* TargetProperty, uint8* TargetAddress) const;
-	uint8* GetAddress(FStateTreeDataView InStructView, const FStateTreePropertyAccess& InAccess) const;
-	FString GetPathAsString(const FStateTreePropertyPath& Path, const int32 HighlightedSegment = INDEX_NONE, const TCHAR* HighlightPrefix = nullptr, const TCHAR* HighlightPostfix = nullptr);
+	void PerformCopy(const FStateTreePropCopy& Copy, const uint8* SourceAddress, uint8* TargetAddress) const;
+	uint8* GetAddress(FStateTreeDataView InStructView, const FStateTreePropertyIndirection& FirstIndirection, const FProperty* LeafProperty) const;
+	FString GetPathAsString(const FStateTreePropertySegment& FirstPathSegment, const FStateTreePropertySegment* HighlightedSegment = nullptr, const TCHAR* HighlightPrefix = nullptr, const TCHAR* HighlightPostfix = nullptr);
 
 	/** Array of expected source structs. */
 	UPROPERTY()
@@ -484,10 +461,6 @@ protected:
 	UPROPERTY()
 	TArray<FStateTreePropertyBinding> PropertyBindings;
 
-	/** Array of property bindings, indexed by property paths. */
-	UPROPERTY()
-	TArray<FStateTreePropertyPath> PropertyPaths;
-
 	/** Array of property segments, indexed by property paths. */
 	UPROPERTY()
 	TArray<FStateTreePropertySegment> PropertySegments;
@@ -496,10 +469,6 @@ protected:
 	UPROPERTY(Transient)
 	TArray<FStateTreePropCopy> PropertyCopies;
 
-	/** Array of property accesses, indexed by copies*/
-	UPROPERTY(Transient)
-	TArray<FStateTreePropertyAccess> PropertyAccesses;
-
 	/** Array of property indirections, indexed by accesses*/
 	UPROPERTY(Transient)
 	TArray<FStateTreePropertyIndirection> PropertyIndirections;
diff --git a/Engine/Plugins/Runtime/StateTree/Source/StateTreeModule/Public/StateTreeTypes.h b/Engine/Plugins/Runtime/StateTree/Source/StateTreeModule/Public/StateTreeTypes.h
--- a/Engine/Plugins/Runtime/StateTree/Source/StateTreeModule/Public/StateTreeTypes.h
+++ b/Engine/Plugins/Runtime/StateTree/Source/StateTreeModule/Public/StateTreeTypes.h
@@ -89,29 +89,27 @@ enum class EStateTreeTransitionEvent : uint8
 };
 ENUM_CLASS_FLAGS(EStateTreeTransitionEvent)
 
-/**
- * Handle that is used to refer compact state tree data.
- */
+/** Handle to a StateTree state */
 USTRUCT(BlueprintType)
-struct STATETREEMODULE_API FStateTreeHandle
+struct STATETREEMODULE_API FStateTreeStateHandle
 {
 	GENERATED_BODY()
 
-	static const uint16 InvalidIndex = uint16(-1);		// Index value indicating invalid item.
-	static const uint16 SucceededIndex = uint16(-2);	// Index value indicating a Succeeded item.
-	static const uint16 FailedIndex = uint16(-3);		// Index value indicating a Failed item.
+	static constexpr uint16 InvalidIndex = uint16(-1);		// Index value indicating invalid item.
+	static constexpr uint16 SucceededIndex = uint16(-2);	// Index value indicating a Succeeded item.
+	static constexpr uint16 FailedIndex = uint16(-3);		// Index value indicating a Failed item.
 	
-	static const FStateTreeHandle Invalid;
-	static const FStateTreeHandle Succeeded;
-	static const FStateTreeHandle Failed;
+	static const FStateTreeStateHandle Invalid;
+	static const FStateTreeStateHandle Succeeded;
+	static const FStateTreeStateHandle Failed;
 
-	FStateTreeHandle() = default;
-	explicit FStateTreeHandle(uint16 InIndex) : Index(InIndex) {}
+	FStateTreeStateHandle() = default;
+	explicit FStateTreeStateHandle(uint16 InIndex) : Index(InIndex) {}
 
 	bool IsValid() const { return Index != InvalidIndex; }
 
-	bool operator==(const FStateTreeHandle& RHS) const { return Index == RHS.Index; }
-	bool operator!=(const FStateTreeHandle& RHS) const { return Index != RHS.Index; }
+	bool operator==(const FStateTreeStateHandle& RHS) const { return Index == RHS.Index; }
+	bool operator!=(const FStateTreeStateHandle& RHS) const { return Index != RHS.Index; }
 
 	FString Describe() const
 	{
@@ -128,6 +126,86 @@ struct STATETREEMODULE_API FStateTreeHandle
 	uint16 Index = InvalidIndex;
 };
 
+/** uint16 index that can be invalid. */
+USTRUCT(BlueprintType)
+struct STATETREEMODULE_API FStateTreeIndex16
+{
+	GENERATED_BODY()
+
+	static constexpr uint16 InvalidValue = MAX_uint16;
+	static const FStateTreeIndex16 Invalid;
+
+	/** @return true if the given index can be represented by the type. */
+	static bool IsValidIndex(const int32 Index)
+	{
+		return Index >= 0 && Index < (int32)MAX_uint16;
+	}
+
+	FStateTreeIndex16() = default;
+	
+	explicit FStateTreeIndex16(const int32 InIndex)
+	{
+		check(InIndex == INDEX_NONE || IsValidIndex(InIndex));
+		Value = InIndex == INDEX_NONE ? InvalidValue : (uint16)InIndex;
+	}
+
+	/** @retrun value of the index. */
+	uint16 Get() const { return Value; }
+	
+	/** @return the index value as int32, mapping invalid value to INDEX_NONE. */
+	int32 AsInt32() const { return Value == InvalidValue ? INDEX_NONE : Value; }
+
+	/** @return true if the index is valid. */
+	bool IsValid() const { return Value != InvalidValue; }
+
+	bool operator==(const FStateTreeIndex16& RHS) const { return Value == RHS.Value; }
+	bool operator!=(const FStateTreeIndex16& RHS) const { return Value != RHS.Value; }
+
+protected:
+	UPROPERTY()
+	uint16 Value = InvalidValue;
+};
+	
+
+USTRUCT(BlueprintType)
+struct STATETREEMODULE_API FStateTreeIndex8
+{
+	GENERATED_BODY()
+
+	static constexpr uint8 InvalidValue = MAX_uint8;
+	static const FStateTreeIndex8 Invalid;
+	
+	/** @return true if the given index can be represented by the type. */
+	static bool IsValidIndex(const int32 Index)
+	{
+		return Index >= 0 && Index < (int32)MAX_uint8;
+	}
+	
+	FStateTreeIndex8() = default;
+	
+	explicit FStateTreeIndex8(const int32 InIndex)
+	{
+		check(InIndex == INDEX_NONE || IsValidIndex(InIndex));
+		Value = InIndex == INDEX_NONE ? InvalidValue : (uint8)InIndex;
+	}
+
+	/** @retrun value of the index. */
+	uint8 Get() const { return Value; }
+
+	/** @return the index value as int32, mapping invalid value to INDEX_NONE. */
+	int32 AsInt32() const { return Value == InvalidValue ? INDEX_NONE : Value; }
+	
+	/** @return true if the index is valid. */
+	bool IsValid() const { return Value != InvalidValue; }
+
+	bool operator==(const FStateTreeIndex8& RHS) const { return Value == RHS.Value; }
+	bool operator!=(const FStateTreeIndex8& RHS) const { return Value != RHS.Value; }
+
+protected:
+	UPROPERTY()
+	uint8 Value = InvalidValue;
+};
+
 
 /**
  * Describes an array of active states in a State Tree.
@@ -141,7 +219,7 @@ struct STATETREEMODULE_API FStateTreeActiveStates
 
 	FStateTreeActiveStates() = default;
 	
-	explicit FStateTreeActiveStates(const FStateTreeHandle StateHandle)
+	explicit FStateTreeActiveStates(const FStateTreeStateHandle StateHandle)
 	{
 		Push(StateHandle);
 	}
@@ -153,7 +231,7 @@ struct STATETREEMODULE_API FStateTreeActiveStates
 	}
 
 	/** Pushes new state at the back of the array and returns true if there was enough space. */
-	bool Push(const FStateTreeHandle StateHandle)
+	bool Push(const FStateTreeStateHandle StateHandle)
 	{
 		if ((NumStates + 1) > MaxStates)
 		{
@@ -166,7 +244,7 @@ struct STATETREEMODULE_API FStateTreeActiveStates
 	}
 
 	/** Pushes new state at the front of the array and returns true if there was enough space. */
-	bool PushFront(const FStateTreeHandle StateHandle)
+	bool PushFront(const FStateTreeStateHandle StateHandle)
 	{
 		if ((NumStates + 1) > MaxStates)
 		{
@@ -184,14 +262,14 @@ struct STATETREEMODULE_API FStateTreeActiveStates
 	}
 
 	/** Pops a state from the back of the array and returns the popped value, or invalid handle if the array was empty. */
-	FStateTreeHandle Pop()
+	FStateTreeStateHandle Pop()
 	{
 		if (NumStates == 0)
 		{
-			return FStateTreeHandle::Invalid;			
+			return FStateTreeStateHandle::Invalid;			
 		}
 
-		const FStateTreeHandle Ret = States[NumStates - 1];
+		const FStateTreeStateHandle Ret = States[NumStates - 1];
 		NumStates--;
 		return Ret;
 	}
@@ -204,16 +282,16 @@ struct STATETREEMODULE_API FStateTreeActiveStates
 		{
 			for (int32 Index = NumStates; Index < NewNum; Index++)
 			{
-				States[Index] = FStateTreeHandle::Invalid;
+				States[Index] = FStateTreeStateHandle::Invalid;
 			}
 		}
 		NumStates = NewNum;
 	}
 
 	/** Returns true of the array contains specified state. */
-	bool Contains(const FStateTreeHandle StateHandle) const
+	bool Contains(const FStateTreeStateHandle StateHandle) const
 	{
-		for (const FStateTreeHandle& Handle : *this)
+		for (const FStateTreeStateHandle& Handle : *this)
 		{
 			if (Handle == StateHandle)
 			{
@@ -224,7 +302,7 @@ struct STATETREEMODULE_API FStateTreeActiveStates
 	}
 
 	/** Returns index of a state, searching in reverse order. */
-	int32 IndexOfReverse(const FStateTreeHandle StateHandle) const
+	int32 IndexOfReverse(const FStateTreeStateHandle StateHandle) const
 	{
 		for (int32 Index = (int32)NumStates - 1; Index >= 0; Index--)
 		{
@@ -235,7 +313,7 @@ struct STATETREEMODULE_API FStateTreeActiveStates
 	}
 	
 	/** Returns last state in the array, or invalid state if the array is empty. */
-	FStateTreeHandle Last() const { return NumStates > 0 ? States[NumStates - 1] : FStateTreeHandle::Invalid; }
+	FStateTreeStateHandle Last() const { return NumStates > 0 ? States[NumStates - 1] : FStateTreeStateHandle::Invalid; }
 	
 	/** Returns number of states in the array. */
 	int32 Num() const { return NumStates; }
@@ -247,37 +325,37 @@ struct STATETREEMODULE_API FStateTreeActiveStates
 	bool IsEmpty() const { return NumStates == 0; } 
 
 	/** Returns a specified state in the array. */
-	FORCEINLINE FStateTreeHandle operator[](const int32 Index) const
+	FORCEINLINE FStateTreeStateHandle operator[](const int32 Index) const
 	{
 		check(Index >= 0 && Index < (int32)NumStates);
 		return States[Index];
 	}
 
 	/** Returns mutable reference to a specified state in the array. */
-	FORCEINLINE FStateTreeHandle& operator[](const int32 Index)
+	FORCEINLINE FStateTreeStateHandle& operator[](const int32 Index)
 	{
 		check(Index >= 0 && Index < (int32)NumStates);
 		return States[Index];
 	}
 
-	/** Returns a specified state in the array, or FStateTreeHandle::Invalid if Index is out of array bounds. */
-	FStateTreeHandle GetStateSafe(const int32 Index) const
+	/** Returns a specified state in the array, or FStateTreeStateHandle::Invalid if Index is out of array bounds. */
+	FStateTreeStateHandle GetStateSafe(const int32 Index) const
 	{
-		return (Index >= 0 && Index < (int32)NumStates) ? States[Index] : FStateTreeHandle::Invalid;
+		return (Index >= 0 && Index < (int32)NumStates) ? States[Index] : FStateTreeStateHandle::Invalid;
 	}
 
 	/**
 	 * DO NOT USE DIRECTLY
 	 * STL-like iterators to enable range-based for loop support.
 	 */
-	FORCEINLINE FStateTreeHandle* begin() { return &States[0]; }
-	FORCEINLINE FStateTreeHandle* end  () { return &States[0] + Num(); }
-	FORCEINLINE const FStateTreeHandle* begin() const { return &States[0]; }
-	FORCEINLINE const FStateTreeHandle* end  () const { return &States[0] + Num(); }
+	FORCEINLINE FStateTreeStateHandle* begin() { return &States[0]; }
+	FORCEINLINE FStateTreeStateHandle* end  () { return &States[0] + Num(); }
+	FORCEINLINE const FStateTreeStateHandle* begin() const { return &States[0]; }
+	FORCEINLINE const FStateTreeStateHandle* end  () const { return &States[0] + Num(); }
 
 
 	UPROPERTY(EditDefaultsOnly, Category = Default)
-	FStateTreeHandle States[MaxStates];
+	FStateTreeStateHandle States[MaxStates];
 
 	UPROPERTY(EditDefaultsOnly, Category = Default)
 	uint8 NumStates = 0;
@@ -305,7 +383,7 @@ struct STATETREEMODULE_API FStateTreeTransitionResult
 
 	/** Transition target state */
 	UPROPERTY(EditDefaultsOnly, Category = Default, BlueprintReadOnly)
-	FStateTreeHandle TargetState = FStateTreeHandle::Invalid;
+	FStateTreeStateHandle TargetState = FStateTreeStateHandle::Invalid;
 
 	/** States selected as result of the transition. */
 	UPROPERTY(EditDefaultsOnly, Category = Default, BlueprintReadOnly)
@@ -313,7 +391,7 @@ struct STATETREEMODULE_API FStateTreeTransitionResult
 
 	/** The current state being executed. On enter/exit callbacks this is the state of the task or evaluator. */
 	UPROPERTY(EditDefaultsOnly, Category = Default, BlueprintReadOnly)
-	FStateTreeHandle CurrentState = FStateTreeHandle::Invalid;
+	FStateTreeStateHandle CurrentState = FStateTreeStateHandle::Invalid;
 };
 
 
@@ -328,7 +406,7 @@ struct STATETREEMODULE_API FCompactStateTransition
 	UPROPERTY()
 	uint16 ConditionsBegin = 0;							// Index to first condition to test
 	UPROPERTY()
-	FStateTreeHandle State = FStateTreeHandle::Invalid;	// Target state of the transition
+	FStateTreeStateHandle State = FStateTreeStateHandle::Invalid;	// Target state of the transition
 	UPROPERTY()
 	EStateTreeTransitionType Type = EStateTreeTransitionType::NotSet;	// Type of the transition.
 	UPROPERTY()
@@ -357,10 +435,10 @@ struct STATETREEMODULE_API FCompactStateTreeState
 	FName Name;											// Name of the State
 
 	UPROPERTY()
-	FStateTreeHandle LinkedState = FStateTreeHandle::Invalid;	// Linked state 
+	FStateTreeStateHandle LinkedState = FStateTreeStateHandle::Invalid;	// Linked state 
 
 	UPROPERTY()
-	FStateTreeHandle Parent = FStateTreeHandle::Invalid;		// Parent state
+	FStateTreeStateHandle Parent = FStateTreeStateHandle::Invalid;		// Parent state
 	UPROPERTY()
 	uint16 ChildrenBegin = 0;							// Index to first child state
 	UPROPERTY()
@@ -374,9 +452,9 @@ struct STATETREEMODULE_API FCompactStateTreeState
 	uint16 TasksBegin = 0;								// Index to first task
 
 	UPROPERTY()
-	uint16 ParameterInstanceIndex = MAX_uint16;			// Index to state instance data
+	FStateTreeIndex16 ParameterInstanceIndex = FStateTreeIndex16::Invalid;	// Index to state instance data
 	UPROPERTY()
-	uint16 ParameterDataViewIndex = MAX_uint16;			// Data view index of the input parameters
+	FStateTreeIndex16 ParameterDataViewIndex = FStateTreeIndex16::Invalid;	// Data view index of the input parameters
 
 	UPROPERTY()
 	uint8 EnterConditionsNum = 0;						// Number of enter conditions
@@ -394,7 +472,7 @@ struct STATETREEMODULE_API FCompactStateTreeParameters
 	GENERATED_BODY()
 
 	UPROPERTY()
-	FStateTreeHandle BindingsBatch = FStateTreeHandle::Invalid;
+	FStateTreeIndex16 BindingsBatch = FStateTreeIndex16::Invalid;
 
 	UPROPERTY()
 	FInstancedPropertyBag Parameters;
@@ -429,14 +507,12 @@ struct STATETREEMODULE_API FStateTreeExternalDataHandle
 	GENERATED_BODY()
 
 	static const FStateTreeExternalDataHandle Invalid;
-	static constexpr uint8 IndexNone = MAX_uint8;
-
-	static bool IsValidIndex(const int32 Index) { return Index >= 0 && Index < (int32)IndexNone; }
-
-	bool IsValid() const { return DataViewIndex != IndexNone; }
+	
+	static bool IsValidIndex(const int32 Index) { return FStateTreeIndex8::IsValidIndex(Index); }
+	bool IsValid() const { return DataViewIndex.IsValid(); }
 
 	UPROPERTY()
-	uint8 DataViewIndex = IndexNone;
+	FStateTreeIndex8 DataViewIndex = FStateTreeIndex8::Invalid;
 };
 
 /**
@@ -494,14 +570,10 @@ struct STATETREEMODULE_API FStateTreeInstanceDataPropertyHandle
 {
 	GENERATED_BODY()
 
-	static constexpr uint8 IndexNone = MAX_uint8;
-
-	static bool IsValidIndex(const int32 Index) { return Index >= 0 && Index < (int32)IndexNone; }
-
-	bool IsValid() const { return DataViewIndex != IndexNone; }
+	bool IsValid() const { return DataViewIndex.IsValid(); }
 
 	uint16 PropertyOffset = 0;
-	uint8 DataViewIndex = IndexNone;
+	FStateTreeIndex8 DataViewIndex = FStateTreeIndex8::Invalid;
 	EStateTreePropertyIndirection Type = EStateTreePropertyIndirection::Offset;
 };
 
