Fix warning for incorrectly included header.

diff --git a/Engine/Plugins/Runtime/StateTree/Source/StateTreeModule/Private/StateTree.cpp b/Engine/Plugins/Runtime/StateTree/Source/StateTreeModule/Private/StateTree.cpp
--- a/Engine/Plugins/Runtime/StateTree/Source/StateTreeModule/Private/StateTree.cpp
+++ b/Engine/Plugins/Runtime/StateTree/Source/StateTreeModule/Private/StateTree.cpp
@@ -5,7 +5,7 @@
 #include "StateTreeNodeBase.h"
 #include "StateTreeTaskBase.h"
 #include "StateTreeEvaluatorBase.h"
-#include "AssetData.h"
+#include "AssetRegistry/AssetData.h"
 
 const FGuid FStateTreeCustomVersion::GUID(0x28E21331, 0x501F4723, 0x8110FA64, 0xEA10DA1E);
 FCustomVersionRegistration GRegisterStateTreeCustomVersion(FStateTreeCustomVersion::GUID, FStateTreeCustomVersion::LatestVersion, TEXT("StateTreeAsset"));
