StateTree: Fix property binding for direct struct access

diff --git a/Engine/Plugins/Runtime/StateTree/Source/StateTreeEditorModule/Private/StateTreePropertyBindingCompiler.cpp b/Engine/Plugins/Runtime/StateTree/Source/StateTreeEditorModule/Private/StateTreePropertyBindingCompiler.cpp
--- a/Engine/Plugins/Runtime/StateTree/Source/StateTreeEditorModule/Private/StateTreePropertyBindingCompiler.cpp
+++ b/Engine/Plugins/Runtime/StateTree/Source/StateTreeEditorModule/Private/StateTreePropertyBindingCompiler.cpp
@@ -77,7 +77,12 @@ bool FStateTreePropertyBindingCompiler::CompileBatch(const FStateTreeBindableStr
 
 		auto StorePropertyPath = [this, &TargetStruct](FStateTreePropertySegment& FirstPathSegment, TArray<FStateTreePropertySegment>& Segments) -> bool
 		{
-			check(Segments.Num() > 0);
+			// The path is empty when directly bound to the target struct.
+			if (Segments.IsEmpty())
+			{
+				return true;
+			}
+			
 			FirstPathSegment = Segments[0];
 			FStateTreePropertySegment* PrevSegment = &FirstPathSegment; 
 			for (int32 Index = 1; Index < Segments.Num(); Index++)
diff --git a/Engine/Plugins/Runtime/StateTree/Source/StateTreeModule/Private/StateTreePropertyBindings.cpp b/Engine/Plugins/Runtime/StateTree/Source/StateTreeModule/Private/StateTreePropertyBindings.cpp
--- a/Engine/Plugins/Runtime/StateTree/Source/StateTreeModule/Private/StateTreePropertyBindings.cpp
+++ b/Engine/Plugins/Runtime/StateTree/Source/StateTreeModule/Private/StateTreePropertyBindings.cpp
@@ -76,7 +76,7 @@ bool FStateTreePropertyBindings::ResolvePath(const UStruct* Struct, const FState
 	TArray<FStateTreePropertyIndirection, TInlineAllocator<16>> TempIndirections;
 	const FStateTreePropertySegment* Segment = &FirstPathSegment;
 
-	while (Segment != nullptr)
+	while (Segment != nullptr && !Segment->IsEmpty())
 	{
 		FStateTreePropertyIndirection& Indirection = TempIndirections.AddDefaulted_GetRef();
 
@@ -152,33 +152,42 @@ bool FStateTreePropertyBindings::ResolvePath(const UStruct* Struct, const FState
 		Segment = Segment->NextIndex.IsValid() ? &PropertySegments[Segment->NextIndex.Get()] : nullptr;
 	}
 
-	// Collapse adjacent offset indirections
-	for (int32 Index = 0; Index < TempIndirections.Num(); Index++)
+	if (TempIndirections.Num() > 0)
 	{
-		FStateTreePropertyIndirection& Indirection = TempIndirections[Index];
-		if (Indirection.Type == EStateTreePropertyAccessType::Offset && (Index + 1) < TempIndirections.Num())
+		// Collapse adjacent offset indirections
+		for (int32 Index = 0; Index < TempIndirections.Num(); Index++)
 		{
-			const FStateTreePropertyIndirection& NextIndirection = TempIndirections[Index + 1];
-			if (NextIndirection.Type == EStateTreePropertyAccessType::Offset)
+			FStateTreePropertyIndirection& Indirection = TempIndirections[Index];
+			if (Indirection.Type == EStateTreePropertyAccessType::Offset && (Index + 1) < TempIndirections.Num())
 			{
-				Indirection.Offset += NextIndirection.Offset;
-				TempIndirections.RemoveAt(Index + 1);
-				Index--;
+				const FStateTreePropertyIndirection& NextIndirection = TempIndirections[Index + 1];
+				if (NextIndirection.Type == EStateTreePropertyAccessType::Offset)
+				{
+					Indirection.Offset += NextIndirection.Offset;
+					TempIndirections.RemoveAt(Index + 1);
+					Index--;
+				}
 			}
 		}
-	}
 
-	check(TempIndirections.Num() > 0);
 
-	// Store indirections
-	OutFirstIndirection = TempIndirections[0];
-	FStateTreePropertyIndirection* PrevIndirection = &OutFirstIndirection;
-	for (int32 Index = 1; Index < TempIndirections.Num(); Index++)
+		// Store indirections
+		OutFirstIndirection = TempIndirections[0];
+		FStateTreePropertyIndirection* PrevIndirection = &OutFirstIndirection;
+		for (int32 Index = 1; Index < TempIndirections.Num(); Index++)
+		{
+			const int32 IndirectionIndex = PropertyIndirections.Num();
+			FStateTreePropertyIndirection& NewIndirection = PropertyIndirections.Add_GetRef(TempIndirections[Index]);
+			PrevIndirection->NextIndex = FStateTreeIndex16(IndirectionIndex);
+			PrevIndirection = &NewIndirection;
+		}
+	}
+	else
 	{
-		const int32 IndirectionIndex = PropertyIndirections.Num();
-		FStateTreePropertyIndirection& NewIndirection = PropertyIndirections.Add_GetRef(TempIndirections[Index]);
-		PrevIndirection->NextIndex = FStateTreeIndex16(IndirectionIndex);
-		PrevIndirection = &NewIndirection;
+		// Indirections can be empty in case we're directly binding to source structs.
+		// Zero offset will return the struct itself.
+		OutFirstIndirection.Offset = 0;
+		OutFirstIndirection.Type = EStateTreePropertyAccessType::Offset;
 	}
 	
 	OutLeafProperty = LeafProperty;
diff --git a/Engine/Plugins/Runtime/StateTree/Source/StateTreeModule/Public/StateTreePropertyBindings.h b/Engine/Plugins/Runtime/StateTree/Source/StateTreeModule/Public/StateTreePropertyBindings.h
--- a/Engine/Plugins/Runtime/StateTree/Source/StateTreeModule/Public/StateTreePropertyBindings.h
+++ b/Engine/Plugins/Runtime/StateTree/Source/StateTreeModule/Public/StateTreePropertyBindings.h
@@ -268,6 +268,9 @@ struct STATETREEMODULE_API FStateTreePropertySegment
 {
 	GENERATED_BODY()
 
+	/** @return true if the segment is empty. */
+	bool IsEmpty() const { return Name.IsNone(); }
+	
 	/** Property name. */
 	UPROPERTY()
 	FName Name;
