StateTree: Prevent linking empty StateTrees (will cause unrelated error)

diff --git a/Engine/Plugins/Runtime/StateTree/Source/StateTreeModule/Private/StateTree.cpp b/Engine/Plugins/Runtime/StateTree/Source/StateTreeModule/Private/StateTree.cpp
--- a/Engine/Plugins/Runtime/StateTree/Source/StateTreeModule/Private/StateTree.cpp
+++ b/Engine/Plugins/Runtime/StateTree/Source/StateTreeModule/Private/StateTree.cpp
@@ -115,6 +115,12 @@ bool UStateTree::Link()
 	// This data will be used to allocate runtime instance on all StateTree users.
 	ResetLinked();
 
+	// Do not try to link empty trees.
+	if (States.Num() == 0 || Nodes.Num() == 0)
+	{
+		return true;
+	}
+
 	if (!DefaultInstanceData.IsValid())
 	{
 		UE_LOG(LogStateTree, Error, TEXT("%s: StartTree does not have instance data. Please recompile the StateTree asset."), *GetName());
