StateTree: Fixed max instance data memory usage estimate calculation
Root states accumulate max child state memory usage (e.g. most expensive chain) not whole tree.

diff --git a/Engine/Plugins/Runtime/StateTree/Source/StateTreeModule/Private/StateTree.cpp b/Engine/Plugins/Runtime/StateTree/Source/StateTreeModule/Private/StateTree.cpp
--- a/Engine/Plugins/Runtime/StateTree/Source/StateTreeModule/Private/StateTree.cpp
+++ b/Engine/Plugins/Runtime/StateTree/Source/StateTreeModule/Private/StateTree.cpp
@@ -245,6 +245,7 @@ TArray<FStateTreeMemoryUsage> UStateTree::CalculateEstimatedMemoryUsage() const
 	}
 
 	const int32 TreeMemUsageIndex = MemoryUsages.Emplace(TEXT("State Tree Max"));
+	const int32 InstanceMemUsageIndex = MemoryUsages.Emplace(TEXT("Instance Overhead"));
 	const int32 EvalMemUsageIndex = MemoryUsages.Emplace(TEXT("Evaluators"));
 	const int32 SharedMemUsageIndex = MemoryUsages.Emplace(TEXT("Shared Data"));
 
@@ -273,14 +274,19 @@ TArray<FStateTreeMemoryUsage> UStateTree::CalculateEstimatedMemoryUsage() const
 		return MemoryUsages.Emplace(TEXT("State ") + CompactState.Name.ToString(), InStateHandle);
 	};
 
+	// Calculate memory usage per state.
+	TArray<FStateTreeMemoryUsage> TempStateMemoryUsages;
+	TempStateMemoryUsages.SetNum(States.Num());
+
 	for (int32 Index = 0; Index < States.Num(); Index++)
 	{
 		const FStateTreeStateHandle StateHandle((uint16)Index);
 		const FCompactStateTreeState& CompactState = States[Index];
 		const FStateTreeStateHandle ParentHandle = GetRootStateHandle(StateHandle);
 		const int32 ParentUsageIndex = GetUsageIndexForState(ParentHandle);
-		FStateTreeMemoryUsage& MemUsage = MemoryUsages[ParentUsageIndex];
-
+		
+		FStateTreeMemoryUsage& MemUsage = CompactState.Parent.IsValid() ? TempStateMemoryUsages[Index] : MemoryUsages[GetUsageIndexForState(StateHandle)];
+		
 		MemUsage.NodeCount += CompactState.TasksNum;
 
 		if (CompactState.Type == EStateTreeStateType::Linked)
@@ -308,6 +314,31 @@ TArray<FStateTreeMemoryUsage> UStateTree::CalculateEstimatedMemoryUsage() const
 		}
 	}
 
+	// Combine max child usage to parents. Iterate backwards to update children first.
+	for (int32 Index = States.Num() - 1; Index >= 0; Index--)
+	{
+		const FStateTreeStateHandle StateHandle((uint16)Index);
+		const FCompactStateTreeState& CompactState = States[Index];
+
+		FStateTreeMemoryUsage& MemUsage = CompactState.Parent.IsValid() ? TempStateMemoryUsages[Index] : MemoryUsages[GetUsageIndexForState(StateHandle)];
+
+		int32 MaxChildStateMem = 0;
+		int32 MaxChildStateNodes = 0;
+		
+		for (uint16 ChildState = CompactState.ChildrenBegin; ChildState < CompactState.ChildrenEnd; ChildState = States[ChildState].GetNextSibling())
+		{
+			const FStateTreeMemoryUsage& ChildMemUsage = TempStateMemoryUsages[ChildState];
+			if (ChildMemUsage.EstimatedMemoryUsage > MaxChildStateMem)
+			{
+				MaxChildStateMem = ChildMemUsage.EstimatedMemoryUsage;
+				MaxChildStateNodes = ChildMemUsage.NodeCount;
+			}
+		}
+
+		MemUsage.EstimatedMemoryUsage += MaxChildStateMem;
+		MemUsage.NodeCount += MaxChildStateNodes;
+	}
+
 	// Accumulate linked states.
 	for (int32 Index = StateLinks.Num() - 1; Index >= 0; Index--)
 	{
@@ -334,6 +365,7 @@ TArray<FStateTreeMemoryUsage> UStateTree::CalculateEstimatedMemoryUsage() const
 		{
 			EvalMemUsage.AddUsage(DefaultInstanceData.GetMutableObject(Eval.InstanceIndex.Get()));
 		}
+		EvalMemUsage.NodeCount++;
 	}
 
 	// Estimate highest combined usage.
@@ -346,10 +378,14 @@ TArray<FStateTreeMemoryUsage> UStateTree::CalculateEstimatedMemoryUsage() const
 	TreeMemUsage.EstimatedMemoryUsage += EvalMemUsage.EstimatedMemoryUsage;
 	TreeMemUsage.NodeCount += EvalMemUsage.NodeCount;
 
+	FStateTreeMemoryUsage& InstanceMemUsage = MemoryUsages[InstanceMemUsageIndex];
 	// FStateTreeInstanceData overhead.
-	TreeMemUsage.EstimatedMemoryUsage += sizeof(FStateTreeInstanceData);
+	InstanceMemUsage.EstimatedMemoryUsage += sizeof(FStateTreeInstanceData);
 	// FInstancedStructArray overhead.
-	TreeMemUsage.EstimatedMemoryUsage += TreeMemUsage.NodeCount * 16;
+	constexpr int32 ItemSize = 16; // sizeof(FInstancedStructArray::FItem);
+	InstanceMemUsage.EstimatedMemoryUsage += TreeMemUsage.NodeCount * ItemSize;
+
+	TreeMemUsage.EstimatedMemoryUsage += InstanceMemUsage.EstimatedMemoryUsage;
 	
 	int32 MaxSubtreeUsage = 0;
 	int32 MaxSubtreeNodeCount = 0;
