StateTree Editor: Make tree rows as wide as the view.

diff --git a/Engine/Plugins/Runtime/StateTree/Source/StateTreeEditorModule/Private/SStateTreeView.cpp b/Engine/Plugins/Runtime/StateTree/Source/StateTreeEditorModule/Private/SStateTreeView.cpp
--- a/Engine/Plugins/Runtime/StateTree/Source/StateTreeEditorModule/Private/SStateTreeView.cpp
+++ b/Engine/Plugins/Runtime/StateTree/Source/StateTreeEditorModule/Private/SStateTreeView.cpp
@@ -137,7 +137,7 @@ void SStateTreeView::Construct(const FArguments& InArgs, TSharedRef<FStateTreeVi
 			.FillWidth(1.0f)
 			.Padding(0.0f)
 			[
-				SNew(SScrollBox)
+				SAssignNew(ViewBox, SScrollBox)
 				.Orientation(Orient_Horizontal)
 				.ExternalScrollbar(HorizontalScrollBar)
 				+SScrollBox::Slot()
@@ -290,7 +290,7 @@ void SStateTreeView::HandleModelSelectionChanged(const TArray<UStateTreeState*>&
 
 TSharedRef<ITableRow> SStateTreeView::HandleGenerateRow(UStateTreeState* InState, const TSharedRef<STableViewBase>& InOwnerTableView)
 {
-	return SNew(SStateTreeViewRow, InOwnerTableView, InState, StateTreeViewModel.ToSharedRef());
+	return SNew(SStateTreeViewRow, InOwnerTableView, InState, ViewBox, StateTreeViewModel.ToSharedRef());
 }
 
 void SStateTreeView::HandleGetChildren(UStateTreeState* InParent, TArray<UStateTreeState*>& OutChildren)
diff --git a/Engine/Plugins/Runtime/StateTree/Source/StateTreeEditorModule/Private/SStateTreeView.h b/Engine/Plugins/Runtime/StateTree/Source/StateTreeEditorModule/Private/SStateTreeView.h
--- a/Engine/Plugins/Runtime/StateTree/Source/StateTreeEditorModule/Private/SStateTreeView.h
+++ b/Engine/Plugins/Runtime/StateTree/Source/StateTreeEditorModule/Private/SStateTreeView.h
@@ -8,7 +8,7 @@
 
 class UStateTreeEditorData;
 class UStateTreeState;
-
+class SScrollBox;
 
 class FActionTreeViewDragDrop : public FDragDropOperation
 {
@@ -75,6 +75,7 @@ private:
 
 	TSharedPtr<STreeView<UStateTreeState*>> TreeView;
 	TSharedPtr<SScrollBar> ExternalScrollbar;
+	TSharedPtr<SScrollBox> ViewBox;
 
 	UStateTreeState* RequestedRenameState;
 	bool bItemsDirty;
diff --git a/Engine/Plugins/Runtime/StateTree/Source/StateTreeEditorModule/Private/SStateTreeViewRow.cpp b/Engine/Plugins/Runtime/StateTree/Source/StateTreeEditorModule/Private/SStateTreeViewRow.cpp
--- a/Engine/Plugins/Runtime/StateTree/Source/StateTreeEditorModule/Private/SStateTreeViewRow.cpp
+++ b/Engine/Plugins/Runtime/StateTree/Source/StateTreeEditorModule/Private/SStateTreeViewRow.cpp
@@ -13,7 +13,6 @@
 
 #include "Widgets/DeclarativeSyntaxSupport.h"
 #include "Widgets/SCompoundWidget.h"
-#include "Widgets/Input/SButton.h"
 #include "Widgets/Text/STextBlock.h"
 #include "Widgets/Layout/SBorder.h"
 #include "Widgets/Layout/SScrollBorder.h"
@@ -42,13 +41,13 @@ static void AddUnique(TArray<FText>& Array, const FText& NewItem)
 	
 }} // UE::StateTreeEditor
 
-void SStateTreeViewRow::Construct(const FArguments& InArgs, const TSharedRef<STableViewBase>& InOwnerTableView, UStateTreeState* InState, TSharedRef<FStateTreeViewModel> InStateTreeViewModel)
+void SStateTreeViewRow::Construct(const FArguments& InArgs, const TSharedRef<STableViewBase>& InOwnerTableView, UStateTreeState* InState, const TSharedPtr<SScrollBox>& ViewBox, TSharedRef<FStateTreeViewModel> InStateTreeViewModel)
 {
 	StateTreeViewModel = InStateTreeViewModel;
 	WeakState = InState;
 
 	STableRow<UStateTreeState*>::ConstructInternal(STableRow::FArguments()
-        .Padding(5.f)
+        .Padding(5.0f)
         .OnDragDetected(this, &SStateTreeViewRow::HandleDragDetected)
         .OnCanAcceptDrop(this, &SStateTreeViewRow::HandleCanAcceptDrop)
         .OnAcceptDrop(this, &SStateTreeViewRow::HandleAcceptDrop)
@@ -57,306 +56,315 @@ void SStateTreeViewRow::Construct(const FArguments& InArgs, const TSharedRef<STa
 
 	static const FLinearColor TasksBackground = FLinearColor(FColor(17, 117, 131));
 	static const FLinearColor LinkBackground = FLinearColor(FColor(84, 84, 84));
-	
+
 	this->ChildSlot
     .HAlign(HAlign_Fill)
     [
-        SNew(SHorizontalBox)
-        + SHorizontalBox::Slot()
-        .VAlign(VAlign_Fill)
-        .HAlign(HAlign_Left)
-        .AutoWidth()
-        [
-            SNew(SExpanderArrow, SharedThis(this))
-            .ShouldDrawWires(true)
-            .IndentAmount(32)
-            .BaseIndentLevel(0)
-        ]
-
-        // State Box
-        + SHorizontalBox::Slot()
-        .VAlign(VAlign_Center)
-        .Padding(FMargin(0.0f, 4.0f))
-        .AutoWidth()
-        [
-            SNew(SBox)
-            .HeightOverride(28.0f)
-            .VAlign(VAlign_Fill)
-            [
-                SNew(SBorder)
-                .BorderImage(FAppStyle::GetBrush("WhiteBrush"))
-                .BorderBackgroundColor(this, &SStateTreeViewRow::GetTitleColor)
-                .Padding(FMargin(16.0f, 0.0f, 16.0f, 0.0f))
-                [
-                    // Conditions icon
-                    SNew(SHorizontalBox)
-                    + SHorizontalBox::Slot()
-                    .VAlign(VAlign_Center)
-                    .AutoWidth()
-                    [
-                        SNew(SBox)
-                        .Padding(FMargin(0.0f, 0.0f, 4.0f, 0.0f))
-                        .Visibility(this, &SStateTreeViewRow::GetConditionVisibility)
-                        [
-                            SNew(SImage)
-							.Image(FAppStyle::Get().GetBrush("Icons.Help"))
+	    SNew(SBox)
+	    .MinDesiredWidth_Lambda([ViewBox]()
+			{
+    			// Make the row at least as wide as the view.
+	    		// The -1 is needed or we'll see a scrollbar. 
+				return ViewBox->GetTickSpaceGeometry().GetLocalSize().X - 1;
+			})
+		[
+	        SNew(SHorizontalBox)
+	        + SHorizontalBox::Slot()
+	        .VAlign(VAlign_Fill)
+	        .HAlign(HAlign_Left)
+	        .AutoWidth()
+	        [
+	            SNew(SExpanderArrow, SharedThis(this))
+	            .ShouldDrawWires(true)
+	            .IndentAmount(32)
+	            .BaseIndentLevel(0)
+	        ]
+
+	        // State Box
+	        + SHorizontalBox::Slot()
+	        .VAlign(VAlign_Center)
+	        .Padding(FMargin(0.0f, 4.0f))
+	        .AutoWidth()
+	        [
+	            SNew(SBox)
+	            .HeightOverride(28.0f)
+	            .VAlign(VAlign_Fill)
+	            [
+	                SNew(SBorder)
+	                .BorderImage(FAppStyle::GetBrush("WhiteBrush"))
+	                .BorderBackgroundColor(this, &SStateTreeViewRow::GetTitleColor)
+	                .Padding(FMargin(16.0f, 0.0f, 16.0f, 0.0f))
+	                [
+	                    // Conditions icon
+	                    SNew(SHorizontalBox)
+	                    + SHorizontalBox::Slot()
+	                    .VAlign(VAlign_Center)
+	                    .AutoWidth()
+	                    [
+	                        SNew(SBox)
+	                        .Padding(FMargin(0.0f, 0.0f, 4.0f, 0.0f))
+	                        .Visibility(this, &SStateTreeViewRow::GetConditionVisibility)
+	                        [
+	                            SNew(SImage)
+								.Image(FAppStyle::Get().GetBrush("Icons.Help"))
+							]
 						]
-					]
 
-					// Selector icon
-                    + SHorizontalBox::Slot()
-                    .VAlign(VAlign_Center)
-                    .AutoWidth()
-                    [
-                        SNew(SBox)
-                        .Padding(FMargin(0.0f, 0.0f, 4.0f, 0.0f))
-                        .Visibility(this, &SStateTreeViewRow::GetSelectorVisibility)
-                        [
-                            SNew(STextBlock)
-                            .Text(this, &SStateTreeViewRow::GetSelectorDesc)
-                            .TextStyle(FStateTreeEditorStyle::Get(), "StateTree.Icon")
-                        ]
-                    ]
-
-                    // State Name
-					+ SHorizontalBox::Slot()
-					.VAlign(VAlign_Center)
-					.AutoWidth()
-					[
-						SAssignNew(NameTextBlock, SInlineEditableTextBlock)
-						.Style(FStateTreeEditorStyle::Get(), "StateTree.State.TitleInlineEditableText")
-						.OnVerifyTextChanged(this, &SStateTreeViewRow::VerifyNodeTextChanged)
-						.OnTextCommitted(this, &SStateTreeViewRow::HandleNodeLabelTextCommitted)
-						.Text(this, &SStateTreeViewRow::GetStateDesc)
-						.Clipping(EWidgetClipping::ClipToBounds)
-						.IsSelected(this, &SStateTreeViewRow::IsSelected)
+						// Selector icon
+	                    + SHorizontalBox::Slot()
+	                    .VAlign(VAlign_Center)
+	                    .AutoWidth()
+	                    [
+	                        SNew(SBox)
+	                        .Padding(FMargin(0.0f, 0.0f, 4.0f, 0.0f))
+	                        .Visibility(this, &SStateTreeViewRow::GetSelectorVisibility)
+	                        [
+	                            SNew(STextBlock)
+	                            .Text(this, &SStateTreeViewRow::GetSelectorDesc)
+	                            .TextStyle(FStateTreeEditorStyle::Get(), "StateTree.Icon")
+	                        ]
+	                    ]
+
+	                    // State Name
+						+ SHorizontalBox::Slot()
+						.VAlign(VAlign_Center)
+						.AutoWidth()
+						[
+							SAssignNew(NameTextBlock, SInlineEditableTextBlock)
+							.Style(FStateTreeEditorStyle::Get(), "StateTree.State.TitleInlineEditableText")
+							.OnVerifyTextChanged(this, &SStateTreeViewRow::VerifyNodeTextChanged)
+							.OnTextCommitted(this, &SStateTreeViewRow::HandleNodeLabelTextCommitted)
+							.Text(this, &SStateTreeViewRow::GetStateDesc)
+							.Clipping(EWidgetClipping::ClipToBounds)
+							.IsSelected(this, &SStateTreeViewRow::IsSelected)
+						]
 					]
 				]
 			]
-		]
 
-		// Tasks Box
-		+ SHorizontalBox::Slot()
-		.VAlign(VAlign_Center)
-		.Padding(FMargin(0.0f, 4.0f))
-		.AutoWidth()
-		[
-			SNew(SBox)
-			.HeightOverride(28.0f)
-			.VAlign(VAlign_Fill)
-			.Visibility(this, &SStateTreeViewRow::GetTasksVisibility)
+			// Tasks Box
+			+ SHorizontalBox::Slot()
+			.VAlign(VAlign_Center)
+			.Padding(FMargin(0.0f, 4.0f))
+			.AutoWidth()
 			[
-				SNew(SBorder)
-				.BorderImage(FAppStyle::GetBrush("WhiteBrush"))
-				.BorderBackgroundColor(TasksBackground)
-				.Padding(FMargin(12.0f, 0.0f, 16.0f, 0.0f))
+				SNew(SBox)
+				.HeightOverride(28.0f)
+				.VAlign(VAlign_Fill)
+				.Visibility(this, &SStateTreeViewRow::GetTasksVisibility)
 				[
-					// Task icon
-					SNew(SHorizontalBox)
-					+ SHorizontalBox::Slot()
-					.VAlign(VAlign_Center)
-					.Padding(FMargin(0.0f, 0.0f, 4.0f, 0.0f))
-					.AutoWidth()
+					SNew(SBorder)
+					.BorderImage(FAppStyle::GetBrush("WhiteBrush"))
+					.BorderBackgroundColor(TasksBackground)
+					.Padding(FMargin(12.0f, 0.0f, 16.0f, 0.0f))
 					[
-						SNew(STextBlock)
-						.Text(FEditorFontGlyphs::Paper_Plane)
-						.TextStyle(FStateTreeEditorStyle::Get(), "StateTree.DetailsIcon")
-					]
+						// Task icon
+						SNew(SHorizontalBox)
+						+ SHorizontalBox::Slot()
+						.VAlign(VAlign_Center)
+						.Padding(FMargin(0.0f, 0.0f, 4.0f, 0.0f))
+						.AutoWidth()
+						[
+							SNew(STextBlock)
+							.Text(FEditorFontGlyphs::Paper_Plane)
+							.TextStyle(FStateTreeEditorStyle::Get(), "StateTree.DetailsIcon")
+						]
 
-					// Tasks list
-					+ SHorizontalBox::Slot()
-					.VAlign(VAlign_Center)
-					.AutoWidth()
-					[
-						SNew(STextBlock)
-						.Text(this, &SStateTreeViewRow::GetTasksDesc)
-						.TextStyle(FStateTreeEditorStyle::Get(), "StateTree.Details")
+						// Tasks list
+						+ SHorizontalBox::Slot()
+						.VAlign(VAlign_Center)
+						.AutoWidth()
+						[
+							SNew(STextBlock)
+							.Text(this, &SStateTreeViewRow::GetTasksDesc)
+							.TextStyle(FStateTreeEditorStyle::Get(), "StateTree.Details")
+						]
 					]
 				]
 			]
-		]
-		
-		// Linked State box
-		+ SHorizontalBox::Slot()
-		.VAlign(VAlign_Center)
-		.Padding(FMargin(0.0f, 4.0f))
-		.AutoWidth()
-		[
-			SNew(SBox)
-			.HeightOverride(28.0f)
-			.VAlign(VAlign_Fill)
-			.Visibility(this, &SStateTreeViewRow::GetLinkedStateVisibility)
+			
+			// Linked State box
+			+ SHorizontalBox::Slot()
+			.VAlign(VAlign_Center)
+			.Padding(FMargin(0.0f, 4.0f))
+			.AutoWidth()
 			[
-				SNew(SBorder)
-				.BorderImage(FAppStyle::GetBrush("WhiteBrush"))
-				.BorderBackgroundColor(LinkBackground)
-				.Padding(FMargin(12.0f, 0.0f, 16.0f, 0.0f))
+				SNew(SBox)
+				.HeightOverride(28.0f)
+				.VAlign(VAlign_Fill)
+				.Visibility(this, &SStateTreeViewRow::GetLinkedStateVisibility)
 				[
-					// Link icon
-					SNew(SHorizontalBox)
-					+ SHorizontalBox::Slot()
-					.VAlign(VAlign_Center)
-					.Padding(FMargin(0.0f, 0.0f, 4.0f, 0.0f))
-					.AutoWidth()
+					SNew(SBorder)
+					.BorderImage(FAppStyle::GetBrush("WhiteBrush"))
+					.BorderBackgroundColor(LinkBackground)
+					.Padding(FMargin(12.0f, 0.0f, 16.0f, 0.0f))
 					[
-						SNew(STextBlock)
-						.Text(FEditorFontGlyphs::Link)
-						.TextStyle(FStateTreeEditorStyle::Get(), "StateTree.DetailsIcon")
-					]
+						// Link icon
+						SNew(SHorizontalBox)
+						+ SHorizontalBox::Slot()
+						.VAlign(VAlign_Center)
+						.Padding(FMargin(0.0f, 0.0f, 4.0f, 0.0f))
+						.AutoWidth()
+						[
+							SNew(STextBlock)
+							.Text(FEditorFontGlyphs::Link)
+							.TextStyle(FStateTreeEditorStyle::Get(), "StateTree.DetailsIcon")
+						]
 
-					// Linked State
-					+ SHorizontalBox::Slot()
-					.VAlign(VAlign_Center)
-					.AutoWidth()
-					[
-						SNew(STextBlock)
-						.Text(this, &SStateTreeViewRow::GetLinkedStateDesc)
-						.TextStyle(FStateTreeEditorStyle::Get(), "StateTree.Details")
+						// Linked State
+						+ SHorizontalBox::Slot()
+						.VAlign(VAlign_Center)
+						.AutoWidth()
+						[
+							SNew(STextBlock)
+							.Text(this, &SStateTreeViewRow::GetLinkedStateDesc)
+							.TextStyle(FStateTreeEditorStyle::Get(), "StateTree.Details")
+						]
 					]
 				]
 			]
-		]
-		
-		// Completed transitions
-		+ SHorizontalBox::Slot()
-		.VAlign(VAlign_Center)
-		.AutoWidth()
-		[
-			SNew(SHorizontalBox)
+			
+			// Completed transitions
 			+ SHorizontalBox::Slot()
 			.VAlign(VAlign_Center)
 			.AutoWidth()
-			.Padding(FMargin(8.0f, 0.0f, 0, 0.0f))
 			[
-				SNew(STextBlock)
-				.Text(this, &SStateTreeViewRow::GetCompletedTransitionsIcon)
-				.TextStyle(FStateTreeEditorStyle::Get(), "StateTree.Icon")
-				.Visibility(this, &SStateTreeViewRow::GetCompletedTransitionVisibility)
-			]
-			+ SHorizontalBox::Slot()
-			.VAlign(VAlign_Center)
-			.AutoWidth()
-			.Padding(FMargin(4.0f, 0, 0, 0))
-			[
-				SNew(STextBlock)
-				.Text(this, &SStateTreeViewRow::GetCompletedTransitionsDesc)
-				.TextStyle(FStateTreeEditorStyle::Get(), "StateTree.Details")
-				.Visibility(this, &SStateTreeViewRow::GetCompletedTransitionVisibility)
+				SNew(SHorizontalBox)
+				+ SHorizontalBox::Slot()
+				.VAlign(VAlign_Center)
+				.AutoWidth()
+				.Padding(FMargin(8.0f, 0.0f, 0, 0.0f))
+				[
+					SNew(STextBlock)
+					.Text(this, &SStateTreeViewRow::GetCompletedTransitionsIcon)
+					.TextStyle(FStateTreeEditorStyle::Get(), "StateTree.Icon")
+					.Visibility(this, &SStateTreeViewRow::GetCompletedTransitionVisibility)
+				]
+				+ SHorizontalBox::Slot()
+				.VAlign(VAlign_Center)
+				.AutoWidth()
+				.Padding(FMargin(4.0f, 0, 0, 0))
+				[
+					SNew(STextBlock)
+					.Text(this, &SStateTreeViewRow::GetCompletedTransitionsDesc)
+					.TextStyle(FStateTreeEditorStyle::Get(), "StateTree.Details")
+					.Visibility(this, &SStateTreeViewRow::GetCompletedTransitionVisibility)
+				]
 			]
-		]
 
-		// Succeeded transitions
-		+ SHorizontalBox::Slot()
-		.VAlign(VAlign_Center)
-		.AutoWidth()
-		[
-			SNew(SHorizontalBox)
-			+ SHorizontalBox::Slot()
-			.VAlign(VAlign_Center)
-			.Padding(FMargin(8.0f, 0.0f, 0, 0))
-			.AutoWidth()
-			[
-				SNew(STextBlock)
-				.Text(FEditorFontGlyphs::Check_Circle)
-                .ColorAndOpacity(FLinearColor(FColor(110,143,67)))
-				.TextStyle(FStateTreeEditorStyle::Get(), "StateTree.Icon")
-				.Visibility(this, &SStateTreeViewRow::GetSucceededTransitionVisibility)
-			]
+			// Succeeded transitions
 			+ SHorizontalBox::Slot()
 			.VAlign(VAlign_Center)
-			.Padding(FMargin(4.0f, 0.0f, 0, 0))
 			.AutoWidth()
 			[
-				SNew(STextBlock)
-				.Text(this, &SStateTreeViewRow::GetSucceededTransitionIcon)
-				.TextStyle(FStateTreeEditorStyle::Get(), "StateTree.Icon")
-				.Visibility(this, &SStateTreeViewRow::GetSucceededTransitionVisibility)
-			]
-			+ SHorizontalBox::Slot()
-			.VAlign(VAlign_Center)
-			.AutoWidth()
-			.Padding(FMargin(4.0f, 0, 0, 0))
-			[
-				SNew(STextBlock)
-				.Text(this, &SStateTreeViewRow::GetSucceededTransitionDesc)
-				.TextStyle(FStateTreeEditorStyle::Get(), "StateTree.Details")
-				.Visibility(this, &SStateTreeViewRow::GetSucceededTransitionVisibility)
+				SNew(SHorizontalBox)
+				+ SHorizontalBox::Slot()
+				.VAlign(VAlign_Center)
+				.Padding(FMargin(8.0f, 0.0f, 0, 0))
+				.AutoWidth()
+				[
+					SNew(STextBlock)
+					.Text(FEditorFontGlyphs::Check_Circle)
+	                .ColorAndOpacity(FLinearColor(FColor(110,143,67)))
+					.TextStyle(FStateTreeEditorStyle::Get(), "StateTree.Icon")
+					.Visibility(this, &SStateTreeViewRow::GetSucceededTransitionVisibility)
+				]
+				+ SHorizontalBox::Slot()
+				.VAlign(VAlign_Center)
+				.Padding(FMargin(4.0f, 0.0f, 0, 0))
+				.AutoWidth()
+				[
+					SNew(STextBlock)
+					.Text(this, &SStateTreeViewRow::GetSucceededTransitionIcon)
+					.TextStyle(FStateTreeEditorStyle::Get(), "StateTree.Icon")
+					.Visibility(this, &SStateTreeViewRow::GetSucceededTransitionVisibility)
+				]
+				+ SHorizontalBox::Slot()
+				.VAlign(VAlign_Center)
+				.AutoWidth()
+				.Padding(FMargin(4.0f, 0, 0, 0))
+				[
+					SNew(STextBlock)
+					.Text(this, &SStateTreeViewRow::GetSucceededTransitionDesc)
+					.TextStyle(FStateTreeEditorStyle::Get(), "StateTree.Details")
+					.Visibility(this, &SStateTreeViewRow::GetSucceededTransitionVisibility)
+				]
 			]
-		]
 
-		// Failed transitions
-		+ SHorizontalBox::Slot()
-		.VAlign(VAlign_Center)
-		.AutoWidth()
-		[
-			SNew(SHorizontalBox)
-			+ SHorizontalBox::Slot()
-			.VAlign(VAlign_Center)
-			.Padding(FMargin(8.0f, 0.0f, 0, 0))
-			.AutoWidth()
-			[
-				SNew(STextBlock)
-				.Text(FEditorFontGlyphs::Times_Circle)
-                .ColorAndOpacity(FLinearColor(FColor(187,77,42)))
-				.TextStyle(FStateTreeEditorStyle::Get(), "StateTree.Icon")
-				.Visibility(this, &SStateTreeViewRow::GetFailedTransitionVisibility)
-			]
+			// Failed transitions
 			+ SHorizontalBox::Slot()
 			.VAlign(VAlign_Center)
-			.Padding(FMargin(4.0f, 0.0f, 0, 0))
 			.AutoWidth()
 			[
-				SNew(STextBlock)
-				.Text(this, &SStateTreeViewRow::GetFailedTransitionIcon)
-				.TextStyle(FStateTreeEditorStyle::Get(), "StateTree.Icon")
-				.Visibility(this, &SStateTreeViewRow::GetFailedTransitionVisibility)
+				SNew(SHorizontalBox)
+				+ SHorizontalBox::Slot()
+				.VAlign(VAlign_Center)
+				.Padding(FMargin(8.0f, 0.0f, 0, 0))
+				.AutoWidth()
+				[
+					SNew(STextBlock)
+					.Text(FEditorFontGlyphs::Times_Circle)
+	                .ColorAndOpacity(FLinearColor(FColor(187,77,42)))
+					.TextStyle(FStateTreeEditorStyle::Get(), "StateTree.Icon")
+					.Visibility(this, &SStateTreeViewRow::GetFailedTransitionVisibility)
+				]
+				+ SHorizontalBox::Slot()
+				.VAlign(VAlign_Center)
+				.Padding(FMargin(4.0f, 0.0f, 0, 0))
+				.AutoWidth()
+				[
+					SNew(STextBlock)
+					.Text(this, &SStateTreeViewRow::GetFailedTransitionIcon)
+					.TextStyle(FStateTreeEditorStyle::Get(), "StateTree.Icon")
+					.Visibility(this, &SStateTreeViewRow::GetFailedTransitionVisibility)
+				]
+				+ SHorizontalBox::Slot()
+				.VAlign(VAlign_Center)
+				.AutoWidth()
+				.Padding(FMargin(4.0f, 0, 0, 0))
+				[
+					SNew(STextBlock)
+					.Text(this, &SStateTreeViewRow::GetFailedTransitionDesc)
+					.TextStyle(FStateTreeEditorStyle::Get(), "StateTree.Details")
+					.Visibility(this, &SStateTreeViewRow::GetFailedTransitionVisibility)
+				]
 			]
+			// Transitions
 			+ SHorizontalBox::Slot()
 			.VAlign(VAlign_Center)
 			.AutoWidth()
-			.Padding(FMargin(4.0f, 0, 0, 0))
 			[
-				SNew(STextBlock)
-				.Text(this, &SStateTreeViewRow::GetFailedTransitionDesc)
-				.TextStyle(FStateTreeEditorStyle::Get(), "StateTree.Details")
-				.Visibility(this, &SStateTreeViewRow::GetFailedTransitionVisibility)
-			]
-		]
-		// Transitions
-		+ SHorizontalBox::Slot()
-		.VAlign(VAlign_Center)
-		.AutoWidth()
-		[
-			SNew(SHorizontalBox)
-			+ SHorizontalBox::Slot()
-			.VAlign(VAlign_Center)
-			.Padding(FMargin(8.0f, 0.0f, 0, 0))
-			.AutoWidth()
-			[
-				SNew(SImage)
-				.Image(FAppStyle::Get().GetBrush("Icons.Help"))
-                .ColorAndOpacity(FLinearColor(FColor(31,151,167)))
-				.Visibility(this, &SStateTreeViewRow::GetConditionalTransitionsVisibility)
-			]
-            + SHorizontalBox::Slot()
-            .VAlign(VAlign_Center)
-            .Padding(FMargin(4.0f, 0.0f, 0, 0))
-            .AutoWidth()
-            [
-                SNew(STextBlock)
-                .Text(FEditorFontGlyphs::Long_Arrow_Right)
-                .TextStyle(FStateTreeEditorStyle::Get(), "StateTree.Icon")
-                .Visibility(this, &SStateTreeViewRow::GetConditionalTransitionsVisibility)
-            ]
-            + SHorizontalBox::Slot()
-			.VAlign(VAlign_Center)
-			.AutoWidth()
-			.Padding(FMargin(4.0f, 0, 0, 0))
-			[
-				SNew(STextBlock)
-				.Text(this, &SStateTreeViewRow::GetConditionalTransitionsDesc)
-				.TextStyle(FStateTreeEditorStyle::Get(), "StateTree.Details")
-				.Visibility(this, &SStateTreeViewRow::GetConditionalTransitionsVisibility)
+				SNew(SHorizontalBox)
+				+ SHorizontalBox::Slot()
+				.VAlign(VAlign_Center)
+				.Padding(FMargin(8.0f, 0.0f, 0, 0))
+				.AutoWidth()
+				[
+					SNew(SImage)
+					.Image(FAppStyle::Get().GetBrush("Icons.Help"))
+	                .ColorAndOpacity(FLinearColor(FColor(31,151,167)))
+					.Visibility(this, &SStateTreeViewRow::GetConditionalTransitionsVisibility)
+				]
+	            + SHorizontalBox::Slot()
+	            .VAlign(VAlign_Center)
+	            .Padding(FMargin(4.0f, 0.0f, 0, 0))
+	            .AutoWidth()
+	            [
+	                SNew(STextBlock)
+	                .Text(FEditorFontGlyphs::Long_Arrow_Right)
+	                .TextStyle(FStateTreeEditorStyle::Get(), "StateTree.Icon")
+	                .Visibility(this, &SStateTreeViewRow::GetConditionalTransitionsVisibility)
+	            ]
+	            + SHorizontalBox::Slot()
+				.VAlign(VAlign_Center)
+				.AutoWidth()
+				.Padding(FMargin(4.0f, 0, 0, 0))
+				[
+					SNew(STextBlock)
+					.Text(this, &SStateTreeViewRow::GetConditionalTransitionsDesc)
+					.TextStyle(FStateTreeEditorStyle::Get(), "StateTree.Details")
+					.Visibility(this, &SStateTreeViewRow::GetConditionalTransitionsVisibility)
+				]
 			]
 		]
 	];
diff --git a/Engine/Plugins/Runtime/StateTree/Source/StateTreeEditorModule/Private/SStateTreeViewRow.h b/Engine/Plugins/Runtime/StateTree/Source/StateTreeEditorModule/Private/SStateTreeViewRow.h
--- a/Engine/Plugins/Runtime/StateTree/Source/StateTreeEditorModule/Private/SStateTreeViewRow.h
+++ b/Engine/Plugins/Runtime/StateTree/Source/StateTreeEditorModule/Private/SStateTreeViewRow.h
@@ -9,6 +9,7 @@
 class UStateTreeEditorData;
 class SStateTreeView;
 class SInlineEditableTextBlock;
+class SScrollBox;
 class FStateTreeViewModel;
 
 class SStateTreeViewRow : public STableRow<UStateTreeState*>
@@ -20,7 +21,7 @@ public:
 	}
 	SLATE_END_ARGS()
 
-	void Construct(const FArguments& InArgs, const TSharedRef<STableViewBase>& InOwnerTableView, UStateTreeState* InState, TSharedRef<FStateTreeViewModel> InStateTreeViewModel);
+	void Construct(const FArguments& InArgs, const TSharedRef<STableViewBase>& InOwnerTableView, UStateTreeState* InState, const TSharedPtr<SScrollBox>& ViewBox, TSharedRef<FStateTreeViewModel> InStateTreeViewModel);
 	void RequestRename();
 
 private:
