StateTree: Added schema picker to new tree creation workflow
- Added BP-like class picker for schema when StateTree asset is created
- Added "CommonSchema" meta tag for schemas to allow them to be added to the common list
- Allow the asset on StateTree actor component to EditAnywhere
- Fixed compilation indication for empty StateTrees (was showing red)

diff --git a/Engine/Plugins/AI/MassAI/Source/MassAIBehavior/Public/MassStateTreeSchema.h b/Engine/Plugins/AI/MassAI/Source/MassAIBehavior/Public/MassStateTreeSchema.h
--- a/Engine/Plugins/AI/MassAI/Source/MassAIBehavior/Public/MassStateTreeSchema.h
+++ b/Engine/Plugins/AI/MassAI/Source/MassAIBehavior/Public/MassStateTreeSchema.h
@@ -6,9 +6,9 @@
 #include "MassStateTreeSchema.generated.h"
 
 /**
- * StateTree Schema for Mass behaviors.
+ * StateTree for Mass behaviors.
  */
-UCLASS(BlueprintType, EditInlineNew, CollapseCategories, meta = (DisplayName = "Mass Behavior"))
+UCLASS(BlueprintType, EditInlineNew, CollapseCategories, meta = (DisplayName = "Mass Behavior", CommonSchema))
 class MASSAIBEHAVIOR_API UMassStateTreeSchema : public UStateTreeSchema
 {
 	GENERATED_BODY()
diff --git a/Engine/Plugins/Runtime/GameplayStateTree/Source/GameplayStateTreeModule/Public/Components/StateTreeComponent.h b/Engine/Plugins/Runtime/GameplayStateTree/Source/GameplayStateTreeModule/Public/Components/StateTreeComponent.h
--- a/Engine/Plugins/Runtime/GameplayStateTree/Source/GameplayStateTreeModule/Public/Components/StateTreeComponent.h
+++ b/Engine/Plugins/Runtime/GameplayStateTree/Source/GameplayStateTreeModule/Public/Components/StateTreeComponent.h
@@ -50,7 +50,7 @@ protected:
 
 	bool SetContextRequirements();
 	
-	UPROPERTY(EditDefaultsOnly, Category = AI, meta=(RequiredAssetDataTags="Schema=StateTreeComponentSchema"))
+	UPROPERTY(EditAnywhere, BlueprintReadWrite, Category = AI, meta=(RequiredAssetDataTags="Schema=StateTreeComponentSchema"))
 	UStateTree* StateTree;
 
 	UPROPERTY(Transient)
diff --git a/Engine/Plugins/Runtime/GameplayStateTree/Source/GameplayStateTreeModule/Public/Components/StateTreeComponentSchema.h b/Engine/Plugins/Runtime/GameplayStateTree/Source/GameplayStateTreeModule/Public/Components/StateTreeComponentSchema.h
--- a/Engine/Plugins/Runtime/GameplayStateTree/Source/GameplayStateTreeModule/Public/Components/StateTreeComponentSchema.h
+++ b/Engine/Plugins/Runtime/GameplayStateTree/Source/GameplayStateTreeModule/Public/Components/StateTreeComponentSchema.h
@@ -5,8 +5,10 @@
 #include "StateTreeSchema.h"
 #include "StateTreeComponentSchema.generated.h"
 
-
-UCLASS(BlueprintType, EditInlineNew, CollapseCategories, meta = (DisplayName = "StateTree Component"))
+/**
+ * StateTree for Actors with StateTree component. 
+ */
+UCLASS(BlueprintType, EditInlineNew, CollapseCategories, meta = (DisplayName = "StateTree Component", CommonSchema))
 class UStateTreeComponentSchema : public UStateTreeSchema
 {
 	GENERATED_BODY()
diff --git a/Engine/Plugins/Runtime/StateTree/Source/StateTreeEditorModule/Private/StateTreeEditor.cpp b/Engine/Plugins/Runtime/StateTree/Source/StateTreeEditorModule/Private/StateTreeEditor.cpp
--- a/Engine/Plugins/Runtime/StateTree/Source/StateTreeEditorModule/Private/StateTreeEditor.cpp
+++ b/Engine/Plugins/Runtime/StateTree/Source/StateTreeEditorModule/Private/StateTreeEditor.cpp
@@ -755,7 +755,7 @@ void FStateTreeEditor::Compile()
 		return;
 	}
 
-	// Note: If the compilation process changes, also update UStateTreeCompileAllCommandlet.
+	// Note: If the compilation process changes, also update UStateTreeCompileAllCommandlet and UStateTreeFactory::FactoryCreateNew.
 	
 	UpdateAsset();
 	
diff --git a/Engine/Plugins/Runtime/StateTree/Source/StateTreeEditorModule/Private/StateTreeEditorModule.cpp b/Engine/Plugins/Runtime/StateTree/Source/StateTreeEditorModule/Private/StateTreeEditorModule.cpp
--- a/Engine/Plugins/Runtime/StateTree/Source/StateTreeEditorModule/Private/StateTreeEditorModule.cpp
+++ b/Engine/Plugins/Runtime/StateTree/Source/StateTreeEditorModule/Private/StateTreeEditorModule.cpp
@@ -28,6 +28,8 @@
 
 #define LOCTEXT_NAMESPACE "StateTreeEditor"
 
+DEFINE_LOG_CATEGORY(LogStateTreeEditor);
+
 IMPLEMENT_MODULE(FStateTreeEditorModule, StateTreeEditorModule)
 
 void FStateTreeEditorModule::StartupModule()
@@ -99,6 +101,13 @@ void FStateTreeEditorModule::ShutdownModule()
 }
 
 TSharedRef<IStateTreeEditor> FStateTreeEditorModule::CreateStateTreeEditor(const EToolkitMode::Type Mode, const TSharedPtr< class IToolkitHost >& InitToolkitHost, UStateTree* StateTree)
+{
+	TSharedRef<FStateTreeEditor> NewEditor(new FStateTreeEditor());
+	NewEditor->InitEditor(Mode, InitToolkitHost, StateTree);
+	return NewEditor;
+}
+
+TSharedPtr<FStateTreeNodeClassCache> FStateTreeEditorModule::GetNodeClassCache()
 {
 	if (!NodeClassCache.IsValid())
 	{
@@ -109,11 +118,11 @@ TSharedRef<IStateTreeEditor> FStateTreeEditorModule::CreateStateTreeEditor(const
 		NodeClassCache->AddRootClass(UStateTreeEvaluatorBlueprintBase::StaticClass());
 		NodeClassCache->AddRootClass(UStateTreeTaskBlueprintBase::StaticClass());
 		NodeClassCache->AddRootClass(UStateTreeConditionBlueprintBase::StaticClass());
+		NodeClassCache->AddRootClass(UStateTreeSchema::StaticClass());
 	}
-	
-	TSharedRef<FStateTreeEditor> NewEditor(new FStateTreeEditor());
-	NewEditor->InitEditor(Mode, InitToolkitHost, StateTree);
-	return NewEditor;
+
+	return NodeClassCache;
 }
 
+
 #undef LOCTEXT_NAMESPACE
diff --git a/Engine/Plugins/Runtime/StateTree/Source/StateTreeEditorModule/Private/StateTreeFactory.cpp b/Engine/Plugins/Runtime/StateTree/Source/StateTreeEditorModule/Private/StateTreeFactory.cpp
--- a/Engine/Plugins/Runtime/StateTree/Source/StateTreeEditorModule/Private/StateTreeFactory.cpp
+++ b/Engine/Plugins/Runtime/StateTree/Source/StateTreeEditorModule/Private/StateTreeFactory.cpp
@@ -2,9 +2,45 @@
 
 #include "StateTreeFactory.h"
 #include "StateTree.h"
+#include "Kismet2/SClassPickerDialog.h"
+#include "ClassViewerFilter.h"
+#include "StateTreeEditorModule.h"
+#include "StateTreeNodeClassCache.h"
+#include "StateTreeEditorData.h"
+#include "StateTreeEditor.h"
+#include "StateTreeCompiler.h"
+#include "StateTreeCompilerLog.h"
+#include "StateTreeDelegates.h"
+#include "Modules/ModuleManager.h"
 
 #define LOCTEXT_NAMESPACE "StateTreeEditor"
 
+/////////////////////////////////////////////////////
+// FStateTreeClassFilter
+
+class FStateTreeClassFilter : public IClassViewerFilter
+{
+public:
+	/** All children of these classes will be included unless filtered out by another setting. */
+	TSet<const UClass*> AllowedChildrenOfClasses;
+
+	/** Disallowed class flags. */
+	EClassFlags DisallowedClassFlags = CLASS_None;
+
+	virtual bool IsClassAllowed(const FClassViewerInitializationOptions& InInitOptions, const UClass* InClass, TSharedRef< FClassViewerFilterFuncs > InFilterFuncs) override
+	{
+		return !InClass->HasAnyClassFlags(DisallowedClassFlags)
+			&& InFilterFuncs->IfInChildOfClassesSet(AllowedChildrenOfClasses, InClass) != EFilterReturn::Failed;
+	}
+	
+	virtual bool IsUnloadedClassAllowed(const FClassViewerInitializationOptions& InInitOptions, const TSharedRef< const IUnloadedBlueprintData > InUnloadedClassData, TSharedRef< FClassViewerFilterFuncs > InFilterFuncs) override
+	{
+		return !InUnloadedClassData->HasAnyClassFlags(DisallowedClassFlags)
+			&& InFilterFuncs->IfInChildOfClassesSet(AllowedChildrenOfClasses, InUnloadedClassData) != EFilterReturn::Failed;
+	}
+
+};
+
 /////////////////////////////////////////////////////
 // UStateTreeFactory
 
@@ -16,9 +52,94 @@ UStateTreeFactory::UStateTreeFactory(const FObjectInitializer& ObjectInitializer
 	SupportedClass = UStateTree::StaticClass();
 }
 
+bool UStateTreeFactory::ConfigureProperties()
+{
+	// Load the classviewer module to display a class picker
+	FClassViewerModule& ClassViewerModule = FModuleManager::LoadModuleChecked<FClassViewerModule>("ClassViewer");
+
+	// Load StateTree class cache to get all schema classes
+	FStateTreeEditorModule& EditorModule = FModuleManager::GetModuleChecked<FStateTreeEditorModule>(TEXT("StateTreeEditorModule"));
+	FStateTreeNodeClassCache* ClassCache = EditorModule.GetNodeClassCache().Get();
+	check(ClassCache);
+
+	// Fill in options
+	FClassViewerInitializationOptions Options;
+	Options.DisplayMode = EClassViewerDisplayMode::Type::TreeView;
+	Options.Mode = EClassViewerMode::ClassPicker;
+	Options.bShowNoneOption = false;
+	Options.bExpandAllNodes = true;
+
+	TSharedPtr<FStateTreeClassFilter> Filter = MakeShareable(new FStateTreeClassFilter);
+	Options.ClassFilters.Add(Filter.ToSharedRef());
+
+	// Add all schemas which are tagged as "CommonSchema" to the common section.
+	TArray<TSharedPtr<FStateTreeNodeClassData>> AvailableClasses;
+	ClassCache->GetClasses(UStateTreeSchema::StaticClass(), AvailableClasses);
+	for (const TSharedPtr<FStateTreeNodeClassData>& ClassData : AvailableClasses)
+	{
+		if (FStateTreeNodeClassData* Data = ClassData.Get())
+		{
+			if (UClass* Class = Data->GetClass())
+			{
+				if (Class->HasMetaData("CommonSchema"))
+				{
+					Options.ExtraPickerCommonClasses.Add(Class);					
+				}
+			}
+		}
+	}
+
+	Filter->DisallowedClassFlags = CLASS_Deprecated | CLASS_NewerVersionExists | CLASS_Abstract | CLASS_HideDropDown;
+	Filter->AllowedChildrenOfClasses.Add(UStateTreeSchema::StaticClass());
+
+	const FText TitleText = LOCTEXT("CreateStateTree", "Pick Schema for State Tree");
+
+	UClass* ChosenClass = nullptr;
+	const bool bResult = SClassPickerDialog::PickClass(TitleText, Options, ChosenClass, UStateTreeSchema::StaticClass());
+	StateTreeSchemaClass = ChosenClass;
+
+	return bResult;
+}
+
 UObject* UStateTreeFactory::FactoryCreateNew(UClass* Class, UObject* InParent, FName Name, EObjectFlags Flags, UObject* Context, FFeedbackContext* Warn)
 {
+	if (StateTreeSchemaClass == nullptr)
+	{
+		return nullptr;
+	}
+
+	// Create new asset
 	UStateTree* NewStateTree = NewObject<UStateTree>(InParent, Class, Name, Flags);
+
+	// Create and init new editor data.
+	UStateTreeEditorData* EditorData = NewObject<UStateTreeEditorData>(NewStateTree, FName(), RF_Transactional);
+	EditorData->AddRootState();
+	EditorData->Schema = NewObject<UStateTreeSchema>(this, StateTreeSchemaClass);
+	
+	NewStateTree->EditorData = EditorData;
+
+	// Compile the StateTree asset, so that it is valid initially.
+	UE::StateTree::Editor::ValidateAsset(*NewStateTree);
+
+	FStateTreeCompilerLog Log;
+	FStateTreeCompiler Compiler(Log);
+
+	const bool bSuccess = Compiler.Compile(*NewStateTree);
+
+	if (bSuccess)
+	{
+		// Success
+		const uint32 EditorDataHash = UE::StateTree::Editor::CalcAssetHash(*NewStateTree);
+		NewStateTree->LastCompiledEditorDataHash = EditorDataHash;
+		UE::StateTree::Delegates::OnPostCompile.Broadcast(*NewStateTree);
+	}
+	else
+	{
+		// Should not happen.
+		Log.DumpToLog(LogStateTreeEditor);
+		NewStateTree = nullptr;
+	}
+	
 	return NewStateTree;
 }
 
diff --git a/Engine/Plugins/Runtime/StateTree/Source/StateTreeEditorModule/Private/StateTreeFactory.h b/Engine/Plugins/Runtime/StateTree/Source/StateTreeEditorModule/Private/StateTreeFactory.h
--- a/Engine/Plugins/Runtime/StateTree/Source/StateTreeEditorModule/Private/StateTreeFactory.h
+++ b/Engine/Plugins/Runtime/StateTree/Source/StateTreeEditorModule/Private/StateTreeFactory.h
@@ -2,11 +2,12 @@
 
 #pragma once
 
-#include "CoreMinimal.h"
 #include "UObject/ObjectMacros.h"
 #include "Factories/Factory.h"
 #include "StateTreeFactory.generated.h"
 
+class UStateTreeSchema;
+
 /**
  * Factory for UStateTree
  */
@@ -17,6 +18,12 @@ class STATETREEEDITORMODULE_API UStateTreeFactory : public UFactory
 	GENERATED_UCLASS_BODY()
 
 	// UFactory interface
+	virtual bool ConfigureProperties() override;
 	virtual UObject* FactoryCreateNew(UClass* Class, UObject* InParent, FName Name, EObjectFlags Flags, UObject* Context, FFeedbackContext* Warn) override;
 	// End of UFactory interface
+
+protected:
+	
+	UPROPERTY(Transient)
+	TObjectPtr<UClass> StateTreeSchemaClass;
 };
diff --git a/Engine/Plugins/Runtime/StateTree/Source/StateTreeEditorModule/Public/StateTreeEditorModule.h b/Engine/Plugins/Runtime/StateTree/Source/StateTreeEditorModule/Public/StateTreeEditorModule.h
--- a/Engine/Plugins/Runtime/StateTree/Source/StateTreeEditorModule/Public/StateTreeEditorModule.h
+++ b/Engine/Plugins/Runtime/StateTree/Source/StateTreeEditorModule/Public/StateTreeEditorModule.h
@@ -11,6 +11,8 @@ class IStateTreeEditor;
 class FAssetTypeActions_Base;
 struct FStateTreeNodeClassCache;
 
+STATETREEEDITORMODULE_API DECLARE_LOG_CATEGORY_EXTERN(LogStateTreeEditor, Warning, All);
+
 /**
 * The public interface to this module
 */
@@ -27,7 +29,7 @@ public:
 	virtual TSharedPtr<FExtensibilityManager> GetMenuExtensibilityManager() override { return MenuExtensibilityManager; }
 	virtual TSharedPtr<FExtensibilityManager> GetToolBarExtensibilityManager() override { return ToolBarExtensibilityManager; }
 
-	TSharedPtr<FStateTreeNodeClassCache> GetNodeClassCache() const { return NodeClassCache; }
+	TSharedPtr<FStateTreeNodeClassCache> GetNodeClassCache();
 	
 protected:
 	TSharedPtr<FExtensibilityManager> MenuExtensibilityManager;
diff --git a/Engine/Plugins/Runtime/StateTree/Source/StateTreeModule/Private/StateTree.cpp b/Engine/Plugins/Runtime/StateTree/Source/StateTreeModule/Private/StateTree.cpp
--- a/Engine/Plugins/Runtime/StateTree/Source/StateTreeModule/Private/StateTree.cpp
+++ b/Engine/Plugins/Runtime/StateTree/Source/StateTreeModule/Private/StateTree.cpp
@@ -115,9 +115,10 @@ bool UStateTree::Link()
 	// This data will be used to allocate runtime instance on all StateTree users.
 	ResetLinked();
 
-	// Do not try to link empty trees.
+	// Do not try to link empty trees any further.
 	if (States.Num() == 0 || Nodes.Num() == 0)
 	{
+		bIsLinked = true;
 		return true;
 	}
 
diff --git a/Engine/Plugins/Runtime/StateTree/Source/StateTreeTestSuite/Private/StateTreeTest.h b/Engine/Plugins/Runtime/StateTree/Source/StateTreeTestSuite/Private/StateTreeTest.h
--- a/Engine/Plugins/Runtime/StateTree/Source/StateTreeTestSuite/Private/StateTreeTest.h
+++ b/Engine/Plugins/Runtime/StateTree/Source/StateTreeTestSuite/Private/StateTreeTest.h
@@ -5,7 +5,7 @@
 #include "StateTreeSchema.h"
 #include "StateTreeTest.generated.h"
 
-UCLASS()
+UCLASS(HideDropdown)
 class UStateTreeTestSchema : public UStateTreeSchema
 {
 	GENERATED_BODY()
