StateTree: Fixed failing test and missing validation
- fail compilation if linking to non-subtree
- fixed crash in linker if subtree does not have parameters
- fixed subtree test (subtree state was missing type)

diff --git a/Engine/Plugins/Runtime/StateTree/Source/StateTreeEditorModule/Private/StateTreeCompiler.cpp b/Engine/Plugins/Runtime/StateTree/Source/StateTreeEditorModule/Private/StateTreeCompiler.cpp
--- a/Engine/Plugins/Runtime/StateTree/Source/StateTreeEditorModule/Private/StateTreeCompiler.cpp
+++ b/Engine/Plugins/Runtime/StateTree/Source/StateTreeEditorModule/Private/StateTreeCompiler.cpp
@@ -468,6 +468,24 @@ bool FStateTreeCompiler::CreateStateTransitions()
 					*LinkedParentState->Name.ToString());
 				return false;
 			}
+
+			// The linked state must be a subtree.
+			const UStateTreeState* TargetState = GetState(SourceState->LinkedState.ID);
+			if (TargetState == nullptr)
+			{
+				Log.Reportf(EMessageSeverity::Error,
+					TEXT("Failed to resolve linked state '%s'."),
+					*SourceState->LinkedState.Name.ToString());
+				return false;
+			}
+			
+			if (TargetState->Type != EStateTreeStateType::Subtree)
+			{
+				Log.Reportf(EMessageSeverity::Error,
+					TEXT("State '%s' is linked to state '%s', which is not a subtree."),
+					*SourceState->Name.ToString(), *TargetState->Name.ToString());
+				return false;
+			}
 			
 			CompactState.LinkedState = GetStateHandle(SourceState->LinkedState.ID);
 			
diff --git a/Engine/Plugins/Runtime/StateTree/Source/StateTreeModule/Private/StateTree.cpp b/Engine/Plugins/Runtime/StateTree/Source/StateTreeModule/Private/StateTree.cpp
--- a/Engine/Plugins/Runtime/StateTree/Source/StateTreeModule/Private/StateTree.cpp
+++ b/Engine/Plugins/Runtime/StateTree/Source/StateTreeModule/Private/StateTree.cpp
@@ -175,8 +175,11 @@ bool UStateTree::Link()
 				return false;
 			}
 
-			FStateTreePropCopyBatch& Batch = CopyBatches[Params.BindingsBatch.Get()];
-			Batch.TargetStruct.Struct = Params.Parameters.GetPropertyBagStruct();
+			if (Params.BindingsBatch.IsValid())
+			{
+				FStateTreePropCopyBatch& Batch = CopyBatches[Params.BindingsBatch.Get()];
+				Batch.TargetStruct.Struct = Params.Parameters.GetPropertyBagStruct();
+			}
 		}
 	}
 	
diff --git a/Engine/Plugins/Runtime/StateTree/Source/StateTreeTestSuite/Private/StateTreeTest.cpp b/Engine/Plugins/Runtime/StateTree/Source/StateTreeTestSuite/Private/StateTreeTest.cpp
--- a/Engine/Plugins/Runtime/StateTree/Source/StateTreeTestSuite/Private/StateTreeTest.cpp
+++ b/Engine/Plugins/Runtime/StateTree/Source/StateTreeTestSuite/Private/StateTreeTest.cpp
@@ -413,7 +413,7 @@ struct FStateTreeTest_SubTree : FAITestBase
 		UStateTreeState& Root = EditorData.AddSubTree(FName(TEXT("Root")));
 		UStateTreeState& State1 = Root.AddChildState(FName(TEXT("State1")), EStateTreeStateType::Linked);
 		UStateTreeState& State2 = Root.AddChildState(FName(TEXT("State2")));
-		UStateTreeState& State3 = Root.AddChildState(FName(TEXT("State3")));
+		UStateTreeState& State3 = Root.AddChildState(FName(TEXT("State3")), EStateTreeStateType::Subtree);
 		UStateTreeState& State3A = State3.AddChildState(FName(TEXT("State3A")));
 		UStateTreeState& State3B = State3.AddChildState(FName(TEXT("State3B")));
 
