StateTree: Added deprecation message to obsole evals.

diff --git a/Engine/Plugins/Runtime/StateTree/Source/StateTreeModule/Public/Evaluators/CooldownGuardEvaluator.h b/Engine/Plugins/Runtime/StateTree/Source/StateTreeModule/Public/Evaluators/CooldownGuardEvaluator.h
--- a/Engine/Plugins/Runtime/StateTree/Source/StateTreeModule/Public/Evaluators/CooldownGuardEvaluator.h
+++ b/Engine/Plugins/Runtime/StateTree/Source/StateTreeModule/Public/Evaluators/CooldownGuardEvaluator.h
@@ -39,7 +39,7 @@ struct STATETREEMODULE_API FCooldownGuardEvaluatorInstanceData
 };
 
 USTRUCT(DisplayName="Cooldown Guard")
-struct STATETREEMODULE_API FCooldownGuardEvaluator : public FStateTreeEvaluatorCommonBase
+struct UE_DEPRECATED(5.1, "This function will be removed for 5.1.") FCooldownGuardEvaluator : public FStateTreeEvaluatorCommonBase
 {
 	GENERATED_BODY()
 
diff --git a/Engine/Plugins/Runtime/StateTree/Source/StateTreeModule/Public/Evaluators/SelectionGuardEvaluator.h b/Engine/Plugins/Runtime/StateTree/Source/StateTreeModule/Public/Evaluators/SelectionGuardEvaluator.h
--- a/Engine/Plugins/Runtime/StateTree/Source/StateTreeModule/Public/Evaluators/SelectionGuardEvaluator.h
+++ b/Engine/Plugins/Runtime/StateTree/Source/StateTreeModule/Public/Evaluators/SelectionGuardEvaluator.h
@@ -33,7 +33,7 @@ struct STATETREEMODULE_API FSelectionGuardEvaluatorInstanceData
 };
 
 USTRUCT(DisplayName="Selection Guard")
-struct STATETREEMODULE_API FSelectionGuardEvaluator : public FStateTreeEvaluatorCommonBase
+struct UE_DEPRECATED(5.1, "This function will be removed for 5.1.") FSelectionGuardEvaluator : public FStateTreeEvaluatorCommonBase
 {
 	GENERATED_BODY()
 
