StateTree:
- Changed existing conditions to use the simpler instance data getter
- Added UObject based conditions
- Renamed condition file for consistency

diff --git a/Engine/Plugins/Runtime/GameplayStateTree/Source/GameplayStateTreeModule/Private/Components/StateTreeComponent.cpp b/Engine/Plugins/Runtime/GameplayStateTree/Source/GameplayStateTreeModule/Private/Components/StateTreeComponent.cpp
--- a/Engine/Plugins/Runtime/GameplayStateTree/Source/GameplayStateTreeModule/Private/Components/StateTreeComponent.cpp
+++ b/Engine/Plugins/Runtime/GameplayStateTree/Source/GameplayStateTreeModule/Private/Components/StateTreeComponent.cpp
@@ -6,7 +6,6 @@
 #include "VisualLogger/VisualLogger.h"
 #include "StateTree.h"
 #include "StateTreeEvaluatorBase.h"
-#include "Conditions/StateTreeCondition_Common.h"
 #include "AIController.h"
 #include "GameFramework/CharacterMovementComponent.h"
 #include "Subsystems/WorldSubsystem.h"
diff --git a/Engine/Plugins/Runtime/StateTree/Config/BaseStateTree.ini b/Engine/Plugins/Runtime/StateTree/Config/BaseStateTree.ini
--- a/Engine/Plugins/Runtime/StateTree/Config/BaseStateTree.ini
+++ b/Engine/Plugins/Runtime/StateTree/Config/BaseStateTree.ini
@@ -27,3 +27,21 @@
 
 +EnumRedirects=(OldName="EStateTreePropertyCopyType",ValueChanges=(("Plain","CopyPlain"),("Complex","CopyComplex"),("Bool","CopyBool"),("Struct","CopyStruct"),("Object","CopyObject"),("Name","CopyName"),("FixedArray","CopyFixedArray")))
 +EnumRedirects=(OldName="EStateTreePropertyAccessType",ValueChanges=(("Array","IndexArray")))
+
++StructRedirects=(OldName="StateTreeCondition_CompareIntInstanceData",NewName="StateTreeCompareIntConditionInstanceData")
++StructRedirects=(OldName="StateTreeCondition_CompareInt",NewName="StateTreeCompareIntCondition")
+
++StructRedirects=(OldName="StateTreeCondition_CompareFloatInstanceData",NewName="StateTreeCompareFloatConditionInstanceData")
++StructRedirects=(OldName="StateTreeCondition_CompareFloat",NewName="StateTreeCompareFloatCondition")
+
++StructRedirects=(OldName="StateTreeCondition_CompareBoolInstanceData",NewName="StateTreeCompareBoolConditionInstanceData")
++StructRedirects=(OldName="StateTreeCondition_CompareBool",NewName="StateTreeCompareBoolCondition")
+
++StructRedirects=(OldName="StateTreeCondition_CompareEnumInstanceData",NewName="StateTreeCompareEnumConditionInstanceData")
++StructRedirects=(OldName="StateTreeCondition_CompareEnum",NewName="StateTreeCompareEnumCondition")
+
++StructRedirects=(OldName="StateTreeCondition_CompareDistanceInstanceData",NewName="StateTreeCompareDistanceConditionInstanceData")
++StructRedirects=(OldName="StateTreeCondition_CompareDistance",NewName="StateTreeCompareDistanceCondition")
+
++StructRedirects=(OldName="StateTreeCondition_RandomInstanceData",NewName="StateTreeRandomConditionInstanceData")
++StructRedirects=(OldName="StateTreeCondition_Random",NewName="StateTreeRandomCondition")
\ No newline at end of file
diff --git a/Engine/Plugins/Runtime/StateTree/Source/StateTreeEditorModule/Private/StateTreeCompiler.cpp b/Engine/Plugins/Runtime/StateTree/Source/StateTreeEditorModule/Private/StateTreeCompiler.cpp
--- a/Engine/Plugins/Runtime/StateTree/Source/StateTreeEditorModule/Private/StateTreeCompiler.cpp
+++ b/Engine/Plugins/Runtime/StateTree/Source/StateTreeEditorModule/Private/StateTreeCompiler.cpp
@@ -4,7 +4,7 @@
 #include "StateTree.h"
 #include "StateTreeEditorData.h"
 #include "StateTreeTypes.h"
-#include "Conditions/StateTreeCondition_Common.h"
+#include "Conditions/StateTreeCommonConditions.h"
 #include "StateTreeEvaluatorBase.h"
 #include "StateTreeTaskBase.h"
 #include "StateTreeConditionBase.h"
diff --git a/Engine/Plugins/Runtime/StateTree/Source/StateTreeModule/Private/Conditions/GameplayTagConditions.cpp b/Engine/Plugins/Runtime/StateTree/Source/StateTreeModule/Private/Conditions/GameplayTagConditions.cpp
deleted file mode 100644
diff --git a/Engine/Plugins/Runtime/StateTree/Source/StateTreeModule/Private/Conditions/StateTreeCommonConditions.cpp b/Engine/Plugins/Runtime/StateTree/Source/StateTreeModule/Private/Conditions/StateTreeCommonConditions.cpp
new file mode 100644
--- /dev/null
+++ b/Engine/Plugins/Runtime/StateTree/Source/StateTreeModule/Private/Conditions/StateTreeCommonConditions.cpp
@@ -0,0 +1,201 @@
+// Copyright Epic Games, Inc. All Rights Reserved.
+
+#include "Conditions/StateTreeCommonConditions.h"
+
+#include "StateTreeExecutionContext.h"
+#include "StateTreePropertyBindings.h"
+#include "StateTreeLinker.h"
+
+#define LOCTEXT_NAMESPACE "StateTreeEditor"
+
+namespace UE::StateTree::Conditions
+{
+
+#if WITH_EDITOR
+FText GetOperatorText(const EGenericAICheck Operator)
+{
+	switch (Operator)
+	{
+	case EGenericAICheck::Equal:
+		return FText::FromString(TEXT("=="));
+		break;
+	case EGenericAICheck::NotEqual:
+		return FText::FromString(TEXT("!="));
+		break;
+	case EGenericAICheck::Less:
+		return FText::FromString(TEXT("&lt;"));
+		break;
+	case EGenericAICheck::LessOrEqual:
+		return FText::FromString(TEXT("&lt;="));
+		break;
+	case EGenericAICheck::Greater:
+		return FText::FromString(TEXT("&gt;"));
+		break;
+	case EGenericAICheck::GreaterOrEqual:
+		return FText::FromString(TEXT("&gt;="));
+		break;
+	default:
+		return FText::FromString(TEXT("??"));
+		break;
+	}
+	return FText::GetEmpty();
+}
+#endif // WITH_EDITOR
+
+template<typename T>
+bool CompareNumbers(const T Left, const T Right, const EGenericAICheck Operator)
+{
+	switch (Operator)
+	{
+	case EGenericAICheck::Equal:
+		return Left == Right;
+		break;
+	case EGenericAICheck::NotEqual:
+		return Left != Right;
+		break;
+	case EGenericAICheck::Less:
+		return Left < Right;
+		break;
+	case EGenericAICheck::LessOrEqual:
+		return Left <= Right;
+		break;
+	case EGenericAICheck::Greater:
+		return Left > Right;
+		break;
+	case EGenericAICheck::GreaterOrEqual:
+		return Left >= Right;
+		break;
+	default:
+		ensureMsgf(false, TEXT("Unhandled operator %d"), Operator);
+		return false;
+		break;
+	}
+	return false;
+}
+
+} // UE::StateTree::Conditions
+
+
+//----------------------------------------------------------------------//
+//  FStateTreeCompareIntCondition
+//----------------------------------------------------------------------//
+
+bool FStateTreeCompareIntCondition::TestCondition(FStateTreeExecutionContext& Context) const
+{
+	const InstanceDataType& InstanceData = Context.GetInstanceData<InstanceDataType>(*this);
+
+	const bool bResult = UE::StateTree::Conditions::CompareNumbers<int32>(InstanceData.Left, InstanceData.Right, Operator);
+	return bResult ^ bInvert;
+}
+
+
+//----------------------------------------------------------------------//
+//  FStateTreeCompareFloatCondition
+//----------------------------------------------------------------------//
+
+bool FStateTreeCompareFloatCondition::TestCondition(FStateTreeExecutionContext& Context) const
+{
+	const InstanceDataType& InstanceData = Context.GetInstanceData<InstanceDataType>(*this);
+
+	const bool bResult = UE::StateTree::Conditions::CompareNumbers<float>(InstanceData.Left, InstanceData.Right, Operator);
+	return bResult ^ bInvert;
+}
+
+
+//----------------------------------------------------------------------//
+//  FStateTreeCompareBoolCondition
+//----------------------------------------------------------------------//
+
+bool FStateTreeCompareBoolCondition::TestCondition(FStateTreeExecutionContext& Context) const
+{
+	const InstanceDataType& InstanceData = Context.GetInstanceData<InstanceDataType>(*this);
+	
+	return (InstanceData.bLeft == InstanceData.bRight) ^ bInvert;
+}
+
+
+//----------------------------------------------------------------------//
+//  FStateTreeCompareEnumCondition
+//----------------------------------------------------------------------//
+
+bool FStateTreeCompareEnumCondition::TestCondition(FStateTreeExecutionContext& Context) const
+{
+	const InstanceDataType& InstanceData = Context.GetInstanceData<InstanceDataType>(*this);
+
+	return (InstanceData.Left == InstanceData.Right) ^ bInvert;
+}
+
+#if WITH_EDITOR
+void FStateTreeCompareEnumCondition::OnBindingChanged(const FGuid& ID, FStateTreeDataView InstanceData, const FStateTreeEditorPropertyPath& SourcePath, const FStateTreeEditorPropertyPath& TargetPath, const IStateTreeBindingLookup& BindingLookup)
+{
+	if (!TargetPath.IsValid())
+	{
+		return;
+	}
+
+	FStateTreeCompareEnumConditionInstanceData& Instance = InstanceData.GetMutable<FStateTreeCompareEnumConditionInstanceData>();
+
+	// Left has changed, update enums from the leaf property.
+	if (TargetPath.Path.Last() == GET_MEMBER_NAME_STRING_CHECKED(FStateTreeCompareEnumConditionInstanceData, Left))
+	{
+		if (const FProperty* LeafProperty = BindingLookup.GetPropertyPathLeafProperty(SourcePath))
+		{
+			// Handle both old stype namespace enums and new class enum properties.
+			UEnum* NewEnum = nullptr;
+			if (const FByteProperty* ByteProperty = CastField<FByteProperty>(LeafProperty))
+			{
+				NewEnum = ByteProperty->GetIntPropertyEnum();
+			}
+			else if (const FEnumProperty* EnumProperty = CastField<FEnumProperty>(LeafProperty))
+			{
+				NewEnum = EnumProperty->GetEnum();
+			}
+
+			if (Instance.Left.Enum != NewEnum)
+			{
+				Instance.Left.Initialize(NewEnum);
+			}
+		}
+		else
+		{
+			Instance.Left.Initialize(nullptr);
+		}
+
+		if (Instance.Right.Enum != Instance.Left.Enum)
+		{
+			Instance.Right.Initialize(Instance.Left.Enum);
+		}
+	}
+}
+
+#endif
+
+
+//----------------------------------------------------------------------//
+//  FStateTreeCompareDistanceCondition
+//----------------------------------------------------------------------//
+
+bool FStateTreeCompareDistanceCondition::TestCondition(FStateTreeExecutionContext& Context) const
+{
+	const InstanceDataType& InstanceData = Context.GetInstanceData<InstanceDataType>(*this);
+
+	const float Left = FVector::DistSquared(InstanceData.Source, InstanceData.Target);
+	const float Right = FMath::Square(InstanceData.Distance);
+	const bool bResult = UE::StateTree::Conditions::CompareNumbers<float>(Left, Right, Operator);
+	return bResult ^ bInvert;
+}
+
+
+//----------------------------------------------------------------------//
+//  FStateTreeRandomCondition
+//----------------------------------------------------------------------//
+
+bool FStateTreeRandomCondition::TestCondition(FStateTreeExecutionContext& Context) const
+{
+	const InstanceDataType& InstanceData = Context.GetInstanceData<InstanceDataType>(*this);
+
+	return FMath::FRandRange(0.0f, 1.0f) < InstanceData.Threshold;
+}
+
+
+#undef LOCTEXT_NAMESPACE
diff --git a/Engine/Plugins/Runtime/StateTree/Source/StateTreeModule/Private/Conditions/StateTreeCondition_Common.cpp b/Engine/Plugins/Runtime/StateTree/Source/StateTreeModule/Private/Conditions/StateTreeCondition_Common.cpp
deleted file mode 100644
diff --git a/Engine/Plugins/Runtime/StateTree/Source/StateTreeModule/Private/Conditions/StateTreeGameplayTagConditions.cpp b/Engine/Plugins/Runtime/StateTree/Source/StateTreeModule/Private/Conditions/StateTreeGameplayTagConditions.cpp
new file mode 100644
--- /dev/null
+++ b/Engine/Plugins/Runtime/StateTree/Source/StateTreeModule/Private/Conditions/StateTreeGameplayTagConditions.cpp
@@ -0,0 +1,96 @@
+﻿// Copyright Epic Games, Inc. All Rights Reserved.
+
+#include "Conditions/StateTreeGameplayTagConditions.h"
+#include "GameplayTagContainer.h"
+#include "StateTreeExecutionContext.h"
+
+#if WITH_EDITOR
+#define LOCTEXT_NAMESPACE "StateTreeEditor"
+
+namespace UE::StateTree::Conditions
+{
+	FText GetContainerAsText(const FGameplayTagContainer& TagContainer, const int ApproxMaxLength = 60)
+	{
+		FString Combined;
+		for (const FGameplayTag& Tag : TagContainer)
+		{
+			FString TagString = Tag.ToString();
+
+			if (Combined.Len() > 0)
+			{
+				Combined += TEXT(", ");
+			}
+			
+			if (Combined.Len() + TagString.Len() > ApproxMaxLength)
+			{
+				// Overflow
+				if (Combined.Len() == 0)
+				{
+					Combined += TagString.Left(ApproxMaxLength);
+				}
+				Combined += TEXT("...");
+				break;
+			}
+
+			Combined += TagString;
+		}
+
+		return FText::FromString(Combined);
+	}
+	
+}
+
+#endif// WITH_EDITOR
+
+
+//----------------------------------------------------------------------//
+//  FGameplayTagMatchCondition
+//----------------------------------------------------------------------//
+
+bool FGameplayTagMatchCondition::TestCondition(FStateTreeExecutionContext& Context) const
+{
+	const InstanceDataType& InstanceData = Context.GetInstanceData<InstanceDataType>(*this);
+	
+	return (bExactMatch ? InstanceData.TagContainer.HasTagExact(InstanceData.Tag) : InstanceData.TagContainer.HasTag(InstanceData.Tag)) ^ bInvert;
+}
+
+//----------------------------------------------------------------------//
+//  FGameplayTagContainerMatchCondition
+//----------------------------------------------------------------------//
+
+bool FGameplayTagContainerMatchCondition::TestCondition(FStateTreeExecutionContext& Context) const
+{
+	const InstanceDataType& InstanceData = Context.GetInstanceData<InstanceDataType>(*this);
+
+	bool bResult = false;
+	switch (MatchType)
+	{
+	case EGameplayContainerMatchType::Any:
+		bResult = bExactMatch ? InstanceData.TagContainer.HasAnyExact(InstanceData.OtherContainer) : InstanceData.TagContainer.HasAny(InstanceData.OtherContainer);
+		break;
+	case EGameplayContainerMatchType::All:
+		bResult = bExactMatch ? InstanceData.TagContainer.HasAllExact(InstanceData.OtherContainer) : InstanceData.TagContainer.HasAll(InstanceData.OtherContainer);
+		break;
+	default:
+		ensureMsgf(false, TEXT("Unhandled match type %s."), *UEnum::GetValueAsString(MatchType));
+	}
+	
+	return bResult ^ bInvert;
+}
+
+
+//----------------------------------------------------------------------//
+//  FGameplayTagQueryCondition
+//----------------------------------------------------------------------//
+
+bool FGameplayTagQueryCondition::TestCondition(FStateTreeExecutionContext& Context) const
+{
+	const InstanceDataType& InstanceData = Context.GetInstanceData<InstanceDataType>(*this);
+
+	return TagQuery.Matches(InstanceData.TagContainer) ^ bInvert;
+}
+
+
+#if WITH_EDITOR
+#undef LOCTEXT_NAMESPACE
+#endif // WITH_EDITOR
diff --git a/Engine/Plugins/Runtime/StateTree/Source/StateTreeModule/Private/Conditions/StateTreeObjectConditions.cpp b/Engine/Plugins/Runtime/StateTree/Source/StateTreeModule/Private/Conditions/StateTreeObjectConditions.cpp
new file mode 100644
--- /dev/null
+++ b/Engine/Plugins/Runtime/StateTree/Source/StateTreeModule/Private/Conditions/StateTreeObjectConditions.cpp
@@ -0,0 +1,43 @@
+﻿// Copyright Epic Games, Inc. All Rights Reserved.
+
+#include "Conditions/StateTreeObjectConditions.h"
+#include "StateTreeExecutionContext.h"
+
+#define LOCTEXT_NAMESPACE "StateTreeEditor"
+
+
+//----------------------------------------------------------------------//
+//  FStateTreeCondition_ObjectIsValid
+//----------------------------------------------------------------------//
+
+bool FStateTreeObjectIsValidCondition::TestCondition(FStateTreeExecutionContext& Context) const
+{
+	const InstanceDataType& InstanceData = Context.GetInstanceData<InstanceDataType>(*this);
+	const bool bResult = IsValid(InstanceData.Object);
+	return bResult ^ bInvert;
+}
+
+//----------------------------------------------------------------------//
+//  FStateTreeCondition_ObjectEquals
+//----------------------------------------------------------------------//
+
+bool FStateTreeObjectEqualsCondition::TestCondition(FStateTreeExecutionContext& Context) const
+{
+	const InstanceDataType& InstanceData = Context.GetInstanceData<InstanceDataType>(*this);
+	const bool bResult = InstanceData.Left == InstanceData.Right;
+	return bResult ^ bInvert;
+}
+
+//----------------------------------------------------------------------//
+//  FStateTreeCondition_ObjectIsChildOfClass
+//----------------------------------------------------------------------//
+
+bool FStateTreeObjectIsChildOfClassCondition::TestCondition(FStateTreeExecutionContext& Context) const
+{
+	const InstanceDataType& InstanceData = Context.GetInstanceData<InstanceDataType>(*this);
+	const bool bResult = InstanceData.Object && InstanceData.Class
+						&& InstanceData.Object->GetClass()->IsChildOf(InstanceData.Class);
+	return bResult ^ bInvert;
+}
+
+#undef LOCTEXT_NAMESPACE
\ No newline at end of file
diff --git a/Engine/Plugins/Runtime/StateTree/Source/StateTreeModule/Private/Tasks/StateTreeDelayTask.cpp b/Engine/Plugins/Runtime/StateTree/Source/StateTreeModule/Private/Tasks/StateTreeDelayTask.cpp
--- a/Engine/Plugins/Runtime/StateTree/Source/StateTreeModule/Private/Tasks/StateTreeDelayTask.cpp
+++ b/Engine/Plugins/Runtime/StateTree/Source/StateTreeModule/Private/Tasks/StateTreeDelayTask.cpp
@@ -4,25 +4,20 @@
 #include "StateTreeExecutionContext.h"
 #include "StateTreeLinker.h"
 
-bool FStateTreeDelayTask::Link(FStateTreeLinker& Linker)
-{
-	Linker.LinkInstanceDataProperty(DurationHandle, STATETREE_INSTANCEDATA_PROPERTY(InstanceDataType, Duration));
-	Linker.LinkInstanceDataProperty(TimeHandle, STATETREE_INSTANCEDATA_PROPERTY(InstanceDataType, Time));
-	return true;
-}
-
 EStateTreeRunStatus FStateTreeDelayTask::EnterState(FStateTreeExecutionContext& Context, const EStateTreeStateChangeType ChangeType, const FStateTreeTransitionResult& Transition) const
 {
-	float& Time = Context.GetInstanceData(TimeHandle);
-	Time = 0.f;
+	InstanceDataType& InstanceData = Context.GetInstanceData<InstanceDataType>(*this);
+
+	InstanceData.Time = 0.f;
+	
 	return EStateTreeRunStatus::Running;
 }
 
 EStateTreeRunStatus FStateTreeDelayTask::Tick(FStateTreeExecutionContext& Context, const float DeltaTime) const
 {
-	float& Time = Context.GetInstanceData(TimeHandle);
-	const float Duration = Context.GetInstanceData(DurationHandle);
-	Time += DeltaTime;
+	InstanceDataType& InstanceData = Context.GetInstanceData<InstanceDataType>(*this);
+
+	InstanceData.Time += DeltaTime;
 	
-	return Duration <= 0.0f ? EStateTreeRunStatus::Running : (Time < Duration ? EStateTreeRunStatus::Running : EStateTreeRunStatus::Succeeded);
+	return InstanceData.Duration <= 0.0f ? EStateTreeRunStatus::Running : (InstanceData.Time < InstanceData.Duration ? EStateTreeRunStatus::Running : EStateTreeRunStatus::Succeeded);
 }
diff --git a/Engine/Plugins/Runtime/StateTree/Source/StateTreeModule/Private/Tasks/StateTreeDelayTask.h b/Engine/Plugins/Runtime/StateTree/Source/StateTreeModule/Private/Tasks/StateTreeDelayTask.h
--- a/Engine/Plugins/Runtime/StateTree/Source/StateTreeModule/Private/Tasks/StateTreeDelayTask.h
+++ b/Engine/Plugins/Runtime/StateTree/Source/StateTreeModule/Private/Tasks/StateTreeDelayTask.h
@@ -33,10 +33,6 @@ struct STATETREEMODULE_API FStateTreeDelayTask : public FStateTreeTaskCommonBase
 
 	virtual const UStruct* GetInstanceDataType() const override { return InstanceDataType::StaticStruct(); }
 
-	virtual bool Link(FStateTreeLinker& Linker) override;
 	virtual EStateTreeRunStatus EnterState(FStateTreeExecutionContext& Context, const EStateTreeStateChangeType ChangeType, const FStateTreeTransitionResult& Transition) const override;
 	virtual EStateTreeRunStatus Tick(FStateTreeExecutionContext& Context, const float DeltaTime) const override;;
-
-	TStateTreeInstanceDataPropertyHandle<float> DurationHandle;
-	TStateTreeInstanceDataPropertyHandle<float> TimeHandle;
 };
diff --git a/Engine/Plugins/Runtime/StateTree/Source/StateTreeModule/Public/Conditions/GameplayTagConditions.h b/Engine/Plugins/Runtime/StateTree/Source/StateTreeModule/Public/Conditions/GameplayTagConditions.h
deleted file mode 100644
diff --git a/Engine/Plugins/Runtime/StateTree/Source/StateTreeModule/Public/Conditions/StateTreeCommonConditions.h b/Engine/Plugins/Runtime/StateTree/Source/StateTreeModule/Public/Conditions/StateTreeCommonConditions.h
new file mode 100644
--- /dev/null
+++ b/Engine/Plugins/Runtime/StateTree/Source/StateTreeModule/Public/Conditions/StateTreeCommonConditions.h
@@ -0,0 +1,237 @@
+// Copyright Epic Games, Inc. All Rights Reserved.
+
+#pragma once
+
+#include "AITypes.h"
+#include "InstancedStruct.h"
+#include "StateTreeConditionBase.h"
+#include "StateTreeAnyEnum.h"
+#include "StateTreeCommonConditions.generated.h"
+
+/**
+ * Condition comparing two integers.
+ */
+USTRUCT()
+struct STATETREEMODULE_API FStateTreeCompareIntConditionInstanceData
+{
+	GENERATED_BODY()
+	
+	UPROPERTY(EditAnywhere, Category = Input)
+	int32 Left = 0;
+
+	UPROPERTY(EditAnywhere, Category = Parameter)
+	int32 Right = 0;
+};
+STATETREE_POD_INSTANCEDATA(FStateTreeCompareIntConditionInstanceData);
+
+USTRUCT(DisplayName="Integer Compare")
+struct STATETREEMODULE_API FStateTreeCompareIntCondition : public FStateTreeConditionCommonBase
+{
+	GENERATED_BODY()
+
+	typedef FStateTreeCompareIntConditionInstanceData InstanceDataType;
+
+	FStateTreeCompareIntCondition() = default;
+	explicit FStateTreeCompareIntCondition(const EGenericAICheck InOperator, const EStateTreeCompare InInverts = EStateTreeCompare::Default)
+		: bInvert(InInverts == EStateTreeCompare::Invert)
+		, Operator(InOperator)
+	{}
+
+	virtual const UStruct* GetInstanceDataType() const override { return FStateTreeCompareIntConditionInstanceData::StaticStruct(); }
+	virtual bool TestCondition(FStateTreeExecutionContext& Context) const override;
+
+	UPROPERTY(EditAnywhere, Category = Parameter)
+	bool bInvert = false;
+
+	UPROPERTY(EditAnywhere, Category = Parameter)
+	EGenericAICheck Operator = EGenericAICheck::Equal;
+};
+
+/**
+ * Condition comparing two floats.
+ */
+USTRUCT()
+struct STATETREEMODULE_API FStateTreeCompareFloatConditionInstanceData
+{
+	GENERATED_BODY()
+
+	UPROPERTY(EditAnywhere, Category = Input)
+	float Left = 0.0f;
+
+	UPROPERTY(EditAnywhere, Category = Parameter)
+	float Right = 0.0f;
+};
+STATETREE_POD_INSTANCEDATA(FStateTreeCompareFloatConditionInstanceData);
+
+USTRUCT(DisplayName = "Float Compare")
+struct STATETREEMODULE_API FStateTreeCompareFloatCondition : public FStateTreeConditionCommonBase
+{
+	GENERATED_BODY()
+
+	typedef FStateTreeCompareFloatConditionInstanceData InstanceDataType;
+
+	FStateTreeCompareFloatCondition() = default;
+	explicit FStateTreeCompareFloatCondition(const EGenericAICheck InOperator, const EStateTreeCompare InInverts = EStateTreeCompare::Default)
+		: bInvert(InInverts == EStateTreeCompare::Invert)
+		, Operator(InOperator)
+	{}
+
+	virtual const UStruct* GetInstanceDataType() const override { return FStateTreeCompareFloatConditionInstanceData::StaticStruct(); }
+	virtual bool TestCondition(FStateTreeExecutionContext& Context) const override;
+
+	UPROPERTY(EditAnywhere, Category = Parameter)
+	bool bInvert = false;
+
+	UPROPERTY(EditAnywhere, Category = Parameter)
+	EGenericAICheck Operator = EGenericAICheck::Equal;
+};
+
+/**
+ * Condition comparing two booleans.
+ */
+USTRUCT()
+struct STATETREEMODULE_API FStateTreeCompareBoolConditionInstanceData
+{
+	GENERATED_BODY()
+
+	UPROPERTY(EditAnywhere, Category = Input)
+	bool bLeft = false;
+
+	UPROPERTY(EditAnywhere, Category = Parameter)
+	bool bRight = false;
+};
+STATETREE_POD_INSTANCEDATA(FStateTreeCompareBoolConditionInstanceData);
+
+USTRUCT(DisplayName = "Bool Compare")
+struct STATETREEMODULE_API FStateTreeCompareBoolCondition : public FStateTreeConditionCommonBase
+{
+	GENERATED_BODY()
+
+	typedef FStateTreeCompareBoolConditionInstanceData InstanceDataType;
+
+	FStateTreeCompareBoolCondition() = default;
+	explicit FStateTreeCompareBoolCondition(const EStateTreeCompare InInverts)
+		: bInvert(InInverts == EStateTreeCompare::Invert)
+	{}
+
+	FStateTreeCompareBoolCondition(const bool bInInverts)
+		: bInvert(bInInverts)
+	{}
+
+	virtual const UStruct* GetInstanceDataType() const override { return FStateTreeCompareBoolConditionInstanceData::StaticStruct(); }
+	virtual bool TestCondition(FStateTreeExecutionContext& Context) const override;
+
+	UPROPERTY(EditAnywhere, Category = Parameter)
+	bool bInvert = false;
+};
+
+/**
+ * Condition comparing two enums.
+ */
+USTRUCT()
+struct STATETREEMODULE_API FStateTreeCompareEnumConditionInstanceData
+{
+	GENERATED_BODY()
+
+	UPROPERTY(EditAnywhere, Category = Input)
+	FStateTreeAnyEnum Left;
+
+	UPROPERTY(EditAnywhere, Category = Parameter)
+	FStateTreeAnyEnum Right;
+};
+STATETREE_POD_INSTANCEDATA(FStateTreeCompareEnumConditionInstanceData);
+
+USTRUCT(DisplayName = "Enum Compare")
+struct STATETREEMODULE_API FStateTreeCompareEnumCondition : public FStateTreeConditionCommonBase
+{
+	GENERATED_BODY()
+
+	typedef FStateTreeCompareEnumConditionInstanceData InstanceDataType;
+
+	FStateTreeCompareEnumCondition() = default;
+	explicit FStateTreeCompareEnumCondition(const EStateTreeCompare InInverts)
+		: bInvert(InInverts == EStateTreeCompare::Invert)
+	{}
+
+	FStateTreeCompareEnumCondition(const bool bInInverts)
+		: bInvert(bInInverts)
+	{}
+
+	virtual const UStruct* GetInstanceDataType() const override { return FStateTreeCompareEnumConditionInstanceData::StaticStruct(); }
+	virtual bool TestCondition(FStateTreeExecutionContext& Context) const override;
+
+#if WITH_EDITOR
+	virtual void OnBindingChanged(const FGuid& ID, FStateTreeDataView InstanceData, const FStateTreeEditorPropertyPath& SourcePath, const FStateTreeEditorPropertyPath& TargetPath, const IStateTreeBindingLookup& BindingLookup) override;
+#endif
+
+	UPROPERTY(EditAnywhere, Category = Parameter)
+	bool bInvert = false;
+};
+
+/**
+ * Condition comparing distance between two vectors.
+ */
+USTRUCT()
+struct STATETREEMODULE_API FStateTreeCompareDistanceConditionInstanceData
+{
+	GENERATED_BODY()
+
+	UPROPERTY(EditAnywhere, Category = Input)
+	FVector Source = FVector(EForceInit::ForceInitToZero);
+
+	UPROPERTY(EditAnywhere, Category = Parameter)
+	FVector Target = FVector(EForceInit::ForceInitToZero);
+
+	UPROPERTY(EditAnywhere, Category = Parameter)
+	float Distance = 0.0f;
+};
+STATETREE_POD_INSTANCEDATA(FStateTreeCompareDistanceConditionInstanceData);
+
+USTRUCT(DisplayName = "Distance Compare")
+struct STATETREEMODULE_API FStateTreeCompareDistanceCondition : public FStateTreeConditionCommonBase
+{
+	GENERATED_BODY()
+
+	typedef FStateTreeCompareDistanceConditionInstanceData InstanceDataType;
+
+	FStateTreeCompareDistanceCondition() = default;
+	explicit FStateTreeCompareDistanceCondition(const EGenericAICheck InOperator, const EStateTreeCompare InInverts = EStateTreeCompare::Default)
+		: bInvert(InInverts == EStateTreeCompare::Invert)
+		, Operator(InOperator)
+	{}
+	
+	virtual const UStruct* GetInstanceDataType() const override { return FStateTreeCompareDistanceConditionInstanceData::StaticStruct(); }
+	virtual bool TestCondition(FStateTreeExecutionContext& Context) const override;
+
+	UPROPERTY(EditAnywhere, Category = Condition)
+	bool bInvert = false;
+
+	UPROPERTY(EditAnywhere, Category = Condition)
+	EGenericAICheck Operator = EGenericAICheck::Equal;
+};
+
+/**
+* Random condition
+*/
+
+USTRUCT()
+struct STATETREEMODULE_API FStateTreeRandomConditionInstanceData
+{
+	GENERATED_BODY()
+	UPROPERTY(EditAnywhere, Category = Parameter, meta = (ClampMin = "0.0", ClampMax = "1.0", UIMin = "0.0", UIMax = "1.0"))
+	float Threshold = 0.5f;
+};
+STATETREE_POD_INSTANCEDATA(FStateTreeRandomConditionInstanceData);
+
+USTRUCT(DisplayName = "Random")
+struct STATETREEMODULE_API FStateTreeRandomCondition : public FStateTreeConditionBase
+{
+	GENERATED_BODY()
+
+	typedef FStateTreeRandomConditionInstanceData InstanceDataType;
+
+	FStateTreeRandomCondition() = default;
+	
+	virtual const UStruct* GetInstanceDataType() const override { return FStateTreeRandomConditionInstanceData::StaticStruct(); }
+	virtual bool TestCondition(FStateTreeExecutionContext& Context) const override;
+};
diff --git a/Engine/Plugins/Runtime/StateTree/Source/StateTreeModule/Public/Conditions/StateTreeCondition_Common.h b/Engine/Plugins/Runtime/StateTree/Source/StateTreeModule/Public/Conditions/StateTreeCondition_Common.h
deleted file mode 100644
diff --git a/Engine/Plugins/Runtime/StateTree/Source/StateTreeModule/Public/Conditions/StateTreeGameplayTagConditions.h b/Engine/Plugins/Runtime/StateTree/Source/StateTreeModule/Public/Conditions/StateTreeGameplayTagConditions.h
new file mode 100644
--- /dev/null
+++ b/Engine/Plugins/Runtime/StateTree/Source/StateTreeModule/Public/Conditions/StateTreeGameplayTagConditions.h
@@ -0,0 +1,109 @@
+﻿// Copyright Epic Games, Inc. All Rights Reserved.
+
+#pragma once
+
+#include "StateTreeConditionBase.h"
+#include "GameplayTagContainer.h"
+#include "StateTreeGameplayTagConditions.generated.h"
+
+/**
+ * Gameplay Tag match condition.
+ */
+USTRUCT()
+struct STATETREEMODULE_API FGameplayTagMatchConditionInstanceData
+{
+	GENERATED_BODY()
+	
+	UPROPERTY(EditAnywhere, Category = Parameter)
+	FGameplayTagContainer TagContainer;
+
+	UPROPERTY(EditAnywhere, Category = Parameter)
+	FGameplayTag Tag;
+};
+
+USTRUCT(DisplayName="Gameplay Tag Match")
+struct STATETREEMODULE_API FGameplayTagMatchCondition : public FStateTreeConditionCommonBase
+{
+	GENERATED_BODY()
+
+	typedef FGameplayTagMatchConditionInstanceData InstanceDataType;
+
+	FGameplayTagMatchCondition() = default;
+	
+	virtual const UStruct* GetInstanceDataType() const override { return InstanceDataType::StaticStruct(); }
+	virtual bool TestCondition(FStateTreeExecutionContext& Context) const override;
+
+	UPROPERTY(EditAnywhere, Category = Condition)
+	bool bExactMatch = false;
+
+	UPROPERTY(EditAnywhere, Category = Condition)
+	bool bInvert = false;
+};
+
+/**
+* Gameplay Tag Container match condition.
+*/
+USTRUCT()
+struct STATETREEMODULE_API FGameplayTagContainerMatchConditionInstanceData
+{
+	GENERATED_BODY()
+	
+	UPROPERTY(EditAnywhere, Category = Input)
+	FGameplayTagContainer TagContainer;
+
+	UPROPERTY(EditAnywhere, Category = Parameter)
+	FGameplayTagContainer OtherContainer;
+};
+
+USTRUCT(DisplayName="Gameplay Tag Container Match")
+struct STATETREEMODULE_API FGameplayTagContainerMatchCondition : public FStateTreeConditionCommonBase
+{
+	GENERATED_BODY()
+
+	typedef FGameplayTagContainerMatchConditionInstanceData InstanceDataType;
+
+	FGameplayTagContainerMatchCondition() = default;
+	
+	virtual const UStruct* GetInstanceDataType() const override { return InstanceDataType::StaticStruct(); }
+	virtual bool TestCondition(FStateTreeExecutionContext& Context) const override;
+
+	UPROPERTY(EditAnywhere, Category = Condition)
+	EGameplayContainerMatchType MatchType = EGameplayContainerMatchType::Any;
+
+	UPROPERTY(EditAnywhere, Category = Condition)
+	bool bExactMatch = false;
+
+	UPROPERTY(EditAnywhere, Category = Condition)
+	bool bInvert = false;
+};
+
+/**
+* Gameplay Tag Query match condition.
+*/
+USTRUCT()
+struct STATETREEMODULE_API FGameplayTagQueryConditionInstanceData
+{
+	GENERATED_BODY()
+	
+	UPROPERTY(EditAnywhere, Category = Input)
+	FGameplayTagContainer TagContainer;
+};
+
+USTRUCT(DisplayName="Gameplay Tag Query")
+struct STATETREEMODULE_API FGameplayTagQueryCondition : public FStateTreeConditionCommonBase
+{
+	GENERATED_BODY()
+
+	typedef FGameplayTagQueryConditionInstanceData InstanceDataType;
+
+	FGameplayTagQueryCondition() = default;
+	
+	virtual const UStruct* GetInstanceDataType() const override { return InstanceDataType::StaticStruct(); }
+	virtual bool TestCondition(FStateTreeExecutionContext& Context) const override;
+
+	UPROPERTY(EditAnywhere, Category = Condition)
+	FGameplayTagQuery TagQuery;
+
+	UPROPERTY(EditAnywhere, Category = Condition)
+	bool bInvert = false;
+};
\ No newline at end of file
diff --git a/Engine/Plugins/Runtime/StateTree/Source/StateTreeModule/Public/Conditions/StateTreeObjectConditions.h b/Engine/Plugins/Runtime/StateTree/Source/StateTreeModule/Public/Conditions/StateTreeObjectConditions.h
new file mode 100644
--- /dev/null
+++ b/Engine/Plugins/Runtime/StateTree/Source/StateTreeModule/Public/Conditions/StateTreeObjectConditions.h
@@ -0,0 +1,106 @@
+﻿// Copyright Epic Games, Inc. All Rights Reserved.
+
+#pragma once
+
+#include "UObject/Object.h"
+#include "StateTreeConditionBase.h"
+#include "StateTreeObjectConditions.generated.h"
+
+/**
+ * Condition testing if specified object is valid.
+ */
+USTRUCT()
+struct STATETREEMODULE_API FStateTreeObjectIsValidConditionInstanceData
+{
+	GENERATED_BODY()
+	
+	UPROPERTY(EditAnywhere, Category = Input)
+	TObjectPtr<UObject> Object = nullptr;
+};
+
+USTRUCT(DisplayName="Object Is Valid")
+struct STATETREEMODULE_API FStateTreeObjectIsValidCondition : public FStateTreeConditionCommonBase
+{
+	GENERATED_BODY()
+
+	typedef FStateTreeObjectIsValidConditionInstanceData InstanceDataType;
+
+	FStateTreeObjectIsValidCondition() = default;
+	explicit FStateTreeObjectIsValidCondition(const EStateTreeCompare InInverts)
+		: bInvert(InInverts == EStateTreeCompare::Invert)
+	{}
+
+	virtual const UStruct* GetInstanceDataType() const override { return InstanceDataType::StaticStruct(); }
+	virtual bool TestCondition(FStateTreeExecutionContext& Context) const override;
+
+	UPROPERTY(EditAnywhere, Category = Parameter)
+	bool bInvert = false;
+};
+
+/**
+ * Condition testing if two object pointers point to the same object.
+ */
+USTRUCT()
+struct STATETREEMODULE_API FStateTreeObjectEqualsConditionInstanceData
+{
+	GENERATED_BODY()
+	
+	UPROPERTY(EditAnywhere, Category = Input)
+	TObjectPtr<UObject> Left = nullptr;
+
+	UPROPERTY(EditAnywhere, Category = Parameter)
+	TObjectPtr<UObject> Right = nullptr;
+};
+
+USTRUCT(DisplayName="Object Equals")
+struct STATETREEMODULE_API FStateTreeObjectEqualsCondition : public FStateTreeConditionCommonBase
+{
+	GENERATED_BODY()
+
+	typedef FStateTreeObjectEqualsConditionInstanceData InstanceDataType;
+
+	FStateTreeObjectEqualsCondition() = default;
+	explicit FStateTreeObjectEqualsCondition(const EStateTreeCompare InInverts)
+		: bInvert(InInverts == EStateTreeCompare::Invert)
+	{}
+
+	virtual const UStruct* GetInstanceDataType() const override { return InstanceDataType::StaticStruct(); }
+	virtual bool TestCondition(FStateTreeExecutionContext& Context) const override;
+
+	UPROPERTY(EditAnywhere, Category = Parameter)
+	bool bInvert = false;
+};
+
+/**
+ * Condition testing if object is child of specified class.
+ */
+USTRUCT()
+struct STATETREEMODULE_API FStateTreeObjectIsChildOfClassConditionInstanceData
+{
+	GENERATED_BODY()
+	
+	UPROPERTY(EditAnywhere, Category = Input)
+	TObjectPtr<UObject> Object = nullptr;
+
+	UPROPERTY(EditAnywhere, Category = Parameter)
+	TObjectPtr<UClass> Class = nullptr;
+};
+
+USTRUCT(DisplayName="Object Class Is")
+struct STATETREEMODULE_API FStateTreeObjectIsChildOfClassCondition : public FStateTreeConditionCommonBase
+{
+	GENERATED_BODY()
+
+	typedef FStateTreeObjectIsChildOfClassConditionInstanceData InstanceDataType;
+
+	FStateTreeObjectIsChildOfClassCondition() = default;
+	explicit FStateTreeObjectIsChildOfClassCondition(const EStateTreeCompare InInverts)
+		: bInvert(InInverts == EStateTreeCompare::Invert)
+	{}
+
+	virtual const UStruct* GetInstanceDataType() const override { return InstanceDataType::StaticStruct(); }
+	virtual bool TestCondition(FStateTreeExecutionContext& Context) const override;
+
+	UPROPERTY(EditAnywhere, Category = Parameter)
+	bool bInvert = false;
+};
diff --git a/Engine/Plugins/Runtime/StateTree/Source/StateTreeModule/Public/StateTreeConditionBase.h b/Engine/Plugins/Runtime/StateTree/Source/StateTreeModule/Public/StateTreeConditionBase.h
--- a/Engine/Plugins/Runtime/StateTree/Source/StateTreeModule/Public/StateTreeConditionBase.h
+++ b/Engine/Plugins/Runtime/StateTree/Source/StateTreeModule/Public/StateTreeConditionBase.h
@@ -15,6 +15,12 @@ struct IStateTreeBindingLookup;
 struct FStateTreeEditorPropertyPath;
 #endif
 
+enum class EStateTreeCompare : uint8
+{
+	Default,
+	Invert,
+};
+
 /**
  * Base struct for all conditions.
  */
diff --git a/Engine/Plugins/Runtime/StateTree/Source/StateTreeTestSuite/Private/StateTreeTest.cpp b/Engine/Plugins/Runtime/StateTree/Source/StateTreeTestSuite/Private/StateTreeTest.cpp
--- a/Engine/Plugins/Runtime/StateTree/Source/StateTreeTestSuite/Private/StateTreeTest.cpp
+++ b/Engine/Plugins/Runtime/StateTree/Source/StateTreeTestSuite/Private/StateTreeTest.cpp
@@ -5,7 +5,7 @@
 #include "StateTreeEditorData.h"
 #include "StateTreeState.h"
 #include "StateTreeCompiler.h"
-#include "Conditions/StateTreeCondition_Common.h"
+#include "Conditions/StateTreeCommonConditions.h"
 #include "StateTree.h"
 #include "StateTreeTestTypes.h"
 #include "StateTreeExecutionContext.h"
@@ -48,7 +48,7 @@ struct FStateTreeTest_MakeAndBakeStateTree : FAITestBase
 		auto& TaskB1 = StateA.AddTask<FTestTask_B>();
 		EditorData.AddPropertyBinding(FStateTreeEditorPropertyPath(EvalA.ID, TEXT("IntA")), FStateTreeEditorPropertyPath(TaskB1.ID, TEXT("IntB")));
 
-		auto& IntCond = StateA.AddEnterCondition<FStateTreeCondition_CompareInt>(EGenericAICheck::Less);
+		auto& IntCond = StateA.AddEnterCondition<FStateTreeCompareIntCondition>(EGenericAICheck::Less);
 		IntCond.GetInstance().Right = 2;
 
 		EditorData.AddPropertyBinding(FStateTreeEditorPropertyPath(EvalA.ID, TEXT("IntA")), FStateTreeEditorPropertyPath(IntCond.ID, TEXT("Left")));
@@ -59,8 +59,8 @@ struct FStateTreeTest_MakeAndBakeStateTree : FAITestBase
 		auto& TaskB2 = StateB.AddTask<FTestTask_B>();
 		EditorData.AddPropertyBinding(FStateTreeEditorPropertyPath(EvalA.ID, TEXT("bBoolA")), FStateTreeEditorPropertyPath(TaskB2.ID, TEXT("bBoolB")));
 
-		FStateTreeTransition& Trans = StateB.AddTransition(EStateTreeTransitionEvent::OnCondition, EStateTreeTransitionType::GotoState, &Root);
-		auto& TransFloatCond = Trans.AddCondition<FStateTreeCondition_CompareFloat>(EGenericAICheck::Less);
+		FStateTreeTransition& Trans = StateB.AddTransition({}, EStateTreeTransitionType::GotoState, &Root);
+		auto& TransFloatCond = Trans.AddCondition<FStateTreeCompareFloatCondition>(EGenericAICheck::Less);
 		TransFloatCond.GetInstance().Right = 13.0f;
 		EditorData.AddPropertyBinding(FStateTreeEditorPropertyPath(EvalA.ID, TEXT("FloatA")), FStateTreeEditorPropertyPath(TransFloatCond.ID, TEXT("Left")));
 
@@ -105,16 +105,16 @@ struct FStateTreeTest_WanderLoop : FAITestBase
 		//      |- UseSmartObjectOnLane
 		auto& ReserveSOTask = UseSmartObjectOnLane.AddTask<FTestTask_ReserveSmartObject>(FName(TEXT("ReserveSOTask")));
 		EditorData.AddPropertyBinding(FStateTreeEditorPropertyPath(SmartObjectEval.ID, TEXT("PotentialSmartObjects")), FStateTreeEditorPropertyPath(ReserveSOTask.ID, TEXT("PotentialSmartObjects")));
-		auto& ReserveHasSmartObjects = UseSmartObjectOnLane.AddEnterCondition<FStateTreeCondition_CompareBool>();
+		auto& ReserveHasSmartObjects = UseSmartObjectOnLane.AddEnterCondition<FStateTreeCompareBoolCondition>();
 		ReserveHasSmartObjects.GetInstance().bRight = true;
 		EditorData.AddPropertyBinding(FStateTreeEditorPropertyPath(SmartObjectEval.ID, TEXT("bHasSmartObjects")), FStateTreeEditorPropertyPath(ReserveHasSmartObjects.ID, TEXT("bLeft")));
-		UseSmartObjectOnLane.AddTransition(EStateTreeTransitionEvent::OnCompleted, EStateTreeTransitionType::GotoState, &Wander); // This catches child states too
+		UseSmartObjectOnLane.AddTransition(EStateTreeTransitionTrigger::OnStateCompleted, EStateTreeTransitionType::GotoState, &Wander); // This catches child states too
 
 		//      |  |- WalkToSO
 		auto& MoveToSOTask = WalkToSO.AddTask<FTestTask_MoveTo>(FName(TEXT("MoveToSOTask")));
 		MoveToSOTask.GetItem().TicksToCompletion = 2;
 		EditorData.AddPropertyBinding(FStateTreeEditorPropertyPath(ReserveSOTask.ID, TEXT("ReservedSmartObjectLocation")), FStateTreeEditorPropertyPath(MoveToSOTask.ID, TEXT("MoveLocation")));
-		WalkToSO.AddTransition(EStateTreeTransitionEvent::OnSucceeded, EStateTreeTransitionType::NextState);
+		WalkToSO.AddTransition(EStateTreeTransitionTrigger::OnStateCompleted, EStateTreeTransitionType::NextState);
 
 		//      |  \- UseSO
 		auto& UseSOTask = UseSO.AddTask<FTestTask_UseSmartObject>(FName(TEXT("UseSOTask")));
@@ -126,15 +126,15 @@ struct FStateTreeTest_WanderLoop : FAITestBase
 		auto& MoveAlongLaneTask = WalkAlongLane.AddTask<FTestTask_MoveTo>(FName(TEXT("MoveAlongLaneTask")));
 		MoveAlongLaneTask.GetItem().TicksToCompletion = 2;
 		EditorData.AddPropertyBinding(FStateTreeEditorPropertyPath(WanderEval.ID, TEXT("WanderLocation")), FStateTreeEditorPropertyPath(MoveAlongLaneTask.ID, TEXT("MoveLocation")));
-		auto& MoveHasWanderLoc = WalkAlongLane.AddEnterCondition<FStateTreeCondition_CompareBool>();
+		auto& MoveHasWanderLoc = WalkAlongLane.AddEnterCondition<FStateTreeCompareBoolCondition>();
 		MoveHasWanderLoc.GetInstance().bRight = true;
 		EditorData.AddPropertyBinding(FStateTreeEditorPropertyPath(WanderEval.ID, TEXT("bHasWanderLocation")), FStateTreeEditorPropertyPath(MoveHasWanderLoc.ID, TEXT("bLeft")));
-		WalkAlongLane.AddTransition(EStateTreeTransitionEvent::OnCompleted, EStateTreeTransitionType::GotoState, &Wander);
+		WalkAlongLane.AddTransition(EStateTreeTransitionTrigger::OnStateCompleted, EStateTreeTransitionType::GotoState, &Wander);
 
 		//      \- StandOnLane
 		auto& StandTask = StandOnLane.AddTask<FTestTask_Stand>(FName(TEXT("StandTask")));
 		StandTask.GetItem().TicksToCompletion = 2;
-		StandOnLane.AddTransition(EStateTreeTransitionEvent::OnCompleted, EStateTreeTransitionType::GotoState, &Wander);
+		StandOnLane.AddTransition(EStateTreeTransitionTrigger::OnStateCompleted, EStateTreeTransitionType::GotoState, &Wander);
 
 		FStateTreeCompilerLog Log;
 		FStateTreeCompiler Compiler(Log);
