StateTree: editor fixes
- enable binding extension on evaluator
- fix schema outer (schema would not get saved)
- removed need to tag bindable properties as parameters
- disallow binding to statetree nodes directly

diff --git a/Engine/Plugins/Runtime/StateTree/Source/StateTreeEditorModule/Private/Customizations/StateTreeBindingExtension.cpp b/Engine/Plugins/Runtime/StateTree/Source/StateTreeEditorModule/Private/Customizations/StateTreeBindingExtension.cpp
--- a/Engine/Plugins/Runtime/StateTree/Source/StateTreeEditorModule/Private/Customizations/StateTreeBindingExtension.cpp
+++ b/Engine/Plugins/Runtime/StateTree/Source/StateTreeEditorModule/Private/Customizations/StateTreeBindingExtension.cpp
@@ -13,6 +13,7 @@
 #include "StateTreeAnyEnum.h"
 #include "StateTreePropertyBindingCompiler.h"
 #include "StateTreeCompiler.h"
+#include "Blueprint/StateTreeItemBlueprintBase.h"
 #include "PropertyEditor/Private/PropertyNode.h"
 
 #define LOCTEXT_NAMESPACE "StateTreeEditor"
@@ -89,12 +90,8 @@ EStateTreePropertyUsage ParsePropertyUsage(TSharedPtr<const IPropertyHandle> InP
 	{
 		return EStateTreePropertyUsage::Output;
 	}
-	if (Category == TEXT("Parameter"))
-	{
-		return EStateTreePropertyUsage::Parameter;
-	}
-	
-	return EStateTreePropertyUsage::Invalid;
+
+	return EStateTreePropertyUsage::Parameter;
 }
 
 FText GetSectionNameFromDataSource(const EStateTreeBindableStructSource Source)
@@ -264,6 +261,13 @@ void FStateTreeBindingExtension::ExtendWidgetRow(FDetailWidgetRow& InWidgetRow,
 
 	Args.OnCanBindToContextStruct = FOnCanBindToContextStruct::CreateLambda([InPropertyHandle](const UStruct* InStruct)
 		{
+			// Do not allow to bind directly StateTree nodes
+			if (InStruct->IsChildOf(UStateTreeItemBlueprintBase::StaticClass())
+				|| InStruct->IsChildOf(FStateTreeNodeBase::StaticStruct()))
+			{
+				return false;
+			}
+		
 			if (const FStructProperty* StructProperty = CastField<FStructProperty>(InPropertyHandle->GetProperty()))
 			{
 				return StructProperty->Struct == InStruct;
diff --git a/Engine/Plugins/Runtime/StateTree/Source/StateTreeEditorModule/Private/StateTreeEditor.cpp b/Engine/Plugins/Runtime/StateTree/Source/StateTreeEditorModule/Private/StateTreeEditor.cpp
--- a/Engine/Plugins/Runtime/StateTree/Source/StateTreeEditorModule/Private/StateTreeEditor.cpp
+++ b/Engine/Plugins/Runtime/StateTree/Source/StateTreeEditorModule/Private/StateTreeEditor.cpp
@@ -310,6 +310,8 @@ TSharedRef<SDockTab> FStateTreeEditor::SpawnTab_AssetDetails(const FSpawnTabArgs
 	AssetDetailsView->SetObject(StateTree ? StateTree->EditorData : nullptr);
 	AssetDetailsView->OnFinishedChangingProperties().AddSP(this, &FStateTreeEditor::OnAssetFinishedChangingProperties);
 
+	AssetDetailsView->SetExtensionHandler(MakeShared<FStateTreeBindingExtension>());
+
 	TSharedRef<SDockTab> SpawnedTab = SNew(SDockTab)
 		.Label(NSLOCTEXT("StateTreeEditor", "AssetDetailsTabLabel", "StateTree"))
 		[
diff --git a/Engine/Plugins/Runtime/StateTree/Source/StateTreeEditorModule/Private/StateTreeFactory.cpp b/Engine/Plugins/Runtime/StateTree/Source/StateTreeEditorModule/Private/StateTreeFactory.cpp
--- a/Engine/Plugins/Runtime/StateTree/Source/StateTreeEditorModule/Private/StateTreeFactory.cpp
+++ b/Engine/Plugins/Runtime/StateTree/Source/StateTreeEditorModule/Private/StateTreeFactory.cpp
@@ -114,7 +114,7 @@ UObject* UStateTreeFactory::FactoryCreateNew(UClass* Class, UObject* InParent, F
 	// Create and init new editor data.
 	UStateTreeEditorData* EditorData = NewObject<UStateTreeEditorData>(NewStateTree, FName(), RF_Transactional);
 	EditorData->AddRootState();
-	EditorData->Schema = NewObject<UStateTreeSchema>(this, StateTreeSchemaClass);
+	EditorData->Schema = NewObject<UStateTreeSchema>(EditorData, StateTreeSchemaClass);
 	
 	NewStateTree->EditorData = EditorData;
 
