StateTree: Fix potential nullptr reference

diff --git a/Engine/Plugins/Runtime/StateTree/Source/StateTreeEditorModule/Private/Customizations/StateTreeBindingExtension.cpp b/Engine/Plugins/Runtime/StateTree/Source/StateTreeEditorModule/Private/Customizations/StateTreeBindingExtension.cpp
--- a/Engine/Plugins/Runtime/StateTree/Source/StateTreeEditorModule/Private/Customizations/StateTreeBindingExtension.cpp
+++ b/Engine/Plugins/Runtime/StateTree/Source/StateTreeEditorModule/Private/Customizations/StateTreeBindingExtension.cpp
@@ -262,10 +262,13 @@ void FStateTreeBindingExtension::ExtendWidgetRow(FDetailWidgetRow& InWidgetRow,
 	Args.OnCanBindToContextStruct = FOnCanBindToContextStruct::CreateLambda([InPropertyHandle](const UStruct* InStruct)
 		{
 			// Do not allow to bind directly StateTree nodes
-			if (InStruct->IsChildOf(UStateTreeItemBlueprintBase::StaticClass())
-				|| InStruct->IsChildOf(FStateTreeNodeBase::StaticStruct()))
+			if (InStruct != nullptr)
 			{
-				return false;
+				if (InStruct->IsChildOf(UStateTreeItemBlueprintBase::StaticClass())
+					|| InStruct->IsChildOf(FStateTreeNodeBase::StaticStruct()))
+				{
+					return false;
+				}
 			}
 		
 			if (const FStructProperty* StructProperty = CastField<FStructProperty>(InPropertyHandle->GetProperty()))
