StateTree Events
- Added StateTree events, which allows e.g. transitions based on internal or external events
- Clarified transition nomenclature
- Update UI to support transition events
- BP support for events
- Renamed UStateTreeItemBase to UStateTreeNodeBase for consistency with FStateTreeNodeBase

diff --git a/Engine/Plugins/Runtime/GameplayStateTree/Source/GameplayStateTreeModule/Private/Components/StateTreeComponent.cpp b/Engine/Plugins/Runtime/GameplayStateTree/Source/GameplayStateTreeModule/Private/Components/StateTreeComponent.cpp
--- a/Engine/Plugins/Runtime/GameplayStateTree/Source/GameplayStateTreeModule/Private/Components/StateTreeComponent.cpp
+++ b/Engine/Plugins/Runtime/GameplayStateTree/Source/GameplayStateTreeModule/Private/Components/StateTreeComponent.cpp
@@ -319,6 +319,18 @@ void UStateTreeComponent::OnGameplayTaskInitialized(UGameplayTask& Task)
 	}
 }
 
+void UStateTreeComponent::SendStateTreeEvent(const FStateTreeEvent& Event)
+{
+	if (!bIsRunning)
+	{
+		STATETREE_LOG(Warning, TEXT("%s: Trying to send even to a StateTree that is not started yet."), ANSI_TO_TCHAR(__FUNCTION__));
+		return;
+	}
+
+	StateTreeContext.SendExternalEvent(Event);
+}
+
+
 #if WITH_GAMEPLAY_DEBUGGER
 FString UStateTreeComponent::GetDebugInfoString() const
 {
diff --git a/Engine/Plugins/Runtime/GameplayStateTree/Source/GameplayStateTreeModule/Public/Components/StateTreeComponent.h b/Engine/Plugins/Runtime/GameplayStateTree/Source/GameplayStateTreeModule/Public/Components/StateTreeComponent.h
--- a/Engine/Plugins/Runtime/GameplayStateTree/Source/GameplayStateTreeModule/Public/Components/StateTreeComponent.h
+++ b/Engine/Plugins/Runtime/GameplayStateTree/Source/GameplayStateTreeModule/Public/Components/StateTreeComponent.h
@@ -45,6 +45,10 @@ public:
 	virtual void OnGameplayTaskInitialized(UGameplayTask& Task) override;
 	// END IGameplayTaskOwnerInterface
 
+	/** Sends event to the running StateTree. */
+	UFUNCTION(BlueprintCallable, Category = "Gameplay|StateTree")
+	void SendStateTreeEvent(const FStateTreeEvent& Event);
+	
 #if WITH_GAMEPLAY_DEBUGGER
 	virtual FString GetDebugInfoString() const override;
 #endif // WITH_GAMEPLAY_DEBUGGER
diff --git a/Engine/Plugins/Runtime/StateTree/Config/BaseStateTree.ini b/Engine/Plugins/Runtime/StateTree/Config/BaseStateTree.ini
--- a/Engine/Plugins/Runtime/StateTree/Config/BaseStateTree.ini
+++ b/Engine/Plugins/Runtime/StateTree/Config/BaseStateTree.ini
@@ -44,4 +44,12 @@
 +StructRedirects=(OldName="StateTreeCondition_CompareDistance",NewName="/script/StateTreeModule.StateTreeCompareDistanceCondition")
 
 +StructRedirects=(OldName="StateTreeCondition_RandomInstanceData",NewName="/script/StateTreeModule.StateTreeRandomConditionInstanceData")
-+StructRedirects=(OldName="StateTreeCondition_Random",NewName="/script/StateTreeModule.StateTreeRandomCondition")
\ No newline at end of file
++StructRedirects=(OldName="StateTreeCondition_Random",NewName="/script/StateTreeModule.StateTreeRandomCondition")
+
++EnumRedirects=(OldName="EStateTreeTransitionEvent",NewName="/Script/StateTreeModule.EStateTreeTransitionTrigger")
+
++EnumRedirects=(OldName="EStateTreeTransitionTrigger",ValueChanges=(("OnCompleted", "OnStateCompleted"), ("OnSucceeded", "OnStateSucceeded"), ("OnFailed", "OnStateFailed"), ("OnCondition", "OnTick")))
+
++PropertyRedirects=(OldName="StateTreeTransition.Event",NewName="StateTreeTransition.Trigger")
++PropertyRedirects=(OldName="CompactStateTransition.Event",NewName="CompactStateTransition.Trigger")
++ClassRedirects=(OldName="/Script/StateTreeModule.StateTreeItemBlueprintBase",NewName="/Script/StateTreeModule.StateTreeNodeBlueprintBase")
diff --git a/Engine/Plugins/Runtime/StateTree/Source/StateTreeEditorModule/Private/Customizations/StateTreeBindingExtension.cpp b/Engine/Plugins/Runtime/StateTree/Source/StateTreeEditorModule/Private/Customizations/StateTreeBindingExtension.cpp
--- a/Engine/Plugins/Runtime/StateTree/Source/StateTreeEditorModule/Private/Customizations/StateTreeBindingExtension.cpp
+++ b/Engine/Plugins/Runtime/StateTree/Source/StateTreeEditorModule/Private/Customizations/StateTreeBindingExtension.cpp
@@ -13,7 +13,7 @@
 #include "StateTreeAnyEnum.h"
 #include "StateTreePropertyBindingCompiler.h"
 #include "StateTreeCompiler.h"
-#include "Blueprint/StateTreeItemBlueprintBase.h"
+#include "Blueprint/StateTreeNodeBlueprintBase.h"
 #include "PropertyEditor/Private/PropertyNode.h"
 
 #define LOCTEXT_NAMESPACE "StateTreeEditor"
@@ -264,7 +264,7 @@ void FStateTreeBindingExtension::ExtendWidgetRow(FDetailWidgetRow& InWidgetRow,
 			// Do not allow to bind directly StateTree nodes
 			if (InStruct != nullptr)
 			{
-				if (InStruct->IsChildOf(UStateTreeItemBlueprintBase::StaticClass())
+				if (InStruct->IsChildOf(UStateTreeNodeBlueprintBase::StaticClass())
 					|| InStruct->IsChildOf(FStateTreeNodeBase::StaticStruct()))
 				{
 					return false;
diff --git a/Engine/Plugins/Runtime/StateTree/Source/StateTreeEditorModule/Private/Customizations/StateTreeTransitionDetails.cpp b/Engine/Plugins/Runtime/StateTree/Source/StateTreeEditorModule/Private/Customizations/StateTreeTransitionDetails.cpp
--- a/Engine/Plugins/Runtime/StateTree/Source/StateTreeEditorModule/Private/Customizations/StateTreeTransitionDetails.cpp
+++ b/Engine/Plugins/Runtime/StateTree/Source/StateTreeEditorModule/Private/Customizations/StateTreeTransitionDetails.cpp
@@ -24,7 +24,8 @@ void FStateTreeTransitionDetails::CustomizeHeader(TSharedRef<class IPropertyHand
 	StructProperty = StructPropertyHandle;
 	PropUtils = StructCustomizationUtils.GetPropertyUtilities().Get();
 
-	EventProperty = StructProperty->GetChildHandle(TEXT("Event"));
+	TriggerProperty = StructProperty->GetChildHandle(TEXT("Trigger"));
+	EventTagProperty = StructProperty->GetChildHandle(TEXT("EventTag"));
 	StateProperty = StructProperty->GetChildHandle(TEXT("State"));
 	GateDelayProperty = StructProperty->GetChildHandle(TEXT("GateDelay"));
 	ConditionsProperty = StructProperty->GetChildHandle(TEXT("Conditions"));
@@ -53,11 +54,25 @@ void FStateTreeTransitionDetails::CustomizeHeader(TSharedRef<class IPropertyHand
 
 void FStateTreeTransitionDetails::CustomizeChildren(TSharedRef<class IPropertyHandle> StructPropertyHandle, class IDetailChildrenBuilder& StructBuilder, IPropertyTypeCustomizationUtils& StructCustomizationUtils)
 {
-	if (EventProperty)
+	if (TriggerProperty)
 	{
-		StructBuilder.AddProperty(EventProperty.ToSharedRef());
+		StructBuilder.AddProperty(TriggerProperty.ToSharedRef());
 	}
 
+	if (EventTagProperty)
+	{
+		IDetailPropertyRow& EventTagRow = StructBuilder.AddProperty(EventTagProperty.ToSharedRef());
+		EventTagRow.Visibility(TAttribute<EVisibility>::Create(TAttribute<EVisibility>::FGetter::CreateLambda([this]()
+		{
+			uint8 TriggerValue = 0; 
+			if (TriggerProperty.IsValid())
+			{
+				TriggerProperty->GetValue(TriggerValue);
+			}
+			return TriggerValue == (uint8)EStateTreeTransitionTrigger::OnEvent ? EVisibility::Visible : EVisibility::Collapsed;
+		})));
+	}
+	
 	if (GateDelayProperty)
 	{
 		StructBuilder.AddProperty(GateDelayProperty.ToSharedRef());
diff --git a/Engine/Plugins/Runtime/StateTree/Source/StateTreeEditorModule/Private/Customizations/StateTreeTransitionDetails.h b/Engine/Plugins/Runtime/StateTree/Source/StateTreeEditorModule/Private/Customizations/StateTreeTransitionDetails.h
--- a/Engine/Plugins/Runtime/StateTree/Source/StateTreeEditorModule/Private/Customizations/StateTreeTransitionDetails.h
+++ b/Engine/Plugins/Runtime/StateTree/Source/StateTreeEditorModule/Private/Customizations/StateTreeTransitionDetails.h
@@ -28,7 +28,8 @@ private:
 
 	FText GetDescription() const;
 
-	TSharedPtr<IPropertyHandle> EventProperty;
+	TSharedPtr<IPropertyHandle> TriggerProperty;
+	TSharedPtr<IPropertyHandle> EventTagProperty;
 	TSharedPtr<IPropertyHandle> StateProperty;
 	TSharedPtr<IPropertyHandle> GateDelayProperty;
 	TSharedPtr<IPropertyHandle> ConditionsProperty;
diff --git a/Engine/Plugins/Runtime/StateTree/Source/StateTreeEditorModule/Private/SStateTreeViewRow.cpp b/Engine/Plugins/Runtime/StateTree/Source/StateTreeEditorModule/Private/SStateTreeViewRow.cpp
--- a/Engine/Plugins/Runtime/StateTree/Source/StateTreeEditorModule/Private/SStateTreeViewRow.cpp
+++ b/Engine/Plugins/Runtime/StateTree/Source/StateTreeEditorModule/Private/SStateTreeViewRow.cpp
@@ -29,18 +29,6 @@
 
 #define LOCTEXT_NAMESPACE "StateTreeEditor"
 
-namespace UE { namespace StateTreeEditor {
-
-static void AddUnique(TArray<FText>& Array, const FText& NewItem)
-{
-	if (!Array.FindByPredicate([&NewItem](const FText& Item) { return Item.IdenticalTo(NewItem); }))
-	{
-		Array.Add(NewItem);
-	}
-}
-	
-}} // UE::StateTreeEditor
-
 void SStateTreeViewRow::Construct(const FArguments& InArgs, const TSharedRef<STableViewBase>& InOwnerTableView, UStateTreeState* InState, const TSharedPtr<SScrollBox>& ViewBox, TSharedRef<FStateTreeViewModel> InStateTreeViewModel)
 {
 	StateTreeViewModel = InStateTreeViewModel;
@@ -132,11 +120,14 @@ void SStateTreeViewRow::Construct(const FArguments& InArgs, const TSharedRef<STa
 						[
 							SAssignNew(NameTextBlock, SInlineEditableTextBlock)
 							.Style(FStateTreeEditorStyle::Get(), "StateTree.State.TitleInlineEditableText")
-							.OnVerifyTextChanged(this, &SStateTreeViewRow::VerifyNodeTextChanged)
+							.OnVerifyTextChanged_Lambda([](const FText& NewLabel, FText& OutErrorMessage)
+							{
+								return !NewLabel.IsEmptyOrWhitespace();
+							})
 							.OnTextCommitted(this, &SStateTreeViewRow::HandleNodeLabelTextCommitted)
 							.Text(this, &SStateTreeViewRow::GetStateDesc)
 							.Clipping(EWidgetClipping::ClipToBounds)
-							.IsSelected(this, &SStateTreeViewRow::IsSelected)
+							.IsSelected(this, &SStateTreeViewRow::IsStateSelected)
 						]
 					]
 				]
@@ -370,7 +361,7 @@ void SStateTreeViewRow::Construct(const FArguments& InArgs, const TSharedRef<STa
 	];
 }
 
-void SStateTreeViewRow::RequestRename()
+void SStateTreeViewRow::RequestRename() const
 {
 	if (NameTextBlock)
 	{
@@ -518,25 +509,26 @@ FText SStateTreeViewRow::GetLinkedStateDesc() const
 	return FText::GetEmpty();
 }
 
-bool SStateTreeViewRow::HasParentTransitionForEvent(const UStateTreeState& State, const EStateTreeTransitionEvent Event) const
+bool SStateTreeViewRow::HasParentTransitionForTrigger(const UStateTreeState& State, const EStateTreeTransitionTrigger Trigger) const
 {
-	EStateTreeTransitionEvent CombinedEvents = EStateTreeTransitionEvent::None;
+	EStateTreeTransitionTrigger CombinedTrigger = EStateTreeTransitionTrigger::None;
 	for (const UStateTreeState* ParentState = State.Parent; ParentState != nullptr; ParentState = ParentState->Parent)
 	{
 		for (const FStateTreeTransition& Transition : ParentState->Transitions)
 		{
-			CombinedEvents |= Transition.Event;
+			CombinedTrigger |= Transition.Trigger;
 		}
 	}
-	return EnumHasAllFlags(CombinedEvents, Event);
+	return EnumHasAllFlags(CombinedTrigger, Trigger);
 }
 
-FText SStateTreeViewRow::GetTransitionsDesc(const UStateTreeState& State, const EStateTreeTransitionEvent Event) const
+FText SStateTreeViewRow::GetTransitionsDesc(const UStateTreeState& State, const EStateTreeTransitionTrigger Trigger, const bool bUseMask) const
 {
 	TArray<FText> DescItems;
 	for (const FStateTreeTransition& Transition : State.Transitions)
 	{
-		if (Transition.Event == Event)
+		const bool bMatch = bUseMask ? EnumHasAnyFlags(Transition.Trigger, Trigger) : Transition.Trigger == Trigger;
+		if (bMatch)
 		{
 			switch (Transition.State.Type)
 			{
@@ -565,9 +557,9 @@ FText SStateTreeViewRow::GetTransitionsDesc(const UStateTreeState& State, const
 	if (State.Children.Num() == 0
 		&& State.Type == EStateTreeStateType::State
 		&& DescItems.Num() == 0
-		&& Event != EStateTreeTransitionEvent::OnCondition)
+		&& EnumHasAnyFlags(Trigger, EStateTreeTransitionTrigger::OnStateCompleted))
 	{
-		if (HasParentTransitionForEvent(State, Event))
+		if (HasParentTransitionForTrigger(State, Trigger))
 		{
 			DescItems.Add(LOCTEXT("TransitionActionHandleInParentStyled", "[Parent]"));
 		}
@@ -580,7 +572,7 @@ FText SStateTreeViewRow::GetTransitionsDesc(const UStateTreeState& State, const
 	return FText::Join(FText::FromString(TEXT(", ")), DescItems);
 }
 
-FText SStateTreeViewRow::GetTransitionsIcon(const UStateTreeState& State, const EStateTreeTransitionEvent Event) const
+FText SStateTreeViewRow::GetTransitionsIcon(const UStateTreeState& State, const EStateTreeTransitionTrigger Trigger, const bool bUseMask) const
 {
 	enum EIconType
 	{
@@ -595,7 +587,8 @@ FText SStateTreeViewRow::GetTransitionsIcon(const UStateTreeState& State, const
 	for (const FStateTreeTransition& Transition : State.Transitions)
 	{
 		// The icons here depict "transition direction", not the type specifically.
-		if (Transition.Event == Event)
+		const bool bMatch = bUseMask ? EnumHasAnyFlags(Transition.Trigger, Trigger) : Transition.Trigger == Trigger;
+		if (bMatch)
 		{
 			switch (Transition.State.Type)
 			{
@@ -630,9 +623,9 @@ FText SStateTreeViewRow::GetTransitionsIcon(const UStateTreeState& State, const
 	if (State.Children.Num() == 0
 		&& State.Type == EStateTreeStateType::State
         && IconType == IconNone
-        && Event != EStateTreeTransitionEvent::OnCondition)
+        && EnumHasAnyFlags(Trigger, EStateTreeTransitionTrigger::OnStateCompleted))
 	{
-		if (HasParentTransitionForEvent(State, Event))
+		if (HasParentTransitionForTrigger(State, Trigger))
 		{
 			IconType = IconLevelUp;
 		}
@@ -659,38 +652,47 @@ FText SStateTreeViewRow::GetTransitionsIcon(const UStateTreeState& State, const
 	return FText::GetEmpty();
 }
 
-EVisibility SStateTreeViewRow::GetTransitionsVisibility(const UStateTreeState& State, const EStateTreeTransitionEvent Event) const
+EVisibility SStateTreeViewRow::GetTransitionsVisibility(const UStateTreeState& State, const EStateTreeTransitionTrigger Trigger) const
 {
-	TStaticArray<int32, 5> Counts;
+	if (!EnumHasAnyFlags(Trigger, EStateTreeTransitionTrigger::OnStateCompleted))
+	{
+		for (const FStateTreeTransition& Transition : State.Transitions)
+		{
+			if (EnumHasAnyFlags(Trigger, Transition.Trigger))
+			{
+				return EVisibility::Visible;
+			}
+		}
+		return EVisibility::Collapsed;
+	}
+	
+	EStateTreeTransitionTrigger HandledTriggers = EStateTreeTransitionTrigger::None;
+	TStaticArray<int32, static_cast<uint32>(EStateTreeTransitionTrigger::MAX)> Counts;
 	Algo::ForEach(Counts, [](int32& Count) { Count = 0; }); 
 
 	for (const FStateTreeTransition& Transition : State.Transitions)
 	{
-		Counts[static_cast<uint8>(Transition.Event)]++;
+		HandledTriggers |= Transition.Trigger;
+		Counts[static_cast<uint8>(Transition.Trigger)]++;
 	}
 
+	// Assume that leaf states should have completion transitions.
 	if (State.Children.Num() == 0
-	    && Counts[static_cast<uint8>(Event)] == 0
-	    && Event != EStateTreeTransitionEvent::OnCondition)
+	    && Counts[static_cast<uint8>(Trigger)] == 0)
 	{
 		// Find the missing transition type, note: Completed = Succeeded|Failed.
-		EStateTreeTransitionEvent HandledEvents = EStateTreeTransitionEvent::None;
-		HandledEvents |= Counts[static_cast<uint8>(EStateTreeTransitionEvent::OnCompleted)] > 0 ? EStateTreeTransitionEvent::OnCompleted : EStateTreeTransitionEvent::None;
-		HandledEvents |= Counts[static_cast<uint8>(EStateTreeTransitionEvent::OnSucceeded)] > 0 ? EStateTreeTransitionEvent::OnSucceeded : EStateTreeTransitionEvent::None;
-		HandledEvents |= Counts[static_cast<uint8>(EStateTreeTransitionEvent::OnFailed)] > 0 ? EStateTreeTransitionEvent::OnFailed : EStateTreeTransitionEvent::None;
-		const EStateTreeTransitionEvent MissingEvent = HandledEvents ^ EStateTreeTransitionEvent::OnCompleted;
-		
-		return MissingEvent == Event ? EVisibility::Visible : EVisibility::Collapsed;
+		const EStateTreeTransitionTrigger MissingTriggers = HandledTriggers ^ EStateTreeTransitionTrigger::OnStateCompleted;
+		return MissingTriggers == Trigger ? EVisibility::Visible : EVisibility::Collapsed;
 	}
 	
-	return Counts[static_cast<uint8>(Event)] > 0 ? EVisibility::Visible : EVisibility::Collapsed;
+	return Counts[static_cast<uint8>(Trigger)] > 0 ? EVisibility::Visible : EVisibility::Collapsed;
 }
 
 EVisibility SStateTreeViewRow::GetCompletedTransitionVisibility() const
 {
 	if (const UStateTreeState* State = WeakState.Get())
 	{
-		return GetTransitionsVisibility(*State, EStateTreeTransitionEvent::OnCompleted);
+		return GetTransitionsVisibility(*State, EStateTreeTransitionTrigger::OnStateCompleted);
 	}
 	return EVisibility::Visible;
 }
@@ -699,7 +701,7 @@ FText SStateTreeViewRow::GetCompletedTransitionsDesc() const
 {
 	if (const UStateTreeState* State = WeakState.Get())
 	{
-		return GetTransitionsDesc(*State, EStateTreeTransitionEvent::OnCompleted);
+		return GetTransitionsDesc(*State, EStateTreeTransitionTrigger::OnStateCompleted);
 	}
 	return LOCTEXT("Invalid", "Invalid");
 }
@@ -708,7 +710,7 @@ FText SStateTreeViewRow::GetCompletedTransitionsIcon() const
 {
 	if (const UStateTreeState* State = WeakState.Get())
 	{
-		return GetTransitionsIcon(*State, EStateTreeTransitionEvent::OnCompleted);
+		return GetTransitionsIcon(*State, EStateTreeTransitionTrigger::OnStateCompleted);
 	}
 	return FText::GetEmpty();
 }
@@ -717,7 +719,7 @@ EVisibility SStateTreeViewRow::GetSucceededTransitionVisibility() const
 {
 	if (const UStateTreeState* State = WeakState.Get())
 	{
-		return GetTransitionsVisibility(*State, EStateTreeTransitionEvent::OnSucceeded);
+		return GetTransitionsVisibility(*State, EStateTreeTransitionTrigger::OnStateSucceeded);
 	}
 	return EVisibility::Collapsed;
 }
@@ -726,7 +728,7 @@ FText SStateTreeViewRow::GetSucceededTransitionDesc() const
 {
 	if (const UStateTreeState* State = WeakState.Get())
 	{
-		return GetTransitionsDesc(*State, EStateTreeTransitionEvent::OnSucceeded);
+		return GetTransitionsDesc(*State, EStateTreeTransitionTrigger::OnStateSucceeded);
 	}
 	return FText::GetEmpty();
 }
@@ -735,7 +737,7 @@ FText SStateTreeViewRow::GetSucceededTransitionIcon() const
 {
 	if (const UStateTreeState* State = WeakState.Get())
 	{
-		return GetTransitionsIcon(*State, EStateTreeTransitionEvent::OnSucceeded);
+		return GetTransitionsIcon(*State, EStateTreeTransitionTrigger::OnStateSucceeded);
 	}
 	return FText::GetEmpty();
 }
@@ -744,7 +746,7 @@ EVisibility SStateTreeViewRow::GetFailedTransitionVisibility() const
 {
 	if (const UStateTreeState* State = WeakState.Get())
 	{
-		return GetTransitionsVisibility(*State, EStateTreeTransitionEvent::OnFailed);
+		return GetTransitionsVisibility(*State, EStateTreeTransitionTrigger::OnStateFailed);
 	}
 	return EVisibility::Collapsed;
 }
@@ -753,7 +755,7 @@ FText SStateTreeViewRow::GetFailedTransitionDesc() const
 {
 	if (const UStateTreeState* State = WeakState.Get())
 	{
-		return GetTransitionsDesc(*State, EStateTreeTransitionEvent::OnFailed);
+		return GetTransitionsDesc(*State, EStateTreeTransitionTrigger::OnStateFailed);
 	}
 	return LOCTEXT("Invalid", "Invalid");
 }
@@ -762,7 +764,7 @@ FText SStateTreeViewRow::GetFailedTransitionIcon() const
 {
 	if (const UStateTreeState* State = WeakState.Get())
 	{
-		return GetTransitionsIcon(*State, EStateTreeTransitionEvent::OnFailed);
+		return GetTransitionsIcon(*State, EStateTreeTransitionTrigger::OnStateFailed);
 	}
 	return FEditorFontGlyphs::Ban;
 }
@@ -771,19 +773,18 @@ EVisibility SStateTreeViewRow::GetConditionalTransitionsVisibility() const
 {
 	if (const UStateTreeState* State = WeakState.Get())
 	{
-		return GetTransitionsVisibility(*State, EStateTreeTransitionEvent::OnCondition);
+		return GetTransitionsVisibility(*State, EStateTreeTransitionTrigger::OnTick | EStateTreeTransitionTrigger::OnEvent);
 	}
 	return EVisibility::Collapsed;
 }
 
 FText SStateTreeViewRow::GetConditionalTransitionsDesc() const
 {
-	TArray<FText> DescItems;
 	if (const UStateTreeState* State = WeakState.Get())
 	{
-		return GetTransitionsDesc(*State, EStateTreeTransitionEvent::OnCondition);
+		return GetTransitionsDesc(*State, EStateTreeTransitionTrigger::OnTick | EStateTreeTransitionTrigger::OnEvent, /*bUseMask*/true);
 	}
-	return FText::Join(FText::FromString(TEXT(", ")), DescItems);
+	return FText::GetEmpty();
 }
 
 bool SStateTreeViewRow::IsRootState() const
@@ -793,7 +794,7 @@ bool SStateTreeViewRow::IsRootState() const
 	return State ? State->Parent == nullptr : false;
 }
 
-bool SStateTreeViewRow::IsSelected() const
+bool SStateTreeViewRow::IsStateSelected() const
 {
 	if (const UStateTreeState* State = WeakState.Get())
 	{
@@ -805,13 +806,7 @@ bool SStateTreeViewRow::IsSelected() const
 	return false;
 }
 
-
-bool SStateTreeViewRow::VerifyNodeTextChanged(const FText& NewLabel, FText& OutErrorMessage)
-{
-	return !NewLabel.IsEmptyOrWhitespace();
-}
-
-void SStateTreeViewRow::HandleNodeLabelTextCommitted(const FText& NewLabel, ETextCommit::Type CommitType)
+void SStateTreeViewRow::HandleNodeLabelTextCommitted(const FText& NewLabel, ETextCommit::Type CommitType) const
 {
 	if (StateTreeViewModel)
 	{
@@ -822,15 +817,14 @@ void SStateTreeViewRow::HandleNodeLabelTextCommitted(const FText& NewLabel, ETex
 	}
 }
 
-
-FReply SStateTreeViewRow::HandleDragDetected(const FGeometry&, const FPointerEvent&)
+FReply SStateTreeViewRow::HandleDragDetected(const FGeometry&, const FPointerEvent&) const
 {
 	return FReply::Handled().BeginDragDrop(FActionTreeViewDragDrop::New(WeakState.Get()));
 }
 
-TOptional<EItemDropZone> SStateTreeViewRow::HandleCanAcceptDrop(const FDragDropEvent& DragDropEvent, EItemDropZone DropZone, UStateTreeState* TargetState)
+TOptional<EItemDropZone> SStateTreeViewRow::HandleCanAcceptDrop(const FDragDropEvent& DragDropEvent, EItemDropZone DropZone, UStateTreeState* TargetState) const
 {
-	TSharedPtr<FActionTreeViewDragDrop> DragDropOperation = DragDropEvent.GetOperationAs<FActionTreeViewDragDrop>();
+	const TSharedPtr<FActionTreeViewDragDrop> DragDropOperation = DragDropEvent.GetOperationAs<FActionTreeViewDragDrop>();
 	if (DragDropOperation.IsValid())
 	{
 		// Cannot drop on selection or child of selection.
@@ -845,9 +839,9 @@ TOptional<EItemDropZone> SStateTreeViewRow::HandleCanAcceptDrop(const FDragDropE
 	return TOptional<EItemDropZone>();
 }
 
-FReply SStateTreeViewRow::HandleAcceptDrop(const FDragDropEvent& DragDropEvent, EItemDropZone DropZone, UStateTreeState* TargetState)
+FReply SStateTreeViewRow::HandleAcceptDrop(const FDragDropEvent& DragDropEvent, EItemDropZone DropZone, UStateTreeState* TargetState) const
 {
-	TSharedPtr<FActionTreeViewDragDrop> DragDropOperation = DragDropEvent.GetOperationAs<FActionTreeViewDragDrop>();
+	const TSharedPtr<FActionTreeViewDragDrop> DragDropOperation = DragDropEvent.GetOperationAs<FActionTreeViewDragDrop>();
 	if (DragDropOperation.IsValid())
 	{
 		if (StateTreeViewModel)
diff --git a/Engine/Plugins/Runtime/StateTree/Source/StateTreeEditorModule/Private/SStateTreeViewRow.h b/Engine/Plugins/Runtime/StateTree/Source/StateTreeEditorModule/Private/SStateTreeViewRow.h
--- a/Engine/Plugins/Runtime/StateTree/Source/StateTreeEditorModule/Private/SStateTreeViewRow.h
+++ b/Engine/Plugins/Runtime/StateTree/Source/StateTreeEditorModule/Private/SStateTreeViewRow.h
@@ -22,7 +22,7 @@ public:
 	SLATE_END_ARGS()
 
 	void Construct(const FArguments& InArgs, const TSharedRef<STableViewBase>& InOwnerTableView, UStateTreeState* InState, const TSharedPtr<SScrollBox>& ViewBox, TSharedRef<FStateTreeViewModel> InStateTreeViewModel);
-	void RequestRename();
+	void RequestRename() const;
 
 private:
 
@@ -54,21 +54,20 @@ private:
 	EVisibility GetConditionalTransitionsVisibility() const;
 	FText GetConditionalTransitionsDesc() const;
 
-	FText GetTransitionsDesc(const UStateTreeState& State, const EStateTreeTransitionEvent Event) const;
-	FText GetTransitionsIcon(const UStateTreeState& State, const EStateTreeTransitionEvent Event) const;
-	EVisibility GetTransitionsVisibility(const UStateTreeState& State, const EStateTreeTransitionEvent Event) const;
+	FText GetTransitionsDesc(const UStateTreeState& State, const EStateTreeTransitionTrigger Trigger, const bool bUseMask = false) const;
+	FText GetTransitionsIcon(const UStateTreeState& State, const EStateTreeTransitionTrigger Trigger, const bool bUseMask = false) const;
+	EVisibility GetTransitionsVisibility(const UStateTreeState& State, const EStateTreeTransitionTrigger Trigger) const;
 
-	bool HasParentTransitionForEvent(const UStateTreeState& State, const EStateTreeTransitionEvent Event) const;
+	bool HasParentTransitionForTrigger(const UStateTreeState& State, const EStateTreeTransitionTrigger Trigger) const;
 
 	bool IsRootState() const;
-	bool IsSelected() const;
+	bool IsStateSelected() const;
 
-	bool VerifyNodeTextChanged(const FText& NewLabel, FText& OutErrorMessage);
-	void HandleNodeLabelTextCommitted(const FText& NewLabel, ETextCommit::Type CommitType);
+	void HandleNodeLabelTextCommitted(const FText& NewLabel, ETextCommit::Type CommitType) const;
 
-	FReply HandleDragDetected(const FGeometry&, const FPointerEvent&);
-	TOptional<EItemDropZone> HandleCanAcceptDrop(const FDragDropEvent& DragDropEvent, EItemDropZone DropZone, UStateTreeState* TargetState);
-	FReply HandleAcceptDrop(const FDragDropEvent& DragDropEvent, EItemDropZone DropZone, UStateTreeState* TargetState);
+	FReply HandleDragDetected(const FGeometry&, const FPointerEvent&) const;
+	TOptional<EItemDropZone> HandleCanAcceptDrop(const FDragDropEvent& DragDropEvent, EItemDropZone DropZone, UStateTreeState* TargetState) const;
+	FReply HandleAcceptDrop(const FDragDropEvent& DragDropEvent, EItemDropZone DropZone, UStateTreeState* TargetState) const;
 	
 	TSharedPtr<FStateTreeViewModel> StateTreeViewModel;
 	TWeakObjectPtr<UStateTreeState> WeakState;
diff --git a/Engine/Plugins/Runtime/StateTree/Source/StateTreeEditorModule/Private/StateTreeCompiler.cpp b/Engine/Plugins/Runtime/StateTree/Source/StateTreeEditorModule/Private/StateTreeCompiler.cpp
--- a/Engine/Plugins/Runtime/StateTree/Source/StateTreeEditorModule/Private/StateTreeCompiler.cpp
+++ b/Engine/Plugins/Runtime/StateTree/Source/StateTreeEditorModule/Private/StateTreeCompiler.cpp
@@ -510,7 +510,8 @@ bool FStateTreeCompiler::CreateStateTransitions()
 		for (FStateTreeTransition& Transition : SourceState->Transitions)
 		{
 			FCompactStateTransition& CompactTransition = StateTree->Transitions.AddDefaulted_GetRef();
-			CompactTransition.Event = Transition.Event;
+			CompactTransition.Trigger = Transition.Trigger;
+			CompactTransition.EventTag = Transition.EventTag;
 			CompactTransition.Type = Transition.State.Type;
 			CompactTransition.GateDelay = (uint8)FMath::Clamp(FMath::CeilToInt(Transition.GateDelay * 10.0f), 0, 255);
 			CompactTransition.State = FStateTreeStateHandle::Invalid;
diff --git a/Engine/Plugins/Runtime/StateTree/Source/StateTreeEditorModule/Private/StateTreeState.cpp b/Engine/Plugins/Runtime/StateTree/Source/StateTreeEditorModule/Private/StateTreeState.cpp
--- a/Engine/Plugins/Runtime/StateTree/Source/StateTreeEditorModule/Private/StateTreeState.cpp
+++ b/Engine/Plugins/Runtime/StateTree/Source/StateTreeEditorModule/Private/StateTreeState.cpp
@@ -28,9 +28,16 @@ void FStateTreeStateLink::Set(const EStateTreeTransitionType InType, const UStat
 //////////////////////////////////////////////////////////////////////////
 // FStateTreeTransition
 
-FStateTreeTransition::FStateTreeTransition(const EStateTreeTransitionEvent InEvent, const EStateTreeTransitionType InType, const UStateTreeState* InState)
+FStateTreeTransition::FStateTreeTransition(const EStateTreeTransitionTrigger InTrigger, const EStateTreeTransitionType InType, const UStateTreeState* InState)
+	: Trigger(InTrigger)
+{
+	State.Set(InType, InState);
+}
+
+FStateTreeTransition::FStateTreeTransition(const EStateTreeTransitionTrigger InTrigger, const FGameplayTag InEventTag, const EStateTreeTransitionType InType, const UStateTreeState* InState)
+	: Trigger(InTrigger)
+	, EventTag(InEventTag)
 {
-	Event = InEvent;
 	State.Set(InType, InState);
 }
 
@@ -167,8 +174,9 @@ void UStateTreeState::PostLoad()
 		SetFlags(RF_Transactional);
 	}
 
-	// Move deprecated evaluators to editor data.
+#if WITH_EDITORONLY_DATA
 	PRAGMA_DISABLE_DEPRECATION_WARNINGS
+	// Move deprecated evaluators to editor data.
 	if (Evaluators_DEPRECATED.Num() > 0)
 	{
 		if (UStateTreeEditorData* TreeData = GetTypedOuter<UStateTreeEditorData>())
@@ -178,6 +186,7 @@ void UStateTreeState::PostLoad()
 		}		
 	}
 	PRAGMA_ENABLE_DEPRECATION_WARNINGS
+#endif
 }
 
 void UStateTreeState::UpdateParametersFromLinkedState()
diff --git a/Engine/Plugins/Runtime/StateTree/Source/StateTreeEditorModule/Public/StateTreeState.h b/Engine/Plugins/Runtime/StateTree/Source/StateTreeEditorModule/Public/StateTreeState.h
--- a/Engine/Plugins/Runtime/StateTree/Source/StateTreeEditorModule/Public/StateTreeState.h
+++ b/Engine/Plugins/Runtime/StateTree/Source/StateTreeEditorModule/Public/StateTreeState.h
@@ -21,7 +21,7 @@ struct STATETREEEDITORMODULE_API FStateTreeStateLink
 	GENERATED_BODY()
 
 	FStateTreeStateLink() = default;
-	FStateTreeStateLink(EStateTreeTransitionType InType) : Type(InType) {}
+	FStateTreeStateLink(const EStateTreeTransitionType InType) : Type(InType) {}
 
 	void Set(const EStateTreeTransitionType InType, const UStateTreeState* InState = nullptr);
 	void Set(const UStateTreeState* InState) { Set(EStateTreeTransitionType::GotoState, InState); }
@@ -47,7 +47,8 @@ struct STATETREEEDITORMODULE_API FStateTreeTransition
 	GENERATED_BODY()
 
 	FStateTreeTransition() = default;
-	FStateTreeTransition(const EStateTreeTransitionEvent InEvent, const EStateTreeTransitionType InType, const UStateTreeState* InState = nullptr);
+	FStateTreeTransition(const EStateTreeTransitionTrigger InTrigger, const EStateTreeTransitionType InType, const UStateTreeState* InState = nullptr);
+	FStateTreeTransition(const EStateTreeTransitionTrigger InTrigger, const FGameplayTag InEventTag, const EStateTreeTransitionType InType, const UStateTreeState* InState = nullptr);
 
 	template<typename T, typename... TArgs>
 	TStateTreeEditorNode<T>& AddCondition(TArgs&&... InArgs)
@@ -64,9 +65,12 @@ struct STATETREEEDITORMODULE_API FStateTreeTransition
 	}
 
 	UPROPERTY(EditDefaultsOnly, Category = Transition)
-	EStateTreeTransitionEvent Event = EStateTreeTransitionEvent::OnCompleted;
+	EStateTreeTransitionTrigger Trigger = EStateTreeTransitionTrigger::OnStateCompleted;
 
 	UPROPERTY(EditDefaultsOnly, Category = Transition)
+	FGameplayTag EventTag;
+
+	UPROPERTY(EditDefaultsOnly, Category = Transition, meta=(DisplayName="Transition To"))
 	FStateTreeStateLink State;
 
 	// Gate delay in seconds.
@@ -171,10 +175,16 @@ public:
 	 * Adds Transition.
 	 * @return reference to the new Transition. 
 	 */
-	FStateTreeTransition& AddTransition(const EStateTreeTransitionEvent InEvent, const EStateTreeTransitionType InType, const UStateTreeState* InState = nullptr)
+	FStateTreeTransition& AddTransition(const EStateTreeTransitionTrigger InTrigger, const EStateTreeTransitionType InType, const UStateTreeState* InState = nullptr)
+	{
+		return Transitions.Emplace_GetRef(InTrigger, InType, InState);
+	}
+
+	FStateTreeTransition& AddTransition(const EStateTreeTransitionTrigger InTrigger, const FGameplayTag InEventTag, const EStateTreeTransitionType InType, const UStateTreeState* InState = nullptr)
 	{
-		return Transitions.Emplace_GetRef(InEvent, InType, InState);
+		return Transitions.Emplace_GetRef(InTrigger, InEventTag, InType, InState);
 	}
+
  
 	// ~StateTree Builder API
 
@@ -196,9 +206,11 @@ public:
 	UPROPERTY(EditDefaultsOnly, Category = "Enter Conditions", meta = (BaseStruct = "/Script/StateTreeModule.StateTreeConditionBase", BaseClass = "/Script/StateTreeModule.StateTreeConditionBlueprintBase"))
 	TArray<FStateTreeEditorNode> EnterConditions;
 
+#if WITH_EDITORONLY_DATA
 	UE_DEPRECATED(5.1, "Evaluators are moved into UStateTreeEditorData. This property will be removed for 5.1.")
 	UPROPERTY(meta = (DeprecatedProperty, BaseStruct = "/Script/StateTreeModule.StateTreeEvaluatorBase", BaseClass = "/Script/StateTreeModule.StateTreeEvaluatorBlueprintBase"))
 	TArray<FStateTreeEditorNode> Evaluators_DEPRECATED;
+#endif
 
 	UPROPERTY(EditDefaultsOnly, Category = "Tasks", meta = (BaseStruct = "/Script/StateTreeModule.StateTreeTaskBase", BaseClass = "/Script/StateTreeModule.StateTreeTaskBlueprintBase"))
 	TArray<FStateTreeEditorNode> Tasks;
diff --git a/Engine/Plugins/Runtime/StateTree/Source/StateTreeModule/Private/Blueprint/StateTreeConditionBlueprintBase.cpp b/Engine/Plugins/Runtime/StateTree/Source/StateTreeModule/Private/Blueprint/StateTreeConditionBlueprintBase.cpp
--- a/Engine/Plugins/Runtime/StateTree/Source/StateTreeModule/Private/Blueprint/StateTreeConditionBlueprintBase.cpp
+++ b/Engine/Plugins/Runtime/StateTree/Source/StateTreeModule/Private/Blueprint/StateTreeConditionBlueprintBase.cpp
@@ -11,6 +11,7 @@
 
 bool UStateTreeConditionBlueprintBase::TestCondition(FStateTreeExecutionContext& Context) const
 {
+	FScopedCurrentContext(*this, Context);
 	AActor* OwnerActor = GetOwnerActor(Context);
 	return ReceiveTestCondition(OwnerActor);
 }
@@ -32,7 +33,7 @@ bool FStateTreeBlueprintConditionWrapper::Link(FStateTreeLinker& Linker)
 
 bool FStateTreeBlueprintConditionWrapper::TestCondition(FStateTreeExecutionContext& Context) const
 {
-	if (UStateTreeConditionBlueprintBase* Instance = Context.GetInstanceObjectInternal<UStateTreeConditionBlueprintBase>(DataViewIndex))
+	if (UStateTreeConditionBlueprintBase* Instance = Context.GetInstanceDataPtr<UStateTreeConditionBlueprintBase>(*this))
 	{
 		Instance->CopyExternalData(Context, ExternalDataHandles);
 		return Instance->TestCondition(Context);
diff --git a/Engine/Plugins/Runtime/StateTree/Source/StateTreeModule/Private/Blueprint/StateTreeEvaluatorBlueprintBase.cpp b/Engine/Plugins/Runtime/StateTree/Source/StateTreeModule/Private/Blueprint/StateTreeEvaluatorBlueprintBase.cpp
--- a/Engine/Plugins/Runtime/StateTree/Source/StateTreeModule/Private/Blueprint/StateTreeEvaluatorBlueprintBase.cpp
+++ b/Engine/Plugins/Runtime/StateTree/Source/StateTreeModule/Private/Blueprint/StateTreeEvaluatorBlueprintBase.cpp
@@ -1,7 +1,6 @@
 ﻿// Copyright Epic Games, Inc. All Rights Reserved.
 
 #include "Blueprint/StateTreeEvaluatorBlueprintBase.h"
-#include "CoreMinimal.h"
 #include "StateTreeExecutionContext.h"
 
 //----------------------------------------------------------------------//
@@ -10,6 +9,7 @@
 
 void UStateTreeEvaluatorBlueprintBase::EnterState(FStateTreeExecutionContext& Context, const EStateTreeStateChangeType ChangeType, const FStateTreeTransitionResult& Transition)
 {
+	FScopedCurrentContext(*this, Context);
 	AActor* OwnerActor = GetOwnerActor(Context);
 	PRAGMA_DISABLE_DEPRECATION_WARNINGS
 	ReceiveEnterState(OwnerActor, ChangeType, Transition);
@@ -18,6 +18,7 @@ void UStateTreeEvaluatorBlueprintBase::EnterState(FStateTreeExecutionContext& Co
 
 void UStateTreeEvaluatorBlueprintBase::ExitState(FStateTreeExecutionContext& Context, const EStateTreeStateChangeType ChangeType, const FStateTreeTransitionResult& Transition)
 {
+	FScopedCurrentContext(*this, Context);
 	AActor* OwnerActor = GetOwnerActor(Context);
 	PRAGMA_DISABLE_DEPRECATION_WARNINGS
 	ReceiveExitState(OwnerActor, ChangeType, Transition);
@@ -26,6 +27,7 @@ void UStateTreeEvaluatorBlueprintBase::ExitState(FStateTreeExecutionContext& Con
 
 void UStateTreeEvaluatorBlueprintBase::StateCompleted(FStateTreeExecutionContext& Context, const EStateTreeRunStatus CompletionStatus, const FStateTreeActiveStates& CompletedActiveStates)
 {
+	FScopedCurrentContext(*this, Context);
 	AActor* OwnerActor = GetOwnerActor(Context);
 	PRAGMA_DISABLE_DEPRECATION_WARNINGS
 	ReceiveStateCompleted(OwnerActor, CompletionStatus, CompletedActiveStates);
@@ -34,18 +36,21 @@ void UStateTreeEvaluatorBlueprintBase::StateCompleted(FStateTreeExecutionContext
 
 void UStateTreeEvaluatorBlueprintBase::TreeStart(FStateTreeExecutionContext& Context)
 {
+	FScopedCurrentContext(*this, Context);
 	AActor* OwnerActor = GetOwnerActor(Context);
 	ReceiveTreeStart(OwnerActor);
 }
 
 void UStateTreeEvaluatorBlueprintBase::TreeStop(FStateTreeExecutionContext& Context)
 {
+	FScopedCurrentContext(*this, Context);
 	AActor* OwnerActor = GetOwnerActor(Context);
 	ReceiveTreeStop(OwnerActor);
 }
 
 void UStateTreeEvaluatorBlueprintBase::Tick(FStateTreeExecutionContext& Context, const float DeltaTime)
 {
+	FScopedCurrentContext(*this, Context);
 	AActor* OwnerActor = GetOwnerActor(Context);
 	PRAGMA_DISABLE_DEPRECATION_WARNINGS
 	ReceiveEvaluate(OwnerActor, EStateTreeEvaluationType::Tick, DeltaTime);
@@ -70,7 +75,7 @@ bool FStateTreeBlueprintEvaluatorWrapper::Link(FStateTreeLinker& Linker)
 
 void FStateTreeBlueprintEvaluatorWrapper::EnterState(FStateTreeExecutionContext& Context, const EStateTreeStateChangeType ChangeType, const FStateTreeTransitionResult& Transition) const
 {
-	UStateTreeEvaluatorBlueprintBase* Instance = Context.GetInstanceObjectInternal<UStateTreeEvaluatorBlueprintBase>(DataViewIndex);
+	UStateTreeEvaluatorBlueprintBase* Instance = Context.GetInstanceDataPtr<UStateTreeEvaluatorBlueprintBase>(*this);
 	check(Instance);
 	
 	Instance->CopyExternalData(Context, ExternalDataHandles);
@@ -81,7 +86,7 @@ void FStateTreeBlueprintEvaluatorWrapper::EnterState(FStateTreeExecutionContext&
 
 void FStateTreeBlueprintEvaluatorWrapper::ExitState(FStateTreeExecutionContext& Context, const EStateTreeStateChangeType ChangeType, const FStateTreeTransitionResult& Transition) const
 {
-	UStateTreeEvaluatorBlueprintBase* Instance = Context.GetInstanceObjectInternal<UStateTreeEvaluatorBlueprintBase>(DataViewIndex);
+	UStateTreeEvaluatorBlueprintBase* Instance = Context.GetInstanceDataPtr<UStateTreeEvaluatorBlueprintBase>(*this);
 	check(Instance);
 	
 	Instance->CopyExternalData(Context, ExternalDataHandles);
@@ -92,7 +97,7 @@ void FStateTreeBlueprintEvaluatorWrapper::ExitState(FStateTreeExecutionContext&
 
 void FStateTreeBlueprintEvaluatorWrapper::StateCompleted(FStateTreeExecutionContext& Context, const EStateTreeRunStatus CompletionStatus, const FStateTreeActiveStates& CompletedActiveStates) const
 {
-	UStateTreeEvaluatorBlueprintBase* Instance = Context.GetInstanceObjectInternal<UStateTreeEvaluatorBlueprintBase>(DataViewIndex);
+	UStateTreeEvaluatorBlueprintBase* Instance = Context.GetInstanceDataPtr<UStateTreeEvaluatorBlueprintBase>(*this);
 	check(Instance);
 	
 	Instance->CopyExternalData(Context, ExternalDataHandles);
@@ -103,7 +108,7 @@ void FStateTreeBlueprintEvaluatorWrapper::StateCompleted(FStateTreeExecutionCont
 
 void FStateTreeBlueprintEvaluatorWrapper::TreeStart(FStateTreeExecutionContext& Context) const
 {
-	UStateTreeEvaluatorBlueprintBase* Instance = Context.GetInstanceObjectInternal<UStateTreeEvaluatorBlueprintBase>(DataViewIndex);
+	UStateTreeEvaluatorBlueprintBase* Instance = Context.GetInstanceDataPtr<UStateTreeEvaluatorBlueprintBase>(*this);
 	check(Instance);
 	
 	Instance->CopyExternalData(Context, ExternalDataHandles);
@@ -112,7 +117,7 @@ void FStateTreeBlueprintEvaluatorWrapper::TreeStart(FStateTreeExecutionContext&
 
 void FStateTreeBlueprintEvaluatorWrapper::TreeStop(FStateTreeExecutionContext& Context) const
 {
-	UStateTreeEvaluatorBlueprintBase* Instance = Context.GetInstanceObjectInternal<UStateTreeEvaluatorBlueprintBase>(DataViewIndex);
+	UStateTreeEvaluatorBlueprintBase* Instance = Context.GetInstanceDataPtr<UStateTreeEvaluatorBlueprintBase>(*this);
 	check(Instance);
 	
 	Instance->CopyExternalData(Context, ExternalDataHandles);
@@ -121,7 +126,7 @@ void FStateTreeBlueprintEvaluatorWrapper::TreeStop(FStateTreeExecutionContext& C
 
 void FStateTreeBlueprintEvaluatorWrapper::Tick(FStateTreeExecutionContext& Context, const float DeltaTime) const
 {
-	UStateTreeEvaluatorBlueprintBase* Instance = Context.GetInstanceObjectInternal<UStateTreeEvaluatorBlueprintBase>(DataViewIndex);
+	UStateTreeEvaluatorBlueprintBase* Instance = Context.GetInstanceDataPtr<UStateTreeEvaluatorBlueprintBase>(*this);
 	check(Instance);
 	
 	Instance->CopyExternalData(Context, ExternalDataHandles);
diff --git a/Engine/Plugins/Runtime/StateTree/Source/StateTreeModule/Private/Blueprint/StateTreeItemBlueprintBase.cpp b/Engine/Plugins/Runtime/StateTree/Source/StateTreeModule/Private/Blueprint/StateTreeItemBlueprintBase.cpp
deleted file mode 100644
diff --git a/Engine/Plugins/Runtime/StateTree/Source/StateTreeModule/Private/Blueprint/StateTreeNodeBlueprintBase.cpp b/Engine/Plugins/Runtime/StateTree/Source/StateTreeModule/Private/Blueprint/StateTreeNodeBlueprintBase.cpp
new file mode 100644
--- /dev/null
+++ b/Engine/Plugins/Runtime/StateTree/Source/StateTreeModule/Private/Blueprint/StateTreeNodeBlueprintBase.cpp
@@ -0,0 +1,129 @@
+// Copyright Epic Games, Inc. All Rights Reserved.
+
+#include "Blueprint/StateTreeNodeBlueprintBase.h"
+#include "AIController.h"
+#include "StateTreeExecutionContext.h"
+#include "StateTreeLinker.h"
+#include "VisualLogger/VisualLogger.h"
+
+UWorld* UStateTreeNodeBlueprintBase::GetWorld() const
+{
+	// The items are duplicated as the StateTreeExecution context as outer, so this should be essentially the same as GetWorld() on StateTree context.
+	// The CDO is used by the BP editor to check for certain functionality, make it return nullptr so that the GetWorld() passes as overridden. 
+	if (!HasAnyFlags(RF_ClassDefaultObject))
+	{
+		if (UObject* Outer = GetOuter())
+		{
+			return Outer->GetWorld();
+		}
+	}
+	
+	return nullptr;
+}
+
+AActor* UStateTreeNodeBlueprintBase::GetOwnerActor(const FStateTreeExecutionContext& Context) const
+{
+	if (const AAIController* Controller = Cast<AAIController>(Context.GetOwner()))
+	{
+		return Controller->GetPawn();
+	}
+	
+	return Cast<AActor>(Context.GetOwner());
+}
+
+#if WITH_EDITOR
+void UStateTreeNodeBlueprintBase::PostCDOCompiled(const FPostCDOCompiledContext& Context)
+{
+	if (Context.bIsSkeletonOnly)
+	{
+		return;
+	}
+
+	// Save some property data for runtime use.
+	PropertyInfos.Reset();
+
+	static const FName CategoryName(TEXT("Category"));
+
+	for (TFieldIterator<FProperty> It(GetClass()); It; ++It)
+	{
+		const FString Category = It->GetMetaData(CategoryName);
+		if (Category == TEXT("Input"))
+		{
+			PropertyInfos.Emplace(It->GetFName(), EStateTreeBlueprintPropertyCategory::Input);
+		}
+		else if (Category == TEXT("Output"))
+		{
+			PropertyInfos.Emplace(It->GetFName(), EStateTreeBlueprintPropertyCategory::Output);
+		}
+		else if (Category == TEXT("Parameter"))
+		{
+			PropertyInfos.Emplace(It->GetFName(), EStateTreeBlueprintPropertyCategory::Parameter);
+		}
+		else if (Category == TEXT("ExternalData"))
+		{
+			PropertyInfos.Emplace(It->GetFName(), EStateTreeBlueprintPropertyCategory::ExternalData);
+		}
+	}
+}
+#endif
+
+void UStateTreeNodeBlueprintBase::LinkExternalData(FStateTreeLinker& Linker, TArray<FStateTreeBlueprintExternalDataHandle>& OutExternalDataHandles) const
+{
+	const UClass* Class = GetClass();
+
+	// Find properties which are external data requirements.
+	for (const FStateTreeBlueprintPropertyInfo& PropertyInfo : PropertyInfos)
+	{
+		if (PropertyInfo.Category == EStateTreeBlueprintPropertyCategory::ExternalData)
+		{
+			const FProperty* Property = Class->FindPropertyByName(PropertyInfo.PropertyName);
+			if (const FObjectProperty* ObjectProperty = CastField<FObjectProperty>(Property))
+			{
+				FStateTreeBlueprintExternalDataHandle& ExtData = OutExternalDataHandles.AddDefaulted_GetRef();
+				ExtData.Property = Property;
+				Linker.LinkExternalData(ExtData.Handle, ObjectProperty->PropertyClass, EStateTreeExternalDataRequirement::Required);
+			}
+			else if (const FStructProperty* StructProperty = CastField<FStructProperty>(Property))
+			{
+				FStateTreeBlueprintExternalDataHandle& ExtData = OutExternalDataHandles.AddDefaulted_GetRef();
+				ExtData.Property = Property;
+				Linker.LinkExternalData(ExtData.Handle, StructProperty->Struct, EStateTreeExternalDataRequirement::Required);
+			}
+		}
+	}
+}
+
+void UStateTreeNodeBlueprintBase::CopyExternalData(FStateTreeExecutionContext& Context, TConstArrayView<FStateTreeBlueprintExternalDataHandle> ExternalDataHandles)
+{
+	for (const FStateTreeBlueprintExternalDataHandle& ExtData : ExternalDataHandles)
+	{
+		FStateTreeDataView View = Context.GetExternalDataView(ExtData.Handle);
+		if (View.IsValid())
+		{
+			if (const FObjectProperty* ObjectProperty = CastField<FObjectProperty>(ExtData.Property))
+			{
+				check(View.GetStruct()->IsChildOf(ObjectProperty->PropertyClass));
+				// Copy pointer
+				uint8* Address = ObjectProperty->ContainerPtrToValuePtr<uint8>(this);
+				ObjectProperty->SetObjectPropertyValue(Address, View.GetMutablePtr<UObject>());
+			}
+			else if (const FStructProperty* StructProperty = CastField<FStructProperty>(ExtData.Property))
+			{
+				check(View.GetStruct()->IsChildOf(StructProperty->Struct));
+				// Copy value
+				uint8* Address = StructProperty->ContainerPtrToValuePtr<uint8>(this);
+				StructProperty->CopyCompleteValue(Address, View.GetMemory());
+			}
+		}
+	}
+}
+
+void UStateTreeNodeBlueprintBase::SendEvent(const FStateTreeEvent& Event)
+{
+	if (CurrentContext == nullptr)
+	{
+		UE_VLOG_UELOG(this, LogStateTree, Error, TEXT("Trying to call SendEvent() outside StateTree tick. Use SendEvent() on UStateTreeComponent instead for sending signals externally."));
+		return;
+	}
+	CurrentContext->SendEvent(Event);
+}
diff --git a/Engine/Plugins/Runtime/StateTree/Source/StateTreeModule/Private/Blueprint/StateTreeTaskBlueprintBase.cpp b/Engine/Plugins/Runtime/StateTree/Source/StateTreeModule/Private/Blueprint/StateTreeTaskBlueprintBase.cpp
--- a/Engine/Plugins/Runtime/StateTree/Source/StateTreeModule/Private/Blueprint/StateTreeTaskBlueprintBase.cpp
+++ b/Engine/Plugins/Runtime/StateTree/Source/StateTreeModule/Private/Blueprint/StateTreeTaskBlueprintBase.cpp
@@ -15,24 +15,28 @@ EStateTreeRunStatus UStateTreeTaskBlueprintBase::ReceiveEnterState_Implementatio
 
 EStateTreeRunStatus UStateTreeTaskBlueprintBase::EnterState(FStateTreeExecutionContext& Context, const EStateTreeStateChangeType ChangeType, const FStateTreeTransitionResult& Transition)
 {
+	FScopedCurrentContext(*this, Context);
 	AActor* OwnerActor = GetOwnerActor(Context);
 	return ReceiveEnterState(OwnerActor, ChangeType, Transition);
 }
 
 void UStateTreeTaskBlueprintBase::ExitState(FStateTreeExecutionContext& Context, const EStateTreeStateChangeType ChangeType, const FStateTreeTransitionResult& Transition)
 {
+	FScopedCurrentContext(*this, Context);
 	AActor* OwnerActor = GetOwnerActor(Context);
 	ReceiveExitState(OwnerActor, ChangeType, Transition);
 }
 
 void UStateTreeTaskBlueprintBase::StateCompleted(FStateTreeExecutionContext& Context, const EStateTreeRunStatus CompletionStatus, const FStateTreeActiveStates& CompletedActiveStates)
 {
+	FScopedCurrentContext(*this, Context);
 	AActor* OwnerActor = GetOwnerActor(Context);
 	ReceiveStateCompleted(OwnerActor, CompletionStatus, CompletedActiveStates);
 }
 
 EStateTreeRunStatus UStateTreeTaskBlueprintBase::Tick(FStateTreeExecutionContext& Context, const float DeltaTime)
 {
+	FScopedCurrentContext(*this, Context);
 	AActor* OwnerActor = GetOwnerActor(Context);
 	return ReceiveTick(OwnerActor, DeltaTime);
 }
@@ -54,7 +58,7 @@ bool FStateTreeBlueprintTaskWrapper::Link(FStateTreeLinker& Linker)
 
 EStateTreeRunStatus FStateTreeBlueprintTaskWrapper::EnterState(FStateTreeExecutionContext& Context, const EStateTreeStateChangeType ChangeType, const FStateTreeTransitionResult& Transition) const
 {
-	UStateTreeTaskBlueprintBase* Instance = Context.GetInstanceObjectInternal<UStateTreeTaskBlueprintBase>(DataViewIndex);
+	UStateTreeTaskBlueprintBase* Instance = Context.GetInstanceDataPtr<UStateTreeTaskBlueprintBase>(*this);
 	check(Instance);
 	
 	Instance->CopyExternalData(Context, ExternalDataHandles);
@@ -63,7 +67,7 @@ EStateTreeRunStatus FStateTreeBlueprintTaskWrapper::EnterState(FStateTreeExecuti
 
 void FStateTreeBlueprintTaskWrapper::ExitState(FStateTreeExecutionContext& Context, const EStateTreeStateChangeType ChangeType, const FStateTreeTransitionResult& Transition) const
 {
-	UStateTreeTaskBlueprintBase* Instance = Context.GetInstanceObjectInternal<UStateTreeTaskBlueprintBase>(DataViewIndex);
+	UStateTreeTaskBlueprintBase* Instance = Context.GetInstanceDataPtr<UStateTreeTaskBlueprintBase>(*this);
 	check(Instance);
 	
 	Instance->CopyExternalData(Context, ExternalDataHandles);
@@ -72,7 +76,7 @@ void FStateTreeBlueprintTaskWrapper::ExitState(FStateTreeExecutionContext& Conte
 
 void FStateTreeBlueprintTaskWrapper::StateCompleted(FStateTreeExecutionContext& Context, const EStateTreeRunStatus CompletionStatus, const FStateTreeActiveStates& CompletedActiveStates) const
 {
-	UStateTreeTaskBlueprintBase* Instance = Context.GetInstanceObjectInternal<UStateTreeTaskBlueprintBase>(DataViewIndex);
+	UStateTreeTaskBlueprintBase* Instance = Context.GetInstanceDataPtr<UStateTreeTaskBlueprintBase>(*this);
 	check(Instance);
 	
 	Instance->CopyExternalData(Context, ExternalDataHandles);
@@ -81,7 +85,7 @@ void FStateTreeBlueprintTaskWrapper::StateCompleted(FStateTreeExecutionContext&
 
 EStateTreeRunStatus FStateTreeBlueprintTaskWrapper::Tick(FStateTreeExecutionContext& Context, const float DeltaTime) const
 {
-	UStateTreeTaskBlueprintBase* Instance = Context.GetInstanceObjectInternal<UStateTreeTaskBlueprintBase>(DataViewIndex);
+	UStateTreeTaskBlueprintBase* Instance = Context.GetInstanceDataPtr<UStateTreeTaskBlueprintBase>(*this);
 	check(Instance);
 	
 	Instance->CopyExternalData(Context, ExternalDataHandles);
diff --git a/Engine/Plugins/Runtime/StateTree/Source/StateTreeModule/Private/StateTreeExecutionContext.cpp b/Engine/Plugins/Runtime/StateTree/Source/StateTreeModule/Private/StateTreeExecutionContext.cpp
--- a/Engine/Plugins/Runtime/StateTree/Source/StateTreeModule/Private/StateTreeExecutionContext.cpp
+++ b/Engine/Plugins/Runtime/StateTree/Source/StateTreeModule/Private/StateTreeExecutionContext.cpp
@@ -15,8 +15,9 @@
 
 namespace UE::StateTree
 {
-	constexpr int32 DebugIndentSize = 2;
-}
+	constexpr int32 MaxActiveEvents = 16;	// Maximum number of events that can be accumulated.
+	constexpr int32 DebugIndentSize = 2;	// Debug printing indent for hierarchical data.
+}; // UE::StateTree
 
 FStateTreeExecutionContext::FStateTreeExecutionContext()
 {
@@ -75,6 +76,48 @@ void FStateTreeExecutionContext::SetParameters(const FInstancedPropertyBag& Para
 	}
 }
 
+bool FStateTreeExecutionContext::AreExternalDataViewsValid() const
+{
+	check(StateTree);
+	bool bResult = true;
+	for (const FStateTreeExternalDataDesc& DataDesc : StateTree->ExternalDataDescs)
+	{
+		const FStateTreeDataView& DataView = DataViews[DataDesc.Handle.DataViewIndex.Get()];
+			
+		if (DataDesc.Requirement == EStateTreeExternalDataRequirement::Required)
+		{
+			// Required items must have valid pointer of the expected type.  
+			if (!DataView.IsValid() || !DataView.GetStruct()->IsChildOf(DataDesc.Struct))
+			{
+				bResult = false;
+				break;
+			}
+		}
+		else
+		{
+			// Optional items must have same type if they are set.
+			if (DataView.IsValid() && !DataView.GetStruct()->IsChildOf(DataDesc.Struct))
+			{
+				bResult = false;
+				break;
+			}
+		}
+	}
+
+	for (const FStateTreeExternalDataDesc& DataDesc : StateTree->GetNamedExternalDataDescs())
+	{
+		const FStateTreeDataView& DataView = DataViews[DataDesc.Handle.DataViewIndex.Get()];
+
+		// Items must have valid pointer of the expected type.  
+		if (!DataView.IsValid() || !DataView.GetStruct()->IsChildOf(DataDesc.Struct))
+		{
+			bResult = false;
+			break;
+		}
+	}
+	return bResult;
+}
+
 void FStateTreeExecutionContext::Reset()
 {
 	InternalInstanceData.Reset();
@@ -143,26 +186,21 @@ EStateTreeRunStatus FStateTreeExecutionContext::Start(FStateTreeInstanceData* Ex
 		}
 	}
 	
+	FScopedCurrentInstanceData InstanceDataScope(*this, InstanceData); 
 	FStateTreeExecutionState* Exec = &GetExecState(InstanceData);
 
+	// Stop if still running previous state.
+	if (Exec->TreeRunStatus == EStateTreeRunStatus::Running)
+	{
+		Stop(ExternalInstanceData);
+	}
+
 	// Call TreeStart on evaluators.
 	StartEvaluators(InstanceData);
 
 	// First tick
 	TickEvaluators(InstanceData, 0.0f);
 
-	// Stop if still running previous state.
-	if (Exec->TreeRunStatus == EStateTreeRunStatus::Running)
-	{
-		FStateTreeTransitionResult Transition;
-		Transition.TargetState = FStateTreeStateHandle::Succeeded;
-		Transition.CurrentActiveStates = Exec->ActiveStates;
-		Transition.CurrentRunStatus = Exec->LastTickStatus;
-		Transition.NextActiveStates = FStateTreeActiveStates(FStateTreeStateHandle::Succeeded);
-
-		ExitState(InstanceData, Transition);
-	}
-	
 	// Initialize to unset running state.
 	Exec->TreeRunStatus = EStateTreeRunStatus::Running;
 	Exec->ActiveStates.Reset();
@@ -227,9 +265,14 @@ EStateTreeRunStatus FStateTreeExecutionContext::Stop(FStateTreeInstanceData* Ext
 	{
 		return EStateTreeRunStatus::Failed;
 	}
-
+	
+	FScopedCurrentInstanceData InstanceDataScope(*this, InstanceData); 
 	FStateTreeExecutionState& Exec = GetExecState(InstanceData);
 
+	// Capture events added between ticks.
+	EventsToProcess = InstanceData.GetEvents();
+	InstanceData.GetEvents().Reset();
+
 	TickEvaluators(InstanceData, 0.0f);
 
 	// Exit states if still in some valid state.
@@ -282,9 +325,14 @@ EStateTreeRunStatus FStateTreeExecutionContext::Tick(const float DeltaTime, FSta
 			ANSI_TO_TCHAR(__FUNCTION__), *GetNameSafe(Owner), *GetNameSafe(StateTree));
 		return EStateTreeRunStatus::Failed;
 	}
-	
+
+	FScopedCurrentInstanceData InstanceDataScope(*this, InstanceData); 
 	FStateTreeExecutionState* Exec = &GetExecState(InstanceData);
 
+	// Capture events added between ticks.
+	EventsToProcess = InstanceData.GetEvents();
+	InstanceData.GetEvents().Reset();
+	
 	// No ticking of the tree is done or stopped.
 	if (Exec->TreeRunStatus != EStateTreeRunStatus::Running)
 	{
@@ -312,6 +360,11 @@ EStateTreeRunStatus FStateTreeExecutionContext::Tick(const float DeltaTime, FSta
 		}
 	}
 
+	// Append events accumulated during the tick, so that transitions can immediately act on them.
+	// We'll consume the events only if they lead to state change below (EnterState is treated the same as Tick),
+	// or let them be processed next frame if no transition.
+	EventsToProcess.Append(InstanceData.GetEvents());
+
 	// The state selection is repeated up to MaxIteration time. This allows failed EnterState() to potentially find a new state immediately.
 	// This helps event driven StateTrees to not require another event/tick to find a suitable state.
 	static constexpr int32 MaxIterations = 5;
@@ -321,15 +374,22 @@ EStateTreeRunStatus FStateTreeExecutionContext::Tick(const float DeltaTime, FSta
 		FStateTreeTransitionResult Transition;
 		if (TriggerTransitions(InstanceData, Transition))
 		{
+			// We have committed to state change, consume events that were accumulated during the tick above.
+			InstanceData.GetEvents().Reset();
+
 			ExitState(InstanceData, Transition);
-			
+
 			if (Transition.NextActiveStates.Last() == FStateTreeStateHandle::Succeeded || Transition.NextActiveStates.Last() == FStateTreeStateHandle::Failed)
 			{
 				// Transition to a terminal state (succeeded/failed), or default transition failed.
 				Exec->TreeRunStatus = Transition.NextActiveStates.Last() == FStateTreeStateHandle::Succeeded ? EStateTreeRunStatus::Succeeded : EStateTreeRunStatus::Failed;
 				return Exec->TreeRunStatus;
 			}
-			
+
+			// Append and consume the events accumulated during the state exit.
+			EventsToProcess.Append(InstanceData.GetEvents());
+			InstanceData.GetEvents().Reset();
+
 			// Enter state tasks can fail/succeed, treat it same as tick.
 			const EStateTreeRunStatus LastTickStatus = EnterState(InstanceData, Transition);
 
@@ -350,7 +410,7 @@ EStateTreeRunStatus FStateTreeExecutionContext::Tick(const float DeltaTime, FSta
 			break;
 		}
 	}
-	
+
 	if (Exec->ActiveStates.IsEmpty())
 	{
 		// Should not happen. This may happen if a state completion transition could not be selected. 
@@ -360,9 +420,74 @@ EStateTreeRunStatus FStateTreeExecutionContext::Tick(const float DeltaTime, FSta
 		return Exec->TreeRunStatus;
 	}
 
+	EventsToProcess.Reset();
+
 	return Exec->TreeRunStatus;
 }
 
+void FStateTreeExecutionContext::SendExternalEvent(const FStateTreeEvent& Event, FStateTreeInstanceData* ExternalInstanceData)
+{
+	CSV_SCOPED_TIMING_STAT_EXCLUSIVE(StateTree_SendEvent);
+
+	if (!Event.Tag.IsValid())
+	{
+		STATETREE_LOG(Warning, TEXT("%s: Event tag cannot be empty, '%s' using StateTree '%s'."),
+			ANSI_TO_TCHAR(__FUNCTION__), *GetNameSafe(Owner), *GetNameSafe(StateTree));
+		return;
+	}
+	
+	if (!Owner || !StateTree)
+	{
+		return;
+	}
+	FStateTreeInstanceData& InstanceData = SelectMutableInstanceData(ExternalInstanceData);
+	if (!InstanceData.IsValid())
+	{
+		return;
+	}
+
+	if (InstanceData.GetEvents().Num() >= UE::StateTree::MaxActiveEvents)
+	{
+		STATETREE_LOG(Warning, TEXT("%s: Too many events send on '%s' using StateTree '%s'."),
+			ANSI_TO_TCHAR(__FUNCTION__), *GetNameSafe(Owner), *GetNameSafe(StateTree));
+		return;
+	}
+
+	InstanceData.GetEvents().Add(Event);
+}
+
+void FStateTreeExecutionContext::SendEvent(const FStateTreeEvent& Event)
+{
+	CSV_SCOPED_TIMING_STAT_EXCLUSIVE(StateTree_SendEvent);
+
+	if (ensureMsgf(!CurrentInstanceData, TEXT("SendEvent can only be used during StateTree tick, use SendExternalEvent() when sending events outside the tick.")))
+	{
+		return;
+	}
+	
+	if (!Event.Tag.IsValid())
+	{
+		STATETREE_LOG(Warning, TEXT("%s: Event tag cannot be empty, '%s' using StateTree '%s'."),
+			ANSI_TO_TCHAR(__FUNCTION__), *GetNameSafe(Owner), *GetNameSafe(StateTree));
+		return;
+	}
+	
+	if (!Owner || !StateTree)
+	{
+		return;
+	}
+
+	constexpr int32 MaxActiveEvents = 16;
+	if (CurrentInstanceData->GetEvents().Num() >= MaxActiveEvents)
+	{
+		STATETREE_LOG(Warning, TEXT("%s: Too many events send on '%s' using StateTree '%s'."),
+			ANSI_TO_TCHAR(__FUNCTION__), *GetNameSafe(Owner), *GetNameSafe(StateTree));
+		return;
+	}
+
+	CurrentInstanceData->GetEvents().Add(Event);
+}
+
 void FStateTreeExecutionContext::UpdateInstanceData(FStateTreeInstanceData& InstanceData, const FStateTreeActiveStates& CurrentActiveStates, const FStateTreeActiveStates& NextActiveStates)
 {
 	check(Owner);
@@ -1014,16 +1139,30 @@ bool FStateTreeExecutionContext::TriggerTransitions(FStateTreeInstanceData& Inst
 	CSV_SCOPED_TIMING_STAT_EXCLUSIVE(StateTree_TriggerTransition);
 
 	FStateTreeExecutionState& Exec = GetExecState(InstanceData);
-	EStateTreeTransitionEvent Event = EStateTreeTransitionEvent::OnCondition;
+	
+	EStateTreeTransitionTrigger CompletionTrigger = EStateTreeTransitionTrigger::None;
 	if (Exec.LastTickStatus == EStateTreeRunStatus::Succeeded)
 	{
-		Event = EStateTreeTransitionEvent::OnSucceeded;
+		CompletionTrigger = EStateTreeTransitionTrigger::OnStateSucceeded;
 	}
 	else if (Exec.LastTickStatus == EStateTreeRunStatus::Failed)
 	{
-		Event = EStateTreeTransitionEvent::OnFailed;
+		CompletionTrigger = EStateTreeTransitionTrigger::OnStateFailed;
 	}
 
+	auto HasEvent = [this](const FGameplayTag QueriedTag)
+	{
+		if (EventsToProcess.IsEmpty())
+		{
+			return false;
+		}
+		
+		return EventsToProcess.ContainsByPredicate([QueriedTag](const FStateTreeEvent& Event)
+		{
+			return Event.Tag == QueriedTag;
+		});
+	};
+
 	// Walk towards root and check all transitions along the way.
 	for (int32 StateIndex = Exec.ActiveStates.Num() - 1; StateIndex >= 0; StateIndex--)
 	{
@@ -1034,7 +1173,12 @@ bool FStateTreeExecutionContext::TriggerTransitions(FStateTreeInstanceData& Inst
 			// All transition conditions must pass
 			const int16 TransitionIndex = State.TransitionsBegin + i;
 			const FCompactStateTransition& Transition = StateTree->Transitions[TransitionIndex];
-			if (EnumHasAllFlags(Transition.Event, Event) && TestAllConditions(Transition.ConditionsBegin, Transition.ConditionsNum))
+
+			const bool bShouldCheck = EnumHasAnyFlags(Transition.Trigger, CompletionTrigger)
+									|| Transition.Trigger == EStateTreeTransitionTrigger::OnTick
+									|| (Transition.Trigger == EStateTreeTransitionTrigger::OnEvent && HasEvent(Transition.EventTag));
+			
+			if (bShouldCheck && TestAllConditions(Transition.ConditionsBegin, Transition.ConditionsNum))
 			{
 				// If a transition has delay, we stop testing other transitions, but the transition will not pass the condition until the delay time passes.
 				if (Transition.GateDelay > 0)
@@ -1408,14 +1552,15 @@ void FStateTreeExecutionContext::DebugPrintInternalLayout(const FStateTreeInstan
 	StateTree->PropertyBindings.DebugPrintInternalLayout(DebugString);
 
 	// Transitions
-	DebugString += FString::Printf(TEXT("\nTransitions(%d)\n  [ %-3s | %15s | %-40s | %-40s | %-8s ]\n"), StateTree->Transitions.Num()
-		, TEXT("Idx"), TEXT("State"), TEXT("Transition Type"), TEXT("Transition Event"), TEXT("Num Cond"));
+	DebugString += FString::Printf(TEXT("\nTransitions(%d)\n  [ %-3s | %15s | %-20s | %-20s | %-40s | %-8s ]\n"), StateTree->Transitions.Num()
+		, TEXT("Idx"), TEXT("State"), TEXT("Transition Type"), TEXT("Transition Trigger"), TEXT("Transition Event Tag"), TEXT("Num Cond"));
 	for (const FCompactStateTransition& Transition : StateTree->Transitions)
 	{
-		DebugString += FString::Printf(TEXT("  | %3d | %15s | %-40s | %-40s | %8d |\n"),
+		DebugString += FString::Printf(TEXT("  | %3d | %15s | %-20s | %-20s | %-40s | %8d |\n"),
 									Transition.ConditionsBegin, *Transition.State.Describe(),
-									*UEnum::GetValueAsString(Transition.Type),
-									*UEnum::GetValueAsString(Transition.Event),
+									*UEnum::GetDisplayValueAsText(Transition.Type).ToString(),
+									*UEnum::GetDisplayValueAsText(Transition.Trigger).ToString(),
+									*Transition.EventTag.ToString(),
 									Transition.ConditionsNum);
 	}
 
diff --git a/Engine/Plugins/Runtime/StateTree/Source/StateTreeModule/Private/StateTreeSchema.cpp b/Engine/Plugins/Runtime/StateTree/Source/StateTreeModule/Private/StateTreeSchema.cpp
--- a/Engine/Plugins/Runtime/StateTree/Source/StateTreeModule/Private/StateTreeSchema.cpp
+++ b/Engine/Plugins/Runtime/StateTree/Source/StateTreeModule/Private/StateTreeSchema.cpp
@@ -1,9 +1,9 @@
 ﻿// Copyright Epic Games, Inc. All Rights Reserved.
 
 #include "StateTreeSchema.h"
-#include "Blueprint/StateTreeItemBlueprintBase.h"
+#include "Blueprint/StateTreeNodeBlueprintBase.h"
 
 bool UStateTreeSchema::IsChildOfBlueprintBase(const UClass* InClass) const
 {
-	return InClass->IsChildOf<UStateTreeItemBlueprintBase>();
+	return InClass->IsChildOf<UStateTreeNodeBlueprintBase>();
 }
diff --git a/Engine/Plugins/Runtime/StateTree/Source/StateTreeModule/Public/Blueprint/StateTreeConditionBlueprintBase.h b/Engine/Plugins/Runtime/StateTree/Source/StateTreeModule/Public/Blueprint/StateTreeConditionBlueprintBase.h
--- a/Engine/Plugins/Runtime/StateTree/Source/StateTreeModule/Public/Blueprint/StateTreeConditionBlueprintBase.h
+++ b/Engine/Plugins/Runtime/StateTree/Source/StateTreeModule/Public/Blueprint/StateTreeConditionBlueprintBase.h
@@ -7,7 +7,7 @@
 #include "StateTreeTypes.h"
 #include "StateTreeConditionBase.h"
 #include "StateTreePropertyBindings.h"
-#include "StateTreeItemBlueprintBase.h"
+#include "StateTreeNodeBlueprintBase.h"
 #include "StateTreeConditionBlueprintBase.generated.h"
 
 struct FStateTreeExecutionContext;
@@ -16,7 +16,7 @@ struct FStateTreeExecutionContext;
  * Base class for Blueprint based Conditions. 
  */
 UCLASS(Abstract, Blueprintable)
-class STATETREEMODULE_API UStateTreeConditionBlueprintBase : public UStateTreeItemBlueprintBase
+class STATETREEMODULE_API UStateTreeConditionBlueprintBase : public UStateTreeNodeBlueprintBase
 {
 	GENERATED_BODY()
 public:
diff --git a/Engine/Plugins/Runtime/StateTree/Source/StateTreeModule/Public/Blueprint/StateTreeEvaluatorBlueprintBase.h b/Engine/Plugins/Runtime/StateTree/Source/StateTreeModule/Public/Blueprint/StateTreeEvaluatorBlueprintBase.h
--- a/Engine/Plugins/Runtime/StateTree/Source/StateTreeModule/Public/Blueprint/StateTreeEvaluatorBlueprintBase.h
+++ b/Engine/Plugins/Runtime/StateTree/Source/StateTreeModule/Public/Blueprint/StateTreeEvaluatorBlueprintBase.h
@@ -6,8 +6,8 @@
 #include "Templates/SubclassOf.h"
 #include "StateTreeTypes.h"
 #include "StateTreeEvaluatorBase.h"
-#include "StateTreePropertyBindings.h"
-#include "StateTreeItemBlueprintBase.h"
+#include "StateTreeEvents.h"
+#include "StateTreeNodeBlueprintBase.h"
 #include "StateTreeEvaluatorBlueprintBase.generated.h"
 
 struct FStateTreeExecutionContext;
@@ -16,7 +16,7 @@ struct FStateTreeExecutionContext;
  * Base class for Blueprint based evaluators. 
  */
 UCLASS(Abstract, Blueprintable)
-class STATETREEMODULE_API UStateTreeEvaluatorBlueprintBase : public UStateTreeItemBlueprintBase
+class STATETREEMODULE_API UStateTreeEvaluatorBlueprintBase : public UStateTreeNodeBlueprintBase
 {
 	GENERATED_BODY()
 public:
@@ -45,7 +45,7 @@ public:
 
 	UFUNCTION(BlueprintImplementableEvent, meta = (DisplayName = "Tick"))
 	void ReceiveTick(AActor* OwnerActor, const float DeltaTime);
-	
+
 	UE_DEPRECATED(5.1, "This function will be removed for 5.1.")
 	virtual void EnterState(FStateTreeExecutionContext& Context, const EStateTreeStateChangeType ChangeType, const FStateTreeTransitionResult& Transition);
 	UE_DEPRECATED(5.1, "This function will be removed for 5.1.")
diff --git a/Engine/Plugins/Runtime/StateTree/Source/StateTreeModule/Public/Blueprint/StateTreeItemBlueprintBase.h b/Engine/Plugins/Runtime/StateTree/Source/StateTreeModule/Public/Blueprint/StateTreeItemBlueprintBase.h
deleted file mode 100644
diff --git a/Engine/Plugins/Runtime/StateTree/Source/StateTreeModule/Public/Blueprint/StateTreeNodeBlueprintBase.h b/Engine/Plugins/Runtime/StateTree/Source/StateTreeModule/Public/Blueprint/StateTreeNodeBlueprintBase.h
new file mode 100644
--- /dev/null
+++ b/Engine/Plugins/Runtime/StateTree/Source/StateTreeModule/Public/Blueprint/StateTreeNodeBlueprintBase.h
@@ -0,0 +1,93 @@
+// Copyright Epic Games, Inc. All Rights Reserved.
+
+#pragma once
+
+#include "StateTreeTypes.h"
+#include "StateTreeEvents.h"
+#include "StateTreeNodeBlueprintBase.generated.h"
+
+struct FStateTreeLinker;
+struct FStateTreeExecutionContext;
+
+UENUM()
+enum class EStateTreeBlueprintPropertyCategory : uint8
+{
+	NotSet,
+	Input,	
+	Parameter,
+	Output,
+	ExternalData,
+};
+
+/** Additional information for properties, captured from property metadata. */
+USTRUCT()
+struct STATETREEMODULE_API FStateTreeBlueprintPropertyInfo
+{
+	GENERATED_BODY()
+
+	FStateTreeBlueprintPropertyInfo() = default;
+
+	FStateTreeBlueprintPropertyInfo(const FName InPropertyName, const EStateTreeBlueprintPropertyCategory InCategory)
+		: PropertyName(InPropertyName), Category(InCategory)
+	{}
+	
+	UPROPERTY()
+	FName PropertyName;
+
+	UPROPERTY()
+	EStateTreeBlueprintPropertyCategory Category = EStateTreeBlueprintPropertyCategory::NotSet;
+};
+
+/** Struct use to copy external data to the Blueprint item instance, resolved during StateTree linking. */
+struct STATETREEMODULE_API FStateTreeBlueprintExternalDataHandle
+{
+	const FProperty* Property = nullptr;
+	FStateTreeExternalDataHandle Handle;
+};
+
+
+UCLASS(Abstract)
+class STATETREEMODULE_API UStateTreeNodeBlueprintBase : public UObject
+{
+	GENERATED_BODY()
+
+public:
+	/** Called during link to collect external data handles. */
+	void LinkExternalData(FStateTreeLinker& Linker, TArray<FStateTreeBlueprintExternalDataHandle>& OutExternalDataHandles) const;
+	/** Called before call to logic functions to copy data from */
+	void CopyExternalData(FStateTreeExecutionContext& Context, TConstArrayView<FStateTreeBlueprintExternalDataHandle> ExternalDataHandles);
+
+	/** Sends event to the StateTree */
+	UFUNCTION(BlueprintCallable, Category = "StateTree", meta = (HideSelfPin = "true", DisplayName = "StateTree Send Event"))
+	void SendEvent(const FStateTreeEvent& Event);
+	
+protected:
+	virtual UWorld* GetWorld() const override;
+	AActor* GetOwnerActor(const FStateTreeExecutionContext& Context) const;
+
+#if WITH_EDITOR
+	virtual void PostCDOCompiled(const FPostCDOCompiledContext& Context) override;
+#endif
+
+	struct FScopedCurrentContext
+	{
+		FScopedCurrentContext(const UStateTreeNodeBlueprintBase& InNode, FStateTreeExecutionContext& Context) :
+			Node(InNode)
+		{
+			Node.CurrentContext = &Context;
+		}
+
+		~FScopedCurrentContext()
+		{
+			Node.CurrentContext = nullptr;
+		}
+		const UStateTreeNodeBlueprintBase& Node;
+	};
+	
+	/** Cached execution context during Tick() and other functions. */
+	mutable FStateTreeExecutionContext* CurrentContext = nullptr;
+	
+	/** Metadata for properties */
+	UPROPERTY()
+	TArray<FStateTreeBlueprintPropertyInfo> PropertyInfos;
+};
diff --git a/Engine/Plugins/Runtime/StateTree/Source/StateTreeModule/Public/Blueprint/StateTreeTaskBlueprintBase.h b/Engine/Plugins/Runtime/StateTree/Source/StateTreeModule/Public/Blueprint/StateTreeTaskBlueprintBase.h
--- a/Engine/Plugins/Runtime/StateTree/Source/StateTreeModule/Public/Blueprint/StateTreeTaskBlueprintBase.h
+++ b/Engine/Plugins/Runtime/StateTree/Source/StateTreeModule/Public/Blueprint/StateTreeTaskBlueprintBase.h
@@ -6,8 +6,8 @@
 #include "Templates/SubclassOf.h"
 #include "StateTreeTypes.h"
 #include "StateTreeTaskBase.h"
-#include "StateTreePropertyBindings.h"
-#include "StateTreeItemBlueprintBase.h"
+#include "StateTreeEvents.h"
+#include "StateTreeNodeBlueprintBase.h"
 #include "StateTreeTaskBlueprintBase.generated.h"
 
 struct FStateTreeExecutionContext;
@@ -16,7 +16,7 @@ struct FStateTreeExecutionContext;
  * Base class for Blueprint based Tasks. 
  */
 UCLASS(Abstract, Blueprintable)
-class STATETREEMODULE_API UStateTreeTaskBlueprintBase : public UStateTreeItemBlueprintBase
+class STATETREEMODULE_API UStateTreeTaskBlueprintBase : public UStateTreeNodeBlueprintBase
 {
 	GENERATED_BODY()
 public:
diff --git a/Engine/Plugins/Runtime/StateTree/Source/StateTreeModule/Public/StateTree.h b/Engine/Plugins/Runtime/StateTree/Source/StateTreeModule/Public/StateTree.h
--- a/Engine/Plugins/Runtime/StateTree/Source/StateTreeModule/Public/StateTree.h
+++ b/Engine/Plugins/Runtime/StateTree/Source/StateTreeModule/Public/StateTree.h
@@ -27,6 +27,8 @@ struct STATETREEMODULE_API FStateTreeCustomVersion
 		InstanceDataArrays,
 		// Added index types.
 		IndexTypes,
+		// Added events.
+		AddedEvents,
 
 		// -----<new versions can be added above this line>-------------------------------------------------
 		VersionPlusOne,
diff --git a/Engine/Plugins/Runtime/StateTree/Source/StateTreeModule/Public/StateTreeEvents.h b/Engine/Plugins/Runtime/StateTree/Source/StateTreeModule/Public/StateTreeEvents.h
new file mode 100644
--- /dev/null
+++ b/Engine/Plugins/Runtime/StateTree/Source/StateTreeModule/Public/StateTreeEvents.h
@@ -0,0 +1,43 @@
+﻿// Copyright Epic Games, Inc. All Rights Reserved.
+
+#pragma once
+
+#include "GameplayTagContainer.h"
+#include "StateTreeEvents.generated.h"
+
+/** Enum used for flow control during event iteration. */
+UENUM()
+enum class EStateTreeLoopEvents : uint8
+{
+	/** Continues to next event. */
+	Next,
+	/** Stops the event handling loop. */
+	Break,
+};
+
+/**
+ * StateTree event with payload.
+ */
+USTRUCT(BlueprintType)
+struct STATETREEMODULE_API FStateTreeEvent
+{
+	GENERATED_BODY()
+
+	FStateTreeEvent() = default;
+	FStateTreeEvent(const FGameplayTag& InTag)
+		: Tag(InTag)
+	{
+	}
+	
+	/** Tag describing the event */
+	UPROPERTY(EditAnywhere, BlueprintReadWrite, Category = "Default", meta=(Categories="StateTreeEvent"))
+	FGameplayTag Tag;
+
+	/** Optional payload for the event. */ 
+	UPROPERTY(EditAnywhere, Category = "Default")
+	FInstancedStruct Payload;
+
+	/** Optional info to describe who sent the event. */
+	UPROPERTY(EditAnywhere, BlueprintReadWrite, Category = "Default")
+	FName Origin;
+};
diff --git a/Engine/Plugins/Runtime/StateTree/Source/StateTreeModule/Public/StateTreeExecutionContext.h b/Engine/Plugins/Runtime/StateTree/Source/StateTreeModule/Public/StateTreeExecutionContext.h
--- a/Engine/Plugins/Runtime/StateTree/Source/StateTreeModule/Public/StateTreeExecutionContext.h
+++ b/Engine/Plugins/Runtime/StateTree/Source/StateTreeModule/Public/StateTreeExecutionContext.h
@@ -11,6 +11,7 @@
 struct FStateTreeEvaluatorBase;
 struct FStateTreeTaskBase;
 struct FStateTreeConditionBase;
+struct FStateTreeEvent;
 
 USTRUCT()
 struct STATETREEMODULE_API FStateTreeExecutionState
@@ -101,6 +102,53 @@ public:
 	/** Tick the state tree logic. */
 	EStateTreeRunStatus Tick(const float DeltaTime, FStateTreeInstanceData* ExternalInstanceData = nullptr);
 
+	
+	/** @return the status of the last tick function */
+	EStateTreeRunStatus GetLastTickStatus(const FStateTreeInstanceData* ExternalInstanceData = nullptr) const;
+
+	/** @return reference to the list of currently active states. */
+	const FStateTreeActiveStates& GetActiveStates(const FStateTreeInstanceData* ExternalInstanceData = nullptr) const;
+
+#if WITH_GAMEPLAY_DEBUGGER
+	/** @return Debug string describing the current state of the execution */
+	FString GetDebugInfoString(const FStateTreeInstanceData* ExternalInstanceData = nullptr) const;
+#endif // WITH_GAMEPLAY_DEBUGGER
+
+#if WITH_STATETREE_DEBUG
+	int32 GetStateChangeCount(const FStateTreeInstanceData* ExternalInstanceData = nullptr) const;
+
+	void DebugPrintInternalLayout(const FStateTreeInstanceData* ExternalInstanceData = nullptr);
+#endif
+
+	/** @return the name of the active state. */
+	FString GetActiveStateName(const FStateTreeInstanceData* ExternalInstanceData = nullptr) const;
+	
+	/** @return the names of all the active state. */
+	TArray<FName> GetActiveStateNames(const FStateTreeInstanceData* ExternalInstanceData = nullptr) const;
+
+	/** Sends event for the StateTree. */
+	void SendExternalEvent(const FStateTreeEvent& Event, FStateTreeInstanceData* ExternalInstanceData = nullptr);
+
+	/** Sends event for the StateTree. Can only be used during StateTree tick. */
+	void SendEvent(const FStateTreeEvent& Event);
+
+	/** Iterates over all events. Can only be used during StateTree tick. Expects a lambda which takes const FStateTreeEvent& Event, and returns EStateTreeLoopEvents. */
+	template<typename TFunc>
+	void ForEachEvent(TFunc&& Function) const
+	{
+		for (const FStateTreeEvent& Event : EventsToProcess)
+		{
+			if (Function(Event) == EStateTreeLoopEvents::Break)
+			{
+				break;
+			}
+		}
+	}
+
+	/** @return events to process this tick. */
+	TConstArrayView<FStateTreeEvent> GetEventsToProcess() const { return EventsToProcess; }
+
+	
 	/** @return Pointer to a State or null if state not found */ 
 	const FCompactStateTreeState* GetStateFromHandle(const FStateTreeStateHandle StateHandle) const
 	{
@@ -122,47 +170,7 @@ public:
 	}
 
 	/** @return True if all required external data pointers are set. */ 
-	bool AreExternalDataViewsValid() const
-	{
-		check(StateTree);
-		bool bResult = true;
-		for (const FStateTreeExternalDataDesc& DataDesc : StateTree->ExternalDataDescs)
-		{
-			const FStateTreeDataView& DataView = DataViews[DataDesc.Handle.DataViewIndex.Get()];
-			
-			if (DataDesc.Requirement == EStateTreeExternalDataRequirement::Required)
-			{
-				// Required items must have valid pointer of the expected type.  
-				if (!DataView.IsValid() || !DataView.GetStruct()->IsChildOf(DataDesc.Struct))
-				{
-					bResult = false;
-					break;
-				}
-			}
-			else
-			{
-				// Optional items must have same type if they are set.
-				if (DataView.IsValid() && !DataView.GetStruct()->IsChildOf(DataDesc.Struct))
-				{
-					bResult = false;
-					break;
-				}
-			}
-		}
-
-		for (const FStateTreeExternalDataDesc& DataDesc : StateTree->GetNamedExternalDataDescs())
-		{
-			const FStateTreeDataView& DataView = DataViews[DataDesc.Handle.DataViewIndex.Get()];
-
-			// Items must have valid pointer of the expected type.  
-			if (!DataView.IsValid() || !DataView.GetStruct()->IsChildOf(DataDesc.Struct))
-			{
-				bResult = false;
-				break;
-			}
-		}
-		return bResult;
-	}
+	bool AreExternalDataViewsValid() const;
 
 	/** @return Handle to external data of type InStruct, or invalid handle if struct not found. */ 
 	FStateTreeExternalDataHandle GetExternalDataHandleByStruct(const UStruct* InStruct) const
@@ -241,22 +249,6 @@ public:
 		return Handle.IsValid() ? (typename T::DataType*)(DataViews[Handle.DataViewIndex.Get()].GetMemory() + Handle.PropertyOffset) : nullptr;
 	}
 
-	/**
-	 * Used internally by the Blueprint wrappers to get wrapped instance objects. 
-	 * @param DataViewIndex Index to a data view
-	 * @return Pointer to an instance object based.
-	 */
-	template <typename T>
-	T* GetInstanceObjectInternal(const FStateTreeIndex16 DataViewIndex) const
-	{
-		const UStruct* Struct = DataViews[DataViewIndex.Get()].GetStruct();
-		if (Struct != nullptr && Struct->IsChildOf<T>())
-		{
-			return DataViews[DataViewIndex.Get()].template GetMutablePtr<T>();
-		}
-		return nullptr;
-	}
-
 	/** @returns pointer to the instance data of specified node. */
 	template <typename T>
 	T* GetInstanceDataPtr(const FStateTreeNodeBase& Node) const
@@ -271,26 +263,6 @@ public:
 		return DataViews[Node.DataViewIndex.Get()].template GetMutable<T>();
 	}
 
-	/** @return the status of the last tick function */
-	EStateTreeRunStatus GetLastTickStatus(const FStateTreeInstanceData* ExternalInstanceData = nullptr) const;
-
-	/** @return reference to the list of currently active states. */
-	const FStateTreeActiveStates& GetActiveStates(const FStateTreeInstanceData* ExternalInstanceData = nullptr) const;
-
-#if WITH_GAMEPLAY_DEBUGGER
-	/** @return Debug string describing the current state of the execution */
-	FString GetDebugInfoString(const FStateTreeInstanceData* ExternalInstanceData = nullptr) const;
-#endif // WITH_GAMEPLAY_DEBUGGER
-
-#if WITH_STATETREE_DEBUG
-	int32 GetStateChangeCount(const FStateTreeInstanceData* ExternalInstanceData = nullptr) const;
-
-	void DebugPrintInternalLayout(const FStateTreeInstanceData* ExternalInstanceData = nullptr);
-#endif
-
-	FString GetActiveStateName(const FStateTreeInstanceData* ExternalInstanceData = nullptr) const;
-	TArray<FName> GetActiveStateNames(const FStateTreeInstanceData* ExternalInstanceData = nullptr) const;
-
 protected:
 
 	/** @return Prefix that will be used by STATETREE_LOG and STATETREE_CLOG, empty by default. */
@@ -409,6 +381,23 @@ protected:
 	/** @return String describing full path of an activate state for logging and debug. */
 	FString DebugGetStatePath(const FStateTreeActiveStates& ActiveStates, int32 ActiveStateIndex) const;
 
+	// Helper to set and clear CurrentInstanceData.
+	struct FScopedCurrentInstanceData
+	{
+		FScopedCurrentInstanceData(FStateTreeExecutionContext& InContext, FStateTreeInstanceData& InstanceData)
+			: Context(InContext)
+		{
+			Context.CurrentInstanceData = &InstanceData;
+		}
+
+		~FScopedCurrentInstanceData()
+		{
+			Context.CurrentInstanceData = nullptr;
+		}
+
+		FStateTreeExecutionContext& Context;
+	};
+	
 	/** The StateTree asset the context is initialized for */
 	UPROPERTY(Transient)
 	TObjectPtr<const UStateTree> StateTree = nullptr;
@@ -420,9 +409,15 @@ protected:
 	UPROPERTY(Transient)
 	FStateTreeInstanceData InternalInstanceData;
 
+	/** Instance data used during current tick. */
+	FStateTreeInstanceData* CurrentInstanceData = nullptr;
+
 	/** Array of data pointers (external data, tasks, evaluators, conditions), used during evaluation. Initialized to match the number of items in the asset. */
 	TArray<FStateTreeDataView> DataViews;
 
 	/** Storage type of the context */
 	EStateTreeStorage StorageType = EStateTreeStorage::Internal;
+
+	/** Events to process in current tick. */
+	TArray<FStateTreeEvent> EventsToProcess;
 };
diff --git a/Engine/Plugins/Runtime/StateTree/Source/StateTreeModule/Public/StateTreeInstanceData.h b/Engine/Plugins/Runtime/StateTree/Source/StateTreeModule/Public/StateTreeInstanceData.h
--- a/Engine/Plugins/Runtime/StateTree/Source/StateTreeModule/Public/StateTreeInstanceData.h
+++ b/Engine/Plugins/Runtime/StateTree/Source/StateTreeModule/Public/StateTreeInstanceData.h
@@ -4,6 +4,7 @@
 
 #include "InstancedStructArray.h"
 #include "StructView.h"
+#include "StateTreeEvents.h"
 #include "StateTreeInstanceData.generated.h"
 
 /**
@@ -58,6 +59,9 @@ struct STATETREEMODULE_API FStateTreeInstanceData
 	/** @return const pointer to an instance object   */
 	const UObject* GetObject(const int32 Index) const { return InstanceObjects[Index]; }
 
+	/** @return array to store unprocessed events. */
+	TArray<FStateTreeEvent>& GetEvents() { return Events; }
+	
 	int32 GetEstimatedMemoryUsage() const;
 	int32 GetNumItems() const;
 	
@@ -73,6 +77,10 @@ private:
 	/** Object instances. */
 	UPROPERTY()
 	TArray<TObjectPtr<UObject>> InstanceObjects;
+
+	/** Events */
+	UPROPERTY()
+	TArray<FStateTreeEvent> Events;
 };
 
 template<>
diff --git a/Engine/Plugins/Runtime/StateTree/Source/StateTreeModule/Public/StateTreeTypes.h b/Engine/Plugins/Runtime/StateTree/Source/StateTreeModule/Public/StateTreeTypes.h
--- a/Engine/Plugins/Runtime/StateTree/Source/StateTreeModule/Public/StateTreeTypes.h
+++ b/Engine/Plugins/Runtime/StateTree/Source/StateTreeModule/Public/StateTreeTypes.h
@@ -2,8 +2,9 @@
 
 #pragma once
 
-#include "InstancedStruct.h"
+#include "StructView.h"
 #include "PropertyBag.h"
+#include "GameplayTagContainer.h"
 #include "StateTreeTypes.generated.h"
 
 STATETREEMODULE_API DECLARE_LOG_CATEGORY_EXTERN(LogStateTree, Warning, All);
@@ -16,56 +17,92 @@ STATETREEMODULE_API DECLARE_LOG_CATEGORY_EXTERN(LogStateTree, Warning, All);
 UENUM()
 enum class EStateTreeRunStatus : uint8
 {
-	Running,			/// Tree is still running.
-	Failed,				/// Tree execution has stopped on failure.
-	Succeeded,			/// Tree execution has stopped on success.
-	Unset,				/// Status not set.
+	/** Tree is still running. */
+	Running,
+	
+	/** Tree execution has stopped on failure. */
+	Failed,
+	
+	/** Tree execution has stopped on success. */
+	Succeeded,
+	
+	/** Status not set. */
+	Unset,
 };
 
-/**  Evaluator evaluation type. */
+/** Evaluator evaluation type. */
 UENUM()
 enum class EStateTreeEvaluationType : uint8
 {
-	PreSelect,			/// Called during selection process on states that have not been visited yet.
-    Tick,				/// Called during tick on active states.
+	/** Called during selection process on states that have not been visited yet. */
+	PreSelect,
+	
+	/** Called during tick on active states. */
+    Tick,
 };
 
-/**  State change type. Passed to EnterState() and ExitState() to indicate how the state change affects the state and Evaluator or Task is on. */
+/** State change type. Passed to EnterState() and ExitState() to indicate how the state change affects the state and Evaluator or Task is on. */
 UENUM()
 enum class EStateTreeStateChangeType : uint8
 {
-	None,				/// Not an activation
-	Changed,			/// The state became activated or deactivated.
-    Sustained,			/// The state is parent of new active state and sustained previous active state.
+	/** Not an activation */
+	None,
+	
+	/** The state became activated or deactivated. */
+	Changed,
+	
+	/** The state is parent of new active state and sustained previous active state. */
+    Sustained,
 };
 
 /** Transitions behavior. */
 UENUM()
 enum class EStateTreeTransitionType : uint8
 {
-	Succeeded,			/// Signal StateTree execution succeeded.
-	Failed,				/// Signal StateTree execution failed.
-	GotoState,			/// Transition to specified state.
-	NotSet,				/// No transition.
-	NextState,			/// Goto next sibling state.
+	/** Stop StateTree and mark execution succeeded. */
+	Succeeded,
+	
+	/** Stop StateTree and mark execution failed. */
+	Failed,
+	
+	/** Transition to specified state. */
+	GotoState,
+	
+	/** No transition. */
+	NotSet,
+	
+	/** Goto next sibling state. */
+	NextState,
 };
 
 /** Operand between conditions */
 UENUM()
 enum class EStateTreeConditionOperand : uint8
 {
-	Copy UMETA(Hidden),	/// Copy result
-	And,				/// Combine results with AND.
-	Or,					/// Combine results with OR.
+	/** Copy result */
+	Copy UMETA(Hidden),
+	
+	/** Combine results with AND. */
+	And,
+	
+	/** Combine results with OR. */
+	Or,
 };
 
 UENUM()
 enum class EStateTreeStateType : uint8
 {
-	State,		// A State containing tasks and evaluators.
-	Group,		// A State containing just sub states.
-	Linked,		// A State that is linked to another state in the tree (the execution continues on the linked state).
-	Subtree,	// A subtree that can be linked to.
+	/** A State containing tasks and evaluators. */
+	State,
+	
+	/** A State containing just sub states. */
+	Group,
+	
+	/** A State that is linked to another state in the tree (the execution continues on the linked state). */
+	Linked,
+	
+	/** A subtree that can be linked to. */
+	Subtree,
 };
 
 namespace UE::StateTree
@@ -77,17 +114,30 @@ namespace UE::StateTree
 
 
 
-/** Transitions event. */
+/** Transitions trigger. */
 UENUM()
-enum class EStateTreeTransitionEvent : uint8
+enum class EStateTreeTransitionTrigger : uint8
 {
 	None = 0 UMETA(Hidden),
-	OnCompleted = 0x1 | 0x2,
-    OnSucceeded = 0x1,
-    OnFailed = 0x2,
-    OnCondition = 0x4,
+
+	/** Try trigger transition when a state succeeded or failed. */
+	OnStateCompleted = 0x1 | 0x2,
+
+	/** Try trigger transition when a state succeeded. */
+    OnStateSucceeded = 0x1,
+
+	/** Try trigger transition when a state failed. */
+    OnStateFailed = 0x2,
+
+	/** Try trigger transition each StateTree tick. */
+    OnTick = 0x4,
+	
+	/** Try trigger transition on specific event. */
+	OnEvent = 0x8,
+
+	MAX
 };
-ENUM_CLASS_FLAGS(EStateTreeTransitionEvent)
+ENUM_CLASS_FLAGS(EStateTreeTransitionTrigger)
 
 /** Handle to a StateTree state */
 USTRUCT(BlueprintType)
@@ -116,7 +166,7 @@ struct STATETREEMODULE_API FStateTreeStateHandle
 		switch (Index)
 		{
 		case InvalidIndex:		return TEXT("Invalid Item");
-		case SucceededIndex: 	return TEXT("Succeeded Item");
+		case SucceededIndex:	return TEXT("Succeeded Item");
 		case FailedIndex: 		return TEXT("Failed Item");
 		default: 				return FString::Printf(TEXT("%d"), Index);
 		}
@@ -167,6 +217,7 @@ protected:
 };
 	
 
+/** uint8 index that can be invalid. */
 USTRUCT(BlueprintType)
 struct STATETREEMODULE_API FStateTreeIndex8
 {
@@ -403,18 +454,33 @@ struct STATETREEMODULE_API FCompactStateTransition
 {
 	GENERATED_BODY()
 
+	/** Transition event tag, used when trigger type is event. */
 	UPROPERTY()
-	uint16 ConditionsBegin = 0;							// Index to first condition to test
+	FGameplayTag EventTag;
+
+	/** Index to first condition to test */
+	UPROPERTY()
+	uint16 ConditionsBegin = 0;
+
+	/** Target state of the transition */
 	UPROPERTY()
-	FStateTreeStateHandle State = FStateTreeStateHandle::Invalid;	// Target state of the transition
+	FStateTreeStateHandle State = FStateTreeStateHandle::Invalid;
+
+	/** Type of the transition. */
 	UPROPERTY()
-	EStateTreeTransitionType Type = EStateTreeTransitionType::NotSet;	// Type of the transition.
+	EStateTreeTransitionType Type = EStateTreeTransitionType::NotSet;
+ 
+	/* Type of the transition trigger. */
 	UPROPERTY()
-	EStateTreeTransitionEvent Event = EStateTreeTransitionEvent::None;	// Type of the transition event.
+	EStateTreeTransitionTrigger Trigger = EStateTreeTransitionTrigger::None;
+
+	/** The time the conditions need to hold true for the transition to become active, in tenths of a seconds. */
 	UPROPERTY()
-	uint8 GateDelay = 0;								// The time the conditions need to hold true for the transition to become active, in tenths of a seconds.
+	uint8 GateDelay = 0;
+
+	/** Number of conditions to test. */
 	UPROPERTY()
-	uint8 ConditionsNum = 0;							// Number of conditions to test.
+	uint8 ConditionsNum = 0;
 };
 
 /**
diff --git a/Engine/Plugins/Runtime/StateTree/Source/StateTreeTestSuite/Private/StateTreeTest.cpp b/Engine/Plugins/Runtime/StateTree/Source/StateTreeTestSuite/Private/StateTreeTest.cpp
--- a/Engine/Plugins/Runtime/StateTree/Source/StateTreeTestSuite/Private/StateTreeTest.cpp
+++ b/Engine/Plugins/Runtime/StateTree/Source/StateTreeTestSuite/Private/StateTreeTest.cpp
@@ -53,7 +53,7 @@ struct FStateTreeTest_MakeAndBakeStateTree : FAITestBase
 
 		EditorData.AddPropertyBinding(FStateTreeEditorPropertyPath(EvalA.ID, TEXT("IntA")), FStateTreeEditorPropertyPath(IntCond.ID, TEXT("Left")));
 
-		StateA.AddTransition(EStateTreeTransitionEvent::OnCompleted, EStateTreeTransitionType::GotoState, &StateB);
+		StateA.AddTransition(EStateTreeTransitionTrigger::OnStateCompleted, EStateTreeTransitionType::GotoState, &StateB);
 
 		// State B
 		auto& TaskB2 = StateB.AddTask<FTestTask_B>();
@@ -64,7 +64,7 @@ struct FStateTreeTest_MakeAndBakeStateTree : FAITestBase
 		TransFloatCond.GetInstance().Right = 13.0f;
 		EditorData.AddPropertyBinding(FStateTreeEditorPropertyPath(EvalA.ID, TEXT("FloatA")), FStateTreeEditorPropertyPath(TransFloatCond.ID, TEXT("Left")));
 
-		StateB.AddTransition(EStateTreeTransitionEvent::OnCompleted, EStateTreeTransitionType::Succeeded);
+		StateB.AddTransition(EStateTreeTransitionTrigger::OnStateCompleted, EStateTreeTransitionType::Succeeded);
 
 		FStateTreeCompilerLog Log;
 		FStateTreeCompiler Compiler(Log);
@@ -224,10 +224,10 @@ struct FStateTreeTest_Sequence : FAITestBase
 		UStateTreeState& State2 = Root.AddChildState(FName(TEXT("State2")));
 
 		auto& Task1 = State1.AddTask<FTestTask_Stand>(FName(TEXT("Task1")));
-		State1.AddTransition(EStateTreeTransitionEvent::OnCompleted, EStateTreeTransitionType::NextState);
+		State1.AddTransition(EStateTreeTransitionTrigger::OnStateCompleted, EStateTreeTransitionType::NextState);
 
 		auto& Task2 = State2.AddTask<FTestTask_Stand>(FName(TEXT("Task2")));
-		State2.AddTransition(EStateTreeTransitionEvent::OnCompleted, EStateTreeTransitionType::Succeeded);
+		State2.AddTransition(EStateTreeTransitionTrigger::OnStateCompleted, EStateTreeTransitionType::Succeeded);
 
 		FStateTreeCompilerLog Log;
 		FStateTreeCompiler Compiler(Log);
@@ -290,7 +290,7 @@ struct FStateTreeTest_Select : FAITestBase
 
 		auto& Task1A = State1A.AddTask<FTestTask_Stand>(FName(TEXT("Task1A")));
 		Task1A.GetItem().TicksToCompletion = 2;
-		State1A.AddTransition(EStateTreeTransitionEvent::OnCompleted, EStateTreeTransitionType::GotoState, &State1);
+		State1A.AddTransition(EStateTreeTransitionTrigger::OnStateCompleted, EStateTreeTransitionType::GotoState, &State1);
 
 		FStateTreeCompilerLog Log;
 		FStateTreeCompiler Compiler(Log);
@@ -362,7 +362,7 @@ struct FStateTreeTest_FailEnterState : FAITestBase
 		auto& Task3 = State1.AddTask<FTestTask_Stand>(FName(TEXT("Task3")));
 
 		auto& Task1A = State1A.AddTask<FTestTask_Stand>(FName(TEXT("Task1A")));
-		State1A.AddTransition(EStateTreeTransitionEvent::OnCompleted, EStateTreeTransitionType::GotoState, &State1);
+		State1A.AddTransition(EStateTreeTransitionTrigger::OnStateCompleted, EStateTreeTransitionType::GotoState, &State1);
 
 		FStateTreeCompilerLog Log;
 		FStateTreeCompiler Compiler(Log);
@@ -422,7 +422,7 @@ struct FStateTreeTest_SubTree : FAITestBase
 		auto& Task2 = State2.AddTask<FTestTask_Stand>(FName(TEXT("Task2")));
 
 		auto& Task3A = State3A.AddTask<FTestTask_Stand>(FName(TEXT("Task3A")));
-		State3A.AddTransition(EStateTreeTransitionEvent::OnCompleted, EStateTreeTransitionType::GotoState, &State3B);
+		State3A.AddTransition(EStateTreeTransitionTrigger::OnStateCompleted, EStateTreeTransitionType::GotoState, &State3B);
 
 		auto& Task3B = State3B.AddTask<FTestTask_Stand>(FName(TEXT("Task3B")));
 
diff --git a/Engine/Plugins/Runtime/StateTree/Source/StateTreeTestSuite/StateTreeTestSuite.Build.cs b/Engine/Plugins/Runtime/StateTree/Source/StateTreeTestSuite/StateTreeTestSuite.Build.cs
--- a/Engine/Plugins/Runtime/StateTree/Source/StateTreeTestSuite/StateTreeTestSuite.Build.cs
+++ b/Engine/Plugins/Runtime/StateTree/Source/StateTreeTestSuite/StateTreeTestSuite.Build.cs
@@ -23,6 +23,7 @@ namespace UnrealBuildTool.Rules
 					"StateTreeEditorModule",
 					"AITestSuite",
 					"StructUtils",
+					"GameplayTags",
 				}
 			);
 
