StateTree: Fix missing include

diff --git a/Engine/Plugins/Runtime/StateTree/Source/StateTreeModule/Public/StateTreeEvents.h b/Engine/Plugins/Runtime/StateTree/Source/StateTreeModule/Public/StateTreeEvents.h
--- a/Engine/Plugins/Runtime/StateTree/Source/StateTreeModule/Public/StateTreeEvents.h
+++ b/Engine/Plugins/Runtime/StateTree/Source/StateTreeModule/Public/StateTreeEvents.h
@@ -3,6 +3,7 @@
 #pragma once
 
 #include "GameplayTagContainer.h"
+#include "InstancedStruct.h"
 #include "StateTreeEvents.generated.h"
 
 /** Enum used for flow control during event iteration. */
