Do not error when a state tree needs recompiling in the editor

diff --git a/Engine/Plugins/Runtime/StateTree/Source/StateTreeModule/Private/StateTree.cpp b/Engine/Plugins/Runtime/StateTree/Source/StateTreeModule/Private/StateTree.cpp
--- a/Engine/Plugins/Runtime/StateTree/Source/StateTreeModule/Private/StateTree.cpp
+++ b/Engine/Plugins/Runtime/StateTree/Source/StateTreeModule/Private/StateTree.cpp
@@ -73,8 +73,10 @@ void UStateTree::PostLoad()
 	{
 #if WITH_EDITOR
 		ResetCompiled();
+		UE_LOG(LogStateTree, Warning, TEXT("%s: StateTree compiled data in older format. Please resave the StateTree asset."), *GetPathName());
+#else
+		UE_LOG(LogStateTree, Error, TEXT("%s: StateTree compiled data in older format. Please recompile the StateTree asset."), *GetPathName());
 #endif
-		UE_LOG(LogStateTree, Error, TEXT("%s: StateTree compiled data in older format. Please recompile the StateTree asset."), *GetName());
 		return;
 	}
 	
