StateTree:
- Fixed potential crash in StateTree editor row
- Fixed StateTreeComponentSchema actor class not updating property on load

diff --git a/Engine/Plugins/Runtime/GameplayStateTree/Source/GameplayStateTreeModule/Private/Components/StateTreeComponentSchema.cpp b/Engine/Plugins/Runtime/GameplayStateTree/Source/GameplayStateTreeModule/Private/Components/StateTreeComponentSchema.cpp
--- a/Engine/Plugins/Runtime/GameplayStateTree/Source/GameplayStateTreeModule/Private/Components/StateTreeComponentSchema.cpp
+++ b/Engine/Plugins/Runtime/GameplayStateTree/Source/GameplayStateTreeModule/Private/Components/StateTreeComponentSchema.cpp
@@ -39,6 +39,12 @@ TConstArrayView<FStateTreeExternalDataDesc> UStateTreeComponentSchema::GetNamedE
 	return MakeArrayView(&ContextActorDataDesc, 1);
 }
 
+void UStateTreeComponentSchema::PostLoad()
+{
+	Super::PostLoad();
+	ContextActorDataDesc.Struct = ContextActorClass.Get();
+}
+
 #if WITH_EDITOR
 void UStateTreeComponentSchema::PostEditChangeChainProperty(FPropertyChangedChainEvent& PropertyChangedEvent)
 {
diff --git a/Engine/Plugins/Runtime/GameplayStateTree/Source/GameplayStateTreeModule/Public/Components/StateTreeComponentSchema.h b/Engine/Plugins/Runtime/GameplayStateTree/Source/GameplayStateTreeModule/Public/Components/StateTreeComponentSchema.h
--- a/Engine/Plugins/Runtime/GameplayStateTree/Source/GameplayStateTreeModule/Public/Components/StateTreeComponentSchema.h
+++ b/Engine/Plugins/Runtime/GameplayStateTree/Source/GameplayStateTreeModule/Public/Components/StateTreeComponentSchema.h
@@ -27,6 +27,8 @@ protected:
 	
 	virtual TConstArrayView<FStateTreeExternalDataDesc> GetNamedExternalDataDescs() const;
 
+	virtual void PostLoad() override;
+
 #if WITH_EDITOR
 	virtual void PostEditChangeChainProperty(FPropertyChangedChainEvent& PropertyChangedEvent) override;
 #endif // WITH_EDITOR
diff --git a/Engine/Plugins/Runtime/StateTree/Source/StateTreeEditorModule/Private/SStateTreeViewRow.cpp b/Engine/Plugins/Runtime/StateTree/Source/StateTreeEditorModule/Private/SStateTreeViewRow.cpp
--- a/Engine/Plugins/Runtime/StateTree/Source/StateTreeEditorModule/Private/SStateTreeViewRow.cpp
+++ b/Engine/Plugins/Runtime/StateTree/Source/StateTreeEditorModule/Private/SStateTreeViewRow.cpp
@@ -654,38 +654,39 @@ FText SStateTreeViewRow::GetTransitionsIcon(const UStateTreeState& State, const
 
 EVisibility SStateTreeViewRow::GetTransitionsVisibility(const UStateTreeState& State, const EStateTreeTransitionTrigger Trigger) const
 {
-	if (!EnumHasAnyFlags(Trigger, EStateTreeTransitionTrigger::OnStateCompleted))
+	// Handle completed, succeeeded and failed transitions.
+	if (EnumHasAnyFlags(Trigger, EStateTreeTransitionTrigger::OnStateCompleted))
 	{
+		const bool bIsLeafState = (State.Children.Num() == 0);
+		EStateTreeTransitionTrigger HandledTriggers = EStateTreeTransitionTrigger::None;
+		bool bExactMatch = false;
+
 		for (const FStateTreeTransition& Transition : State.Transitions)
 		{
-			if (EnumHasAnyFlags(Trigger, Transition.Trigger))
-			{
-				return EVisibility::Visible;
-			}
+			HandledTriggers |= Transition.Trigger;
+			bExactMatch |= (Transition.Trigger == Trigger);
 		}
-		return EVisibility::Collapsed;
-	}
-	
-	EStateTreeTransitionTrigger HandledTriggers = EStateTreeTransitionTrigger::None;
-	TStaticArray<int32, static_cast<uint32>(EStateTreeTransitionTrigger::MAX)> Counts;
-	Algo::ForEach(Counts, [](int32& Count) { Count = 0; }); 
 
-	for (const FStateTreeTransition& Transition : State.Transitions)
-	{
-		HandledTriggers |= Transition.Trigger;
-		Counts[static_cast<uint8>(Transition.Trigger)]++;
+		// Assume that leaf states should have completion transitions.
+		if (!bExactMatch && bIsLeafState)
+		{
+			// Find the missing transition type, note: Completed = Succeeded|Failed.
+			const EStateTreeTransitionTrigger MissingTriggers = HandledTriggers ^ EStateTreeTransitionTrigger::OnStateCompleted;
+			return MissingTriggers == Trigger ? EVisibility::Visible : EVisibility::Collapsed;
+		}
+		
+		return bExactMatch ? EVisibility::Visible : EVisibility::Collapsed;
 	}
 
-	// Assume that leaf states should have completion transitions.
-	if (State.Children.Num() == 0
-	    && Counts[static_cast<uint8>(Trigger)] == 0)
+	// Handle the test
+	for (const FStateTreeTransition& Transition : State.Transitions)
 	{
-		// Find the missing transition type, note: Completed = Succeeded|Failed.
-		const EStateTreeTransitionTrigger MissingTriggers = HandledTriggers ^ EStateTreeTransitionTrigger::OnStateCompleted;
-		return MissingTriggers == Trigger ? EVisibility::Visible : EVisibility::Collapsed;
+		if (EnumHasAnyFlags(Trigger, Transition.Trigger))
+		{
+			return EVisibility::Visible;
+		}
 	}
-	
-	return Counts[static_cast<uint8>(Trigger)] > 0 ? EVisibility::Visible : EVisibility::Collapsed;
+	return EVisibility::Collapsed;
 }
 
 EVisibility SStateTreeViewRow::GetCompletedTransitionVisibility() const
