StateTree: Threadsafe shared data
- added thread safe access to shared instance data (each thread gets a copy)
- fixed Item->Node in the editor node API

diff --git a/Engine/Plugins/Runtime/StateTree/Source/StateTreeEditorModule/Public/StateTreeEditorNode.h b/Engine/Plugins/Runtime/StateTree/Source/StateTreeEditorModule/Public/StateTreeEditorNode.h
--- a/Engine/Plugins/Runtime/StateTree/Source/StateTreeEditorModule/Public/StateTreeEditorNode.h
+++ b/Engine/Plugins/Runtime/StateTree/Source/StateTreeEditorModule/Public/StateTreeEditorNode.h
@@ -65,7 +65,7 @@ struct STATETREEEDITORMODULE_API FStateTreeEditorNode
 template <typename T>
 struct TStateTreeEditorNode : public FStateTreeEditorNode
 {
-	typedef T NodeType;
-	FORCEINLINE T& GetItem() { return Node.template GetMutable<T>(); }
-	FORCEINLINE typename T::InstanceDataType& GetInstance() { return Instance.template GetMutable<typename T::InstanceDataType>(); }
+	using NodeType = T;
+	FORCEINLINE T& GetNode() { return Node.template GetMutable<T>(); }
+	FORCEINLINE typename T::InstanceDataType& GetInstanceData() { return Instance.template GetMutable<typename T::InstanceDataType>(); }
 };
diff --git a/Engine/Plugins/Runtime/StateTree/Source/StateTreeModule/Private/StateTree.cpp b/Engine/Plugins/Runtime/StateTree/Source/StateTreeModule/Private/StateTree.cpp
--- a/Engine/Plugins/Runtime/StateTree/Source/StateTreeModule/Private/StateTree.cpp
+++ b/Engine/Plugins/Runtime/StateTree/Source/StateTreeModule/Private/StateTree.cpp
@@ -6,6 +6,7 @@
 #include "StateTreeTaskBase.h"
 #include "StateTreeEvaluatorBase.h"
 #include "AssetRegistry/AssetData.h"
+#include "Misc/ScopeRWLock.h"
 
 const FGuid FStateTreeCustomVersion::GUID(0x28E21331, 0x501F4723, 0x8110FA64, 0xEA10DA1E);
 FCustomVersionRegistration GRegisterStateTreeCustomVersion(FStateTreeCustomVersion::GUID, FStateTreeCustomVersion::LatestVersion, TEXT("StateTreeAsset"));
@@ -16,6 +17,42 @@ bool UStateTree::IsReadyToRun() const
 	return States.Num() > 0 && bIsLinked;
 }
 
+TSharedPtr<FStateTreeInstanceData> UStateTree::GetSharedInstanceData() const
+{
+	// Create a unique index for each thread.
+	static std::atomic_int ThreadIndexCounter {0};
+	static thread_local int32 ThreadIndex = ThreadIndexCounter.fetch_add(1);
+
+	// If shared instance data for this thread exists, return it.
+	{
+		FReadScopeLock ReadLock(PerThreadSharedInstanceDataLock);
+		if (ThreadIndex < PerThreadSharedInstanceData.Num())
+		{
+			return PerThreadSharedInstanceData[ThreadIndex];
+		}
+	}
+
+	// Not initialized yet, create new instances up to the index.
+	FWriteScopeLock WriteLock(PerThreadSharedInstanceDataLock);
+
+	// It is possible that multiple threads are waiting for the write lock,
+	// which means that execution may get here so that 'ThreadIndex' is already in valid range.
+	// The loop below is organized to handle that too.
+	
+	const int32 NewNum = ThreadIndex + 1;
+	PerThreadSharedInstanceData.Reserve(NewNum);
+	UStateTree* NonConstThis = const_cast<UStateTree*>(this); 
+	
+	for (int32 Index = PerThreadSharedInstanceData.Num(); Index < NewNum; Index++)
+	{
+		TSharedPtr<FStateTreeInstanceData> SharedData = MakeShared<FStateTreeInstanceData>();
+		SharedData->CopyFrom(*NonConstThis, SharedInstanceData);
+		PerThreadSharedInstanceData.Add(SharedData);
+	}
+
+	return PerThreadSharedInstanceData[ThreadIndex];
+}
+
 #if WITH_EDITOR
 void UStateTree::ResetCompiled()
 {
@@ -109,6 +146,9 @@ void UStateTree::ResetLinked()
 	ExternalDataDescs.Reset();
 	ExternalDataBaseIndex = 0;
 	NumDataViews = 0;
+
+	FWriteScopeLock WriteLock(PerThreadSharedInstanceDataLock);
+	PerThreadSharedInstanceData.Reset();
 }
 
 bool UStateTree::Link()
diff --git a/Engine/Plugins/Runtime/StateTree/Source/StateTreeModule/Private/StateTreeExecutionContext.cpp b/Engine/Plugins/Runtime/StateTree/Source/StateTreeModule/Private/StateTreeExecutionContext.cpp
--- a/Engine/Plugins/Runtime/StateTree/Source/StateTreeModule/Private/StateTreeExecutionContext.cpp
+++ b/Engine/Plugins/Runtime/StateTree/Source/StateTreeModule/Private/StateTreeExecutionContext.cpp
@@ -185,7 +185,10 @@ EStateTreeRunStatus FStateTreeExecutionContext::Start(FStateTreeInstanceData* Ex
 			return EStateTreeRunStatus::Failed;
 		}
 	}
-	
+
+	const TSharedPtr<FStateTreeInstanceData> SharedInstanceData = StateTree->GetSharedInstanceData();
+	check(SharedInstanceData.IsValid());
+
 	FScopedCurrentInstanceData InstanceDataScope(*this, InstanceData); 
 	FStateTreeExecutionState* Exec = &GetExecState(InstanceData);
 
@@ -209,7 +212,7 @@ EStateTreeRunStatus FStateTreeExecutionContext::Start(FStateTreeInstanceData* Ex
 	static const FStateTreeStateHandle RootState = FStateTreeStateHandle(0);
 
 	FStateTreeActiveStates NextActiveStates;
-	if (SelectState(InstanceData, RootState, NextActiveStates))
+	if (SelectState(InstanceData, *SharedInstanceData.Get(), RootState, NextActiveStates))
 	{
 		if (NextActiveStates.Last() == FStateTreeStateHandle::Succeeded || NextActiveStates.Last() == FStateTreeStateHandle::Failed)
 		{
@@ -326,6 +329,9 @@ EStateTreeRunStatus FStateTreeExecutionContext::Tick(const float DeltaTime, FSta
 		return EStateTreeRunStatus::Failed;
 	}
 
+	const TSharedPtr<FStateTreeInstanceData> SharedInstanceData = StateTree->GetSharedInstanceData();
+	check(SharedInstanceData.IsValid());
+
 	FScopedCurrentInstanceData InstanceDataScope(*this, InstanceData); 
 	FStateTreeExecutionState* Exec = &GetExecState(InstanceData);
 
@@ -372,7 +378,7 @@ EStateTreeRunStatus FStateTreeExecutionContext::Tick(const float DeltaTime, FSta
 	{
 		// Trigger conditional transitions or state succeed/failed transitions. First tick transition is handled here too.
 		FStateTreeTransitionResult Transition;
-		if (TriggerTransitions(InstanceData, Transition))
+		if (TriggerTransitions(InstanceData, *SharedInstanceData.Get(), Transition))
 		{
 			// We have committed to state change, consume events that were accumulated during the tick above.
 			InstanceData.GetEvents().Reset();
@@ -1054,7 +1060,7 @@ EStateTreeRunStatus FStateTreeExecutionContext::TickTasks(FStateTreeInstanceData
 	return Result;
 }
 
-bool FStateTreeExecutionContext::TestAllConditions(const int32 ConditionsOffset, const int32 ConditionsNum)
+bool FStateTreeExecutionContext::TestAllConditions(FStateTreeInstanceData& SharedInstanceData, const int32 ConditionsOffset, const int32 ConditionsNum)
 {
 	CSV_SCOPED_TIMING_STAT_EXCLUSIVE(StateTree_TestConditions);
 
@@ -1063,8 +1069,6 @@ bool FStateTreeExecutionContext::TestAllConditions(const int32 ConditionsOffset,
 		return true;
 	}
 
-	const FStateTreeInstanceData& SharedInstanceData = StateTree->GetSharedInstanceData();
-	
 	TStaticArray<EStateTreeConditionOperand, UE::StateTree::MaxConditionIndent + 1> Operands(InPlace, EStateTreeConditionOperand::Copy);
 	TStaticArray<bool, UE::StateTree::MaxConditionIndent + 1> Values(InPlace, false);
 
@@ -1134,7 +1138,7 @@ bool FStateTreeExecutionContext::TestAllConditions(const int32 ConditionsOffset,
 	return Values[0];
 }
 
-bool FStateTreeExecutionContext::TriggerTransitions(FStateTreeInstanceData& InstanceData, FStateTreeTransitionResult& OutTransition)
+bool FStateTreeExecutionContext::TriggerTransitions(FStateTreeInstanceData& InstanceData, FStateTreeInstanceData& SharedInstanceData, FStateTreeTransitionResult& OutTransition)
 {
 	CSV_SCOPED_TIMING_STAT_EXCLUSIVE(StateTree_TriggerTransition);
 
@@ -1178,7 +1182,7 @@ bool FStateTreeExecutionContext::TriggerTransitions(FStateTreeInstanceData& Inst
 									|| Transition.Trigger == EStateTreeTransitionTrigger::OnTick
 									|| (Transition.Trigger == EStateTreeTransitionTrigger::OnEvent && HasEvent(Transition.EventTag));
 			
-			if (bShouldCheck && TestAllConditions(Transition.ConditionsBegin, Transition.ConditionsNum))
+			if (bShouldCheck && TestAllConditions(SharedInstanceData, Transition.ConditionsBegin, Transition.ConditionsNum))
 			{
 				// If a transition has delay, we stop testing other transitions, but the transition will not pass the condition until the delay time passes.
 				if (Transition.GateDelay > 0)
@@ -1206,7 +1210,7 @@ bool FStateTreeExecutionContext::TriggerTransitions(FStateTreeInstanceData& Inst
 					OutTransition.TargetState = Transition.State;
 					OutTransition.NextActiveStates.Reset();
 
-					if (SelectState(InstanceData, Transition.State, OutTransition.NextActiveStates))
+					if (SelectState(InstanceData, SharedInstanceData, Transition.State, OutTransition.NextActiveStates))
 					{
 						STATETREE_LOG(Verbose, TEXT("Transition on state '%s' (%s) -[%s]-> state '%s'"), *GetSafeStateName(Exec.ActiveStates.Last()), *State.Name.ToString(), *GetSafeStateName(Transition.State), *GetSafeStateName(OutTransition.NextActiveStates.Last()));
 						return true;
@@ -1248,14 +1252,14 @@ bool FStateTreeExecutionContext::TriggerTransitions(FStateTreeInstanceData& Inst
 		// Could not trigger completion transition, jump back to start.
 		static const FStateTreeStateHandle RootState = FStateTreeStateHandle(0);
 		OutTransition.TargetState = RootState;
-		return SelectState(InstanceData, RootState, OutTransition.NextActiveStates);
+		return SelectState(InstanceData, SharedInstanceData, RootState, OutTransition.NextActiveStates);
 	}
 
 	// No transition triggered, keep on updating current state.
 	return false;
 }
 
-bool FStateTreeExecutionContext::SelectState(FStateTreeInstanceData& InstanceData, const FStateTreeStateHandle NextState, FStateTreeActiveStates& OutNewActiveState)
+bool FStateTreeExecutionContext::SelectState(FStateTreeInstanceData& InstanceData, FStateTreeInstanceData& SharedInstanceData, const FStateTreeStateHandle NextState, FStateTreeActiveStates& OutNewActiveState)
 {
 	const FStateTreeExecutionState& Exec = GetExecState(InstanceData);
 
@@ -1311,10 +1315,10 @@ bool FStateTreeExecutionContext::SelectState(FStateTreeInstanceData& InstanceDat
 		return false;
 	}
 
-	return SelectStateInternal(InstanceData, NextState, OutNewActiveState);
+	return SelectStateInternal(InstanceData, SharedInstanceData, NextState, OutNewActiveState);
 }
 
-bool FStateTreeExecutionContext::SelectStateInternal(FStateTreeInstanceData& InstanceData, const FStateTreeStateHandle NextState, FStateTreeActiveStates& OutNewActiveState)
+bool FStateTreeExecutionContext::SelectStateInternal(FStateTreeInstanceData& InstanceData, FStateTreeInstanceData& SharedInstanceData, const FStateTreeStateHandle NextState, FStateTreeActiveStates& OutNewActiveState)
 {
 	CSV_SCOPED_TIMING_STAT_EXCLUSIVE(StateTree_SelectState);
 
@@ -1331,7 +1335,7 @@ bool FStateTreeExecutionContext::SelectStateInternal(FStateTreeInstanceData& Ins
 	const FCompactStateTreeState& State = StateTree->States[NextState.Index];
 
 	// Check that the state can be entered
-	if (TestAllConditions(State.EnterConditionsBegin, State.EnterConditionsNum))
+	if (TestAllConditions(SharedInstanceData, State.EnterConditionsBegin, State.EnterConditionsNum))
 	{
 		if (!OutNewActiveState.Push(NextState))
 		{
@@ -1343,7 +1347,7 @@ bool FStateTreeExecutionContext::SelectStateInternal(FStateTreeInstanceData& Ins
 		if (State.LinkedState.IsValid())
 		{
 			// If State is linked, proceed to the linked state.
-			if (SelectStateInternal(InstanceData, State.LinkedState, OutNewActiveState))
+			if (SelectStateInternal(InstanceData, SharedInstanceData, State.LinkedState, OutNewActiveState))
 			{
 				// Selection succeeded
 				return true;
@@ -1354,7 +1358,7 @@ bool FStateTreeExecutionContext::SelectStateInternal(FStateTreeInstanceData& Ins
 			// If the state has children, proceed to select children.
 			for (uint16 ChildState = State.ChildrenBegin; ChildState < State.ChildrenEnd; ChildState = StateTree->States[ChildState].GetNextSibling())
 			{
-				if (SelectStateInternal(InstanceData, FStateTreeStateHandle(ChildState), OutNewActiveState))
+				if (SelectStateInternal(InstanceData, SharedInstanceData, FStateTreeStateHandle(ChildState), OutNewActiveState))
 				{
 					// Selection succeeded
 					return true;
diff --git a/Engine/Plugins/Runtime/StateTree/Source/StateTreeModule/Public/StateTree.h b/Engine/Plugins/Runtime/StateTree/Source/StateTreeModule/Public/StateTree.h
--- a/Engine/Plugins/Runtime/StateTree/Source/StateTreeModule/Public/StateTree.h
+++ b/Engine/Plugins/Runtime/StateTree/Source/StateTreeModule/Public/StateTree.h
@@ -9,6 +9,7 @@
 #include "InstancedStructArray.h"
 #include "StateTreePropertyBindings.h"
 #include "StateTreeInstanceData.h"
+#include "Misc/ScopeRWLock.h"
 #include "StateTree.generated.h"
 
 
@@ -77,11 +78,8 @@ class STATETREEMODULE_API UStateTree : public UDataAsset
 
 public:
 
-	/**
-	 * @todo: This should return different data for each thread.
-	 * @return Shared instance data.
-	 */
-	const FStateTreeInstanceData& GetSharedInstanceData() const { return SharedInstanceData; }
+	/** @return Shared instance data. */
+	TSharedPtr<FStateTreeInstanceData> GetSharedInstanceData() const;
 
 	/** @return Number of data views required for StateTree execution (Evaluators, Tasks, Conditions, External data). */
 	int32 GetNumDataViews() const { return NumDataViews; }
@@ -169,6 +167,9 @@ private:
 	UPROPERTY()
 	FStateTreeInstanceData SharedInstanceData;
 
+	mutable FRWLock PerThreadSharedInstanceDataLock;
+	mutable TArray<TSharedPtr<FStateTreeInstanceData>> PerThreadSharedInstanceData;
+	
 	/** List of names external data enforced by the schema, created at compilation. */
 	UPROPERTY()
 	TArray<FStateTreeExternalDataDesc> NamedExternalDataDescs;
diff --git a/Engine/Plugins/Runtime/StateTree/Source/StateTreeModule/Public/StateTreeExecutionContext.h b/Engine/Plugins/Runtime/StateTree/Source/StateTreeModule/Public/StateTreeExecutionContext.h
--- a/Engine/Plugins/Runtime/StateTree/Source/StateTreeModule/Public/StateTreeExecutionContext.h
+++ b/Engine/Plugins/Runtime/StateTree/Source/StateTreeModule/Public/StateTreeExecutionContext.h
@@ -312,7 +312,7 @@ protected:
 	 * Checks all conditions at given range
 	 * @return True if all conditions pass.
 	 */
-	bool TestAllConditions(const int32 ConditionsOffset, const int32 ConditionsNum);
+	bool TestAllConditions(FStateTreeInstanceData& SharedInstanceData, const int32 ConditionsOffset, const int32 ConditionsNum);
 
 	/**
 	 * Triggers transitions based on current run status. CurrentStatus is used to select which transitions events are triggered.
@@ -321,7 +321,7 @@ protected:
 	 * the actual next state returned by the selector.
 	 * @return Transition result describing the source state, state transitioned to, and next selected state.
 	 */
-	bool TriggerTransitions(FStateTreeInstanceData& InstanceData, FStateTreeTransitionResult& OutTransition);
+	bool TriggerTransitions(FStateTreeInstanceData& InstanceData, FStateTreeInstanceData& SharedInstanceData, FStateTreeTransitionResult& OutTransition);
 
 	/**
 	 * Runs state selection logic starting at the specified state, walking towards the leaf states.
@@ -333,12 +333,12 @@ protected:
 	 * @param OutNewActiveStates Active states that got selected.
 	 * @return True if succeeded to select new active states.
 	 */
-	bool SelectState(FStateTreeInstanceData& InstanceData, const FStateTreeStateHandle NextState, FStateTreeActiveStates& OutNewActiveStates);
+	bool SelectState(FStateTreeInstanceData& InstanceData, FStateTreeInstanceData& SharedInstanceData, const FStateTreeStateHandle NextState, FStateTreeActiveStates& OutNewActiveStates);
 
 	/**
 	 * Used internally to do the recursive part of the SelectState().
 	 */
-	bool SelectStateInternal(FStateTreeInstanceData& InstanceData, const FStateTreeStateHandle NextState, FStateTreeActiveStates& OutNewActiveStates);
+	bool SelectStateInternal(FStateTreeInstanceData& InstanceData, FStateTreeInstanceData& SharedInstanceData, const FStateTreeStateHandle NextState, FStateTreeActiveStates& OutNewActiveStates);
 
 	/** @return Mutable storage based on storage settings. */
 	FStateTreeInstanceData& SelectMutableInstanceData(FStateTreeInstanceData* ExternalInstanceData)
diff --git a/Engine/Plugins/Runtime/StateTree/Source/StateTreeTestSuite/Private/StateTreeTest.cpp b/Engine/Plugins/Runtime/StateTree/Source/StateTreeTestSuite/Private/StateTreeTest.cpp
--- a/Engine/Plugins/Runtime/StateTree/Source/StateTreeTestSuite/Private/StateTreeTest.cpp
+++ b/Engine/Plugins/Runtime/StateTree/Source/StateTreeTestSuite/Private/StateTreeTest.cpp
@@ -10,11 +10,14 @@
 #include "StateTreeTestTypes.h"
 #include "StateTreeExecutionContext.h"
 #include "Engine/World.h"
+#include "Async/ParallelFor.h"
 
 #define LOCTEXT_NAMESPACE "AITestSuite_StateTreeTest"
 
 PRAGMA_DISABLE_OPTIMIZATION
 
+std::atomic<int32> FStateTreeTestConditionInstanceData::GlobalCounter = 0;
+
 namespace UE::StateTree::Tests
 {
 	UStateTree& NewStateTree(UObject* Outer = GetTransientPackage())
@@ -49,7 +52,7 @@ struct FStateTreeTest_MakeAndBakeStateTree : FAITestBase
 		EditorData.AddPropertyBinding(FStateTreeEditorPropertyPath(EvalA.ID, TEXT("IntA")), FStateTreeEditorPropertyPath(TaskB1.ID, TEXT("IntB")));
 
 		auto& IntCond = StateA.AddEnterCondition<FStateTreeCompareIntCondition>(EGenericAICheck::Less);
-		IntCond.GetInstance().Right = 2;
+		IntCond.GetInstanceData().Right = 2;
 
 		EditorData.AddPropertyBinding(FStateTreeEditorPropertyPath(EvalA.ID, TEXT("IntA")), FStateTreeEditorPropertyPath(IntCond.ID, TEXT("Left")));
 
@@ -61,7 +64,7 @@ struct FStateTreeTest_MakeAndBakeStateTree : FAITestBase
 
 		FStateTreeTransition& Trans = StateB.AddTransition({}, EStateTreeTransitionType::GotoState, &Root);
 		auto& TransFloatCond = Trans.AddCondition<FStateTreeCompareFloatCondition>(EGenericAICheck::Less);
-		TransFloatCond.GetInstance().Right = 13.0f;
+		TransFloatCond.GetInstanceData().Right = 13.0f;
 		EditorData.AddPropertyBinding(FStateTreeEditorPropertyPath(EvalA.ID, TEXT("FloatA")), FStateTreeEditorPropertyPath(TransFloatCond.ID, TEXT("Left")));
 
 		StateB.AddTransition(EStateTreeTransitionTrigger::OnStateCompleted, EStateTreeTransitionType::Succeeded);
@@ -106,34 +109,34 @@ struct FStateTreeTest_WanderLoop : FAITestBase
 		auto& ReserveSOTask = UseSmartObjectOnLane.AddTask<FTestTask_ReserveSmartObject>(FName(TEXT("ReserveSOTask")));
 		EditorData.AddPropertyBinding(FStateTreeEditorPropertyPath(SmartObjectEval.ID, TEXT("PotentialSmartObjects")), FStateTreeEditorPropertyPath(ReserveSOTask.ID, TEXT("PotentialSmartObjects")));
 		auto& ReserveHasSmartObjects = UseSmartObjectOnLane.AddEnterCondition<FStateTreeCompareBoolCondition>();
-		ReserveHasSmartObjects.GetInstance().bRight = true;
+		ReserveHasSmartObjects.GetInstanceData().bRight = true;
 		EditorData.AddPropertyBinding(FStateTreeEditorPropertyPath(SmartObjectEval.ID, TEXT("bHasSmartObjects")), FStateTreeEditorPropertyPath(ReserveHasSmartObjects.ID, TEXT("bLeft")));
 		UseSmartObjectOnLane.AddTransition(EStateTreeTransitionTrigger::OnStateCompleted, EStateTreeTransitionType::GotoState, &Wander); // This catches child states too
 
 		//      |  |- WalkToSO
 		auto& MoveToSOTask = WalkToSO.AddTask<FTestTask_MoveTo>(FName(TEXT("MoveToSOTask")));
-		MoveToSOTask.GetItem().TicksToCompletion = 2;
+		MoveToSOTask.GetNode().TicksToCompletion = 2;
 		EditorData.AddPropertyBinding(FStateTreeEditorPropertyPath(ReserveSOTask.ID, TEXT("ReservedSmartObjectLocation")), FStateTreeEditorPropertyPath(MoveToSOTask.ID, TEXT("MoveLocation")));
 		WalkToSO.AddTransition(EStateTreeTransitionTrigger::OnStateCompleted, EStateTreeTransitionType::NextState);
 
 		//      |  \- UseSO
 		auto& UseSOTask = UseSO.AddTask<FTestTask_UseSmartObject>(FName(TEXT("UseSOTask")));
-		UseSOTask.GetItem().TicksToCompletion = 2;
+		UseSOTask.GetNode().TicksToCompletion = 2;
 		EditorData.AddPropertyBinding(FStateTreeEditorPropertyPath(ReserveSOTask.ID, TEXT("ReservedSmartObject")), FStateTreeEditorPropertyPath(UseSOTask.ID, TEXT("SmartObject")));
 		// UseSO uses UseSmartObjectOnLane completed transition.
 
 		//      |- Walk Along Lane
 		auto& MoveAlongLaneTask = WalkAlongLane.AddTask<FTestTask_MoveTo>(FName(TEXT("MoveAlongLaneTask")));
-		MoveAlongLaneTask.GetItem().TicksToCompletion = 2;
+		MoveAlongLaneTask.GetNode().TicksToCompletion = 2;
 		EditorData.AddPropertyBinding(FStateTreeEditorPropertyPath(WanderEval.ID, TEXT("WanderLocation")), FStateTreeEditorPropertyPath(MoveAlongLaneTask.ID, TEXT("MoveLocation")));
 		auto& MoveHasWanderLoc = WalkAlongLane.AddEnterCondition<FStateTreeCompareBoolCondition>();
-		MoveHasWanderLoc.GetInstance().bRight = true;
+		MoveHasWanderLoc.GetInstanceData().bRight = true;
 		EditorData.AddPropertyBinding(FStateTreeEditorPropertyPath(WanderEval.ID, TEXT("bHasWanderLocation")), FStateTreeEditorPropertyPath(MoveHasWanderLoc.ID, TEXT("bLeft")));
 		WalkAlongLane.AddTransition(EStateTreeTransitionTrigger::OnStateCompleted, EStateTreeTransitionType::GotoState, &Wander);
 
 		//      \- StandOnLane
 		auto& StandTask = StandOnLane.AddTask<FTestTask_Stand>(FName(TEXT("StandTask")));
-		StandTask.GetItem().TicksToCompletion = 2;
+		StandTask.GetNode().TicksToCompletion = 2;
 		StandOnLane.AddTransition(EStateTreeTransitionTrigger::OnStateCompleted, EStateTreeTransitionType::GotoState, &Wander);
 
 		FStateTreeCompilerLog Log;
@@ -283,13 +286,13 @@ struct FStateTreeTest_Select : FAITestBase
 		UStateTreeState& State1A = State1.AddChildState(FName(TEXT("State1A")));
 
 		auto& TaskRoot = Root.AddTask<FTestTask_Stand>(FName(TEXT("TaskRoot")));
-		TaskRoot.GetItem().TicksToCompletion = 2;
+		TaskRoot.GetNode().TicksToCompletion = 2;
 
 		auto& Task1 = State1.AddTask<FTestTask_Stand>(FName(TEXT("Task1")));
-		Task1.GetItem().TicksToCompletion = 2;
+		Task1.GetNode().TicksToCompletion = 2;
 
 		auto& Task1A = State1A.AddTask<FTestTask_Stand>(FName(TEXT("Task1A")));
-		Task1A.GetItem().TicksToCompletion = 2;
+		Task1A.GetNode().TicksToCompletion = 2;
 		State1A.AddTransition(EStateTreeTransitionTrigger::OnStateCompleted, EStateTreeTransitionType::GotoState, &State1);
 
 		FStateTreeCompilerLog Log;
@@ -358,7 +361,7 @@ struct FStateTreeTest_FailEnterState : FAITestBase
 
 		auto& Task1 = State1.AddTask<FTestTask_Stand>(FName(TEXT("Task1")));
 		auto& Task2 = State1.AddTask<FTestTask_Stand>(FName(TEXT("Task2")));
-		Task2.GetItem().EnterStateResult = EStateTreeRunStatus::Failed;
+		Task2.GetNode().EnterStateResult = EStateTreeRunStatus::Failed;
 		auto& Task3 = State1.AddTask<FTestTask_Stand>(FName(TEXT("Task3")));
 
 		auto& Task1A = State1A.AddTask<FTestTask_Stand>(FName(TEXT("Task1A")));
@@ -461,6 +464,100 @@ struct FStateTreeTest_SubTree : FAITestBase
 };
 IMPLEMENT_AI_INSTANT_TEST(FStateTreeTest_SubTree, "System.StateTree.SubTree");
 
+
+struct FStateTreeTest_SharedInstanceData : FAITestBase
+{
+	virtual bool InstantTest() override
+	{
+		UStateTree& StateTree = UE::StateTree::Tests::NewStateTree(&GetWorld());
+		UStateTreeEditorData& EditorData = *Cast<UStateTreeEditorData>(StateTree.EditorData);
+		
+		UStateTreeState& Root = EditorData.AddSubTree(FName(TEXT("Root")));
+		auto& IntCond = Root.AddEnterCondition<FStateTreeTestCondition>();
+		IntCond.GetInstanceData().Count = 1;
+
+		auto& Task = Root.AddTask<FTestTask_Stand>(FName(TEXT("Task")));
+		Task.GetNode().TicksToCompletion = 2;
+
+		FStateTreeCompilerLog Log;
+		FStateTreeCompiler Compiler(Log);
+		const bool bResult = Compiler.Compile(StateTree);
+		AITEST_TRUE("StateTree should get compiled", bResult);
+
+		// Init, nothing should access the shared data.
+		constexpr int32 NumConcurrent = 100;
+		FStateTreeTestConditionInstanceData::GlobalCounter = 0;
+
+		bool bInitSucceeded = true;
+		TArray<FTestStateTreeExecutionContext> ExecContexts;
+		ExecContexts.SetNum(NumConcurrent);
+		for (FTestStateTreeExecutionContext& Exec : ExecContexts)
+		{
+			bInitSucceeded &= Exec.Init(StateTree, StateTree, EStateTreeStorage::Internal);
+		}
+		AITEST_TRUE("All StateTree contexts should init", bInitSucceeded);
+		AITEST_EQUAL("Test condition global counter should be 0", (int32)FStateTreeTestConditionInstanceData::GlobalCounter, 0);
+		
+		// Start in parallel
+		// This should create shared data per thread.
+		// We expect that ParallelForWithTaskContext() creates a context per thread.
+		TArray<FStateTreeTestRunContext> RunContexts;
+		
+		ParallelForWithTaskContext(
+			RunContexts,
+			ExecContexts.Num(),
+			[&ExecContexts](FStateTreeTestRunContext& RunContext, int32 Index)
+			{
+				const EStateTreeRunStatus Status = ExecContexts[Index].Start();
+				if (Status == EStateTreeRunStatus::Running)
+				{
+					RunContext.Count++;
+				}
+			}
+		);
+
+		int32 StartTotalRunning = 0;
+		for (FStateTreeTestRunContext RunContext : RunContexts)
+		{
+			StartTotalRunning += RunContext.Count;
+		}
+		AITEST_EQUAL("All StateTree contexts should be running after Start", StartTotalRunning, NumConcurrent);
+		AITEST_EQUAL("Test condition global counter should equal context count after Start", (int32)FStateTreeTestConditionInstanceData::GlobalCounter, ExecContexts.Num());
+
+		
+		// Tick in parallel
+		// This should not recreate the data, so FStateTreeTestConditionInstanceData::GlobalCounter should stay as is.
+		for (FStateTreeTestRunContext RunContext : RunContexts)
+		{
+			RunContext.Count = 0;
+		}
+
+		ParallelForWithTaskContext(
+			RunContexts,
+			ExecContexts.Num(),
+			[&ExecContexts](FStateTreeTestRunContext& RunContext, int32 Index)
+			{
+				const EStateTreeRunStatus Status = ExecContexts[Index].Tick(0.1f);
+				if (Status == EStateTreeRunStatus::Running)
+				{
+					RunContext.Count++;
+				}
+			}
+		);
+
+		int32 TickTotalRunning = 0;
+		for (FStateTreeTestRunContext RunContext : RunContexts)
+		{
+			TickTotalRunning += RunContext.Count;
+		}
+		AITEST_EQUAL("All StateTree contexts should be running after Tick", TickTotalRunning, NumConcurrent);
+		AITEST_EQUAL("Test condition global counter should equal context count after Tick", (int32)FStateTreeTestConditionInstanceData::GlobalCounter, ExecContexts.Num());
+
+		return true;
+	}
+};
+IMPLEMENT_AI_INSTANT_TEST(FStateTreeTest_SharedInstanceData, "System.StateTree.SharedInstanceData");
+
 PRAGMA_ENABLE_OPTIMIZATION
 
 #undef LOCTEXT_NAMESPACE
diff --git a/Engine/Plugins/Runtime/StateTree/Source/StateTreeTestSuite/Private/StateTreeTestTypes.h b/Engine/Plugins/Runtime/StateTree/Source/StateTreeTestSuite/Private/StateTreeTestTypes.h
--- a/Engine/Plugins/Runtime/StateTree/Source/StateTreeTestSuite/Private/StateTreeTestTypes.h
+++ b/Engine/Plugins/Runtime/StateTree/Source/StateTreeTestSuite/Private/StateTreeTestTypes.h
@@ -4,6 +4,7 @@
 
 #include "StateTreeTaskBase.h"
 #include "StateTreeEvaluatorBase.h"
+#include "StateTreeConditionBase.h"
 #include "StateTreeExecutionContext.h"
 #include "StateTreeLinker.h"
 #include "StateTreeTestTypes.generated.h"
@@ -838,3 +839,37 @@ struct FTestTask_UseSmartObject : public FStateTreeTaskBase
 	EStateTreeRunStatus EnterStateResult = EStateTreeRunStatus::Running;
 };
 
+USTRUCT()
+struct FStateTreeTestConditionInstanceData
+{
+	GENERATED_BODY()
+	
+	UPROPERTY(EditAnywhere, Category = Input)
+	int32 Count = 1;
+
+	static std::atomic<int32> GlobalCounter;
+};
+
+USTRUCT()
+struct FStateTreeTestCondition : public FStateTreeConditionCommonBase
+{
+	GENERATED_BODY()
+
+	typedef FStateTreeTestConditionInstanceData InstanceDataType;
+
+	FStateTreeTestCondition() = default;
+
+	virtual const UStruct* GetInstanceDataType() const override { return InstanceDataType::StaticStruct(); }
+	virtual bool TestCondition(FStateTreeExecutionContext& Context) const override
+	{
+		InstanceDataType& InstanceData = Context.GetInstanceData<InstanceDataType>(*this);
+		InstanceData.GlobalCounter.fetch_add(InstanceData.Count);
+		return true;
+	}
+};
+
+struct FStateTreeTestRunContext
+{
+	int32 Count = 0;
+};
+
