StateTree: Fixed the schema filtering
- change schema types to include the module path
- removed unused GetStorageSuperStruct()

diff --git a/Engine/Plugins/AI/MassAI/Source/MassAIBehavior/Private/MassStateTreeSchema.cpp b/Engine/Plugins/AI/MassAI/Source/MassAIBehavior/Private/MassStateTreeSchema.cpp
--- a/Engine/Plugins/AI/MassAI/Source/MassAIBehavior/Private/MassStateTreeSchema.cpp
+++ b/Engine/Plugins/AI/MassAI/Source/MassAIBehavior/Private/MassStateTreeSchema.cpp
@@ -6,12 +6,6 @@
 #include "StateTreeConditionBase.h"
 #include "Subsystems/WorldSubsystem.h"
 
-UScriptStruct* UMassStateTreeSchema::GetStorageSuperStruct() const
-{
-	// The runtime state must be derived from FMassFragment since it's stored in the entity subsystem.
-	return FMassFragment::StaticStruct();
-}
-
 bool UMassStateTreeSchema::IsStructAllowed(const UScriptStruct* InScriptStruct) const
 {
 	// Only allow Mass evals and tasks,and common conditions.
diff --git a/Engine/Plugins/AI/MassAI/Source/MassAIBehavior/Public/MassStateTreeSchema.h b/Engine/Plugins/AI/MassAI/Source/MassAIBehavior/Public/MassStateTreeSchema.h
--- a/Engine/Plugins/AI/MassAI/Source/MassAIBehavior/Public/MassStateTreeSchema.h
+++ b/Engine/Plugins/AI/MassAI/Source/MassAIBehavior/Public/MassStateTreeSchema.h
@@ -15,7 +15,6 @@ class MASSAIBEHAVIOR_API UMassStateTreeSchema : public UStateTreeSchema
 
 protected:
 
-	virtual UScriptStruct* GetStorageSuperStruct() const override;
 	virtual bool IsStructAllowed(const UScriptStruct* InScriptStruct) const override;
 	virtual bool IsExternalItemAllowed(const UStruct& InStruct) const override;
 };
diff --git a/Engine/Plugins/AI/MassAI/Source/MassAIBehavior/Public/MassStateTreeTrait.h b/Engine/Plugins/AI/MassAI/Source/MassAIBehavior/Public/MassStateTreeTrait.h
--- a/Engine/Plugins/AI/MassAI/Source/MassAIBehavior/Public/MassStateTreeTrait.h
+++ b/Engine/Plugins/AI/MassAI/Source/MassAIBehavior/Public/MassStateTreeTrait.h
@@ -21,6 +21,6 @@ protected:
 
 	virtual void ValidateTemplate(FMassEntityTemplateBuildContext& BuildContext, const UWorld& World) const override;
 
-	UPROPERTY(Category="StateTree", EditAnywhere, meta=(RequiredAssetDataTags="Schema=MassStateTreeSchema"))
+	UPROPERTY(Category="StateTree", EditAnywhere, meta=(RequiredAssetDataTags="Schema=/Script/MassAIBehavior/MassStateTreeSchema"))
 	TObjectPtr<UStateTree> StateTree;
 };
diff --git a/Engine/Plugins/Runtime/StateTree/Source/StateTreeEditorModule/Private/Customizations/StateTreeReferenceDetails.cpp b/Engine/Plugins/Runtime/StateTree/Source/StateTreeEditorModule/Private/Customizations/StateTreeReferenceDetails.cpp
--- a/Engine/Plugins/Runtime/StateTree/Source/StateTreeEditorModule/Private/Customizations/StateTreeReferenceDetails.cpp
+++ b/Engine/Plugins/Runtime/StateTree/Source/StateTreeEditorModule/Private/Customizations/StateTreeReferenceDetails.cpp
@@ -27,6 +27,8 @@ void FStateTreeReferenceDetails::CustomizeHeader(const TSharedRef<IPropertyHandl
 	const TSharedPtr<IPropertyHandle> StateTreeProperty = StructPropertyHandle->GetChildHandle(FName(TEXT("StateTree")));
 	check(StateTreeProperty.IsValid());
 
+	const FString SchemaMetaDataValue = InStructPropertyHandle->GetMetaData(UE::StateTree::SchemaTag);
+
 	InHeaderRow
 	.NameContent()
 	[
@@ -38,9 +40,8 @@ void FStateTreeReferenceDetails::CustomizeHeader(const TSharedRef<IPropertyHandl
 		.PropertyHandle(StateTreeProperty)
 		.AllowedClass(UStateTree::StaticClass())
 		.ThumbnailPool(InCustomizationUtils.GetThumbnailPool())
-		.OnShouldFilterAsset_Lambda([InStructPropertyHandle](const FAssetData& InAssetData)
+		.OnShouldFilterAsset_Lambda([SchemaMetaDataValue](const FAssetData& InAssetData)
 		{
-			const FString& SchemaMetaDataValue = InStructPropertyHandle->GetMetaData(UE::StateTree::SchemaTag);
 			return !SchemaMetaDataValue.IsEmpty() && !InAssetData.TagsAndValues.ContainsKeyValue(UE::StateTree::SchemaTag, SchemaMetaDataValue);
 		})
 	]
diff --git a/Engine/Plugins/Runtime/StateTree/Source/StateTreeModule/Public/StateTreeSchema.h b/Engine/Plugins/Runtime/StateTree/Source/StateTreeModule/Public/StateTreeSchema.h
--- a/Engine/Plugins/Runtime/StateTree/Source/StateTreeModule/Public/StateTreeSchema.h
+++ b/Engine/Plugins/Runtime/StateTree/Source/StateTreeModule/Public/StateTreeSchema.h
@@ -21,9 +21,6 @@ class STATETREEMODULE_API UStateTreeSchema : public UObject
 
 public:
 
-	/** @return Returns the script struct the storage struct will be derived from. */
-	virtual UScriptStruct* GetStorageSuperStruct() const { return nullptr; }
-
 	/** @return True if specified struct is supported */
 	virtual bool IsStructAllowed(const UScriptStruct* InScriptStruct) const { return false; }
 
