StateTree: Fix WinRT compiler error.
"dynamic initialization of thread local data not allowed in WinRT code"

diff --git a/Engine/Plugins/Runtime/StateTree/Source/StateTreeModule/Private/StateTree.cpp b/Engine/Plugins/Runtime/StateTree/Source/StateTreeModule/Private/StateTree.cpp
--- a/Engine/Plugins/Runtime/StateTree/Source/StateTreeModule/Private/StateTree.cpp
+++ b/Engine/Plugins/Runtime/StateTree/Source/StateTreeModule/Private/StateTree.cpp
@@ -21,7 +21,11 @@ TSharedPtr<FStateTreeInstanceData> UStateTree::GetSharedInstanceData() const
 {
 	// Create a unique index for each thread.
 	static std::atomic_int ThreadIndexCounter {0};
-	static thread_local int32 ThreadIndex = ThreadIndexCounter.fetch_add(1);
+	static thread_local int32 ThreadIndex = INDEX_NONE; // Cannot init directly on WinRT
+	if (ThreadIndex == INDEX_NONE)
+	{
+		ThreadIndex = ThreadIndexCounter.fetch_add(1);
+	}
 
 	// If shared instance data for this thread exists, return it.
 	{
