StateTree: include mode double conversion for the type promotions

diff --git a/Engine/Plugins/Runtime/StateTree/Source/StateTreeEditorModule/Private/StateTreePropertyBindingCompiler.cpp b/Engine/Plugins/Runtime/StateTree/Source/StateTreeEditorModule/Private/StateTreePropertyBindingCompiler.cpp
--- a/Engine/Plugins/Runtime/StateTree/Source/StateTreeEditorModule/Private/StateTreePropertyBindingCompiler.cpp
+++ b/Engine/Plugins/Runtime/StateTree/Source/StateTreeEditorModule/Private/StateTreePropertyBindingCompiler.cpp
@@ -329,33 +329,53 @@ EStateTreePropertyCopyType FStateTreePropertyBindingCompiler::GetCopyType(const
 				return EStateTreePropertyCopyType::PromoteInt32ToDouble;
 			}
 		}
+		else if (SourceProperty->IsA<FUInt32Property>())
+		{
+			if (TargetProperty->IsA<FInt64Property>())
+			{
+				return EStateTreePropertyCopyType::PromoteUInt32ToInt64;
+			}
+			else if (TargetProperty->IsA<FFloatProperty>())
+			{
+				return EStateTreePropertyCopyType::PromoteUInt32ToFloat;
+			}
+			else if (TargetProperty->IsA<FDoubleProperty>())
+			{
+				return EStateTreePropertyCopyType::PromoteUInt32ToDouble;
+			}
+		}
 		else if (SourceProperty->IsA<FFloatProperty>())
 		{
-			if (TargetProperty->IsA<FDoubleProperty>())
+			if (TargetProperty->IsA<FIntProperty>())
+			{
+				return EStateTreePropertyCopyType::PromoteFloatToInt32;
+			}
+			else if (TargetProperty->IsA<FInt64Property>())
+			{
+				return EStateTreePropertyCopyType::PromoteFloatToInt64;
+			}
+			else if (TargetProperty->IsA<FDoubleProperty>())
 			{
 				return EStateTreePropertyCopyType::PromoteFloatToDouble;
 			}
 		}
 		else if (SourceProperty->IsA<FDoubleProperty>())
 		{
-			if (TargetProperty->IsA<FFloatProperty>())
+			if (TargetProperty->IsA<FIntProperty>())
 			{
-				return EStateTreePropertyCopyType::DemoteDoubleToFloat;
+				return EStateTreePropertyCopyType::DemoteDoubleToInt32;
 			}
-		}
-		else if (SourceProperty->IsA<FUInt32Property>())
-		{
-			if (TargetProperty->IsA<FInt64Property>())
+			else if (TargetProperty->IsA<FInt64Property>())
 			{
-				return EStateTreePropertyCopyType::PromoteUInt32ToInt64;
+				return EStateTreePropertyCopyType::DemoteDoubleToInt64;
 			}
 			else if (TargetProperty->IsA<FFloatProperty>())
 			{
-				return EStateTreePropertyCopyType::PromoteUInt32ToFloat;
+				return EStateTreePropertyCopyType::DemoteDoubleToFloat;
 			}
 		}
 	}
-
+	
 	ensureMsgf(false, TEXT("Couldnt determine property copy type (%s -> %s)"), *SourceProperty->GetNameCPP(), *TargetProperty->GetNameCPP());
 
 	return EStateTreePropertyCopyType::None;
@@ -547,69 +567,100 @@ bool FStateTreePropertyBindingCompiler::ResolvePropertyPath(const FStateTreeBind
 	return true;
 }
 
-EPropertyAccessCompatibility FStateTreePropertyBindingCompiler::GetPropertyCompatibility(const FProperty* InPropertyA, const FProperty* InPropertyB)
+EPropertyAccessCompatibility FStateTreePropertyBindingCompiler::GetPropertyCompatibility(const FProperty* FromProperty, const FProperty* ToProperty)
 {
-	if (InPropertyA == InPropertyB)
+	if (FromProperty == ToProperty)
 	{
 		return EPropertyAccessCompatibility::Compatible;
 	}
 
-	if (InPropertyA == nullptr || InPropertyB == nullptr)
+	if (FromProperty == nullptr || ToProperty == nullptr)
 	{
 		return EPropertyAccessCompatibility::Incompatible;
 	}
 
 	// Special case for object properties since InPropertyA->SameType(InPropertyB) requires both properties to be of the exact same class.
 	// In our case we want to be able to bind a source property if its class is a child of the target property class.
-	if (InPropertyA->IsA<FObjectPropertyBase>() && InPropertyB->IsA<FObjectPropertyBase>())
+	if (FromProperty->IsA<FObjectPropertyBase>() && ToProperty->IsA<FObjectPropertyBase>())
 	{
-		const FObjectPropertyBase* SourceProperty = CastField<FObjectPropertyBase>(InPropertyA);
-		const FObjectPropertyBase* TargetProperty = CastField<FObjectPropertyBase>(InPropertyB);
+		const FObjectPropertyBase* SourceProperty = CastField<FObjectPropertyBase>(FromProperty);
+		const FObjectPropertyBase* TargetProperty = CastField<FObjectPropertyBase>(ToProperty);
 		return (SourceProperty->PropertyClass->IsChildOf(TargetProperty->PropertyClass)) ? EPropertyAccessCompatibility::Compatible : EPropertyAccessCompatibility::Incompatible;
 	}
 
 	// Extract underlying types for enums
-	if (const FEnumProperty* EnumPropertyA = CastField<const FEnumProperty>(InPropertyA))
+	if (const FEnumProperty* EnumPropertyA = CastField<const FEnumProperty>(FromProperty))
 	{
-		InPropertyA = EnumPropertyA->GetUnderlyingProperty();
+		FromProperty = EnumPropertyA->GetUnderlyingProperty();
 	}
 
-	if (const FEnumProperty* EnumPropertyB = CastField<const FEnumProperty>(InPropertyB))
+	if (const FEnumProperty* EnumPropertyB = CastField<const FEnumProperty>(ToProperty))
 	{
-		InPropertyB = EnumPropertyB->GetUnderlyingProperty();
+		ToProperty = EnumPropertyB->GetUnderlyingProperty();
 	}
 
-	if (InPropertyA->SameType(InPropertyB))
+	if (FromProperty->SameType(ToProperty))
 	{
 		return EPropertyAccessCompatibility::Compatible;
 	}
 	else
 	{
 		// Not directly compatible, check for promotions
-		if (InPropertyA->IsA<FBoolProperty>())
+		if (FromProperty->IsA<FBoolProperty>())
+		{
+			if (ToProperty->IsA<FByteProperty>()
+				|| ToProperty->IsA<FIntProperty>()
+				|| ToProperty->IsA<FUInt32Property>()
+				|| ToProperty->IsA<FInt64Property>()
+				|| ToProperty->IsA<FFloatProperty>()
+				|| ToProperty->IsA<FDoubleProperty>())
+			{
+				return EPropertyAccessCompatibility::Promotable;
+			}
+		}
+		else if (FromProperty->IsA<FByteProperty>())
+		{
+			if (ToProperty->IsA<FIntProperty>()
+				|| ToProperty->IsA<FUInt32Property>()
+				|| ToProperty->IsA<FInt64Property>()
+				|| ToProperty->IsA<FFloatProperty>()
+				|| ToProperty->IsA<FDoubleProperty>())
+			{
+				return EPropertyAccessCompatibility::Promotable;
+			}
+		}
+		else if (FromProperty->IsA<FIntProperty>())
 		{
-			if (InPropertyB->IsA<FByteProperty>() || InPropertyB->IsA<FIntProperty>() || InPropertyB->IsA<FUInt32Property>() || InPropertyB->IsA<FInt64Property>() || InPropertyB->IsA<FFloatProperty>())
+			if (ToProperty->IsA<FInt64Property>()
+				|| ToProperty->IsA<FFloatProperty>()
+				|| ToProperty->IsA<FDoubleProperty>())
 			{
 				return EPropertyAccessCompatibility::Promotable;
 			}
 		}
-		else if (InPropertyA->IsA<FByteProperty>())
+		else if (FromProperty->IsA<FUInt32Property>())
 		{
-			if (InPropertyB->IsA<FIntProperty>() || InPropertyB->IsA<FUInt32Property>() || InPropertyB->IsA<FInt64Property>() || InPropertyB->IsA<FFloatProperty>())
+			if (ToProperty->IsA<FInt64Property>()
+				|| ToProperty->IsA<FFloatProperty>()
+				|| ToProperty->IsA<FDoubleProperty>())
 			{
 				return EPropertyAccessCompatibility::Promotable;
 			}
 		}
-		else if (InPropertyA->IsA<FIntProperty>())
+		else if (FromProperty->IsA<FFloatProperty>())
 		{
-			if (InPropertyB->IsA<FInt64Property>() || InPropertyB->IsA<FFloatProperty>())
+			if (ToProperty->IsA<FIntProperty>()
+				|| ToProperty->IsA<FInt64Property>()
+				|| ToProperty->IsA<FDoubleProperty>())
 			{
 				return EPropertyAccessCompatibility::Promotable;
 			}
 		}
-		else if (InPropertyA->IsA<FUInt32Property>())
+		else if (FromProperty->IsA<FDoubleProperty>())
 		{
-			if (InPropertyB->IsA<FInt64Property>() || InPropertyB->IsA<FFloatProperty>())
+			if (ToProperty->IsA<FIntProperty>()
+				|| ToProperty->IsA<FInt64Property>()
+				|| ToProperty->IsA<FFloatProperty>())
 			{
 				return EPropertyAccessCompatibility::Promotable;
 			}
diff --git a/Engine/Plugins/Runtime/StateTree/Source/StateTreeEditorModule/Public/StateTreePropertyBindingCompiler.h b/Engine/Plugins/Runtime/StateTree/Source/StateTreeEditorModule/Public/StateTreePropertyBindingCompiler.h
--- a/Engine/Plugins/Runtime/StateTree/Source/StateTreeEditorModule/Public/StateTreePropertyBindingCompiler.h
+++ b/Engine/Plugins/Runtime/StateTree/Source/StateTreeEditorModule/Public/StateTreePropertyBindingCompiler.h
@@ -75,11 +75,11 @@ struct STATETREEEDITORMODULE_API FStateTreePropertyBindingCompiler
 
 	/**
 	 * Checks if two property types can are compatible for copying.
-	 * @param InPropertyA First property used in the check.
-	 * @param InPropertyB Second property used in the check.
+	 * @param FromProperty Property to copy from.
+	 * @param ToProperty Property to copy to.
 	 * @return Incompatible if the properties cannot be copied, Compatible if they are trivially copyable, or Promotable if numeric values can be promoted to another numeric type.
 	 */
-	static EPropertyAccessCompatibility GetPropertyCompatibility(const FProperty* InPropertyA, const FProperty* InPropertyB);
+	static EPropertyAccessCompatibility GetPropertyCompatibility(const FProperty* FromProperty, const FProperty* ToProperty);
 	
 protected:
 
diff --git a/Engine/Plugins/Runtime/StateTree/Source/StateTreeModule/Private/Conditions/StateTreeCommonConditions.cpp b/Engine/Plugins/Runtime/StateTree/Source/StateTreeModule/Private/Conditions/StateTreeCommonConditions.cpp
--- a/Engine/Plugins/Runtime/StateTree/Source/StateTreeModule/Private/Conditions/StateTreeCommonConditions.cpp
+++ b/Engine/Plugins/Runtime/StateTree/Source/StateTreeModule/Private/Conditions/StateTreeCommonConditions.cpp
@@ -97,7 +97,7 @@ bool FStateTreeCompareFloatCondition::TestCondition(FStateTreeExecutionContext&
 {
 	const InstanceDataType& InstanceData = Context.GetInstanceData<InstanceDataType>(*this);
 
-	const bool bResult = UE::StateTree::Conditions::CompareNumbers<float>(InstanceData.Left, InstanceData.Right, Operator);
+	const bool bResult = UE::StateTree::Conditions::CompareNumbers<double>(InstanceData.Left, InstanceData.Right, Operator);
 	return bResult ^ bInvert;
 }
 
@@ -179,9 +179,9 @@ bool FStateTreeCompareDistanceCondition::TestCondition(FStateTreeExecutionContex
 {
 	const InstanceDataType& InstanceData = Context.GetInstanceData<InstanceDataType>(*this);
 
-	const float Left = FVector::DistSquared(InstanceData.Source, InstanceData.Target);
-	const float Right = FMath::Square(InstanceData.Distance);
-	const bool bResult = UE::StateTree::Conditions::CompareNumbers<float>(Left, Right, Operator);
+	const FVector::FReal Left = FVector::DistSquared(InstanceData.Source, InstanceData.Target);
+	const FVector::FReal Right = FMath::Square(InstanceData.Distance);
+	const bool bResult = UE::StateTree::Conditions::CompareNumbers<FVector::FReal>(Left, Right, Operator);
 	return bResult ^ bInvert;
 }
 
diff --git a/Engine/Plugins/Runtime/StateTree/Source/StateTreeModule/Private/StateTreePropertyBindings.cpp b/Engine/Plugins/Runtime/StateTree/Source/StateTreeModule/Private/StateTreePropertyBindings.cpp
--- a/Engine/Plugins/Runtime/StateTree/Source/StateTreeModule/Private/StateTreePropertyBindings.cpp
+++ b/Engine/Plugins/Runtime/StateTree/Source/StateTreeModule/Private/StateTreePropertyBindings.cpp
@@ -258,6 +258,8 @@ bool FStateTreePropertyBindings::ValidateCopy(FStateTreePropCopy& Copy) const
 			&& TargetStructProperty->Struct == TBaseStructure<FStateTreeStructRef>::Get();
 		break;
 	}
+
+	// Bool promotions		
 	case EStateTreePropertyCopyType::PromoteBoolToByte:
 		bResult = SourceProperty->IsA<FBoolProperty>() && TargetProperty->IsA<FByteProperty>();
 		break;
@@ -276,6 +278,8 @@ bool FStateTreePropertyBindings::ValidateCopy(FStateTreePropCopy& Copy) const
 	case EStateTreePropertyCopyType::PromoteBoolToDouble:
 		bResult = SourceProperty->IsA<FBoolProperty>() && TargetProperty->IsA<FDoubleProperty>();
 		break;
+
+	// Byte promotions
 	case EStateTreePropertyCopyType::PromoteByteToInt32:
 		bResult = SourceProperty->IsA<FByteProperty>() && TargetProperty->IsA<FIntProperty>();
 		break;
@@ -291,6 +295,8 @@ bool FStateTreePropertyBindings::ValidateCopy(FStateTreePropCopy& Copy) const
 	case EStateTreePropertyCopyType::PromoteByteToDouble:
 		bResult = SourceProperty->IsA<FByteProperty>() && TargetProperty->IsA<FDoubleProperty>();
 		break;
+
+	// Int32 promotions
 	case EStateTreePropertyCopyType::PromoteInt32ToInt64:
 		bResult = SourceProperty->IsA<FIntProperty>() && TargetProperty->IsA<FInt64Property>();
 		break;
@@ -300,18 +306,39 @@ bool FStateTreePropertyBindings::ValidateCopy(FStateTreePropCopy& Copy) const
 	case EStateTreePropertyCopyType::PromoteInt32ToDouble:
 		bResult = SourceProperty->IsA<FIntProperty>() && TargetProperty->IsA<FDoubleProperty>();
 		break;
-	case EStateTreePropertyCopyType::PromoteFloatToDouble:
-		bResult = SourceProperty->IsA<FFloatProperty>() && TargetProperty->IsA<FDoubleProperty>();
-		break;
-	case EStateTreePropertyCopyType::DemoteDoubleToFloat:
-		bResult = SourceProperty->IsA<FDoubleProperty>() && TargetProperty->IsA<FFloatProperty>();
-		break;
+
+	// Uint32 promotions
 	case EStateTreePropertyCopyType::PromoteUInt32ToInt64:
 		bResult = SourceProperty->IsA<FUInt32Property>() && TargetProperty->IsA<FInt64Property>();
 		break;
 	case EStateTreePropertyCopyType::PromoteUInt32ToFloat:
 		bResult = SourceProperty->IsA<FUInt32Property>() && TargetProperty->IsA<FFloatProperty>();
 		break;
+	case EStateTreePropertyCopyType::PromoteUInt32ToDouble:
+		bResult = SourceProperty->IsA<FUInt32Property>() && TargetProperty->IsA<FDoubleProperty>();
+		break;
+
+	// Float promotions
+	case EStateTreePropertyCopyType::PromoteFloatToInt32:
+		bResult = SourceProperty->IsA<FFloatProperty>() && TargetProperty->IsA<FIntProperty>();
+		break;
+	case EStateTreePropertyCopyType::PromoteFloatToInt64:
+		bResult = SourceProperty->IsA<FFloatProperty>() && TargetProperty->IsA<FInt64Property>();
+		break;
+	case EStateTreePropertyCopyType::PromoteFloatToDouble:
+		bResult = SourceProperty->IsA<FFloatProperty>() && TargetProperty->IsA<FDoubleProperty>();
+		break;
+
+	// Double promotions
+	case EStateTreePropertyCopyType::DemoteDoubleToInt32:
+		bResult = SourceProperty->IsA<FDoubleProperty>() && TargetProperty->IsA<FIntProperty>();
+		break;
+	case EStateTreePropertyCopyType::DemoteDoubleToInt64:
+		bResult = SourceProperty->IsA<FDoubleProperty>() && TargetProperty->IsA<FInt64Property>();
+		break;
+	case EStateTreePropertyCopyType::DemoteDoubleToFloat:
+		bResult = SourceProperty->IsA<FDoubleProperty>() && TargetProperty->IsA<FFloatProperty>();
+		break;
 	default:
 		UE_LOG(LogStateTree, Error, TEXT("FStateTreePropertyBindings::ValidateCopy: Unhandled copy type %s between '%s' and '%s'"),
 			*StaticEnum<EStateTreePropertyCopyType>()->GetValueAsString(Copy.Type), *SourceProperty->GetNameCPP(), *TargetProperty->GetNameCPP());
@@ -447,6 +474,7 @@ void FStateTreePropertyBindings::PerformCopy(const FStateTreePropCopy& Copy, uin
 		Target->Set(FStructView(SourceStructProperty->Struct, SourceAddress));
 		break;
 	}
+	// Bool promotions
 	case EStateTreePropertyCopyType::PromoteBoolToByte:
 		*reinterpret_cast<uint8*>(TargetAddress) = (uint8)static_cast<const FBoolProperty*>(Copy.SourceLeafProperty)->GetPropertyValue(SourceAddress);
 		break;
@@ -465,6 +493,8 @@ void FStateTreePropertyBindings::PerformCopy(const FStateTreePropCopy& Copy, uin
 	case EStateTreePropertyCopyType::PromoteBoolToDouble:
 		*reinterpret_cast<double*>(TargetAddress) = (double)static_cast<const FBoolProperty*>(Copy.SourceLeafProperty)->GetPropertyValue(SourceAddress);
 		break;
+		
+	// Byte promotions	
 	case EStateTreePropertyCopyType::PromoteByteToInt32:
 		*reinterpret_cast<int32*>(TargetAddress) = (int32)*reinterpret_cast<const uint8*>(SourceAddress);
 		break;
@@ -480,6 +510,8 @@ void FStateTreePropertyBindings::PerformCopy(const FStateTreePropCopy& Copy, uin
 	case EStateTreePropertyCopyType::PromoteByteToDouble:
 		*reinterpret_cast<double*>(TargetAddress) = (double)*reinterpret_cast<const uint8*>(SourceAddress);
 		break;
+
+	// Int32 promotions
 	case EStateTreePropertyCopyType::PromoteInt32ToInt64:
 		*reinterpret_cast<int64*>(TargetAddress) = (int64)*reinterpret_cast<const int32*>(SourceAddress);
 		break;
@@ -489,18 +521,40 @@ void FStateTreePropertyBindings::PerformCopy(const FStateTreePropCopy& Copy, uin
 	case EStateTreePropertyCopyType::PromoteInt32ToDouble:
 		*reinterpret_cast<double*>(TargetAddress) = (double)*reinterpret_cast<const int32*>(SourceAddress);
 		break;
-	case EStateTreePropertyCopyType::PromoteFloatToDouble:
-		*reinterpret_cast<double*>(TargetAddress) = (double)*reinterpret_cast<const float*>(SourceAddress);
-		break;
-	case EStateTreePropertyCopyType::DemoteDoubleToFloat:
-		*reinterpret_cast<float*>(TargetAddress) = (float)*reinterpret_cast<const double*>(SourceAddress);
-		break;
+
+	// Uint32 promotions
 	case EStateTreePropertyCopyType::PromoteUInt32ToInt64:
 		*reinterpret_cast<int64*>(TargetAddress) = (int64)*reinterpret_cast<const uint32*>(SourceAddress);
 		break;
 	case EStateTreePropertyCopyType::PromoteUInt32ToFloat:
 		*reinterpret_cast<float*>(TargetAddress) = (float)*reinterpret_cast<const uint32*>(SourceAddress);
 		break;
+	case EStateTreePropertyCopyType::PromoteUInt32ToDouble:
+		*reinterpret_cast<double*>(TargetAddress) = (double)*reinterpret_cast<const uint32*>(SourceAddress);
+		break;
+
+	// Float promotions
+	case EStateTreePropertyCopyType::PromoteFloatToInt32:
+		*reinterpret_cast<int32*>(TargetAddress) = (int32)*reinterpret_cast<const float*>(SourceAddress);
+		break;
+	case EStateTreePropertyCopyType::PromoteFloatToInt64:
+		*reinterpret_cast<int64*>(TargetAddress) = (int64)*reinterpret_cast<const float*>(SourceAddress);
+		break;
+	case EStateTreePropertyCopyType::PromoteFloatToDouble:
+		*reinterpret_cast<double*>(TargetAddress) = (double)*reinterpret_cast<const float*>(SourceAddress);
+		break;
+
+	// Double promotions
+	case EStateTreePropertyCopyType::DemoteDoubleToInt32:
+		*reinterpret_cast<int32*>(TargetAddress) = (int32)*reinterpret_cast<const double*>(SourceAddress);
+		break;
+	case EStateTreePropertyCopyType::DemoteDoubleToInt64:
+		*reinterpret_cast<int64*>(TargetAddress) = (int64)*reinterpret_cast<const double*>(SourceAddress);
+		break;
+	case EStateTreePropertyCopyType::DemoteDoubleToFloat:
+		*reinterpret_cast<float*>(TargetAddress) = (float)*reinterpret_cast<const double*>(SourceAddress);
+		break;
+
 	default:
 		ensureMsgf(false, TEXT("FStateTreePropertyBindings::PerformCopy: Unhandled copy type %s between '%s' and '%s'"),
 			*StaticEnum<EStateTreePropertyCopyType>()->GetValueAsString(Copy.Type), *Copy.SourceLeafProperty->GetNameCPP(), *Copy.TargetLeafProperty->GetNameCPP());
diff --git a/Engine/Plugins/Runtime/StateTree/Source/StateTreeModule/Public/Conditions/StateTreeCommonConditions.h b/Engine/Plugins/Runtime/StateTree/Source/StateTreeModule/Public/Conditions/StateTreeCommonConditions.h
--- a/Engine/Plugins/Runtime/StateTree/Source/StateTreeModule/Public/Conditions/StateTreeCommonConditions.h
+++ b/Engine/Plugins/Runtime/StateTree/Source/StateTreeModule/Public/Conditions/StateTreeCommonConditions.h
@@ -56,10 +56,10 @@ struct STATETREEMODULE_API FStateTreeCompareFloatConditionInstanceData
 	GENERATED_BODY()
 
 	UPROPERTY(EditAnywhere, Category = Input)
-	float Left = 0.0f;
+	double Left = 0.0;
 
 	UPROPERTY(EditAnywhere, Category = Parameter)
-	float Right = 0.0f;
+	double Right = 0.0;
 };
 STATETREE_POD_INSTANCEDATA(FStateTreeCompareFloatConditionInstanceData);
 
@@ -183,7 +183,7 @@ struct STATETREEMODULE_API FStateTreeCompareDistanceConditionInstanceData
 	FVector Target = FVector(EForceInit::ForceInitToZero);
 
 	UPROPERTY(EditAnywhere, Category = Parameter)
-	float Distance = 0.0f;
+	double Distance = 0.0;
 };
 STATETREE_POD_INSTANCEDATA(FStateTreeCompareDistanceConditionInstanceData);
 
diff --git a/Engine/Plugins/Runtime/StateTree/Source/StateTreeModule/Public/StateTreePropertyBindings.h b/Engine/Plugins/Runtime/StateTree/Source/StateTreeModule/Public/StateTreePropertyBindings.h
--- a/Engine/Plugins/Runtime/StateTree/Source/StateTreeModule/Public/StateTreePropertyBindings.h
+++ b/Engine/Plugins/Runtime/StateTree/Source/StateTreeModule/Public/StateTreePropertyBindings.h
@@ -253,12 +253,18 @@ enum class EStateTreePropertyCopyType : uint8
 
 	/* UInt32 promotions */
 	PromoteUInt32ToInt64,
-    PromoteUInt32ToFloat,	// This is strictly sketchy because of potential data loss, but it is usually OK in the general case
+	PromoteUInt32ToFloat,	// This is strictly sketchy because of potential data loss, but it is usually OK in the general case
+	PromoteUInt32ToDouble,
 
 	/* Float promotions */
-	// LWC_TODO: Float/double should become synonyms in state trees? Certainly BPs.
+	PromoteFloatToInt32,
+	PromoteFloatToInt64,
 	PromoteFloatToDouble,
-	DemoteDoubleToFloat,	// LWC_TODO: This should not ship!
+
+	/* Double promotions */
+	DemoteDoubleToInt32,
+	DemoteDoubleToInt64,
+	DemoteDoubleToFloat,
 };
 
 
