Refactor the StateTree delay task to:
- accept 0s delay
- have a bool to run forever
- support random deviation

diff --git a/Engine/Plugins/Runtime/StateTree/Source/StateTreeModule/Private/Tasks/StateTreeDelayTask.cpp b/Engine/Plugins/Runtime/StateTree/Source/StateTreeModule/Private/Tasks/StateTreeDelayTask.cpp
--- a/Engine/Plugins/Runtime/StateTree/Source/StateTreeModule/Private/Tasks/StateTreeDelayTask.cpp
+++ b/Engine/Plugins/Runtime/StateTree/Source/StateTreeModule/Private/Tasks/StateTreeDelayTask.cpp
@@ -8,7 +8,11 @@ EStateTreeRunStatus FStateTreeDelayTask::EnterState(FStateTreeExecutionContext&
 {
 	InstanceDataType& InstanceData = Context.GetInstanceData<InstanceDataType>(*this);
 
-	InstanceData.Time = 0.f;
+	if (!InstanceData.bRunForever)
+	{
+		InstanceData.RemainingTime = FMath::FRandRange(
+			FMath::Max(0.0f, InstanceData.Duration - InstanceData.RandomDeviation), (InstanceData.Duration + InstanceData.RandomDeviation));
+	}
 	
 	return EStateTreeRunStatus::Running;
 }
@@ -17,7 +21,15 @@ EStateTreeRunStatus FStateTreeDelayTask::Tick(FStateTreeExecutionContext& Contex
 {
 	InstanceDataType& InstanceData = Context.GetInstanceData<InstanceDataType>(*this);
 
-	InstanceData.Time += DeltaTime;
+	if (!InstanceData.bRunForever)
+	{
+		InstanceData.RemainingTime -= DeltaTime;
+
+		if (InstanceData.RemainingTime <= 0.f)
+		{
+			return EStateTreeRunStatus::Succeeded;
+		}
+	}
 	
-	return InstanceData.Duration <= 0.0f ? EStateTreeRunStatus::Running : (InstanceData.Time < InstanceData.Duration ? EStateTreeRunStatus::Running : EStateTreeRunStatus::Succeeded);
+	return EStateTreeRunStatus::Running;
 }
diff --git a/Engine/Plugins/Runtime/StateTree/Source/StateTreeModule/Private/Tasks/StateTreeDelayTask.h b/Engine/Plugins/Runtime/StateTree/Source/StateTreeModule/Private/Tasks/StateTreeDelayTask.h
--- a/Engine/Plugins/Runtime/StateTree/Source/StateTreeModule/Private/Tasks/StateTreeDelayTask.h
+++ b/Engine/Plugins/Runtime/StateTree/Source/StateTreeModule/Private/Tasks/StateTreeDelayTask.h
@@ -10,13 +10,20 @@ struct STATETREEMODULE_API FStateTreeDelayTaskInstanceData
 {
 	GENERATED_BODY()
 	
-	/** Delay before the task ends. Default (0 or any negative) will run indefinitely so it requires a transition in the state tree to stop it. */
+	/** Delay before the task ends. */
+	UPROPERTY(EditAnywhere, Category = Parameter, meta = (EditCondition = "!bRunForever", ClampMin="0.0"))
+	float Duration = 1.f;
+	
+	/** Adds random range to the Duration. */
+	UPROPERTY(EditAnywhere, Category = Parameter, meta = (EditCondition = "!bRunForever", ClampMin="0.0"))
+	float RandomDeviation = 0.f;
+	
+	/** If true the task will run forever until a transition stops it. */
 	UPROPERTY(EditAnywhere, Category = Parameter)
-	float Duration = 0.f;
+	bool bRunForever = false;
 
-	/** Accumulated time used to stop task if duration is set */
-	UPROPERTY()
-	float Time = 0.f;
+	/** Internal countdown in seconds. */
+	float RemainingTime = 0.f;
 };
 
 /**
@@ -34,5 +41,5 @@ struct STATETREEMODULE_API FStateTreeDelayTask : public FStateTreeTaskCommonBase
 	virtual const UStruct* GetInstanceDataType() const override { return InstanceDataType::StaticStruct(); }
 
 	virtual EStateTreeRunStatus EnterState(FStateTreeExecutionContext& Context, const EStateTreeStateChangeType ChangeType, const FStateTreeTransitionResult& Transition) const override;
-	virtual EStateTreeRunStatus Tick(FStateTreeExecutionContext& Context, const float DeltaTime) const override;;
+	virtual EStateTreeRunStatus Tick(FStateTreeExecutionContext& Context, const float DeltaTime) const override;
 };
