StateTree: Compile StateTree on load if in older format.

diff --git a/Engine/Plugins/Runtime/StateTree/Source/StateTreeEditorModule/Private/StateTreeEditorModule.cpp b/Engine/Plugins/Runtime/StateTree/Source/StateTreeEditorModule/Private/StateTreeEditorModule.cpp
--- a/Engine/Plugins/Runtime/StateTree/Source/StateTreeEditorModule/Private/StateTreeEditorModule.cpp
+++ b/Engine/Plugins/Runtime/StateTree/Source/StateTreeEditorModule/Private/StateTreeEditorModule.cpp
@@ -24,7 +24,9 @@
 #include "Blueprint/StateTreeTaskBlueprintBase.h"
 #include "Blueprint/StateTreeConditionBlueprintBase.h"
 #include "StateTreeEditorCommands.h"
-
+#include "StateTreeDelegates.h"
+#include "StateTreeCompiler.h"
+#include "StateTreeCompilerLog.h"
 
 #define LOCTEXT_NAMESPACE "StateTreeEditor"
 
@@ -32,8 +34,46 @@ DEFINE_LOG_CATEGORY(LogStateTreeEditor);
 
 IMPLEMENT_MODULE(FStateTreeEditorModule, StateTreeEditorModule)
 
+namespace UE::StateTree::Editor
+{
+	// @todo Could we make this a IModularFeature?
+	static bool CompileStateTree(UStateTree& StateTree)
+	{
+		// Compile the StateTree asset.
+		UE::StateTree::Editor::ValidateAsset(StateTree);
+		const uint32 EditorDataHash = UE::StateTree::Editor::CalcAssetHash(StateTree);
+
+		FStateTreeCompilerLog Log;
+		FStateTreeCompiler Compiler(Log);
+
+		const bool bSuccess = Compiler.Compile(StateTree);
+
+		if (bSuccess)
+		{
+			// Success
+			StateTree.LastCompiledEditorDataHash = EditorDataHash;
+			UE::StateTree::Delegates::OnPostCompile.Broadcast(StateTree);
+			UE_LOG(LogStateTreeEditor, Log, TEXT("Compile StateTree '%s' succeeded."), *StateTree.GetFullName());
+		}
+		else
+		{
+			// Make sure not to leave stale data on failed compile.
+			StateTree.ResetCompiled();
+			StateTree.LastCompiledEditorDataHash = 0;
+
+			UE_LOG(LogStateTreeEditor, Error, TEXT("Failed to compile '%s', errors follow."), *StateTree.GetFullName());
+			Log.DumpToLog(LogStateTreeEditor);
+		}
+
+		return bSuccess;
+	}
+
+}; // UE::StateTree::Editor
+
 void FStateTreeEditorModule::StartupModule()
 {
+	UE::StateTree::Delegates::OnRequestCompile.BindStatic(&UE::StateTree::Editor::CompileStateTree);
+
 	MenuExtensibilityManager = MakeShareable(new FExtensibilityManager);
 	ToolBarExtensibilityManager = MakeShareable(new FExtensibilityManager);
 
@@ -66,6 +106,8 @@ void FStateTreeEditorModule::StartupModule()
 
 void FStateTreeEditorModule::ShutdownModule()
 {
+	UE::StateTree::Delegates::OnRequestCompile.Unbind();
+	
 	MenuExtensibilityManager.Reset();
 	ToolBarExtensibilityManager.Reset();
 
diff --git a/Engine/Plugins/Runtime/StateTree/Source/StateTreeEditorModule/Public/StateTreeEditorModule.h b/Engine/Plugins/Runtime/StateTree/Source/StateTreeEditorModule/Public/StateTreeEditorModule.h
--- a/Engine/Plugins/Runtime/StateTree/Source/StateTreeEditorModule/Public/StateTreeEditorModule.h
+++ b/Engine/Plugins/Runtime/StateTree/Source/StateTreeEditorModule/Public/StateTreeEditorModule.h
@@ -11,7 +11,7 @@ class IStateTreeEditor;
 class FAssetTypeActions_Base;
 struct FStateTreeNodeClassCache;
 
-STATETREEEDITORMODULE_API DECLARE_LOG_CATEGORY_EXTERN(LogStateTreeEditor, Warning, All);
+STATETREEEDITORMODULE_API DECLARE_LOG_CATEGORY_EXTERN(LogStateTreeEditor, Log, All);
 
 /**
 * The public interface to this module
diff --git a/Engine/Plugins/Runtime/StateTree/Source/StateTreeModule/Private/StateTree.cpp b/Engine/Plugins/Runtime/StateTree/Source/StateTreeModule/Private/StateTree.cpp
--- a/Engine/Plugins/Runtime/StateTree/Source/StateTreeModule/Private/StateTree.cpp
+++ b/Engine/Plugins/Runtime/StateTree/Source/StateTreeModule/Private/StateTree.cpp
@@ -7,6 +7,8 @@
 #include "StateTreeEvaluatorBase.h"
 #include "AssetRegistry/AssetData.h"
 #include "Misc/ScopeRWLock.h"
+#include "StateTreeDelegates.h"
+#include "Logging/LogScopedVerbosityOverride.h"
 
 const FGuid FStateTreeCustomVersion::GUID(0x28E21331, 0x501F4723, 0x8110FA64, 0xEA10DA1E);
 FCustomVersionRegistration GRegisterStateTreeCustomVersion(FStateTreeCustomVersion::GUID, FStateTreeCustomVersion::LatestVersion, TEXT("StateTreeAsset"));
@@ -113,17 +115,27 @@ void UStateTree::PostLoad()
 	if (CurrentVersion < FStateTreeCustomVersion::LatestVersion)
 	{
 #if WITH_EDITOR
-		ResetCompiled();
-		UE_LOG(LogStateTree, Warning, TEXT("%s: StateTree compiled data in older format. Please resave the StateTree asset."), *GetPathName());
+		// Compiled data is in older format, try to compile the StateTree.
+		if (UE::StateTree::Delegates::OnRequestCompile.IsBound())
+		{
+			LOG_SCOPE_VERBOSITY_OVERRIDE(LogStateTree, ELogVerbosity::Log);
+			UE_LOG(LogStateTree, Log, TEXT("%s: compiled data is in older format. Trying to compile the asset..."), *GetFullName());
+			UE::StateTree::Delegates::OnRequestCompile.Execute(*this);
+		}
+		else
+		{
+			ResetCompiled();
+			UE_LOG(LogStateTree, Warning, TEXT("%s: compiled data is in older format. Please resave the StateTree asset."), *GetFullName());
+		}
 #else
-		UE_LOG(LogStateTree, Error, TEXT("%s: StateTree compiled data in older format. Please recompile the StateTree asset."), *GetPathName());
+		UE_LOG(LogStateTree, Error, TEXT("%s: compiled data is in older format. Please recompile the StateTree asset."), *GetFullName());
 #endif
 		return;
 	}
 	
 	if (!Link())
 	{
-		UE_LOG(LogStateTree, Error, TEXT("%s failed to link. Asset will not be usable at runtime."), *GetName());	
+		UE_LOG(LogStateTree, Error, TEXT("%s failed to link. Asset will not be usable at runtime."), *GetFullName());	
 	}
 }
 
diff --git a/Engine/Plugins/Runtime/StateTree/Source/StateTreeModule/Private/StateTreeDelegates.cpp b/Engine/Plugins/Runtime/StateTree/Source/StateTreeModule/Private/StateTreeDelegates.cpp
--- a/Engine/Plugins/Runtime/StateTree/Source/StateTreeModule/Private/StateTreeDelegates.cpp
+++ b/Engine/Plugins/Runtime/StateTree/Source/StateTreeModule/Private/StateTreeDelegates.cpp
@@ -11,6 +11,7 @@ FOnSchemaChanged OnSchemaChanged;
 FOnParametersChanged OnParametersChanged;
 FOnStateParametersChanged OnStateParametersChanged;
 FOnPostCompile OnPostCompile;
+FOnRequestCompile OnRequestCompile;
 #endif
 	
 }; // UE::StateTree::Delegates
diff --git a/Engine/Plugins/Runtime/StateTree/Source/StateTreeModule/Public/StateTreeDelegates.h b/Engine/Plugins/Runtime/StateTree/Source/StateTreeModule/Public/StateTreeDelegates.h
--- a/Engine/Plugins/Runtime/StateTree/Source/StateTreeModule/Public/StateTreeDelegates.h
+++ b/Engine/Plugins/Runtime/StateTree/Source/StateTreeModule/Public/StateTreeDelegates.h
@@ -40,6 +40,10 @@ namespace UE::StateTree::Delegates
 	/** Called when compilation succeeds */
 	DECLARE_MULTICAST_DELEGATE_OneParam(FOnPostCompile, const UStateTree& /*StateTree*/);
 	extern STATETREEMODULE_API FOnPostCompile OnPostCompile;
+
+	/** Request StateTree compilation. Works only in editor. */
+	DECLARE_DELEGATE_RetVal_OneParam(bool, FOnRequestCompile, UStateTree& /*StateTreeToCompile*/);
+	extern STATETREEMODULE_API FOnRequestCompile OnRequestCompile;
 #endif
 
 }; // UE::StateTree::Delegates
