Fix const cast on a TObjectPtr

diff --git a/Engine/Plugins/Runtime/StateTree/Source/StateTreeEditorModule/Private/Customizations/StateTreeEditorDataDetails.cpp b/Engine/Plugins/Runtime/StateTree/Source/StateTreeEditorModule/Private/Customizations/StateTreeEditorDataDetails.cpp
--- a/Engine/Plugins/Runtime/StateTree/Source/StateTreeEditorModule/Private/Customizations/StateTreeEditorDataDetails.cpp
+++ b/Engine/Plugins/Runtime/StateTree/Source/StateTreeEditorModule/Private/Customizations/StateTreeEditorDataDetails.cpp
@@ -1,4 +1,4 @@
-﻿// Copyright Epic Games, Inc. All Rights Reserved.
+// Copyright Epic Games, Inc. All Rights Reserved.
 
 #include "StateTreeEditorDataDetails.h"
 #include "Widgets/DeclarativeSyntaxSupport.h"
@@ -71,7 +71,7 @@ void FStateTreeEditorDataDetails::CustomizeDetails(IDetailLayoutBuilder& DetailB
 			{
 				continue;
 			}
-			PinType.PinSubCategoryObject = const_cast<UStruct*>(ContextData.Struct);
+			PinType.PinSubCategoryObject = const_cast<UStruct*>(ContextData.Struct.Get());
 
 			const UEdGraphSchema_K2* EdGraphSchema = GetDefault<UEdGraphSchema_K2>();
 			const FSlateBrush* Icon = FBlueprintEditorUtils::GetIconFromPin(PinType, true);
