StateTree : FStateTreeRandomCondition change inherit from FStateTreeConditionBase to FStateTreeConditionCommonBase
#tests PIE
#rnx

diff --git a/Engine/Plugins/Runtime/StateTree/Source/StateTreeModule/Public/Conditions/StateTreeCommonConditions.h b/Engine/Plugins/Runtime/StateTree/Source/StateTreeModule/Public/Conditions/StateTreeCommonConditions.h
--- a/Engine/Plugins/Runtime/StateTree/Source/StateTreeModule/Public/Conditions/StateTreeCommonConditions.h
+++ b/Engine/Plugins/Runtime/StateTree/Source/StateTreeModule/Public/Conditions/StateTreeCommonConditions.h
@@ -224,7 +224,7 @@ struct STATETREEMODULE_API FStateTreeRandomConditionInstanceData
 STATETREE_POD_INSTANCEDATA(FStateTreeRandomConditionInstanceData);
 
 USTRUCT(DisplayName = "Random")
-struct STATETREEMODULE_API FStateTreeRandomCondition : public FStateTreeConditionBase
+struct STATETREEMODULE_API FStateTreeRandomCondition : public FStateTreeConditionCommonBase
 {
 	GENERATED_BODY()
 
