StateTree: Fix wrong redirect.

diff --git a/Engine/Plugins/Runtime/StateTree/Config/BaseStateTree.ini b/Engine/Plugins/Runtime/StateTree/Config/BaseStateTree.ini
--- a/Engine/Plugins/Runtime/StateTree/Config/BaseStateTree.ini
+++ b/Engine/Plugins/Runtime/StateTree/Config/BaseStateTree.ini
@@ -54,6 +54,6 @@
 +PropertyRedirects=(OldName="CompactStateTransition.Event",NewName="CompactStateTransition.Trigger")
 +ClassRedirects=(OldName="/Script/StateTreeModule.StateTreeItemBlueprintBase",NewName="/Script/StateTreeModule.StateTreeNodeBlueprintBase")
 
-+EnumRedirects=(OldName="EStateTreeBindableStructSource",ValueChanges=(("TreeData", "ContextObjects"), ("TreeParameter", "Parameter"), ("StartParameter", "State"))
++EnumRedirects=(OldName="EStateTreeBindableStructSource",ValueChanges=(("TreeData", "Context"), ("TreeParameter", "Parameter"), ("StartParameter", "State"))
 
 +PropertyRedirects=(OldName="/Script/StateTreeModule.StateTree.NamedExternalDataDescs",NewName="/Script/StateTreeModule.StateTree.ContextDataDescs")
