StateTree: Handle optionally bound Input properties.

diff --git a/Engine/Plugins/Runtime/StateTree/Source/StateTreeEditorModule/Private/StateTreeCompiler.cpp b/Engine/Plugins/Runtime/StateTree/Source/StateTreeEditorModule/Private/StateTreeCompiler.cpp
--- a/Engine/Plugins/Runtime/StateTree/Source/StateTreeEditorModule/Private/StateTreeCompiler.cpp
+++ b/Engine/Plugins/Runtime/StateTree/Source/StateTreeEditorModule/Private/StateTreeCompiler.cpp
@@ -1120,12 +1120,15 @@ bool FStateTreeCompiler::GetAndValidateBindings(const FStateTreeBindableStructDe
 	for (TFieldIterator<FProperty> It(TargetStruct.Struct, EFieldIterationFlags::None); It; ++It)
 	{
 		const FProperty* Property = *It;
+		check(Property);
 		const FString PropertyName = Property->GetName();
 		const EStateTreePropertyUsage Usage = UE::StateTree::Compiler::GetUsageFromMetaData(Property);
 		if (Usage == EStateTreePropertyUsage::Input)
 		{
-			// Make sure that an Input property is bound.
-			if (!IsPropertyBound(PropertyName))
+			const bool bIsOptional = Property->HasMetaData(TEXT("Optional"));
+			
+			// Make sure that an Input property is bound unless marked optional.
+			if (bIsOptional == false && !IsPropertyBound(PropertyName))
 			{
 				Log.Reportf(EMessageSeverity::Error, TargetStruct,
 					TEXT("Input property '%s:%s' is expected to have a binding."),
