StateTree: fix test and empty parameters
- handle empty subtree parameters
- changed a compiler check() to an error
- fixed tests

diff --git a/Engine/Plugins/Runtime/StateTree/Source/StateTreeEditorModule/Private/StateTreeCompiler.cpp b/Engine/Plugins/Runtime/StateTree/Source/StateTreeEditorModule/Private/StateTreeCompiler.cpp
--- a/Engine/Plugins/Runtime/StateTree/Source/StateTreeEditorModule/Private/StateTreeCompiler.cpp
+++ b/Engine/Plugins/Runtime/StateTree/Source/StateTreeEditorModule/Private/StateTreeCompiler.cpp
@@ -367,49 +367,61 @@ bool FStateTreeCompiler::CreateStateTasksAndParameters()
 
 			if (SourceState->Type == EStateTreeStateType::Subtree)
 			{
-				// Register a binding source
-				const FStateTreeBindableStructDesc SubtreeParamsDesc = {
+				// Register a binding source if we have parameters
+				int32 SourceStructIndex = INDEX_NONE;
+
+				if (SourceState->Parameters.Parameters.IsValid())
+				{
+					const FStateTreeBindableStructDesc SubtreeParamsDesc = {
 						SourceState->Name,
 						SourceState->Parameters.Parameters.GetPropertyBagStruct(),
 						EStateTreeBindableStructSource::State,
 						SourceState->Parameters.ID
 					};
-				const int32 SourceStructIndex = BindingsCompiler.AddSourceStruct(SubtreeParamsDesc);
-				if (const auto Validation = UE::StateTree::Compiler::IsValidIndex16(SourceStructIndex); Validation.DidFail())
-				{
-					Validation.Log(Log, TEXT("SourceStructIndex"), SubtreeParamsDesc);
-					return false;
+					SourceStructIndex = BindingsCompiler.AddSourceStruct(SubtreeParamsDesc);
+					
+					if (const auto Validation = UE::StateTree::Compiler::IsValidIndex16(SourceStructIndex); Validation.DidFail())
+					{
+						Validation.Log(Log, TEXT("SourceStructIndex"), SubtreeParamsDesc);
+						return false;
+					}
 				}
+				
 				CompactState.ParameterDataViewIndex = FStateTreeIndex16(SourceStructIndex);
 			}
 			else if (SourceState->Type == EStateTreeStateType::Linked)
 			{
-				// Binding target
-				FStateTreeBindableStructDesc LinkedParamsDesc = {
+				int32 BatchIndex = INDEX_NONE;
+
+				if (SourceState->Parameters.Parameters.IsValid())
+				{
+					// Binding target
+					FStateTreeBindableStructDesc LinkedParamsDesc = {
 						SourceState->Name,
 						SourceState->Parameters.Parameters.GetPropertyBagStruct(),
 						EStateTreeBindableStructSource::State,
 						SourceState->Parameters.ID
 					};
 
-				// Check that the bindings for this struct are still all valid.
-				TArray<FStateTreeEditorPropertyBinding> Bindings;
-				if (!GetAndValidateBindings(LinkedParamsDesc, Bindings))
-				{
-					return false;
-				}
+					// Check that the bindings for this struct are still all valid.
+					TArray<FStateTreeEditorPropertyBinding> Bindings;
+					if (!GetAndValidateBindings(LinkedParamsDesc, Bindings))
+					{
+						return false;
+					}
 
-				int32 BatchIndex = INDEX_NONE;
-				if (!BindingsCompiler.CompileBatch(LinkedParamsDesc, Bindings, BatchIndex))
-				{
-					return false;
+					if (!BindingsCompiler.CompileBatch(LinkedParamsDesc, Bindings, BatchIndex))
+					{
+						return false;
+					}
+					
+					if (const auto Validation = UE::StateTree::Compiler::IsValidIndex16(BatchIndex); Validation.DidFail())
+					{
+						Validation.Log(Log, TEXT("BatchIndex"), LinkedParamsDesc);
+						return false;
+					}
 				}
 
-				if (const auto Validation = UE::StateTree::Compiler::IsValidIndex16(BatchIndex); Validation.DidFail())
-				{
-					Validation.Log(Log, TEXT("BatchIndex"), LinkedParamsDesc);
-					return false;
-				}
 				CompactParams.BindingsBatch = FStateTreeIndex16(BatchIndex);
 			}
 		}
@@ -1032,7 +1044,13 @@ bool FStateTreeCompiler::ValidateStructRef(const FStateTreeBindableStructDesc& S
 
 bool FStateTreeCompiler::GetAndValidateBindings(const FStateTreeBindableStructDesc& TargetStruct, TArray<FStateTreeEditorPropertyBinding>& OutBindings) const
 {
-	check(TargetStruct.Struct != nullptr);
+	if (TargetStruct.Struct == nullptr)
+	{
+		Log.Reportf(EMessageSeverity::Error, TargetStruct,
+				TEXT("The type of binding target '%s' is invalid."),
+				*TargetStruct.Name.ToString());
+		return false;
+	}
 	
 	OutBindings.Reset();
 	
diff --git a/Engine/Plugins/Runtime/StateTree/Source/StateTreeModule/Private/StateTree.cpp b/Engine/Plugins/Runtime/StateTree/Source/StateTreeModule/Private/StateTree.cpp
--- a/Engine/Plugins/Runtime/StateTree/Source/StateTreeModule/Private/StateTree.cpp
+++ b/Engine/Plugins/Runtime/StateTree/Source/StateTreeModule/Private/StateTree.cpp
@@ -194,17 +194,19 @@ bool UStateTree::Link()
 		{
 			if (State.Type == EStateTreeStateType::Subtree)
 			{
-				if (State.ParameterInstanceIndex.IsValid() == false
-					|| State.ParameterDataViewIndex.IsValid() == false)
+				if (State.ParameterInstanceIndex.IsValid() == false)
 				{
 					UE_LOG(LogStateTree, Error, TEXT("%s: Data for state '%s' is malformed. Please recompile the StateTree asset."), *GetName(), *State.Name.ToString());
 					return false;
 				}
 
 				// Subtree is a bind source, update bag struct.
-				const FCompactStateTreeParameters& Params = DefaultInstanceData.GetMutableStruct(State.ParameterInstanceIndex.Get()).GetMutable<FCompactStateTreeParameters>();
-				FStateTreeBindableStructDesc& Desc = SourceStructs[State.ParameterDataViewIndex.Get()];
-				Desc.Struct = Params.Parameters.GetPropertyBagStruct();
+				if (State.ParameterDataViewIndex.IsValid())
+				{
+					const FCompactStateTreeParameters& Params = DefaultInstanceData.GetMutableStruct(State.ParameterInstanceIndex.Get()).GetMutable<FCompactStateTreeParameters>();
+					FStateTreeBindableStructDesc& Desc = SourceStructs[State.ParameterDataViewIndex.Get()];
+					Desc.Struct = Params.Parameters.GetPropertyBagStruct();
+				}
 			}
 			else if (State.Type == EStateTreeStateType::Linked && State.LinkedState.IsValid())
 			{
diff --git a/Engine/Plugins/Runtime/StateTree/Source/StateTreeModule/Private/StateTreeExecutionContext.cpp b/Engine/Plugins/Runtime/StateTree/Source/StateTreeModule/Private/StateTreeExecutionContext.cpp
--- a/Engine/Plugins/Runtime/StateTree/Source/StateTreeModule/Private/StateTreeExecutionContext.cpp
+++ b/Engine/Plugins/Runtime/StateTree/Source/StateTreeModule/Private/StateTreeExecutionContext.cpp
@@ -132,35 +132,42 @@ void FStateTreeExecutionContext::UpdateLinkedStateParameters(const FStateTreeIns
 	const FStateTreeDataView StateParamsInstance = InstanceData.GetMutableStruct(ParameterInstanceIndex);
 	const FCompactStateTreeParameters& StateParams = StateParamsInstance.GetMutable<FCompactStateTreeParameters>();
 
-	// Parameters property bag
-	const FStateTreeDataView ParametersView(StateParams.Parameters.GetMutableValue());
-	if (StateParams.BindingsBatch.IsValid())
+	// Update parameters if the state has any.
+	if (StateParams.Parameters.IsValid())
 	{
-		StateTree->PropertyBindings.CopyTo(DataViews, StateParams.BindingsBatch, ParametersView);
-	}
+		// Parameters property bag
+		const FStateTreeDataView ParametersView(StateParams.Parameters.GetMutableValue());
+		if (StateParams.BindingsBatch.IsValid())
+		{
+			StateTree->PropertyBindings.CopyTo(DataViews, StateParams.BindingsBatch, ParametersView);
+		}
 
-	// Set the parameters as the input parameters for the linked state.
-	check(State.LinkedState.IsValid());
-	const FCompactStateTreeState& LinkedState = StateTree->States[State.LinkedState.Index];
-	check(LinkedState.ParameterDataViewIndex.IsValid());
-	DataViews[LinkedState.ParameterDataViewIndex.Get()] = ParametersView;
+		// Set the parameters as the input parameters for the linked state.
+		check(State.LinkedState.IsValid());
+		const FCompactStateTreeState& LinkedState = StateTree->States[State.LinkedState.Index];
+		check(LinkedState.ParameterDataViewIndex.IsValid());
+		DataViews[LinkedState.ParameterDataViewIndex.Get()] = ParametersView;
+	}
 }
 
 void FStateTreeExecutionContext::UpdateSubtreeStateParameters(const FStateTreeInstanceData& InstanceData, const FCompactStateTreeState& State)
 {
-	check(State.ParameterDataViewIndex.IsValid());
 	check(State.ParameterInstanceIndex.IsValid());
-	
-	// Usually the subtree parameter view is set by the linked state. If it's not (i.e. transitioned into a parametrized subtree), we'll set the view default params.
-	if (DataViews[State.ParameterDataViewIndex.Get()].IsValid())
+
+	// Update parameters if the state has any.
+	if (State.ParameterDataViewIndex.IsValid())
 	{
-		return;
-	}
+		// Usually the subtree parameter view is set by the linked state. If it's not (i.e. transitioned into a parametrized subtree), we'll set the view default params.
+		if (DataViews[State.ParameterDataViewIndex.Get()].IsValid())
+		{
+			return;
+		}
 
-	// Set view to default parameters.
-	const FStateTreeDataView ParamInstanceView = StateTree->DefaultInstanceData.GetMutableStruct(State.ParameterInstanceIndex.Get()); // These are used as const, so get them from the tree initial values.
-	const FCompactStateTreeParameters& Params = ParamInstanceView.GetMutable<FCompactStateTreeParameters>();
-	DataViews[State.ParameterDataViewIndex.Get()] = FStateTreeDataView(Params.Parameters.GetMutableValue());
+		// Set view to default parameters.
+		const FStateTreeDataView ParamInstanceView = StateTree->DefaultInstanceData.GetMutableStruct(State.ParameterInstanceIndex.Get()); // These are used as const, so get them from the tree initial values.
+		const FCompactStateTreeParameters& Params = ParamInstanceView.GetMutable<FCompactStateTreeParameters>();
+		DataViews[State.ParameterDataViewIndex.Get()] = FStateTreeDataView(Params.Parameters.GetMutableValue());
+	}
 }
 
 EStateTreeRunStatus FStateTreeExecutionContext::Start(FStateTreeInstanceData* ExternalInstanceData)
diff --git a/Engine/Plugins/Runtime/StateTree/Source/StateTreeTestSuite/Private/StateTreeTestTypes.h b/Engine/Plugins/Runtime/StateTree/Source/StateTreeTestSuite/Private/StateTreeTestTypes.h
--- a/Engine/Plugins/Runtime/StateTree/Source/StateTreeTestSuite/Private/StateTreeTestTypes.h
+++ b/Engine/Plugins/Runtime/StateTree/Source/StateTreeTestSuite/Private/StateTreeTestTypes.h
@@ -844,7 +844,7 @@ struct FStateTreeTestConditionInstanceData
 {
 	GENERATED_BODY()
 	
-	UPROPERTY(EditAnywhere, Category = Input)
+	UPROPERTY(EditAnywhere, Category = Parameters)
 	int32 Count = 1;
 
 	static std::atomic<int32> GlobalCounter;
