Fix crash issue with condition instances not being reflected to GC

diff --git a/Engine/Plugins/Runtime/StateTree/Source/StateTreeModule/Private/StateTree.cpp b/Engine/Plugins/Runtime/StateTree/Source/StateTreeModule/Private/StateTree.cpp
--- a/Engine/Plugins/Runtime/StateTree/Source/StateTreeModule/Private/StateTree.cpp
+++ b/Engine/Plugins/Runtime/StateTree/Source/StateTreeModule/Private/StateTree.cpp
@@ -104,6 +104,26 @@ void UStateTree::PostLoadAssetRegistryTags(const FAssetData& InAssetData, TArray
 	}
 }
 
+void UStateTree::AddReferencedObjects(UObject* InThis, FReferenceCollector& Collector)
+{
+	Super::AddReferencedObjects(InThis, Collector);
+
+	const UStateTree* StateTree = Cast<UStateTree>(InThis);
+	check(StateTree);
+	
+	FReadScopeLock ReadLock(StateTree->PerThreadSharedInstanceDataLock);
+
+	for (const TSharedPtr<FStateTreeInstanceData>& InstanceData : StateTree->PerThreadSharedInstanceData)
+	{
+		if (InstanceData.IsValid())
+		{
+			uint8* StructMemory = (uint8*)InstanceData.Get();
+			const UScriptStruct* ScriptStruct = FStateTreeInstanceData::StaticStruct();
+			Collector.AddReferencedObjects(ScriptStruct, StructMemory);
+		}
+	}
+}
+
 #endif // WITH_EDITOR
 
 void UStateTree::PostLoad()
diff --git a/Engine/Plugins/Runtime/StateTree/Source/StateTreeModule/Public/StateTree.h b/Engine/Plugins/Runtime/StateTree/Source/StateTreeModule/Public/StateTree.h
--- a/Engine/Plugins/Runtime/StateTree/Source/StateTreeModule/Public/StateTree.h
+++ b/Engine/Plugins/Runtime/StateTree/Source/StateTreeModule/Public/StateTree.h
@@ -135,6 +135,8 @@ protected:
 	virtual void PostLoadAssetRegistryTags(const FAssetData& InAssetData, TArray<FAssetRegistryTag>& OutTagsAndValuesToUpdate) const;
 #endif
 
+	static void AddReferencedObjects(UObject* InThis, FReferenceCollector& Collector);
+	
 private:
 
 	/**
