Fixed crash when a task fails and one of its children instance data does not get updated anymore because of that, resulting in nullptr access during on completion

diff --git a/Engine/Plugins/Runtime/StateTree/Source/StateTreeModule/Private/StateTreeExecutionContext.cpp b/Engine/Plugins/Runtime/StateTree/Source/StateTreeModule/Private/StateTreeExecutionContext.cpp
--- a/Engine/Plugins/Runtime/StateTree/Source/StateTreeModule/Private/StateTreeExecutionContext.cpp
+++ b/Engine/Plugins/Runtime/StateTree/Source/StateTreeModule/Private/StateTreeExecutionContext.cpp
@@ -917,7 +917,10 @@ EStateTreeRunStatus FStateTreeExecutionContext::TickTasks(const float DeltaTime)
 	int32 InstanceStructIndex = Exec.FirstTaskStructIndex.Get();
 	int32 InstanceObjectIndex = Exec.FirstTaskObjectIndex.Get();
 
-	for (int32 Index = 0; Index < Exec.ActiveStates.Num() && Result != EStateTreeRunStatus::Failed; Index++)
+	// Used to stop ticking tasks after one fails, but we still want to keep updating the data of them
+	bool bShouldTickTasks = true;
+	
+	for (int32 Index = 0; Index < Exec.ActiveStates.Num(); Index++)
 	{
 		const FStateTreeStateHandle CurrentHandle = Exec.ActiveStates[Index];
 		const FCompactStateTreeState& State = StateTree.States[CurrentHandle.Index];
@@ -934,7 +937,7 @@ EStateTreeRunStatus FStateTreeExecutionContext::TickTasks(const float DeltaTime)
 			UpdateSubtreeStateParameters(State);
 		}
 
-		// Tick Tasks
+		// Update Tasks data and tick if possible (ie. if no task has yet failed and so bShouldTickTasks is true)
 		for (int32 TaskIndex = State.TasksBegin; TaskIndex < (State.TasksBegin + State.TasksNum); TaskIndex++)
 		{
 			const FStateTreeTaskBase& Task = StateTree.Nodes[TaskIndex].Get<FStateTreeTaskBase>();
@@ -946,7 +949,8 @@ EStateTreeRunStatus FStateTreeExecutionContext::TickTasks(const float DeltaTime)
 				StateTree.PropertyBindings.CopyTo(DataViews, Task.BindingsBatch, DataViews[Task.DataViewIndex.Get()]);
 			}
 			STATETREE_LOG(VeryVerbose, TEXT("%*s  Tick: '%s'"), Index*UE::StateTree::DebugIndentSize, TEXT(""), *Task.Name.ToString());
-			
+
+			if (bShouldTickTasks)
 			{
 				QUICK_SCOPE_CYCLE_COUNTER(StateTree_Task_Tick);
 				CSV_SCOPED_TIMING_STAT_EXCLUSIVE(StateTree_Task_Tick);
@@ -958,9 +962,10 @@ EStateTreeRunStatus FStateTreeExecutionContext::TickTasks(const float DeltaTime)
 				{
 					Result = TaskResult;
 				}
+				
 				if (TaskResult == EStateTreeRunStatus::Failed)
 				{
-					break;
+					bShouldTickTasks = false;
 				}
 			}
 		}
