StateTree: Editor fixes
- Updated add button to new positive button style
- Change the add-state logic to add to root or subtrees if context is present (e.g. add after)
- Renamed UStateTreeState::LinkedState to LinkedSubtree to indicate what kind of state is expected
- Removed strong reference from the treeview (was preventing a StateTree asset being deleted)
- Fixed bad colors when renaming a state

diff --git a/Engine/Plugins/Runtime/StateTree/Config/BaseStateTree.ini b/Engine/Plugins/Runtime/StateTree/Config/BaseStateTree.ini
--- a/Engine/Plugins/Runtime/StateTree/Config/BaseStateTree.ini
+++ b/Engine/Plugins/Runtime/StateTree/Config/BaseStateTree.ini
@@ -20,6 +20,7 @@
 +PropertyRedirects=(OldName="StateTreeState.EnterConditions2",NewName="StateTreeState.EnterConditions")
 +PropertyRedirects=(OldName="StateTreeState.Evaluators2",NewName="StateTreeState.Evaluators")
 +PropertyRedirects=(OldName="StateTreeState.Tasks2",NewName="StateTreeState.Tasks")
++PropertyRedirects=(OldName="StateTreeState.LinkedState",NewName="StateTreeState.LinkedSubtree")
 
 +PropertyRedirects=(OldName="StateTree.Conditions2",NewName="StateTree.Conditions")
 +PropertyRedirects=(OldName="StateTree.Items",NewName="StateTree.Nodes")
diff --git a/Engine/Plugins/Runtime/StateTree/Source/StateTreeEditorModule/Private/SStateTreeView.cpp b/Engine/Plugins/Runtime/StateTree/Source/StateTreeEditorModule/Private/SStateTreeView.cpp
--- a/Engine/Plugins/Runtime/StateTree/Source/StateTreeEditorModule/Private/SStateTreeView.cpp
+++ b/Engine/Plugins/Runtime/StateTree/Source/StateTreeEditorModule/Private/SStateTreeView.cpp
@@ -8,7 +8,6 @@
 #include "ScopedTransaction.h"
 
 #include "Templates/SharedPointer.h"
-
 #include "Styling/AppStyle.h"
 #include "EditorFontGlyphs.h"
 #include "StateTreeEditorStyle.h"
@@ -22,7 +21,9 @@
 #include "Widgets/Images/SImage.h"
 #include "Widgets/Text/SInlineEditableTextBlock.h"
 #include "Widgets/Layout/SScrollBox.h"
+#include "SPositiveActionButton.h"
 #include "SStateTreeViewRow.h"
+#include "StateTree.h"
 
 #include "StateTreeViewModel.h"
 #include "StateTreeState.h"
@@ -70,10 +71,12 @@ void SStateTreeView::Construct(const FArguments& InArgs, TSharedRef<FStateTreeVi
 		.Orientation(Orient_Vertical)
 		.Thickness(FVector2D(12.0f, 12.0f));
 
-	TreeView = SNew(STreeView<UStateTreeState*>)
+	StateTreeViewModel->GetSubTrees(Subtrees);
+
+	TreeView = SNew(STreeView<TWeakObjectPtr<UStateTreeState>>)
 		.OnGenerateRow(this, &SStateTreeView::HandleGenerateRow)
 		.OnGetChildren(this, &SStateTreeView::HandleGetChildren)
-		.TreeItemsSource(StateTreeViewModel->GetSubTrees())
+		.TreeItemsSource(&Subtrees)
 		.ItemHeight(32)
 		.OnSelectionChanged(this, &SStateTreeView::HandleTreeSelectionChanged)
 		.OnContextMenuOpening(this, &SStateTreeView::HandleContextMenuOpening)
@@ -100,31 +103,11 @@ void SStateTreeView::Construct(const FArguments& InArgs, TSharedRef<FStateTreeVi
 				.Padding(4.0f, 2.0f)
 				.AutoWidth()
 				[
-					SNew(SButton)
-					.ButtonStyle(FAppStyle::Get(), "FlatButton.Success")
+					SNew(SPositiveActionButton)
 					.ToolTipText(LOCTEXT("AddStateToolTip", "Add New State"))
+					.Icon(FAppStyle::Get().GetBrush("Icons.Plus")) 
+					.Text(LOCTEXT("AddState", "Add State"))
 					.OnClicked(this, &SStateTreeView::HandleAddStateButton)
-					[
-						SNew(SHorizontalBox)
-						+ SHorizontalBox::Slot()
-						.VAlign(VAlign_Center)
-						.Padding(0.0f, 0.0f, 4.0f, 0.0f)
-						.AutoWidth()
-						[
-							SNew(STextBlock)
-							.Text(FEditorFontGlyphs::File)
-							.TextStyle(FAppStyle::Get(), "FlatButton.DefaultTextStyle")
-							.Font(FAppStyle::Get().GetFontStyle("FontAwesome.11"))
-						]
-						+ SHorizontalBox::Slot()
-						.VAlign(VAlign_Center)
-						.AutoWidth()
-						[
-							SNew(STextBlock)
-							.Text(LOCTEXT("AddState", "Add State"))
-							.TextStyle(FAppStyle::Get(), "FlatButton.DefaultTextStyle")
-						]
-					]
 				]
 			]
 		]
@@ -170,7 +153,7 @@ void SStateTreeView::SavePersistentExpandedStates()
 		return;
 	}
 
-	TSet<UStateTreeState*> ExpandedStates;
+	TSet<TWeakObjectPtr<UStateTreeState>> ExpandedStates;
 	TreeView->GetExpandedItems(ExpandedStates);
 	StateTreeViewModel->SetPersistentExpandedStates(ExpandedStates);
 }
@@ -199,7 +182,7 @@ void SStateTreeView::UpdateTree(bool bExpandPersistent)
 		return;
 	}
 
-	TSet<UStateTreeState*> ExpandedStates;
+	TSet<TWeakObjectPtr<UStateTreeState>> ExpandedStates;
 	if (bExpandPersistent)
 	{
 		// Get expanded state from the tree data.
@@ -212,14 +195,15 @@ void SStateTreeView::UpdateTree(bool bExpandPersistent)
 	}
 
 	// Remember selection
-	TArray<UStateTreeState*> SelectedStates;
+	TArray<TWeakObjectPtr<UStateTreeState>> SelectedStates;
 	StateTreeViewModel->GetSelectedStates(SelectedStates);
 
 	// Regenerate items
-	TreeView->SetTreeItemsSource(StateTreeViewModel->GetSubTrees());
+	StateTreeViewModel->GetSubTrees(Subtrees);
+	TreeView->SetTreeItemsSource(&Subtrees);
 
 	// Restore expanded state
-	for (UStateTreeState* State : ExpandedStates)
+	for (const TWeakObjectPtr<UStateTreeState>& State : ExpandedStates)
 	{
 		TreeView->SetItemExpansion(State, true);
 	}
@@ -272,7 +256,7 @@ void SStateTreeView::HandleModelStatesChanged(const TSet<UStateTreeState*>& Affe
 	}
 }
 
-void SStateTreeView::HandleModelSelectionChanged(const TArray<UStateTreeState*>& SelectedStates)
+void SStateTreeView::HandleModelSelectionChanged(const TArray<TWeakObjectPtr<UStateTreeState>>& SelectedStates)
 {
 	if (bUpdatingSelection)
 	{
@@ -288,20 +272,20 @@ void SStateTreeView::HandleModelSelectionChanged(const TArray<UStateTreeState*>&
 }
 
 
-TSharedRef<ITableRow> SStateTreeView::HandleGenerateRow(UStateTreeState* InState, const TSharedRef<STableViewBase>& InOwnerTableView)
+TSharedRef<ITableRow> SStateTreeView::HandleGenerateRow(TWeakObjectPtr<UStateTreeState> InState, const TSharedRef<STableViewBase>& InOwnerTableView)
 {
 	return SNew(SStateTreeViewRow, InOwnerTableView, InState, ViewBox, StateTreeViewModel.ToSharedRef());
 }
 
-void SStateTreeView::HandleGetChildren(UStateTreeState* InParent, TArray<UStateTreeState*>& OutChildren)
+void SStateTreeView::HandleGetChildren(TWeakObjectPtr<UStateTreeState> InParent, TArray<TWeakObjectPtr<UStateTreeState>>& OutChildren)
 {
-	if (InParent)
+	if (const UStateTreeState* Parent = InParent.Get())
 	{
-		OutChildren.Append(InParent->Children);
+		OutChildren.Append(Parent->Children);
 	}
 }
 
-void SStateTreeView::HandleTreeSelectionChanged(UStateTreeState* InSelectedItem, ESelectInfo::Type SelectionType)
+void SStateTreeView::HandleTreeSelectionChanged(TWeakObjectPtr<UStateTreeState> InSelectedItem, ESelectInfo::Type SelectionType)
 {
 	if (!StateTreeViewModel)
 	{
@@ -314,7 +298,7 @@ void SStateTreeView::HandleTreeSelectionChanged(UStateTreeState* InSelectedItem,
 		return;
 	}
 
-	TArray<UStateTreeState*> SelectedItems = TreeView->GetSelectedItems();
+	TArray<TWeakObjectPtr<UStateTreeState>> SelectedItems = TreeView->GetSelectedItems();
 
 	bUpdatingSelection = true;
 	StateTreeViewModel->SetSelection(SelectedItems);
@@ -391,21 +375,10 @@ FReply SStateTreeView::HandleAddStateButton()
 			StateTreeViewModel->AddState(FirstSelectedState);
 		}
 	}
-	else if (StateTreeViewModel->GetSubTreeCount() == 0)
-	{
-		// If no root state, create new.
-		StateTreeViewModel->AddState(nullptr);
-	}
 	else
 	{
-		// If no state has been selected, add to root
-		TArray<UStateTreeState*>* SubTrees = StateTreeViewModel->GetSubTrees();
-		UStateTreeState* RootState = (SubTrees != nullptr && SubTrees->Num() > 0) ? SubTrees->GetData()[0] : nullptr;
-		if (RootState != nullptr)
-		{
-			StateTreeViewModel->AddChildState(RootState);
-			TreeView->SetItemExpansion(RootState, true);
-		}
+		// Add root state at the lowest level.
+		StateTreeViewModel->AddState(nullptr);
 	}
 
 	return FReply::Handled();
diff --git a/Engine/Plugins/Runtime/StateTree/Source/StateTreeEditorModule/Private/SStateTreeView.h b/Engine/Plugins/Runtime/StateTree/Source/StateTreeEditorModule/Private/SStateTreeView.h
--- a/Engine/Plugins/Runtime/StateTree/Source/StateTreeEditorModule/Private/SStateTreeView.h
+++ b/Engine/Plugins/Runtime/StateTree/Source/StateTreeEditorModule/Private/SStateTreeView.h
@@ -56,12 +56,12 @@ private:
 	void HandleModelStatesMoved(const TSet<UStateTreeState*>& AffectedParents, const TSet<UStateTreeState*>& MovedStates);
 	void HandleModelStateAdded(UStateTreeState* ParentState, UStateTreeState* NewState);
 	void HandleModelStatesChanged(const TSet<UStateTreeState*>& AffectedStates, const FPropertyChangedEvent& PropertyChangedEvent);
-	void HandleModelSelectionChanged(const TArray<UStateTreeState*>& SelectedStates);
+	void HandleModelSelectionChanged(const TArray<TWeakObjectPtr<UStateTreeState>>& SelectedStates);
 
 	// Treeview handlers
-	TSharedRef<ITableRow> HandleGenerateRow(UStateTreeState* InState, const TSharedRef<STableViewBase>& InOwnerTableView);
-	void HandleGetChildren(UStateTreeState* InParent, TArray<UStateTreeState*>& OutChildren);
-	void HandleTreeSelectionChanged(UStateTreeState* InSelectedItem, ESelectInfo::Type SelectionType);
+	TSharedRef<ITableRow> HandleGenerateRow(TWeakObjectPtr<UStateTreeState> InState, const TSharedRef<STableViewBase>& InOwnerTableView);
+	void HandleGetChildren(TWeakObjectPtr<UStateTreeState> InParent, TArray<TWeakObjectPtr<UStateTreeState>>& OutChildren);
+	void HandleTreeSelectionChanged(TWeakObjectPtr<UStateTreeState> InSelectedItem, ESelectInfo::Type SelectionType);
 	TSharedPtr<SWidget> HandleContextMenuOpening();
 
 	// Action handlers
@@ -73,9 +73,10 @@ private:
 
 	TSharedPtr<FStateTreeViewModel> StateTreeViewModel;
 
-	TSharedPtr<STreeView<UStateTreeState*>> TreeView;
+	TSharedPtr<STreeView<TWeakObjectPtr<UStateTreeState>>> TreeView;
 	TSharedPtr<SScrollBar> ExternalScrollbar;
 	TSharedPtr<SScrollBox> ViewBox;
+	TArray<TWeakObjectPtr<UStateTreeState>> Subtrees;
 
 	UStateTreeState* RequestedRenameState;
 	bool bItemsDirty;
diff --git a/Engine/Plugins/Runtime/StateTree/Source/StateTreeEditorModule/Private/SStateTreeViewRow.cpp b/Engine/Plugins/Runtime/StateTree/Source/StateTreeEditorModule/Private/SStateTreeViewRow.cpp
--- a/Engine/Plugins/Runtime/StateTree/Source/StateTreeEditorModule/Private/SStateTreeViewRow.cpp
+++ b/Engine/Plugins/Runtime/StateTree/Source/StateTreeEditorModule/Private/SStateTreeViewRow.cpp
@@ -29,12 +29,12 @@
 
 #define LOCTEXT_NAMESPACE "StateTreeEditor"
 
-void SStateTreeViewRow::Construct(const FArguments& InArgs, const TSharedRef<STableViewBase>& InOwnerTableView, UStateTreeState* InState, const TSharedPtr<SScrollBox>& ViewBox, TSharedRef<FStateTreeViewModel> InStateTreeViewModel)
+void SStateTreeViewRow::Construct(const FArguments& InArgs, const TSharedRef<STableViewBase>& InOwnerTableView, TWeakObjectPtr<UStateTreeState> InState, const TSharedPtr<SScrollBox>& ViewBox, TSharedRef<FStateTreeViewModel> InStateTreeViewModel)
 {
 	StateTreeViewModel = InStateTreeViewModel;
 	WeakState = InState;
 
-	STableRow<UStateTreeState*>::ConstructInternal(STableRow::FArguments()
+	STableRow<TWeakObjectPtr<UStateTreeState>>::ConstructInternal(STableRow::FArguments()
         .Padding(5.0f)
         .OnDragDetected(this, &SStateTreeViewRow::HandleDragDetected)
         .OnCanAcceptDrop(this, &SStateTreeViewRow::HandleCanAcceptDrop)
@@ -503,7 +503,7 @@ FText SStateTreeViewRow::GetLinkedStateDesc() const
 
 	if (State->Type == EStateTreeStateType::Linked)
 	{
-		return FText::FromName(State->LinkedState.Name);
+		return FText::FromName(State->LinkedSubtree.Name);
 	}
 	
 	return FText::GetEmpty();
@@ -823,13 +823,13 @@ FReply SStateTreeViewRow::HandleDragDetected(const FGeometry&, const FPointerEve
 	return FReply::Handled().BeginDragDrop(FActionTreeViewDragDrop::New(WeakState.Get()));
 }
 
-TOptional<EItemDropZone> SStateTreeViewRow::HandleCanAcceptDrop(const FDragDropEvent& DragDropEvent, EItemDropZone DropZone, UStateTreeState* TargetState) const
+TOptional<EItemDropZone> SStateTreeViewRow::HandleCanAcceptDrop(const FDragDropEvent& DragDropEvent, EItemDropZone DropZone, TWeakObjectPtr<UStateTreeState> TargetState) const
 {
 	const TSharedPtr<FActionTreeViewDragDrop> DragDropOperation = DragDropEvent.GetOperationAs<FActionTreeViewDragDrop>();
 	if (DragDropOperation.IsValid())
 	{
 		// Cannot drop on selection or child of selection.
-		if (StateTreeViewModel && StateTreeViewModel->IsChildOfSelection(TargetState))
+		if (StateTreeViewModel && StateTreeViewModel->IsChildOfSelection(TargetState.Get()))
 		{
 			return TOptional<EItemDropZone>();
 		}
@@ -840,7 +840,7 @@ TOptional<EItemDropZone> SStateTreeViewRow::HandleCanAcceptDrop(const FDragDropE
 	return TOptional<EItemDropZone>();
 }
 
-FReply SStateTreeViewRow::HandleAcceptDrop(const FDragDropEvent& DragDropEvent, EItemDropZone DropZone, UStateTreeState* TargetState) const
+FReply SStateTreeViewRow::HandleAcceptDrop(const FDragDropEvent& DragDropEvent, EItemDropZone DropZone, TWeakObjectPtr<UStateTreeState> TargetState) const
 {
 	const TSharedPtr<FActionTreeViewDragDrop> DragDropOperation = DragDropEvent.GetOperationAs<FActionTreeViewDragDrop>();
 	if (DragDropOperation.IsValid())
@@ -849,15 +849,15 @@ FReply SStateTreeViewRow::HandleAcceptDrop(const FDragDropEvent& DragDropEvent,
 		{
 			if (DropZone == EItemDropZone::AboveItem)
 			{
-				StateTreeViewModel->MoveSelectedStatesBefore(TargetState);
+				StateTreeViewModel->MoveSelectedStatesBefore(TargetState.Get());
 			}
 			else if (DropZone == EItemDropZone::BelowItem)
 			{
-				StateTreeViewModel->MoveSelectedStatesAfter(TargetState);
+				StateTreeViewModel->MoveSelectedStatesAfter(TargetState.Get());
 			}
 			else
 			{
-				StateTreeViewModel->MoveSelectedStatesInto(TargetState);
+				StateTreeViewModel->MoveSelectedStatesInto(TargetState.Get());
 			}
 
 			return FReply::Handled();
diff --git a/Engine/Plugins/Runtime/StateTree/Source/StateTreeEditorModule/Private/SStateTreeViewRow.h b/Engine/Plugins/Runtime/StateTree/Source/StateTreeEditorModule/Private/SStateTreeViewRow.h
--- a/Engine/Plugins/Runtime/StateTree/Source/StateTreeEditorModule/Private/SStateTreeViewRow.h
+++ b/Engine/Plugins/Runtime/StateTree/Source/StateTreeEditorModule/Private/SStateTreeViewRow.h
@@ -12,7 +12,7 @@ class SInlineEditableTextBlock;
 class SScrollBox;
 class FStateTreeViewModel;
 
-class SStateTreeViewRow : public STableRow<UStateTreeState*>
+class SStateTreeViewRow : public STableRow<TWeakObjectPtr<UStateTreeState>>
 {
 public:
 
@@ -21,7 +21,7 @@ public:
 	}
 	SLATE_END_ARGS()
 
-	void Construct(const FArguments& InArgs, const TSharedRef<STableViewBase>& InOwnerTableView, UStateTreeState* InState, const TSharedPtr<SScrollBox>& ViewBox, TSharedRef<FStateTreeViewModel> InStateTreeViewModel);
+	void Construct(const FArguments& InArgs, const TSharedRef<STableViewBase>& InOwnerTableView, TWeakObjectPtr<UStateTreeState> InState, const TSharedPtr<SScrollBox>& ViewBox, TSharedRef<FStateTreeViewModel> InStateTreeViewModel);
 	void RequestRename() const;
 
 private:
@@ -66,8 +66,8 @@ private:
 	void HandleNodeLabelTextCommitted(const FText& NewLabel, ETextCommit::Type CommitType) const;
 
 	FReply HandleDragDetected(const FGeometry&, const FPointerEvent&) const;
-	TOptional<EItemDropZone> HandleCanAcceptDrop(const FDragDropEvent& DragDropEvent, EItemDropZone DropZone, UStateTreeState* TargetState) const;
-	FReply HandleAcceptDrop(const FDragDropEvent& DragDropEvent, EItemDropZone DropZone, UStateTreeState* TargetState) const;
+	TOptional<EItemDropZone> HandleCanAcceptDrop(const FDragDropEvent& DragDropEvent, EItemDropZone DropZone, TWeakObjectPtr<UStateTreeState> TargetState) const;
+	FReply HandleAcceptDrop(const FDragDropEvent& DragDropEvent, EItemDropZone DropZone, TWeakObjectPtr<UStateTreeState> TargetState) const;
 	
 	TSharedPtr<FStateTreeViewModel> StateTreeViewModel;
 	TWeakObjectPtr<UStateTreeState> WeakState;
diff --git a/Engine/Plugins/Runtime/StateTree/Source/StateTreeEditorModule/Private/StateTreeCompiler.cpp b/Engine/Plugins/Runtime/StateTree/Source/StateTreeEditorModule/Private/StateTreeCompiler.cpp
--- a/Engine/Plugins/Runtime/StateTree/Source/StateTreeEditorModule/Private/StateTreeCompiler.cpp
+++ b/Engine/Plugins/Runtime/StateTree/Source/StateTreeEditorModule/Private/StateTreeCompiler.cpp
@@ -501,7 +501,7 @@ bool FStateTreeCompiler::CreateStateTransitions()
 			const UStateTreeState* LinkedParentState = nullptr;
 			for (const UStateTreeState* State = SourceState; State != nullptr; State = State->Parent)
 			{
-				if (State->ID == SourceState->LinkedState.ID)
+				if (State->ID == SourceState->LinkedSubtree.ID)
 				{
 					LinkedParentState = State;
 					break;
@@ -511,36 +511,36 @@ bool FStateTreeCompiler::CreateStateTransitions()
 			if (LinkedParentState != nullptr)
 			{
 				Log.Reportf(EMessageSeverity::Error,
-					TEXT("State is linked to it's parent state '%s', which will create infinite loop."),
+					TEXT("State is linked to it's parent subtree '%s', which will create infinite loop."),
 					*LinkedParentState->Name.ToString());
 				return false;
 			}
 
 			// The linked state must be a subtree.
-			const UStateTreeState* TargetState = GetState(SourceState->LinkedState.ID);
+			const UStateTreeState* TargetState = GetState(SourceState->LinkedSubtree.ID);
 			if (TargetState == nullptr)
 			{
 				Log.Reportf(EMessageSeverity::Error,
-					TEXT("Failed to resolve linked state '%s'."),
-					*SourceState->LinkedState.Name.ToString());
+					TEXT("Failed to resolve linked subtree '%s'."),
+					*SourceState->LinkedSubtree.Name.ToString());
 				return false;
 			}
 			
 			if (TargetState->Type != EStateTreeStateType::Subtree)
 			{
 				Log.Reportf(EMessageSeverity::Error,
-					TEXT("State '%s' is linked to state '%s', which is not a subtree."),
+					TEXT("State '%s' is linked to subtree '%s', which is not a subtree."),
 					*SourceState->Name.ToString(), *TargetState->Name.ToString());
 				return false;
 			}
 			
-			CompactState.LinkedState = GetStateHandle(SourceState->LinkedState.ID);
+			CompactState.LinkedState = GetStateHandle(SourceState->LinkedSubtree.ID);
 			
 			if (!CompactState.LinkedState.IsValid())
 			{
 				Log.Reportf(EMessageSeverity::Error,
-					TEXT("Failed to resolve linked state '%s'."),
-					*SourceState->LinkedState.Name.ToString());
+					TEXT("Failed to resolve linked subtree '%s'."),
+					*SourceState->LinkedSubtree.Name.ToString());
 				return false;
 			}
 		}
diff --git a/Engine/Plugins/Runtime/StateTree/Source/StateTreeEditorModule/Private/StateTreeEditor.cpp b/Engine/Plugins/Runtime/StateTree/Source/StateTreeEditorModule/Private/StateTreeEditor.cpp
--- a/Engine/Plugins/Runtime/StateTree/Source/StateTreeEditorModule/Private/StateTreeEditor.cpp
+++ b/Engine/Plugins/Runtime/StateTree/Source/StateTreeEditorModule/Private/StateTreeEditor.cpp
@@ -383,14 +383,17 @@ void FStateTreeEditor::HandleModelAssetChanged()
 	UpdateAsset();
 }
 
-void FStateTreeEditor::HandleModelSelectionChanged(const TArray<UStateTreeState*>& SelectedStates)
+void FStateTreeEditor::HandleModelSelectionChanged(const TArray<TWeakObjectPtr<UStateTreeState>>& SelectedStates)
 {
 	if (SelectionDetailsView)
 	{
 		TArray<UObject*> Selected;
-		for (UStateTreeState* State : SelectedStates)
+		for (const TWeakObjectPtr<UStateTreeState>& WeakState : SelectedStates)
 		{
-			Selected.Add(State);
+			if (UStateTreeState* State = WeakState.Get())
+			{
+				Selected.Add(State);
+			}
 		}
 		SelectionDetailsView->SetObjects(Selected);
 	}
@@ -457,9 +460,9 @@ void FStateTreeEditor::OnStateParametersChanged(const UStateTree& InStateTree, c
 		{
 			TreeData->VisitHierarchy([&ChangedStateID](UStateTreeState& State, UStateTreeState* /*ParentState*/)
 			{
-				if (State.Type == EStateTreeStateType::Linked && State.LinkedState.ID == ChangedStateID)
+				if (State.Type == EStateTreeStateType::Linked && State.LinkedSubtree.ID == ChangedStateID)
 				{
-					State.UpdateParametersFromLinkedState();
+					State.UpdateParametersFromLinkedSubtree();
 				}
 				return EStateTreeVisitor::Continue;
 			});
@@ -554,7 +557,7 @@ namespace UE::StateTree::Editor::Internal
 		{
 			if (State.Type == EStateTreeStateType::Linked)
 			{
-				FixChangedStateLinkName(State.LinkedState, IDToName);
+				FixChangedStateLinkName(State.LinkedSubtree, IDToName);
 			}
 					
 			for (FStateTreeTransition& Transition : State.Transitions)
@@ -598,7 +601,7 @@ namespace UE::StateTree::Editor::Internal
 		// Clear evaluators if not allowed.
 		if (Schema->AllowEvaluators() == false && TreeData->Evaluators.Num() > 0)
 		{
-			UE_LOG(LogStateTree, Warning, TEXT("%s: Resetting Evaluators due to current schema restrictions."), *GetNameSafe(&StateTree));
+			UE_LOG(LogStateTreeEditor, Warning, TEXT("%s: Resetting Evaluators due to current schema restrictions."), *GetNameSafe(&StateTree));
 			TreeData->Evaluators.Reset();
 		}
 
@@ -608,7 +611,7 @@ namespace UE::StateTree::Editor::Internal
 			// Clear enter conditions if not allowed.
 			if (Schema->AllowEnterConditions() == false && State.EnterConditions.Num() > 0)
 			{
-				UE_LOG(LogStateTree, Warning, TEXT("%s: Resetting Enter Conditions in state %s due to current schema restrictions."), *GetNameSafe(&StateTree), *GetNameSafe(&State));
+				UE_LOG(LogStateTreeEditor, Warning, TEXT("%s: Resetting Enter Conditions in state %s due to current schema restrictions."), *GetNameSafe(&StateTree), *GetNameSafe(&State));
 				State.EnterConditions.Reset();
 			}
 
@@ -618,7 +621,7 @@ namespace UE::StateTree::Editor::Internal
 				if (State.Tasks.Num() > 0)
 				{
 					State.Tasks.Reset();
-					UE_LOG(LogStateTree, Warning, TEXT("%s: Resetting Tasks in state %s due to current schema restrictions."), *GetNameSafe(&StateTree), *GetNameSafe(&State));
+					UE_LOG(LogStateTreeEditor, Warning, TEXT("%s: Resetting Tasks in state %s due to current schema restrictions."), *GetNameSafe(&StateTree), *GetNameSafe(&State));
 				}
 				
 				// Task name is the same as state name.
@@ -632,7 +635,7 @@ namespace UE::StateTree::Editor::Internal
 				if (State.SingleTask.Node.IsValid())
 				{
 					State.SingleTask.Reset();
-					UE_LOG(LogStateTree, Warning, TEXT("%s: Resetting Single Task in state %s due to current schema restrictions."), *GetNameSafe(&StateTree), *GetNameSafe(&State));
+					UE_LOG(LogStateTreeEditor, Warning, TEXT("%s: Resetting Single Task in state %s due to current schema restrictions."), *GetNameSafe(&StateTree), *GetNameSafe(&State));
 				}
 			}
 			
@@ -665,7 +668,7 @@ namespace UE::StateTree::Editor::Internal
 		{
 			if (State.Type == EStateTreeStateType::Linked)
 			{
-				State.UpdateParametersFromLinkedState();
+				State.UpdateParametersFromLinkedSubtree();
 			}
 			return EStateTreeVisitor::Continue;
 		});
diff --git a/Engine/Plugins/Runtime/StateTree/Source/StateTreeEditorModule/Private/StateTreeEditor.h b/Engine/Plugins/Runtime/StateTree/Source/StateTreeEditorModule/Private/StateTreeEditor.h
--- a/Engine/Plugins/Runtime/StateTree/Source/StateTreeEditorModule/Private/StateTreeEditor.h
+++ b/Engine/Plugins/Runtime/StateTree/Source/StateTreeEditorModule/Private/StateTreeEditor.h
@@ -62,7 +62,7 @@ protected:
 private:
 
 	void HandleMessageTokenClicked(const TSharedRef<IMessageToken>& InMessageToken);
-	void HandleModelSelectionChanged(const TArray<UStateTreeState*>& SelectedStates);
+	void HandleModelSelectionChanged(const TArray<TWeakObjectPtr<UStateTreeState>>& SelectedStates);
 
 	/** Spawns the tab with the update graph inside */
 	TSharedRef<SDockTab> SpawnTab_StateTreeView(const FSpawnTabArgs& Args);
diff --git a/Engine/Plugins/Runtime/StateTree/Source/StateTreeEditorModule/Private/StateTreeEditorStyle.cpp b/Engine/Plugins/Runtime/StateTree/Source/StateTreeEditorModule/Private/StateTreeEditorStyle.cpp
--- a/Engine/Plugins/Runtime/StateTree/Source/StateTreeEditorModule/Private/StateTreeEditorStyle.cpp
+++ b/Engine/Plugins/Runtime/StateTree/Source/StateTreeEditorModule/Private/StateTreeEditorStyle.cpp
@@ -61,11 +61,13 @@ void FStateTreeEditorStyle::Initialize()
 
 		FEditableTextBoxStyle StateTitleEditableText = FEditableTextBoxStyle()
 			.SetTextStyle(NormalText)
-			.SetFont(TTF_CORE_FONT("Fonts/Roboto-Bold", 10))
+			.SetFont(TTF_CORE_FONT("Fonts/Roboto-Bold", 12))
 			.SetBackgroundImageNormal(BOX_BRUSH("Common/TextBox", FMargin(4.0f / 16.0f)))
 			.SetBackgroundImageHovered(BOX_BRUSH("Common/TextBox_Hovered", FMargin(4.0f / 16.0f)))
 			.SetBackgroundImageFocused(BOX_BRUSH("Common/TextBox_Hovered", FMargin(4.0f / 16.0f)))
 			.SetBackgroundImageReadOnly(BOX_BRUSH("Common/TextBox_ReadOnly", FMargin(4.0f / 16.0f)))
+			.SetBackgroundColor(FLinearColor(0,0,0,0.1f))
+			.SetPadding(FMargin(0))
 			.SetScrollBarStyle(ScrollBar);
 		StyleSet->Set("StateTree.State.TitleEditableText", StateTitleEditableText);
 
diff --git a/Engine/Plugins/Runtime/StateTree/Source/StateTreeEditorModule/Private/StateTreeState.cpp b/Engine/Plugins/Runtime/StateTree/Source/StateTreeEditorModule/Private/StateTreeState.cpp
--- a/Engine/Plugins/Runtime/StateTree/Source/StateTreeEditorModule/Private/StateTreeState.cpp
+++ b/Engine/Plugins/Runtime/StateTree/Source/StateTreeEditorModule/Private/StateTreeState.cpp
@@ -90,14 +90,14 @@ void UStateTreeState::PostEditChangeChainProperty(FPropertyChangedChainEvent& Pr
 			// If transitioning from linked state, reset the linked state.
 			if (Type != EStateTreeStateType::Linked)
 			{
-				LinkedState = FStateTreeStateLink();
+				LinkedSubtree = FStateTreeStateLink();
 			}
 
 			if (Type == EStateTreeStateType::Linked)
 			{
 				// Linked parameter layout is fixed, and copied from the linked target state.
 				Parameters.bFixedLayout = true;
-				UpdateParametersFromLinkedState();
+				UpdateParametersFromLinkedSubtree();
 			}
 			else if (Type == EStateTreeStateType::Subtree)
 			{
@@ -111,12 +111,12 @@ void UStateTreeState::PostEditChangeChainProperty(FPropertyChangedChainEvent& Pr
 		}
 
 		if (Property->GetOwnerClass() == UStateTreeState::StaticClass()
-			&& Property->GetFName() == GET_MEMBER_NAME_CHECKED(UStateTreeState, LinkedState))
+			&& Property->GetFName() == GET_MEMBER_NAME_CHECKED(UStateTreeState, LinkedSubtree))
 		{
 			// When switching to new state, update the parameters.
 			if (Type == EStateTreeStateType::Linked)
 			{
-				UpdateParametersFromLinkedState();
+				UpdateParametersFromLinkedSubtree();
 			}
 		}
 
@@ -191,11 +191,11 @@ void UStateTreeState::PostLoad()
 #endif
 }
 
-void UStateTreeState::UpdateParametersFromLinkedState()
+void UStateTreeState::UpdateParametersFromLinkedSubtree()
 {
 	if (const UStateTreeEditorData* TreeData = GetTypedOuter<UStateTreeEditorData>())
 	{
-		if (const UStateTreeState* LinkTargetState = TreeData->GetStateByID(LinkedState.ID))
+		if (const UStateTreeState* LinkTargetState = TreeData->GetStateByID(LinkedSubtree.ID))
 		{
 			Parameters.Parameters.MigrateToNewBagInstance(LinkTargetState->Parameters.Parameters);
 		}
diff --git a/Engine/Plugins/Runtime/StateTree/Source/StateTreeEditorModule/Private/StateTreeViewModel.cpp b/Engine/Plugins/Runtime/StateTree/Source/StateTreeEditorModule/Private/StateTreeViewModel.cpp
--- a/Engine/Plugins/Runtime/StateTree/Source/StateTreeEditorModule/Private/StateTreeViewModel.cpp
+++ b/Engine/Plugins/Runtime/StateTree/Source/StateTreeEditorModule/Private/StateTreeViewModel.cpp
@@ -79,7 +79,7 @@ namespace UE::StateTree::Editor
 
 
 FStateTreeViewModel::FStateTreeViewModel()
-	: TreeData(nullptr)
+	: TreeDataWeak(nullptr)
 {
 }
 
@@ -92,7 +92,7 @@ FStateTreeViewModel::~FStateTreeViewModel()
 
 void FStateTreeViewModel::Init(UStateTreeEditorData* InTreeData)
 {
-	TreeData = InTreeData;
+	TreeDataWeak = InTreeData;
 
 	GEditor->RegisterForUndo(this);
 
@@ -101,6 +101,7 @@ void FStateTreeViewModel::Init(UStateTreeEditorData* InTreeData)
 
 void FStateTreeViewModel::HandleIdentifierChanged(const UStateTree& StateTree) const
 {
+	UStateTreeEditorData* TreeData = TreeDataWeak.Get();
 	const UStateTree* OuterStateTree = TreeData ? Cast<UStateTree>(TreeData->GetOuter()) : nullptr;
 	if (OuterStateTree == &StateTree)
 	{
@@ -120,14 +121,28 @@ void FStateTreeViewModel::NotifyStatesChangedExternally(const TSet<UStateTreeSta
 
 TArray<UStateTreeState*>* FStateTreeViewModel::GetSubTrees() const
 {
+	UStateTreeEditorData* TreeData = TreeDataWeak.Get();
 	return TreeData != nullptr ? &ToRawPtrTArrayUnsafe(TreeData->SubTrees) : nullptr;
 }
 
 int32 FStateTreeViewModel::GetSubTreeCount() const
 {
+	UStateTreeEditorData* TreeData = TreeDataWeak.Get();
 	return TreeData != nullptr ? TreeData->SubTrees.Num() : 0;
 }
 
+void FStateTreeViewModel::GetSubTrees(TArray<TWeakObjectPtr<UStateTreeState>>& OutSubtrees) const
+{
+	OutSubtrees.Reset();
+	if (UStateTreeEditorData* TreeData = TreeDataWeak.Get())
+	{
+		for (UStateTreeState* Subtree : TreeData->SubTrees)
+		{
+			OutSubtrees.Add(Subtree);
+		}
+	}
+}
+
 void FStateTreeViewModel::PostUndo(bool bSuccess)
 {
 	// TODO: see if we can narrow this down.
@@ -143,7 +158,7 @@ void FStateTreeViewModel::ClearSelection()
 {
 	SelectedStates.Reset();
 
-	const TArray<UStateTreeState*> SelectedStatesArr;
+	const TArray<TWeakObjectPtr<UStateTreeState>> SelectedStatesArr;
 	OnSelectionChanged.Broadcast(SelectedStatesArr);
 }
 
@@ -153,18 +168,18 @@ void FStateTreeViewModel::SetSelection(UStateTreeState* Selected)
 
 	SelectedStates.Add(Selected);
 
-	TArray<UStateTreeState*> SelectedStatesArr;
+	TArray<TWeakObjectPtr<UStateTreeState>> SelectedStatesArr;
 	SelectedStatesArr.Add(Selected);
 	OnSelectionChanged.Broadcast(SelectedStatesArr);
 }
 
-void FStateTreeViewModel::SetSelection(const TArray<UStateTreeState*>& InSelectedStates)
+void FStateTreeViewModel::SetSelection(const TArray<TWeakObjectPtr<UStateTreeState>>& InSelectedStates)
 {
 	SelectedStates.Reset();
 
-	for (UStateTreeState* State : InSelectedStates)
+	for (const TWeakObjectPtr<UStateTreeState>& State : InSelectedStates)
 	{
-		if (State)
+		if (State.Get())
 		{
 			SelectedStates.Add(State);
 		}
@@ -176,12 +191,13 @@ void FStateTreeViewModel::SetSelection(const TArray<UStateTreeState*>& InSelecte
 
 bool FStateTreeViewModel::IsSelected(const UStateTreeState* State) const
 {
-	return SelectedStates.Contains(State);
+	const TWeakObjectPtr<UStateTreeState> WeakState = const_cast<UStateTreeState*>(State);
+	return SelectedStates.Contains(WeakState);
 }
 
 bool FStateTreeViewModel::IsChildOfSelection(const UStateTreeState* State) const
 {
-	for (const FWeakObjectPtr& WeakSelectedState : SelectedStates)
+	for (const TWeakObjectPtr<UStateTreeState>& WeakSelectedState : SelectedStates)
 	{
 		if (const UStateTreeState* SelectedState = Cast<UStateTreeState>(WeakSelectedState.Get()))
 		{
@@ -202,34 +218,45 @@ bool FStateTreeViewModel::IsChildOfSelection(const UStateTreeState* State) const
 void FStateTreeViewModel::GetSelectedStates(TArray<UStateTreeState*>& OutSelectedStates)
 {
 	OutSelectedStates.Reset();
-	for (FWeakObjectPtr& WeakState : SelectedStates)
+	for (TWeakObjectPtr<UStateTreeState>& WeakState : SelectedStates)
 	{
-		if (UStateTreeState* State = Cast<UStateTreeState>(WeakState.Get()))
+		if (UStateTreeState* State = WeakState.Get())
 		{
 			OutSelectedStates.Add(State);
 		}
 	}
 }
 
+void FStateTreeViewModel::GetSelectedStates(TArray<TWeakObjectPtr<UStateTreeState>>& OutSelectedStates)
+{
+	OutSelectedStates.Reset();
+	for (TWeakObjectPtr<UStateTreeState>& WeakState : SelectedStates)
+	{
+		if (WeakState.Get())
+		{
+			OutSelectedStates.Add(WeakState);
+		}
+	}
+}
+
 bool FStateTreeViewModel::HasSelection() const
 {
 	return SelectedStates.Num() > 0;
 }
 
-void FStateTreeViewModel::GetPersistentExpandedStates(TSet<UStateTreeState*>& OutExpandedStates)
+void FStateTreeViewModel::GetPersistentExpandedStates(TSet<TWeakObjectPtr<UStateTreeState>>& OutExpandedStates)
 {
 	OutExpandedStates.Reset();
-	if (TreeData == nullptr)
-	{
-		return;
-	}
-	for (UStateTreeState* SubTree : TreeData->SubTrees)
+	if (UStateTreeEditorData* TreeData = TreeDataWeak.Get())
 	{
-		GetExpandedStatesRecursive(SubTree, OutExpandedStates);
+		for (UStateTreeState* SubTree : TreeData->SubTrees)
+		{
+			GetExpandedStatesRecursive(SubTree, OutExpandedStates);
+		}
 	}
 }
 
-void FStateTreeViewModel::GetExpandedStatesRecursive(UStateTreeState* State, TSet<UStateTreeState*>& OutExpandedStates)
+void FStateTreeViewModel::GetExpandedStatesRecursive(UStateTreeState* State, TSet<TWeakObjectPtr<UStateTreeState>>& OutExpandedStates)
 {
 	if (State->bExpanded)
 	{
@@ -241,8 +268,10 @@ void FStateTreeViewModel::GetExpandedStatesRecursive(UStateTreeState* State, TSe
 	}
 }
 
-void FStateTreeViewModel::SetPersistentExpandedStates(TSet<UStateTreeState*>& InExpandedStates)
+
+void FStateTreeViewModel::SetPersistentExpandedStates(TSet<TWeakObjectPtr<UStateTreeState>>& InExpandedStates)
 {
+	UStateTreeEditorData* TreeData = TreeDataWeak.Get();
 	if (TreeData == nullptr)
 	{
 		return;
@@ -250,9 +279,9 @@ void FStateTreeViewModel::SetPersistentExpandedStates(TSet<UStateTreeState*>& In
 
 	TreeData->Modify();
 
-	for (UStateTreeState* State : InExpandedStates)
+	for (TWeakObjectPtr<UStateTreeState>& WeakState : InExpandedStates)
 	{
-		if (State != nullptr)
+		if (UStateTreeState* State = WeakState.Get())
 		{
 			State->bExpanded = true;
 		}
@@ -262,6 +291,7 @@ void FStateTreeViewModel::SetPersistentExpandedStates(TSet<UStateTreeState*>& In
 
 void FStateTreeViewModel::AddState(UStateTreeState* AfterState)
 {
+	UStateTreeEditorData* TreeData = TreeDataWeak.Get();
 	if (TreeData == nullptr)
 	{
 		return;
@@ -272,7 +302,27 @@ void FStateTreeViewModel::AddState(UStateTreeState* AfterState)
 	UStateTreeState* NewState = NewObject<UStateTreeState>(TreeData, FName(), RF_Transactional);
 	UStateTreeState* ParentState = nullptr;
 
-	if (AfterState != nullptr)
+	if (AfterState == nullptr)
+	{
+		// If no subtrees, add a subtree, or add to the root state.
+		if (TreeData->SubTrees.IsEmpty())
+		{
+			TreeData->Modify();
+			TreeData->SubTrees.Add(NewState);
+		}
+		else
+		{
+			UStateTreeState* RootState = TreeData->SubTrees[0];
+			if (ensureMsgf(RootState, TEXT("%s: Root state is empty."), *GetNameSafe(TreeData->GetOuter())))
+			{
+				RootState->Modify();
+				RootState->Children.Add(NewState);
+				NewState->Parent = RootState;
+				ParentState = RootState;
+			}
+		}
+	}
+	else
 	{
 		ParentState = AfterState->Parent;
 		if (ParentState != nullptr)
@@ -307,6 +357,7 @@ void FStateTreeViewModel::AddState(UStateTreeState* AfterState)
 
 void FStateTreeViewModel::AddChildState(UStateTreeState* ParentState)
 {
+	UStateTreeEditorData* TreeData = TreeDataWeak.Get();
 	if (TreeData == nullptr || ParentState == nullptr)
 	{
 		return;
@@ -345,6 +396,7 @@ void FStateTreeViewModel::RenameState(UStateTreeState* State, FName NewName)
 
 void FStateTreeViewModel::RemoveSelectedStates()
 {
+	UStateTreeEditorData* TreeData = TreeDataWeak.Get();
 	if (TreeData == nullptr)
 	{
 		return;
@@ -414,6 +466,7 @@ void FStateTreeViewModel::MoveSelectedStatesInto(UStateTreeState* TargetState)
 
 void FStateTreeViewModel::MoveSelectedStates(UStateTreeState* TargetState, const FStateTreeViewModelInsert RelativeLocation)
 {
+	UStateTreeEditorData* TreeData = TreeDataWeak.Get();
 	if (TreeData == nullptr || TargetState == nullptr)
 	{
 		return;
@@ -535,7 +588,13 @@ void FStateTreeViewModel::MoveSelectedStates(UStateTreeState* TargetState, const
 
 		OnStatesMoved.Broadcast(AffectedParents, AffectedStates);
 
-		SetSelection(States);
+		TArray<TWeakObjectPtr<UStateTreeState>> WeakStates;
+		for (UStateTreeState* State : States)
+		{
+			WeakStates.Add(State);
+		}
+
+		SetSelection(WeakStates);
 	}
 }
 
diff --git a/Engine/Plugins/Runtime/StateTree/Source/StateTreeEditorModule/Private/StateTreeViewModel.h b/Engine/Plugins/Runtime/StateTree/Source/StateTreeEditorModule/Private/StateTreeViewModel.h
--- a/Engine/Plugins/Runtime/StateTree/Source/StateTreeEditorModule/Private/StateTreeViewModel.h
+++ b/Engine/Plugins/Runtime/StateTree/Source/StateTreeEditorModule/Private/StateTreeViewModel.h
@@ -31,7 +31,7 @@ public:
 	DECLARE_MULTICAST_DELEGATE_TwoParams(FOnStateAdded, UStateTreeState* /*ParentState*/, UStateTreeState* /*NewState*/);
 	DECLARE_MULTICAST_DELEGATE_OneParam(FOnStatesRemoved, const TSet<UStateTreeState*>& /*AffectedParents*/);
 	DECLARE_MULTICAST_DELEGATE_TwoParams(FOnStatesMoved, const TSet<UStateTreeState*>& /*AffectedParents*/, const TSet<UStateTreeState*>& /*MovedStates*/);
-	DECLARE_MULTICAST_DELEGATE_OneParam(FOnSelectionChanged, const TArray<UStateTreeState*>& /*SelectedStates*/);
+	DECLARE_MULTICAST_DELEGATE_OneParam(FOnSelectionChanged, const TArray<TWeakObjectPtr<UStateTreeState>>& /*SelectedStates*/);
 
 	FStateTreeViewModel();
 	virtual ~FStateTreeViewModel() override;
@@ -45,19 +45,21 @@ public:
 	// Selection handling.
 	void ClearSelection();
 	void SetSelection(UStateTreeState* Selected);
-	void SetSelection(const TArray<UStateTreeState*>& InSelection);
+	void SetSelection(const TArray<TWeakObjectPtr<UStateTreeState>>& InSelection);
 	bool IsSelected(const UStateTreeState* State) const;
 	bool IsChildOfSelection(const UStateTreeState* State) const;
 	void GetSelectedStates(TArray<UStateTreeState*>& OutSelectedStates);
+	void GetSelectedStates(TArray<TWeakObjectPtr<UStateTreeState>>& OutSelectedStates);
 	bool HasSelection() const;
 
 	// Returns array of subtrees to edit.
 	TArray<UStateTreeState*>* GetSubTrees() const;
 	int32 GetSubTreeCount() const;
+	void GetSubTrees(TArray<TWeakObjectPtr<UStateTreeState>>& OutSubtrees) const;
 
 	// Gets and sets StateTree view expansion state store in the asset.
-	void GetPersistentExpandedStates(TSet<UStateTreeState*>& OutExpandedStates);
-	void SetPersistentExpandedStates(TSet<UStateTreeState*>& InExpandedStates);
+	void SetPersistentExpandedStates(TSet<TWeakObjectPtr<UStateTreeState>>& InExpandedStates);
+	void GetPersistentExpandedStates(TSet<TWeakObjectPtr<UStateTreeState>>& OutExpandedStates);
 
 	// State manipulation commands
 	void AddState(UStateTreeState* AfterState);
@@ -91,13 +93,13 @@ public:
 	FOnSelectionChanged& GetOnSelectionChanged() { return OnSelectionChanged; }
 
 protected:
-	void GetExpandedStatesRecursive(UStateTreeState* State, TSet<UStateTreeState*>& ExpandedStates);
+	void GetExpandedStatesRecursive(UStateTreeState* State, TSet<TWeakObjectPtr<UStateTreeState>>& ExpandedStates);
 	void MoveSelectedStates(UStateTreeState* TargetState, const FStateTreeViewModelInsert RelativeLocation);
 
 	void HandleIdentifierChanged(const UStateTree& StateTree) const;
 
-	UStateTreeEditorData* TreeData;	// TODO: is this safe? should we used TWeakObjectPtr here instead?
-	TSet<FWeakObjectPtr> SelectedStates;
+	TWeakObjectPtr<UStateTreeEditorData> TreeDataWeak;
+	TSet<TWeakObjectPtr<UStateTreeState>> SelectedStates;
 	FOnAssetChanged OnAssetChanged;
 	FOnStatesChanged OnStatesChanged;
 	FOnStateAdded OnStateAdded;
diff --git a/Engine/Plugins/Runtime/StateTree/Source/StateTreeEditorModule/Public/StateTreeState.h b/Engine/Plugins/Runtime/StateTree/Source/StateTreeEditorModule/Public/StateTreeState.h
--- a/Engine/Plugins/Runtime/StateTree/Source/StateTreeEditorModule/Public/StateTreeState.h
+++ b/Engine/Plugins/Runtime/StateTree/Source/StateTreeEditorModule/Public/StateTreeState.h
@@ -116,7 +116,7 @@ public:
 #if WITH_EDITOR
 	virtual void PostEditChangeChainProperty(FPropertyChangedChainEvent& PropertyChangedEvent) override;
 	virtual void PostLoad() override;;
-	void UpdateParametersFromLinkedState();
+	void UpdateParametersFromLinkedSubtree();
 #endif
 
 	UStateTreeState* GetNextSiblingState() const;
@@ -195,7 +195,7 @@ public:
 	EStateTreeStateType Type = EStateTreeStateType::State;
 
 	UPROPERTY(EditDefaultsOnly, Category = "State", Meta=(DirectStatesOnly, SubtreesOnly))
-	FStateTreeStateLink LinkedState;
+	FStateTreeStateLink LinkedSubtree;
 
 	UPROPERTY(EditDefaultsOnly, Category = "State")
 	FStateTreeStateParameters Parameters;
diff --git a/Engine/Plugins/Runtime/StateTree/Source/StateTreeEditorModule/StateTreeEditorModule.Build.cs b/Engine/Plugins/Runtime/StateTree/Source/StateTreeEditorModule/StateTreeEditorModule.Build.cs
--- a/Engine/Plugins/Runtime/StateTree/Source/StateTreeEditorModule/StateTreeEditorModule.Build.cs
+++ b/Engine/Plugins/Runtime/StateTree/Source/StateTreeEditorModule/StateTreeEditorModule.Build.cs
@@ -31,7 +31,7 @@ namespace UnrealBuildTool.Rules
 				"PropertyAccessEditor",
 				"StructUtils",
 				"StructUtilsEditor",
-				"GameplayTags"
+				"GameplayTags",
 			}
 			);
 
@@ -43,6 +43,7 @@ namespace UnrealBuildTool.Rules
 				"PropertyPath",
 				"PropertyEditor",
 				"ToolMenus",
+				"ToolWidgets",
 			}
 			);
 		}
diff --git a/Engine/Plugins/Runtime/StateTree/Source/StateTreeTestSuite/Private/StateTreeTest.cpp b/Engine/Plugins/Runtime/StateTree/Source/StateTreeTestSuite/Private/StateTreeTest.cpp
--- a/Engine/Plugins/Runtime/StateTree/Source/StateTreeTestSuite/Private/StateTreeTest.cpp
+++ b/Engine/Plugins/Runtime/StateTree/Source/StateTreeTestSuite/Private/StateTreeTest.cpp
@@ -291,7 +291,7 @@ struct FStateTreeTest_SubTree : FAITestBase
 		UStateTreeState& State3A = State3.AddChildState(FName(TEXT("State3A")));
 		UStateTreeState& State3B = State3.AddChildState(FName(TEXT("State3B")));
 
-		State1.LinkedState.Set(&State3);
+		State1.LinkedSubtree.Set(&State3);
 
 		auto& Task2 = State2.AddTask<FTestTask_Stand>(FName(TEXT("Task2")));
 
