StateTree: Fix duplicate and copy/paste of node
- Generate new IDs for transition conditions nodes on duplicate
- Generate new IDs for nodes on paste
- Prevent pasting between incompatible node containers

diff --git a/Engine/Plugins/Runtime/StateTree/Source/StateTreeEditorModule/Private/Customizations/StateTreeEditorNodeDetails.cpp b/Engine/Plugins/Runtime/StateTree/Source/StateTreeEditorModule/Private/Customizations/StateTreeEditorNodeDetails.cpp
--- a/Engine/Plugins/Runtime/StateTree/Source/StateTreeEditorModule/Private/Customizations/StateTreeEditorNodeDetails.cpp
+++ b/Engine/Plugins/Runtime/StateTree/Source/StateTreeEditorModule/Private/Customizations/StateTreeEditorNodeDetails.cpp
@@ -37,9 +37,29 @@
 #include "Styling/StyleColors.h"
 #include "ObjectEditorUtils.h"
 #include "StateTreeCompiler.h"
+#include "HAL/PlatformApplicationMisc.h"
 
 #define LOCTEXT_NAMESPACE "StateTreeEditor"
 
+/** Helper class to detect if there were issues when calling ImportText() */
+class FStateTreeDefaultValueImportErrorContext : public FOutputDevice
+{
+public:
+
+	int32 NumErrors;
+
+	FStateTreeDefaultValueImportErrorContext()
+		: FOutputDevice()
+		, NumErrors(0)
+	{
+	}
+
+	virtual void Serialize(const TCHAR* V, ELogVerbosity::Type Verbosity, const class FName& Category) override
+	{
+		NumErrors++;
+	}
+};
+
 namespace UE::StateTreeEditor::Internal
 {
 	/** @return text describing the pin type, matches SPinTypeSelector. */
@@ -451,7 +471,72 @@ void FStateTreeEditorNodeDetails::CustomizeHeader(TSharedRef<class IPropertyHand
 				StructPropertyHandle->CreateDefaultPropertyButtonWidgets()
 			]
 		]
-		.OverrideResetToDefault(ResetOverride);
+		.OverrideResetToDefault(ResetOverride)
+		.CopyAction(FUIAction(FExecuteAction::CreateSP(this, &FStateTreeEditorNodeDetails::OnCopyNode)))
+		.PasteAction(FUIAction(FExecuteAction::CreateSP(this, &FStateTreeEditorNodeDetails::OnPasteNode)));
+
+}
+
+void FStateTreeEditorNodeDetails::OnCopyNode()
+{
+	FString Value;
+	if (StructProperty->GetValueAsFormattedString(Value) == FPropertyAccess::Success)
+	{
+		FPlatformApplicationMisc::ClipboardCopy(*Value);
+	}
+}
+
+void FStateTreeEditorNodeDetails::OnPasteNode()
+{
+	FString PastedText;
+	FPlatformApplicationMisc::ClipboardPaste(PastedText);
+
+	if (!PastedText.IsEmpty())
+	{
+		// create node from the 
+		FStateTreeEditorNode TempNode;
+		UScriptStruct* NodeScriptStruct = TBaseStructure<FStateTreeEditorNode>::Get();
+
+		FStateTreeDefaultValueImportErrorContext ErrorPipe;
+		NodeScriptStruct->ImportText(*PastedText, &TempNode, nullptr, EPropertyPortFlags::PPF_None, &ErrorPipe, NodeScriptStruct->GetName());
+
+		if (ErrorPipe.NumErrors == 0)
+		{
+			// Do not allow to mix and match types.
+			// @todo: Check Schema too and warn user about mismatching schema.
+			if (TempNode.Node.GetScriptStruct() && TempNode.Node.GetScriptStruct()->IsChildOf(BaseScriptStruct))
+			{
+				FScopedTransaction Transaction(LOCTEXT("PasteNode", "Paste Node"));
+
+				StructProperty->NotifyPreChange();
+
+				StructProperty->SetValueFromFormattedString(PastedText);
+
+				// Reset GUIDs on paste
+				TArray<void*> RawNodeData;
+				StructProperty->AccessRawData(RawNodeData);
+				for (void* Data : RawNodeData)
+				{
+					if (FStateTreeEditorNode* EditorNode = static_cast<FStateTreeEditorNode*>(Data))
+					{
+						if (FStateTreeNodeBase* Node = EditorNode->Node.GetMutablePtr<FStateTreeNodeBase>())
+						{
+							Node->Name = FName(Node->Name.ToString() + TEXT(" Copy"));
+						}
+						EditorNode->ID = FGuid::NewGuid();
+					}
+				}
+
+				StructProperty->NotifyPostChange(EPropertyChangeType::ValueSet);
+				StructProperty->NotifyFinishedChangingProperties();
+
+				if (PropUtils)
+				{
+					PropUtils->ForceRefresh();
+				}
+			}
+		}
+	}
 }
 
 bool FStateTreeEditorNodeDetails::ShouldResetToDefault(TSharedPtr<IPropertyHandle> PropertyHandle) const
diff --git a/Engine/Plugins/Runtime/StateTree/Source/StateTreeEditorModule/Private/Customizations/StateTreeEditorNodeDetails.h b/Engine/Plugins/Runtime/StateTree/Source/StateTreeEditorModule/Private/Customizations/StateTreeEditorNodeDetails.h
--- a/Engine/Plugins/Runtime/StateTree/Source/StateTreeEditorModule/Private/Customizations/StateTreeEditorNodeDetails.h
+++ b/Engine/Plugins/Runtime/StateTree/Source/StateTreeEditorModule/Private/Customizations/StateTreeEditorNodeDetails.h
@@ -33,6 +33,8 @@ private:
 
 	bool ShouldResetToDefault(TSharedPtr<IPropertyHandle> PropertyHandle) const;
 	void ResetToDefault(TSharedPtr<IPropertyHandle> PropertyHandle);
+	void OnCopyNode();
+	void OnPasteNode();
 	TSharedPtr<IPropertyHandle> GetInstancedObjectValueHandle(TSharedPtr<IPropertyHandle> PropertyHandle);
 
 	FOptionalSize GetIndentSize() const;
diff --git a/Engine/Plugins/Runtime/StateTree/Source/StateTreeEditorModule/Private/StateTreeState.cpp b/Engine/Plugins/Runtime/StateTree/Source/StateTreeEditorModule/Private/StateTreeState.cpp
--- a/Engine/Plugins/Runtime/StateTree/Source/StateTreeEditorModule/Private/StateTreeState.cpp
+++ b/Engine/Plugins/Runtime/StateTree/Source/StateTreeEditorModule/Private/StateTreeState.cpp
@@ -66,6 +66,18 @@ void UStateTreeState::PostEditChangeChainProperty(FPropertyChangedChainEvent& Pr
 		MemberProperty = PropertyChangedEvent.PropertyChain.GetActiveMemberNode()->GetValue();
 	}
 
+	auto TryGetConditionIndex = [&PropertyChangedEvent]() -> int32
+	{
+		check (PropertyChangedEvent.PropertyChain.GetActiveMemberNode());
+		FEditPropertyChain::TDoubleLinkedListNode* ConditionPropertyNode = PropertyChangedEvent.PropertyChain.GetActiveMemberNode()->GetNextNode();
+		FProperty* ConditionsProperty = ConditionPropertyNode ? ConditionPropertyNode->GetValue() : nullptr;
+		if (ConditionsProperty && ConditionsProperty->GetFName() == GET_MEMBER_NAME_CHECKED(FStateTreeTransition, Conditions))
+		{
+			return PropertyChangedEvent.GetArrayIndex(ConditionsProperty->GetFName().ToString());
+		}
+		return INDEX_NONE;
+	};
+	
 	if (Property)
 	{
 		if (Property->GetOwnerClass() == UStateTreeState::StaticClass()
@@ -156,11 +168,31 @@ void UStateTreeState::PostEditChangeChainProperty(FPropertyChangedChainEvent& Pr
 					const int32 ArrayIndex = PropertyChangedEvent.GetArrayIndex(MemberProperty->GetFName().ToString());
 					if (EnterConditions.IsValidIndex(ArrayIndex))
 					{
+						if (FStateTreeConditionBase* Condition = EnterConditions[ArrayIndex].Node.GetMutablePtr<FStateTreeConditionBase>())
+						{
+							Condition->Name = FName(Condition->Name.ToString() + TEXT(" Duplicate"));
+						}
 						EnterConditions[ArrayIndex].ID = FGuid::NewGuid();
 					}
 				}
-				// TODO: Transition conditions.
-				
+				if (MemberProperty->GetFName() == GET_MEMBER_NAME_CHECKED(UStateTreeState, Transitions))
+				{
+					const int32 TransitionsIndex = PropertyChangedEvent.GetArrayIndex(MemberProperty->GetFName().ToString());
+					const int32 ConditionsIndex = TryGetConditionIndex();
+
+					if (Transitions.IsValidIndex(TransitionsIndex))
+					{
+						FStateTreeTransition& Transition = Transitions[TransitionsIndex];
+						if (Transition.Conditions.IsValidIndex(ConditionsIndex))
+						{
+							if (FStateTreeConditionBase* Condition = Transition.Conditions[ConditionsIndex].Node.GetMutablePtr<FStateTreeConditionBase>())
+							{
+								Condition->Name = FName(Condition->Name.ToString() + TEXT(" Duplicate"));
+							}
+							Transition.Conditions[ConditionsIndex].ID = FGuid::NewGuid();
+						}
+					}
+				}
 			}
 		}
 	}
diff --git a/Engine/Plugins/Runtime/StateTree/Source/StateTreeEditorModule/StateTreeEditorModule.Build.cs b/Engine/Plugins/Runtime/StateTree/Source/StateTreeEditorModule/StateTreeEditorModule.Build.cs
--- a/Engine/Plugins/Runtime/StateTree/Source/StateTreeEditorModule/StateTreeEditorModule.Build.cs
+++ b/Engine/Plugins/Runtime/StateTree/Source/StateTreeEditorModule/StateTreeEditorModule.Build.cs
@@ -44,6 +44,7 @@ namespace UnrealBuildTool.Rules
 				"PropertyEditor",
 				"ToolMenus",
 				"ToolWidgets",
+				"ApplicationCore",
 			}
 			);
 		}
