StateTree: Fix comment on execution context SendEvent()

diff --git a/Engine/Plugins/Runtime/StateTree/Source/StateTreeModule/Public/StateTreeExecutionContext.h b/Engine/Plugins/Runtime/StateTree/Source/StateTreeModule/Public/StateTreeExecutionContext.h
--- a/Engine/Plugins/Runtime/StateTree/Source/StateTreeModule/Public/StateTreeExecutionContext.h
+++ b/Engine/Plugins/Runtime/StateTree/Source/StateTreeModule/Public/StateTreeExecutionContext.h
@@ -137,7 +137,7 @@ public:
 	/** @return the names of all the active state. */
 	TArray<FName> GetActiveStateNames() const;
 
-	/** Sends event for the StateTree. Can only be used during StateTree tick. */
+	/** Sends event for the StateTree. */
 	void SendEvent(const FStateTreeEvent& Event);
 
 	/** Iterates over all events. Can only be used during StateTree tick. Expects a lambda which takes const FStateTreeEvent& Event, and returns EStateTreeLoopEvents. */
