StateTree: Fix unity build on StateTreeEditorNodeDetails

diff --git a/Engine/Plugins/Runtime/StateTree/Source/StateTreeEditorModule/Private/Customizations/StateTreeEditorNodeDetails.cpp b/Engine/Plugins/Runtime/StateTree/Source/StateTreeEditorModule/Private/Customizations/StateTreeEditorNodeDetails.cpp
--- a/Engine/Plugins/Runtime/StateTree/Source/StateTreeEditorModule/Private/Customizations/StateTreeEditorNodeDetails.cpp
+++ b/Engine/Plugins/Runtime/StateTree/Source/StateTreeEditorModule/Private/Customizations/StateTreeEditorNodeDetails.cpp
@@ -38,6 +38,7 @@
 #include "ObjectEditorUtils.h"
 #include "StateTreeCompiler.h"
 #include "HAL/PlatformApplicationMisc.h"
+#include "ScopedTransaction.h"
 
 #define LOCTEXT_NAMESPACE "StateTreeEditor"
 
