StateTree: Fix AnyEnum display and binding

diff --git a/Engine/Plugins/Runtime/StateTree/Source/StateTreeEditorModule/Private/Customizations/StateTreeAnyEnumDetails.cpp b/Engine/Plugins/Runtime/StateTree/Source/StateTreeEditorModule/Private/Customizations/StateTreeAnyEnumDetails.cpp
--- a/Engine/Plugins/Runtime/StateTree/Source/StateTreeEditorModule/Private/Customizations/StateTreeAnyEnumDetails.cpp
+++ b/Engine/Plugins/Runtime/StateTree/Source/StateTreeEditorModule/Private/Customizations/StateTreeAnyEnumDetails.cpp
@@ -86,7 +86,8 @@ TSharedRef<SWidget> FStateTreeAnyEnumDetails::OnGetComboContent() const
 	{
 		if (StateTreeEnum.Enum)
 		{
-			for (int32 i = 0; i < StateTreeEnum.Enum->NumEnums(); i++)
+			// This is the number of entry in the enum, - 1, because the last item in an enum is the _MAX item
+			for (int32 i = 0; i < StateTreeEnum.Enum->NumEnums() - 1; i++)
 			{
 				const int64 Value = StateTreeEnum.Enum->GetValueByIndex(i);
 				MenuBuilder.AddMenuEntry(StateTreeEnum.Enum->GetDisplayNameTextByIndex(i), TAttribute<FText>(), FSlateIcon(),
@@ -98,6 +99,7 @@ TSharedRef<SWidget> FStateTreeAnyEnumDetails::OnGetComboContent() const
 					}))
 				);
 			}
+			bSuccess = true;
 		}
 	}
 
diff --git a/Engine/Plugins/Runtime/StateTree/Source/StateTreeEditorModule/Private/Customizations/StateTreeBindingExtension.cpp b/Engine/Plugins/Runtime/StateTree/Source/StateTreeEditorModule/Private/Customizations/StateTreeBindingExtension.cpp
--- a/Engine/Plugins/Runtime/StateTree/Source/StateTreeEditorModule/Private/Customizations/StateTreeBindingExtension.cpp
+++ b/Engine/Plugins/Runtime/StateTree/Source/StateTreeEditorModule/Private/Customizations/StateTreeBindingExtension.cpp
@@ -22,7 +22,7 @@ namespace UE::StateTree::PropertyBinding
 {
 	
 const FName StateTreeNodeIDName(TEXT("StateTreeNodeID"));
-	
+const FName AllowAnyBindingName(TEXT("AllowAnyBinding"));
 
 UObject* FindEditorBindingsOwner(UObject* InObject)
 {
@@ -192,6 +192,9 @@ void FStateTreeBindingExtension::ExtendWidgetRow(FDetailWidgetRow& InWidgetRow,
 			// Note: AnyEnums will need special handling before they can be used for binding.
 			if (bIsAnyEnum)
 			{
+				// If the AnyEnum has AllowAnyBinding, allow to bind to any enum.
+				const bool bAllowAnyBinding = InPropertyHandle->HasMetaData(UE::StateTree::PropertyBinding::AllowAnyBindingName);
+
 				FStateTreeAnyEnum AnyEnum;
 				UE::StateTree::PropertyHelpers::GetStructValue(InPropertyHandle, AnyEnum);
 
@@ -200,12 +203,12 @@ void FStateTreeBindingExtension::ExtendWidgetRow(FDetailWidgetRow& InWidgetRow,
 				{
 					if (UEnum* Enum = ByteProperty->GetIntPropertyEnum())
 					{
-						bCanBind = !AnyEnum.Enum || AnyEnum.Enum == Enum;
+						bCanBind = bAllowAnyBinding || AnyEnum.Enum == Enum;
 					}
 				}
 				else if (FEnumProperty* EnumProperty = CastField<FEnumProperty>(InProperty))
 				{
-					bCanBind = !AnyEnum.Enum || AnyEnum.Enum == EnumProperty->GetEnum();
+					bCanBind = bAllowAnyBinding || AnyEnum.Enum == EnumProperty->GetEnum();
 				}
 			}
 			else if (bIsDataRef && DataRefBaseStruct != nullptr)
diff --git a/Engine/Plugins/Runtime/StateTree/Source/StateTreeEditorModule/Private/StateTreeEditorData.cpp b/Engine/Plugins/Runtime/StateTree/Source/StateTreeEditorModule/Private/StateTreeEditorData.cpp
--- a/Engine/Plugins/Runtime/StateTree/Source/StateTreeEditorModule/Private/StateTreeEditorData.cpp
+++ b/Engine/Plugins/Runtime/StateTree/Source/StateTreeEditorModule/Private/StateTreeEditorData.cpp
@@ -245,18 +245,59 @@ bool UStateTreeEditorData::GetStructByID(const FGuid StructID, FStateTreeBindabl
 {
 	bool bResult = false;
 
-	VisitHierarchyNodes([&bResult, &OutStructDesc, StructID](const UStateTreeState* State, const FGuid& ID, const FName& Name, const EStateTreeNodeType NodeType, const UScriptStruct* NodeStruct, const UStruct* InstanceStruct)
+	// All parameters
+	if (StructID == RootParameters.ID)
+	{
+		OutStructDesc.Struct = RootParameters.Parameters.GetPropertyBagStruct();
+		OutStructDesc.Name = FName(TEXT("Parameters"));
+		OutStructDesc.ID = RootParameters.ID;
+		bResult = true;
+	}
+
+	// All named external data items declared by the schema
+	if (!bResult && Schema != nullptr)
+	{
+		for (const FStateTreeExternalDataDesc& Desc : Schema->GetContextDataDescs())
+		{
+			OutStructDesc.Struct = Desc.Struct;
+			OutStructDesc.Name = Desc.Name;
+			OutStructDesc.ID = Desc.ID;
+			bResult = true;
+			break;
+		}	
+	}
+
+	// Evaluators
+	if (!bResult)
+	{
+		for (const FStateTreeEditorNode& Node : Evaluators)
 		{
-			if (ID == StructID)
+			if (const FStateTreeEvaluatorBase* Evaluator = Node.Node.GetPtr<FStateTreeEvaluatorBase>())
 			{
-				OutStructDesc.Struct = InstanceStruct;
-				OutStructDesc.Name = Name;
-				OutStructDesc.ID = ID;
+				OutStructDesc.Struct = Evaluator->GetInstanceDataType();
+				OutStructDesc.Name = Evaluator->Name;
+				OutStructDesc.ID = Node.ID;
 				bResult = true;
-				return EStateTreeVisitor::Break;
+				break;
 			}
-			return EStateTreeVisitor::Continue;
-		});
+		}
+	}
+
+	if (!bResult)
+	{
+		VisitHierarchyNodes([&bResult, &OutStructDesc, StructID](const UStateTreeState* State, const FGuid& ID, const FName& Name, const EStateTreeNodeType NodeType, const UScriptStruct* NodeStruct, const UStruct* InstanceStruct)
+			{
+				if (ID == StructID)
+				{
+					OutStructDesc.Struct = InstanceStruct;
+					OutStructDesc.Name = Name;
+					OutStructDesc.ID = ID;
+					bResult = true;
+					return EStateTreeVisitor::Break;
+				}
+				return EStateTreeVisitor::Continue;
+			});
+	}
 
 	return bResult;
 }
@@ -315,6 +356,15 @@ void UStateTreeEditorData::GetAllStructIDs(TMap<FGuid, const UStruct*>& AllStruc
 		}	
 	}
 
+	// Evaluators
+	for (const FStateTreeEditorNode& Node : Evaluators)
+	{
+		if (const FStateTreeEvaluatorBase* Evaluator = Node.Node.GetPtr<FStateTreeEvaluatorBase>())
+		{
+			AllStructs.Emplace(Node.ID, Evaluator->GetInstanceDataType());
+		}
+	}
+
 	VisitHierarchyNodes([&AllStructs](const UStateTreeState* State, const FGuid& ID, const FName& Name, const EStateTreeNodeType NodeType, const UScriptStruct* NodeStruct, const UStruct* InstanceStruct)
 		{
 			AllStructs.Add(ID, InstanceStruct);
@@ -445,27 +495,10 @@ void UStateTreeEditorData::VisitHierarchy(TFunctionRef<EStateTreeVisitor(UStateT
 
 void UStateTreeEditorData::VisitHierarchyNodes(TFunctionRef<EStateTreeVisitor(const UStateTreeState* State, const FGuid& ID, const FName& Name, const EStateTreeNodeType NodeType, const UScriptStruct* NodeStruct, const UStruct* InstanceStruct)> InFunc) const
 {
-	// Evaluators
-	bool bContinue = true;
-	for (const FStateTreeEditorNode& Node : Evaluators)
-	{
-		if (const FStateTreeEvaluatorBase* Evaluator = Node.Node.GetPtr<FStateTreeEvaluatorBase>())
-		{
-			if (InFunc(nullptr, Node.ID, Evaluator->Name, EStateTreeNodeType::Task, Node.Node.GetScriptStruct(), Evaluator->GetInstanceDataType()) == EStateTreeVisitor::Break)
-			{
-				bContinue = false;
-				break;
-			}
-		}
-	}
-
-	if (bContinue)
+	VisitHierarchy([this, &InFunc](const UStateTreeState& State, UStateTreeState* /*ParentState*/)
 	{
-		VisitHierarchy([this, &InFunc](const UStateTreeState& State, UStateTreeState* /*ParentState*/)
-		{
-			return VisitStateNodes(State, InFunc);
-		});
-	}
+		return VisitStateNodes(State, InFunc);
+	});
 }
 
 
diff --git a/Engine/Plugins/Runtime/StateTree/Source/StateTreeModule/Public/Conditions/StateTreeCommonConditions.h b/Engine/Plugins/Runtime/StateTree/Source/StateTreeModule/Public/Conditions/StateTreeCommonConditions.h
--- a/Engine/Plugins/Runtime/StateTree/Source/StateTreeModule/Public/Conditions/StateTreeCommonConditions.h
+++ b/Engine/Plugins/Runtime/StateTree/Source/StateTreeModule/Public/Conditions/StateTreeCommonConditions.h
@@ -133,7 +133,7 @@ struct STATETREEMODULE_API FStateTreeCompareEnumConditionInstanceData
 {
 	GENERATED_BODY()
 
-	UPROPERTY(EditAnywhere, Category = Input)
+	UPROPERTY(EditAnywhere, Category = Input, meta=(AllowAnyBinding))
 	FStateTreeAnyEnum Left;
 
 	UPROPERTY(EditAnywhere, Category = Parameter)
