StateTree: Fix compilation of any enum bindings.

diff --git a/Engine/Plugins/Runtime/StateTree/Source/StateTreeEditorModule/Private/StateTreeCompiler.cpp b/Engine/Plugins/Runtime/StateTree/Source/StateTreeEditorModule/Private/StateTreeCompiler.cpp
--- a/Engine/Plugins/Runtime/StateTree/Source/StateTreeEditorModule/Private/StateTreeCompiler.cpp
+++ b/Engine/Plugins/Runtime/StateTree/Source/StateTreeEditorModule/Private/StateTreeCompiler.cpp
@@ -1127,8 +1127,9 @@ bool FStateTreeCompiler::GetAndValidateBindings(const FStateTreeBindableStructDe
 	{
 		return OutBindings.ContainsByPredicate([&PropertyName](const FStateTreeEditorPropertyBinding& Binding)
 			{
-				// We're looping over just the first level of properties on the struct, so we assume that the path is just one item.
-				return Binding.TargetPath.Path.Num() == 1 && Binding.TargetPath.Path[0] == PropertyName;
+				// We're looping over just the first level of properties on the struct, so we assume that the path is just one item
+				// (or two in case of AnyEnum, because we expand the path to Property.Value, see code above).
+				return Binding.TargetPath.Path.Num() >= 1 && Binding.TargetPath.Path[0] == PropertyName;
 			});
 	};
 
