StateTree: Improved event handling
- Moved event queue to it's own struct
- Changed the SendEvent() API to accept struct view to prevent instanced struct copy in common cases
- Fixed event handling in case EnterState() fails.
- Added option for Tasks to disavble ticking, or to be ticked only when there are events
- Added option for Tasks handle when properties are copied
- Changed DebugText to use external Actor reference
- Changed DelayTask to use the new GetInstanceData which does not need type
- Added TStateTreeInstanceDataStructRef which can be used to access struct instance data in delegates

diff --git a/Engine/Plugins/Runtime/GameplayStateTree/Source/GameplayStateTreeModule/Private/Components/StateTreeComponent.cpp b/Engine/Plugins/Runtime/GameplayStateTree/Source/GameplayStateTreeModule/Private/Components/StateTreeComponent.cpp
--- a/Engine/Plugins/Runtime/GameplayStateTree/Source/GameplayStateTreeModule/Private/Components/StateTreeComponent.cpp
+++ b/Engine/Plugins/Runtime/GameplayStateTree/Source/GameplayStateTreeModule/Private/Components/StateTreeComponent.cpp
@@ -323,6 +323,11 @@ void UStateTreeComponent::OnGameplayTaskInitialized(UGameplayTask& Task)
 }
 
 void UStateTreeComponent::SendStateTreeEvent(const FStateTreeEvent& Event)
+{
+	SendStateTreeEvent(Event.Tag, Event.Payload, Event.Origin);
+}
+
+void UStateTreeComponent::SendStateTreeEvent(const FGameplayTag Tag, const FConstStructView Payload, const FName Origin)
 {
 	if (!bIsRunning)
 	{
@@ -333,7 +338,7 @@ void UStateTreeComponent::SendStateTreeEvent(const FStateTreeEvent& Event)
 	FStateTreeExecutionContext Context(*GetOwner(), *StateTreeRef.GetStateTree(), InstanceData);
 	if (Context.IsValid())
 	{
-		Context.SendEvent(Event);
+		Context.SendEvent(Tag, Payload, Origin);
 	}
 }
 
diff --git a/Engine/Plugins/Runtime/GameplayStateTree/Source/GameplayStateTreeModule/Public/Components/StateTreeComponent.h b/Engine/Plugins/Runtime/GameplayStateTree/Source/GameplayStateTreeModule/Public/Components/StateTreeComponent.h
--- a/Engine/Plugins/Runtime/GameplayStateTree/Source/GameplayStateTreeModule/Public/Components/StateTreeComponent.h
+++ b/Engine/Plugins/Runtime/GameplayStateTree/Source/GameplayStateTreeModule/Public/Components/StateTreeComponent.h
@@ -48,7 +48,10 @@ public:
 	/** Sends event to the running StateTree. */
 	UFUNCTION(BlueprintCallable, Category = "Gameplay|StateTree")
 	void SendStateTreeEvent(const FStateTreeEvent& Event);
-	
+
+	/** Sends event to the running StateTree. */
+	void SendStateTreeEvent(const FGameplayTag Tag, const FConstStructView Payload = FConstStructView(), const FName Origin = FName());
+
 #if WITH_GAMEPLAY_DEBUGGER
 	virtual FString GetDebugInfoString() const override;
 #endif // WITH_GAMEPLAY_DEBUGGER
diff --git a/Engine/Plugins/Runtime/StateTree/Source/StateTreeModule/Private/Blueprint/StateTreeNodeBlueprintBase.cpp b/Engine/Plugins/Runtime/StateTree/Source/StateTreeModule/Private/Blueprint/StateTreeNodeBlueprintBase.cpp
--- a/Engine/Plugins/Runtime/StateTree/Source/StateTreeModule/Private/Blueprint/StateTreeNodeBlueprintBase.cpp
+++ b/Engine/Plugins/Runtime/StateTree/Source/StateTreeModule/Private/Blueprint/StateTreeNodeBlueprintBase.cpp
@@ -40,6 +40,6 @@ void UStateTreeNodeBlueprintBase::SendEvent(const FStateTreeEvent& Event)
 		UE_VLOG_UELOG(this, LogStateTree, Error, TEXT("Trying to call SendEvent() outside StateTree tick. Use SendEvent() on UStateTreeComponent instead for sending signals externally."));
 		return;
 	}
-	CurrentContext->SendEvent(Event);
+	CurrentContext->SendEvent(Event.Tag, Event.Payload, Event.Origin);
 }
 
diff --git a/Engine/Plugins/Runtime/StateTree/Source/StateTreeModule/Private/Blueprint/StateTreeTaskBlueprintBase.cpp b/Engine/Plugins/Runtime/StateTree/Source/StateTreeModule/Private/Blueprint/StateTreeTaskBlueprintBase.cpp
--- a/Engine/Plugins/Runtime/StateTree/Source/StateTreeModule/Private/Blueprint/StateTreeTaskBlueprintBase.cpp
+++ b/Engine/Plugins/Runtime/StateTree/Source/StateTreeModule/Private/Blueprint/StateTreeTaskBlueprintBase.cpp
@@ -13,6 +13,8 @@
 
 UStateTreeTaskBlueprintBase::UStateTreeTaskBlueprintBase(const FObjectInitializer& ObjectInitializer)
 	: Super(ObjectInitializer)
+	, bShouldStateChangeOnReselect(true)
+	, bShouldCallTickOnlyOnEvents(false)
 {
 	bHasEnterState = BlueprintNodeHelpers::HasBlueprintFunction(TEXT("ReceiveEnterState"), *this, *StaticClass());
 	bHasExitState = BlueprintNodeHelpers::HasBlueprintFunction(TEXT("ReceiveExitState"), *this, *StaticClass());
@@ -106,6 +108,16 @@ EStateTreeRunStatus FStateTreeBlueprintTaskWrapper::Tick(FStateTreeExecutionCont
 {
 	UStateTreeTaskBlueprintBase* Instance = Context.GetInstanceDataPtr<UStateTreeTaskBlueprintBase>(*this);
 	check(Instance);
-	return Instance->Tick(Context, DeltaTime);
+
+	// @todo: reflect this value into the node.
+	const bool bNeedsTick = !Instance->bShouldCallTickOnlyOnEvents
+		|| (Instance->bShouldCallTickOnlyOnEvents && !Context.GetEventsToProcess().IsEmpty());
+	
+	if (bNeedsTick)
+	{
+		return Instance->Tick(Context, DeltaTime);
+	}
+	
+	return EStateTreeRunStatus::Running;
 }
 
diff --git a/Engine/Plugins/Runtime/StateTree/Source/StateTreeModule/Private/StateTreeEvents.cpp b/Engine/Plugins/Runtime/StateTree/Source/StateTreeModule/Private/StateTreeEvents.cpp
new file mode 100644
--- /dev/null
+++ b/Engine/Plugins/Runtime/StateTree/Source/StateTreeModule/Private/StateTreeEvents.cpp
@@ -0,0 +1,16 @@
+﻿// Copyright Epic Games, Inc. All Rights Reserved.
+
+#include "StateTreeEvents.h"
+#include "StateTreeTypes.h"
+#include "VisualLogger/VisualLogger.h"
+
+void FStateTreeEventQueue::SendEvent(const UObject* Owner, const FGameplayTag& Tag, const FConstStructView Payload, const FName Origin)
+{
+	if (Events.Num() >= MaxActiveEvents)
+	{
+		UE_VLOG_UELOG(Owner, LogStateTree, Error, TEXT("%s: Too many events send on '%s'. Dropping event %s"), ANSI_TO_TCHAR(__FUNCTION__), *GetNameSafe(Owner), *Tag.ToString());
+		return;
+	}
+
+	Events.Emplace(Tag, Payload, Origin);
+}
\ No newline at end of file
diff --git a/Engine/Plugins/Runtime/StateTree/Source/StateTreeModule/Private/StateTreeExecutionContext.cpp b/Engine/Plugins/Runtime/StateTree/Source/StateTreeModule/Private/StateTreeExecutionContext.cpp
--- a/Engine/Plugins/Runtime/StateTree/Source/StateTreeModule/Private/StateTreeExecutionContext.cpp
+++ b/Engine/Plugins/Runtime/StateTree/Source/StateTreeModule/Private/StateTreeExecutionContext.cpp
@@ -17,7 +17,6 @@
 
 namespace UE::StateTree
 {
-	constexpr int32 MaxActiveEvents = 16;	// Maximum number of events that can be accumulated.
 	constexpr int32 DebugIndentSize = 2;	// Debug printing indent for hierarchical data.
 }; // UE::StateTree
 
@@ -261,8 +260,9 @@ EStateTreeRunStatus FStateTreeExecutionContext::Stop()
 	FStateTreeExecutionState& Exec = GetExecState();
 
 	// Capture events added between ticks.
-	EventsToProcess = InstanceData.GetEvents();
-	InstanceData.GetEvents().Reset();
+	FStateTreeEventQueue& EventQueue = InstanceData.GetEventQueue();
+	EventsToProcess = EventQueue.GetEvents();
+	EventQueue.Reset();
 
 	TickEvaluators(0.0f);
 
@@ -322,12 +322,14 @@ EStateTreeRunStatus FStateTreeExecutionContext::Tick(const float DeltaTime)
 
 	const TSharedPtr<FStateTreeInstanceData> SharedInstanceData = StateTree.GetSharedInstanceData();
 	check(SharedInstanceData.IsValid());
-
+	
+	FStateTreeEventQueue& EventQueue = InstanceData.GetEventQueue();
+	
 	FStateTreeExecutionState* Exec = &GetExecState();
 
 	// Capture events added between ticks.
-	EventsToProcess = InstanceData.GetEvents();
-	InstanceData.GetEvents().Reset();
+	EventsToProcess = EventQueue.GetEvents();
+	EventQueue.Reset();
 	
 	// No ticking of the tree is done or stopped.
 	if (Exec->TreeRunStatus != EStateTreeRunStatus::Running)
@@ -356,22 +358,22 @@ EStateTreeRunStatus FStateTreeExecutionContext::Tick(const float DeltaTime)
 		}
 	}
 
-	// Append events accumulated during the tick, so that transitions can immediately act on them.
-	// We'll consume the events only if they lead to state change below (EnterState is treated the same as Tick),
-	// or let them be processed next frame if no transition.
-	EventsToProcess.Append(InstanceData.GetEvents());
-
 	// The state selection is repeated up to MaxIteration time. This allows failed EnterState() to potentially find a new state immediately.
 	// This helps event driven StateTrees to not require another event/tick to find a suitable state.
 	static constexpr int32 MaxIterations = 5;
 	for (int32 Iter = 0; Iter < MaxIterations; Iter++)
 	{
+		// Append events accumulated during the tick, so that transitions can immediately act on them.
+		// We'll consume the events only if they lead to state change below (EnterState is treated the same as Tick),
+		// or let them be processed next frame if no transition.
+		EventsToProcess.Append(EventQueue.GetEvents());
+
 		// Trigger conditional transitions or state succeed/failed transitions. First tick transition is handled here too.
 		FStateTreeTransitionResult Transition;
 		if (TriggerTransitions(*SharedInstanceData.Get(), Transition))
 		{
 			// We have committed to state change, consume events that were accumulated during the tick above.
-			InstanceData.GetEvents().Reset();
+			EventQueue.Reset();
 
 			ExitState(Transition);
 
@@ -384,12 +386,15 @@ EStateTreeRunStatus FStateTreeExecutionContext::Tick(const float DeltaTime)
 			}
 
 			// Append and consume the events accumulated during the state exit.
-			EventsToProcess.Append(InstanceData.GetEvents());
-			InstanceData.GetEvents().Reset();
+			EventsToProcess.Append(EventQueue.GetEvents());
+			EventQueue.Reset();
 
 			// Enter state tasks can fail/succeed, treat it same as tick.
 			const EStateTreeRunStatus LastTickStatus = EnterState(Transition);
 
+			// Consider events so far processed. Events sent during EnterState went into EventQueue, and are processed in next iteration.
+			EventsToProcess.Reset();
+
 			// Need to reacquire the exec state as EnterState may alter the allocation. 
 			Exec = &GetExecState();
 			Exec->LastTickStatus = LastTickStatus;
@@ -422,16 +427,14 @@ EStateTreeRunStatus FStateTreeExecutionContext::Tick(const float DeltaTime)
 	return Exec->TreeRunStatus;
 }
 
-void FStateTreeExecutionContext::SendEvent(const FStateTreeEvent& Event)
+void FStateTreeExecutionContext::SendEvent(const FStateTreeEvent& Event) const
 {
-	CSV_SCOPED_TIMING_STAT_EXCLUSIVE(StateTree_SendEvent);
+	SendEvent(Event.Tag, Event.Payload, Event.Origin);
+}
 
-	if (!Event.Tag.IsValid())
-	{
-		STATETREE_LOG(Warning, TEXT("%s: Event tag cannot be empty, '%s' using StateTree '%s'."),
-			ANSI_TO_TCHAR(__FUNCTION__), *GetNameSafe(&Owner), *GetFullNameSafe(&StateTree));
-		return;
-	}
+void FStateTreeExecutionContext::SendEvent(const FGameplayTag Tag, const FConstStructView Payload, const FName Origin) const
+{
+	CSV_SCOPED_TIMING_STAT_EXCLUSIVE(StateTree_SendEvent);
 
 	if (!IsValid())
 	{
@@ -447,14 +450,8 @@ void FStateTreeExecutionContext::SendEvent(const FStateTreeEvent& Event)
 		return;
 	}
 
-	if (InstanceData.GetEvents().Num() >= UE::StateTree::MaxActiveEvents)
-	{
-		STATETREE_LOG(Warning, TEXT("%s: Too many events send on '%s' using StateTree '%s'."),
-			ANSI_TO_TCHAR(__FUNCTION__), *GetNameSafe(&Owner), *GetFullNameSafe(&StateTree));
-		return;
-	}
-
-	InstanceData.GetEvents().Add(Event);
+	FStateTreeEventQueue& EventQueue = InstanceData.GetEventQueue();
+	EventQueue.SendEvent(&Owner, Tag, Payload, Origin);
 }
 
 
@@ -732,7 +729,7 @@ void FStateTreeExecutionContext::ExitState(const FStateTreeTransitionResult& Tra
 			SetNodeDataView(Task, InstanceStructIndex, InstanceObjectIndex);
 
 			// Copy bound properties.
-			if (Task.BindingsBatch.IsValid())
+			if (Task.BindingsBatch.IsValid() && Task.bShouldCopyBoundPropertiesOnExitState)
 			{
 				StateTree.PropertyBindings.CopyTo(DataViews, Task.BindingsBatch, DataViews[Task.DataViewIndex.Get()]);
 			}
@@ -923,6 +920,8 @@ EStateTreeRunStatus FStateTreeExecutionContext::TickTasks(const float DeltaTime)
 	EStateTreeRunStatus Result = EStateTreeRunStatus::Running;
 	int32 NumTotalTasks = 0;
 
+	const bool bHasEvents = !EventsToProcess.IsEmpty();
+	
 	check(Exec.FirstTaskStructIndex.IsValid() && Exec.FirstTaskObjectIndex.IsValid()); 
 	int32 InstanceStructIndex = Exec.FirstTaskStructIndex.Get();
 	int32 InstanceObjectIndex = Exec.FirstTaskObjectIndex.Get();
@@ -954,13 +953,15 @@ EStateTreeRunStatus FStateTreeExecutionContext::TickTasks(const float DeltaTime)
 			SetNodeDataView(Task, InstanceStructIndex, InstanceObjectIndex);
 
 			// Copy bound properties.
-			if (Task.BindingsBatch.IsValid())
+			if (Task.BindingsBatch.IsValid() && Task.bShouldCopyBoundPropertiesOnTick)
 			{
 				StateTree.PropertyBindings.CopyTo(DataViews, Task.BindingsBatch, DataViews[Task.DataViewIndex.Get()]);
 			}
 			STATETREE_LOG(VeryVerbose, TEXT("%*s  Tick: '%s'"), Index*UE::StateTree::DebugIndentSize, TEXT(""), *Task.Name.ToString());
 
-			if (bShouldTickTasks)
+			const bool bNeedTick = Task.bShouldCallTick || (bHasEvents && Task.bShouldCallTickOnlyOnEvents);
+			
+			if (bShouldTickTasks && bNeedTick)
 			{
 				QUICK_SCOPE_CYCLE_COUNTER(StateTree_Task_Tick);
 				CSV_SCOPED_TIMING_STAT_EXCLUSIVE(StateTree_Task_Tick);
@@ -1069,6 +1070,20 @@ bool FStateTreeExecutionContext::TestAllConditions(FStateTreeInstanceData& Share
 	return Values[0];
 }
 
+FString FStateTreeExecutionContext::DebugGetEventsAsString() const
+{
+	FString Result;
+	for (const FStateTreeEvent& Event : EventsToProcess)
+	{
+		if (!Result.IsEmpty())
+		{
+			Result += TEXT(", ");
+		}
+		Result += Event.Tag.ToString();
+	}
+	return Result;
+}
+
 bool FStateTreeExecutionContext::TriggerTransitions(FStateTreeInstanceData& SharedInstanceData, FStateTreeTransitionResult& OutTransition)
 {
 	CSV_SCOPED_TIMING_STAT_EXCLUSIVE(StateTree_TriggerTransition);
@@ -1085,6 +1100,8 @@ bool FStateTreeExecutionContext::TriggerTransitions(FStateTreeInstanceData& Shar
 		CompletionTrigger = EStateTreeTransitionTrigger::OnStateFailed;
 	}
 
+	STATETREE_CLOG(!EventsToProcess.IsEmpty(), Verbose, TEXT("Trigger transitions with events [%s]"), *DebugGetEventsAsString());
+	
 	auto HasEvent = [this](const FGameplayTag QueriedTag)
 	{
 		if (EventsToProcess.IsEmpty())
diff --git a/Engine/Plugins/Runtime/StateTree/Source/StateTreeModule/Private/StateTreeInstanceData.cpp b/Engine/Plugins/Runtime/StateTree/Source/StateTreeModule/Private/StateTreeInstanceData.cpp
--- a/Engine/Plugins/Runtime/StateTree/Source/StateTreeModule/Private/StateTreeInstanceData.cpp
+++ b/Engine/Plugins/Runtime/StateTree/Source/StateTreeModule/Private/StateTreeInstanceData.cpp
@@ -9,6 +9,11 @@
 //  FStateTreeInstanceData
 //----------------------------------------------------------------//
 
+TArray<FStateTreeEvent>& FStateTreeInstanceData::GetEvents()
+{
+	return EventQueue.GetEventsArray();
+}
+
 int32 FStateTreeInstanceData::GetEstimatedMemoryUsage() const
 {
 	int32 Size = sizeof(FStateTreeInstanceData);
@@ -181,5 +186,6 @@ void FStateTreeInstanceData::Reset()
 {
 	InstanceStructs.Reset();
 	InstanceObjects.Reset();
+	EventQueue.Reset();
 }
 
diff --git a/Engine/Plugins/Runtime/StateTree/Source/StateTreeModule/Private/Tasks/StateTreeDebugTextTask.cpp b/Engine/Plugins/Runtime/StateTree/Source/StateTreeModule/Private/Tasks/StateTreeDebugTextTask.cpp
--- a/Engine/Plugins/Runtime/StateTree/Source/StateTreeModule/Private/Tasks/StateTreeDebugTextTask.cpp
+++ b/Engine/Plugins/Runtime/StateTree/Source/StateTreeModule/Private/Tasks/StateTreeDebugTextTask.cpp
@@ -8,29 +8,19 @@
 
 #include UE_INLINE_GENERATED_CPP_BY_NAME(StateTreeDebugTextTask)
 
-bool FStateTreeDebugTextTask::Link(FStateTreeLinker& Linker)
-{
-	Linker.LinkExternalData(ReferenceActorHandle);
-	return true;
-}
-
 EStateTreeRunStatus FStateTreeDebugTextTask::EnterState(FStateTreeExecutionContext& Context, const FStateTreeTransitionResult& Transition) const
 {
 	if (!bEnabled)
 	{
 		return EStateTreeRunStatus::Succeeded;
 	}
+
+	const FInstanceDataType& InstanceData = Context.GetInstanceData(*this);
 	
 	const UWorld* World = Context.GetWorld();
-
-	AActor* ReferenceActor = Context.GetExternalDataPtr(ReferenceActorHandle);
-	if (World == nullptr && ReferenceActor != nullptr)
-	{
-		World = ReferenceActor->GetWorld();
-	}
-	else if (ReferenceActor == nullptr)
+	if (World == nullptr && InstanceData.ReferenceActor != nullptr)
 	{
-		ReferenceActor = Cast<AActor>(Context.GetOwner());
+		World = InstanceData.ReferenceActor->GetWorld();
 	}
 
 	// Reference actor is not required (offset will be used as a global world location)
@@ -40,7 +30,7 @@ EStateTreeRunStatus FStateTreeDebugTextTask::EnterState(FStateTreeExecutionConte
 		return EStateTreeRunStatus::Failed;
 	}
 
-	DrawDebugString(World, Offset, Text, ReferenceActor, TextColor,	/*Duration*/-1,	/*DrawShadows*/true, FontScale);
+	DrawDebugString(World, Offset, Text, InstanceData.ReferenceActor, TextColor,	/*Duration*/-1,	/*DrawShadows*/true, FontScale);
 	
 	return EStateTreeRunStatus::Running;
 }
@@ -52,18 +42,14 @@ void FStateTreeDebugTextTask::ExitState(FStateTreeExecutionContext& Context, con
 		return;
 	}
 
+	const FInstanceDataType& InstanceData = Context.GetInstanceData(*this);
+	
 	const UWorld* World = Context.GetWorld();
-
-	AActor* ReferenceActor = Context.GetExternalDataPtr(ReferenceActorHandle);
-	if (World == nullptr && ReferenceActor != nullptr)
+	if (World == nullptr && InstanceData.ReferenceActor != nullptr)
 	{
-		World = ReferenceActor->GetWorld();
-	}	
-	else if (ReferenceActor == nullptr)
-	{
-		ReferenceActor = Cast<AActor>(Context.GetOwner());
+		World = InstanceData.ReferenceActor->GetWorld();
 	}
-	
+
 	// Reference actor is not required (offset was used as a global world location)
 	// but a valid world is required.
 	if (World == nullptr)
@@ -72,6 +58,6 @@ void FStateTreeDebugTextTask::ExitState(FStateTreeExecutionContext& Context, con
 	}
 
 	// Drawing an empty text will remove the HUD DebugText entries associated to the target actor
-	DrawDebugString(World, Offset, "",	ReferenceActor);
+	DrawDebugString(World, Offset, "",	InstanceData.ReferenceActor);
 }
 
diff --git a/Engine/Plugins/Runtime/StateTree/Source/StateTreeModule/Private/Tasks/StateTreeDebugTextTask.h b/Engine/Plugins/Runtime/StateTree/Source/StateTreeModule/Private/Tasks/StateTreeDebugTextTask.h
--- a/Engine/Plugins/Runtime/StateTree/Source/StateTreeModule/Private/Tasks/StateTreeDebugTextTask.h
+++ b/Engine/Plugins/Runtime/StateTree/Source/StateTreeModule/Private/Tasks/StateTreeDebugTextTask.h
@@ -9,6 +9,10 @@ USTRUCT()
 struct STATETREEMODULE_API FStateTreeDebugTextTaskInstanceData
 {
 	GENERATED_BODY()
+
+	/** Optional actor where to draw the text at. */
+	UPROPERTY(EditAnywhere, Category = "Parameter")
+	TObjectPtr<AActor> ReferenceActor = nullptr;
 };
 
 /**
@@ -19,30 +23,27 @@ struct STATETREEMODULE_API FStateTreeDebugTextTask : public FStateTreeTaskCommon
 {
 	GENERATED_BODY()
 
-	typedef FStateTreeDebugTextTaskInstanceData InstanceDataType;
+	using FInstanceDataType = FStateTreeDebugTextTaskInstanceData;
 	
 	FStateTreeDebugTextTask() = default;
 
-	virtual const UStruct* GetInstanceDataType() const override { return InstanceDataType::StaticStruct(); }
+	virtual const UStruct* GetInstanceDataType() const override { return FInstanceDataType::StaticStruct(); }
 
-	virtual bool Link(FStateTreeLinker& Linker) override;
 	virtual EStateTreeRunStatus EnterState(FStateTreeExecutionContext& Context, const FStateTreeTransitionResult& Transition) const override;
 	virtual void ExitState(FStateTreeExecutionContext& Context, const FStateTreeTransitionResult& Transition) const override;
 
-	TStateTreeExternalDataHandle<AActor, EStateTreeExternalDataRequirement::Optional> ReferenceActorHandle;
-	
-	UPROPERTY(EditAnywhere, Category = Parameter)
+	UPROPERTY(EditAnywhere, Category = "Parameter")
 	FString Text;
 
-	UPROPERTY(EditAnywhere, Category = Parameter)
+	UPROPERTY(EditAnywhere, Category = "Parameter")
 	FColor TextColor = FColor::White;
 
-	UPROPERTY(EditAnywhere, Category = Parameter, meta=(ClampMin = 0, UIMin = 0))
+	UPROPERTY(EditAnywhere, Category = "Parameter", meta=(ClampMin = 0, UIMin = 0))
 	float FontScale = 1.0f;
 	
-	UPROPERTY(EditAnywhere, Category = Parameter)
+	UPROPERTY(EditAnywhere, Category = "Parameter")
 	FVector Offset = FVector::ZeroVector;
 
-	UPROPERTY(EditAnywhere, Category = Parameter)
+	UPROPERTY(EditAnywhere, Category = "Parameter")
 	bool bEnabled = true;
 };
diff --git a/Engine/Plugins/Runtime/StateTree/Source/StateTreeModule/Private/Tasks/StateTreeDelayTask.cpp b/Engine/Plugins/Runtime/StateTree/Source/StateTreeModule/Private/Tasks/StateTreeDelayTask.cpp
--- a/Engine/Plugins/Runtime/StateTree/Source/StateTreeModule/Private/Tasks/StateTreeDelayTask.cpp
+++ b/Engine/Plugins/Runtime/StateTree/Source/StateTreeModule/Private/Tasks/StateTreeDelayTask.cpp
@@ -8,7 +8,7 @@
 
 EStateTreeRunStatus FStateTreeDelayTask::EnterState(FStateTreeExecutionContext& Context, const FStateTreeTransitionResult& Transition) const
 {
-	InstanceDataType& InstanceData = Context.GetInstanceData<InstanceDataType>(*this);
+	FInstanceDataType& InstanceData = Context.GetInstanceData(*this);
 
 	if (!InstanceData.bRunForever)
 	{
@@ -21,7 +21,7 @@ EStateTreeRunStatus FStateTreeDelayTask::EnterState(FStateTreeExecutionContext&
 
 EStateTreeRunStatus FStateTreeDelayTask::Tick(FStateTreeExecutionContext& Context, const float DeltaTime) const
 {
-	InstanceDataType& InstanceData = Context.GetInstanceData<InstanceDataType>(*this);
+	FInstanceDataType& InstanceData = Context.GetInstanceData(*this);
 
 	if (!InstanceData.bRunForever)
 	{
diff --git a/Engine/Plugins/Runtime/StateTree/Source/StateTreeModule/Private/Tasks/StateTreeDelayTask.h b/Engine/Plugins/Runtime/StateTree/Source/StateTreeModule/Private/Tasks/StateTreeDelayTask.h
--- a/Engine/Plugins/Runtime/StateTree/Source/StateTreeModule/Private/Tasks/StateTreeDelayTask.h
+++ b/Engine/Plugins/Runtime/StateTree/Source/StateTreeModule/Private/Tasks/StateTreeDelayTask.h
@@ -34,11 +34,11 @@ struct STATETREEMODULE_API FStateTreeDelayTask : public FStateTreeTaskCommonBase
 {
 	GENERATED_BODY()
 
-	typedef FStateTreeDelayTaskInstanceData InstanceDataType;
+	using FInstanceDataType = FStateTreeDelayTaskInstanceData;
 	
 	FStateTreeDelayTask() = default;
 
-	virtual const UStruct* GetInstanceDataType() const override { return InstanceDataType::StaticStruct(); }
+	virtual const UStruct* GetInstanceDataType() const override { return FInstanceDataType::StaticStruct(); }
 
 	virtual EStateTreeRunStatus EnterState(FStateTreeExecutionContext& Context, const FStateTreeTransitionResult& Transition) const override;
 	virtual EStateTreeRunStatus Tick(FStateTreeExecutionContext& Context, const float DeltaTime) const override;
diff --git a/Engine/Plugins/Runtime/StateTree/Source/StateTreeModule/Public/Blueprint/StateTreeTaskBlueprintBase.h b/Engine/Plugins/Runtime/StateTree/Source/StateTreeModule/Public/Blueprint/StateTreeTaskBlueprintBase.h
--- a/Engine/Plugins/Runtime/StateTree/Source/StateTreeModule/Public/Blueprint/StateTreeTaskBlueprintBase.h
+++ b/Engine/Plugins/Runtime/StateTree/Source/StateTreeModule/Public/Blueprint/StateTreeTaskBlueprintBase.h
@@ -39,7 +39,9 @@ protected:
 	 * Note: The API has been deprecated. ChangeType is moved into FStateTreeTransitionResult.
 	 * You can configure the task to be only called on state changes (that is, never call sustained changes) by setting bShouldStateChangeOnReselect to true.
 	 */
+	UE_DEPRECATED(5.2, "Use EnterState without ChangeType instead.")
 	virtual EStateTreeRunStatus EnterState(FStateTreeExecutionContext& Context, const EStateTreeStateChangeType ChangeType, const FStateTreeTransitionResult& Transition) final { return EStateTreeRunStatus::Running; }
+	UE_DEPRECATED(5.2, "Use ExitState without ChangeType instead.")
 	virtual void ExitState(FStateTreeExecutionContext& Context, const EStateTreeStateChangeType ChangeType, const FStateTreeTransitionResult& Transition) final {};
 
 	virtual EStateTreeRunStatus EnterState(FStateTreeExecutionContext& Context, const FStateTreeTransitionResult& Transition);
@@ -53,7 +55,11 @@ protected:
 	 * Generally this should be true for action type tasks, like playing animation,
 	 * and false on state like tasks like claiming a resource that is expected to be acquired on child states. */
 	UPROPERTY(EditDefaultsOnly, Category="Default")
-	bool bShouldStateChangeOnReselect = true;
+	uint8 bShouldStateChangeOnReselect : 1;
+
+	/** If set to true, Tick() is called. Default false. */
+	UPROPERTY(EditDefaultsOnly, Category="Default")
+	uint8 bShouldCallTickOnlyOnEvents : 1;
 
 	uint8 bHasEnterState : 1;
 	uint8 bHasExitState : 1;
diff --git a/Engine/Plugins/Runtime/StateTree/Source/StateTreeModule/Public/StateTreeEvents.h b/Engine/Plugins/Runtime/StateTree/Source/StateTreeModule/Public/StateTreeEvents.h
--- a/Engine/Plugins/Runtime/StateTree/Source/StateTreeModule/Public/StateTreeEvents.h
+++ b/Engine/Plugins/Runtime/StateTree/Source/StateTreeModule/Public/StateTreeEvents.h
@@ -4,6 +4,7 @@
 
 #include "GameplayTagContainer.h"
 #include "InstancedStruct.h"
+#include "StructView.h"
 #include "StateTreeEvents.generated.h"
 
 /** Enum used for flow control during event iteration. */
@@ -25,8 +26,16 @@ struct STATETREEMODULE_API FStateTreeEvent
 	GENERATED_BODY()
 
 	FStateTreeEvent() = default;
-	FStateTreeEvent(const FGameplayTag& InTag)
+
+	explicit FStateTreeEvent(const FGameplayTag InTag)
+		: Tag(InTag)
+	{
+	}
+	
+	explicit FStateTreeEvent(const FGameplayTag InTag, const FConstStructView InPayload, const FName InOrigin)
 		: Tag(InTag)
+		, Payload(InPayload)
+		, Origin(InOrigin)
 	{
 	}
 	
@@ -42,3 +51,45 @@ struct STATETREEMODULE_API FStateTreeEvent
 	UPROPERTY(EditAnywhere, BlueprintReadWrite, Category = "Default")
 	FName Origin;
 };
+
+/**
+ * Event queue buffering all the events to be processed by a State Tree.
+ */
+USTRUCT()
+struct STATETREEMODULE_API FStateTreeEventQueue
+{
+	GENERATED_BODY()
+
+	/** Maximum number of events that can be buffered. */
+	static constexpr int32 MaxActiveEvents = 64;
+
+	/** @return const view to all the events in the buffer. */
+	TConstArrayView<FStateTreeEvent> GetEvents() const
+	{
+		return Events;
+	}
+
+	/** Resets the events in the event queue */
+	void Reset()
+	{
+		Events.Reset();
+	}
+
+	/**
+	 * Buffers and event to be sent to the State Tree.
+	 * @param Owner Optional pointer to an owner UObject that is used for logging errors.
+	 * @param Tag tag identifying the event.
+	 * @param Payload Optional reference to the payload struct.
+	 * @param Origin Optional name identifying the origin of the event.
+	 */
+	void SendEvent(const UObject* Owner, const FGameplayTag& Tag, const FConstStructView Payload = FConstStructView(), const FName Origin = FName());
+
+protected:
+	// Used by FStateTreeExecutionState to implement deprecated functionality.
+	TArray<FStateTreeEvent>& GetEventsArray() { return Events; };
+
+	UPROPERTY()
+	TArray<FStateTreeEvent> Events;
+
+	friend struct FStateTreeInstanceData;
+};
diff --git a/Engine/Plugins/Runtime/StateTree/Source/StateTreeModule/Public/StateTreeExecutionContext.h b/Engine/Plugins/Runtime/StateTree/Source/StateTreeModule/Public/StateTreeExecutionContext.h
--- a/Engine/Plugins/Runtime/StateTree/Source/StateTreeModule/Public/StateTreeExecutionContext.h
+++ b/Engine/Plugins/Runtime/StateTree/Source/StateTreeModule/Public/StateTreeExecutionContext.h
@@ -90,12 +90,15 @@ public:
 	/** @return the StateTree asset in use. */
 	const UStateTree* GetStateTree() const { return &StateTree; }
 
-	/** @retrun const references to the instance data in use, or nullptr if the context is not valid. */
+	/** @return const references to the instance data in use, or nullptr if the context is not valid. */
 	const FStateTreeInstanceData* GetInstanceData() const { return &InstanceData; }
 
-	/** @retrun mutable references to the instance data in use, or nullptr if the context is not valid. */
+	/** @retuen mutable references to the instance data in use, or nullptr if the context is not valid. */
 	FStateTreeInstanceData* GetMutableInstanceData() const { return &InstanceData; }
-	
+
+	/** @retuen mutable references to the instance data in use. */
+	FStateTreeEventQueue& GetEventQueue() const { return InstanceData.GetEventQueue(); }
+
 	/** @return The owner of the context */
 	UObject* GetOwner() const { return &Owner; }
 	/** @return The world of the owner or nullptr if the owner is not set. */ 
@@ -138,7 +141,11 @@ public:
 	TArray<FName> GetActiveStateNames() const;
 
 	/** Sends event for the StateTree. */
-	void SendEvent(const FStateTreeEvent& Event);
+	// @todo: remove temporarily - UE_DEPRECATED(5.2, "Use AddEvent() with individual parameters instead.")
+	void SendEvent(const FStateTreeEvent& Event) const;
+
+	/** Sends event for the StateTree. */
+	void SendEvent(const FGameplayTag Tag, const FConstStructView Payload = FConstStructView(), const FName Origin = FName()) const;
 
 	/** Iterates over all events. Can only be used during StateTree tick. Expects a lambda which takes const FStateTreeEvent& Event, and returns EStateTreeLoopEvents. */
 	template<typename TFunc>
@@ -247,6 +254,14 @@ public:
 		return DataViews[Node.DataViewIndex.Get()].template GetMutable<typename T::FInstanceDataType>();
 	}
 
+	/** @returns reference to instance data struct that can be passed to lambdas. See TStateTreeInstanceDataStructRef for usage. */
+	template <typename T>
+	TStateTreeInstanceDataStructRef<typename T::FInstanceDataType> GetInstanceDataStructRef(const T& Node) const
+	{
+		static_assert(TIsDerivedFrom<T, FStateTreeNodeBase>::IsDerived, "Expecting Node to derive from FStateTreeNodeBase.");
+		return TStateTreeInstanceDataStructRef<typename T::FInstanceDataType>(InstanceData, DataViews[Node.DataViewIndex.Get()].template GetMutable<typename T::FInstanceDataType>());
+	}
+	
 protected:
 
 	/** @return Prefix that will be used by STATETREE_LOG and STATETREE_CLOG, empty by default. */
@@ -312,7 +327,7 @@ protected:
 	 * If a state cannot be selected, false is returned. 
 	 * If NextState is a selector state, SelectStateInternal is called recursively (depth-first) to all child states (where NextState will be one of child states).
 	 * If NextState is a leaf state, the active states leading from root to the leaf are returned.
-	 * @param InstanceData Reference to the instance data
+	 * @param SharedInstanceData Reference to the instance data
 	 * @param NextState The state which we try to select next.
 	 * @param OutNewActiveStates Active states that got selected.
 	 * @return True if succeeded to select new active states.
@@ -351,6 +366,9 @@ protected:
 	/** @return String describing full path of an activate state for logging and debug. */
 	FString DebugGetStatePath(const FStateTreeActiveStates& ActiveStates, int32 ActiveStateIndex) const;
 
+	/** @return String describing all events that are currently being processed  for logging and debug. */
+	FString DebugGetEventsAsString() const;
+	
 	/** Helper function to update struct or object dataview of a node. */
 	template<typename T>
 	void SetNodeDataView(T& Node, int32& InstanceStructIndex, int32& InstanceObjectIndex)
diff --git a/Engine/Plugins/Runtime/StateTree/Source/StateTreeModule/Public/StateTreeInstanceData.h b/Engine/Plugins/Runtime/StateTree/Source/StateTreeModule/Public/StateTreeInstanceData.h
--- a/Engine/Plugins/Runtime/StateTree/Source/StateTreeModule/Public/StateTreeInstanceData.h
+++ b/Engine/Plugins/Runtime/StateTree/Source/StateTreeModule/Public/StateTreeInstanceData.h
@@ -44,6 +44,9 @@ struct STATETREEMODULE_API FStateTreeInstanceData
 	/** @return Number of items in the instance data. */
 	int32 NumStructs() const { return InstanceStructs.Num(); }
 
+	/** @return true if the specified index is valid index into the struct data container. */
+	bool IsValidStructIndex(const int32 Index) const { return InstanceStructs.IsValidIndex(Index); }
+	
 	/** @return mutable view to the struct at specified index. */
 	FStructView GetMutableStruct(const int32 Index) const { return InstanceStructs[Index]; }
 
@@ -60,7 +63,11 @@ struct STATETREEMODULE_API FStateTreeInstanceData
 	const UObject* GetObject(const int32 Index) const { return InstanceObjects[Index]; }
 
 	/** @return array to store unprocessed events. */
-	TArray<FStateTreeEvent>& GetEvents() { return Events; }
+	UE_DEPRECATED(5.2, "Use GetEventQueue() instead.")
+	TArray<FStateTreeEvent>& GetEvents();
+
+	/** @return reference to the event queue. */
+	FStateTreeEventQueue& GetEventQueue() { return EventQueue; }
 	
 	int32 GetEstimatedMemoryUsage() const;
 	int32 GetNumItems() const;
@@ -80,7 +87,7 @@ private:
 
 	/** Events */
 	UPROPERTY()
-	TArray<FStateTreeEvent> Events;
+	FStateTreeEventQueue EventQueue;
 };
 
 template<>
@@ -91,3 +98,61 @@ struct TStructOpsTypeTraits<FStateTreeInstanceData> : public TStructOpsTypeTrait
 		WithIdentical = true,
 	};
 };
+
+/**
+ * Stores indexed reference to a instance data struct.
+ * The instance data structs may be relocated when the instance data composition changed. For that reason you cannot store pointers to the instance data.
+ * This is often needed for example when dealing with delegate lambda's. This helper struct stores the instance data as index to the instance data array.
+ * That way we can access the instance data even of the array changes.
+ *
+ * You generally do not use this directly, but via FStateTreeExecutionContext.
+ *
+ *	EStateTreeRunStatus FTestTask::EnterState(FStateTreeExecutionContext& Context, const FStateTreeTransitionResult& Transition) const
+ *	{
+ *		FInstanceDataType& InstanceData = Context.GetInstanceData(*this);
+ *
+ *		Context.GetWorld()->GetTimerManager().SetTimer(
+ *	        InstanceData.TimerHandle,
+ *	        [InstanceDataRef = Context.GetInstanceDataStructRef()]()
+ *	        {
+ *	            FInstanceDataType& InstanceData = *InstanceDataRef;
+ *	            ...
+ *	        },
+ *	        Delay, true);
+ *
+ *	    return EStateTreeRunStatus::Running;
+ *	}
+ */
+template <typename T>
+struct TStateTreeInstanceDataStructRef
+{
+	TStateTreeInstanceDataStructRef(FStateTreeInstanceData& InInstanceData, const T& InstanceDataStruct)
+		: InstanceData(InInstanceData)
+	{
+		const FConstStructView InstanceDataStructView = FConstStructView::template Make(InstanceDataStruct);
+		// Find struct in the instance data.
+		for (int32 Index = 0; Index < InstanceData.NumStructs(); Index++)
+		{
+			if (InstanceData.GetStruct(Index) == InstanceDataStructView)
+			{
+				StructIndex = Index;
+				break;
+			}
+		}
+		check(StructIndex != INDEX_NONE);
+	}
+
+	bool IsValid() const { return InstanceData.IsValidStructIndex(StructIndex); }
+
+	T& operator*() const
+	{
+		check(IsValid());
+		const FStructView Struct = InstanceData.GetMutableStruct(StructIndex);
+		check(Struct.GetScriptStruct() == TBaseStructure<T>::Get());
+		return *reinterpret_cast<T*>(Struct.GetMutableMemory());
+	}
+
+protected:
+	FStateTreeInstanceData& InstanceData;
+	int32 StructIndex = INDEX_NONE;
+};
\ No newline at end of file
diff --git a/Engine/Plugins/Runtime/StateTree/Source/StateTreeModule/Public/StateTreeNodeBase.h b/Engine/Plugins/Runtime/StateTree/Source/StateTreeModule/Public/StateTreeNodeBase.h
--- a/Engine/Plugins/Runtime/StateTree/Source/StateTreeModule/Public/StateTreeNodeBase.h
+++ b/Engine/Plugins/Runtime/StateTree/Source/StateTreeModule/Public/StateTreeNodeBase.h
@@ -25,7 +25,6 @@ struct STATETREEMODULE_API FStateTreeNodeBase
 	/**
 	 * Called when the StateTree asset is linked. Allows to resolve references to other StateTree data.
 	 * @see TStateTreeExternalDataHandle
-	 * @see TStateTreeInstanceDataPropertyHandle
 	 * @param Linker Reference to the linker
 	 * @return true if linking succeeded. 
 	 */
diff --git a/Engine/Plugins/Runtime/StateTree/Source/StateTreeModule/Public/StateTreeTaskBase.h b/Engine/Plugins/Runtime/StateTree/Source/StateTreeModule/Public/StateTreeTaskBase.h
--- a/Engine/Plugins/Runtime/StateTree/Source/StateTreeModule/Public/StateTreeTaskBase.h
+++ b/Engine/Plugins/Runtime/StateTree/Source/StateTreeModule/Public/StateTreeTaskBase.h
@@ -19,6 +19,10 @@ struct STATETREEMODULE_API FStateTreeTaskBase : public FStateTreeNodeBase
 
 	FStateTreeTaskBase()
 		: bShouldStateChangeOnReselect(true)
+		, bShouldCallTick(true)
+		, bShouldCallTickOnlyOnEvents(false)
+		, bShouldCopyBoundPropertiesOnTick(true)
+		, bShouldCopyBoundPropertiesOnExitState(true)
 	{
 	}
 	
@@ -73,6 +77,16 @@ struct STATETREEMODULE_API FStateTreeTaskBase : public FStateTreeNodeBase
 	 * and false on state like tasks like claiming a resource that is expected to be acquired on child states.
 	 * Default value is true. */
 	uint8 bShouldStateChangeOnReselect : 1;
+
+	/** If set to true, Tick() is called. Default true. */
+	uint8 bShouldCallTick : 1;
+	/** If set to true, Tick() is called only when there are events. No effect if bShouldCallTickState is true. Default false. */
+	uint8 bShouldCallTickOnlyOnEvents : 1;
+
+	/** If set to true, copy the values of bound properties before calling Tick(). Default true. */
+	uint8 bShouldCopyBoundPropertiesOnTick : 1;
+	/** If set to true, copy the values of bound properties before calling ExitState(). Default true. */
+	uint8 bShouldCopyBoundPropertiesOnExitState : 1;
 };
 
 /**
diff --git a/Engine/Plugins/Runtime/StateTree/Source/StateTreeModule/Public/StateTreeTypes.h b/Engine/Plugins/Runtime/StateTree/Source/StateTreeModule/Public/StateTreeTypes.h
--- a/Engine/Plugins/Runtime/StateTree/Source/StateTreeModule/Public/StateTreeTypes.h
+++ b/Engine/Plugins/Runtime/StateTree/Source/StateTreeModule/Public/StateTreeTypes.h
@@ -30,17 +30,6 @@ enum class EStateTreeRunStatus : uint8
 	Unset,
 };
 
-/** Evaluator evaluation type. */
-UENUM()
-enum class EStateTreeEvaluationType : uint8
-{
-	/** Called during selection process on states that have not been visited yet. */
-	PreSelect,
-	
-	/** Called during tick on active states. */
-    Tick,
-};
-
 /** State change type. Passed to EnterState() and ExitState() to indicate how the state change affects the state and Evaluator or Task is on. */
 UENUM()
 enum class EStateTreeStateChangeType : uint8
