StateTree: Fix copy pasting UObject based nodes
- not all properties were copied due to use of PPF_PropertyWindow instead of PPF_Copy
- instance object was not duplicated correctly

diff --git a/Engine/Plugins/Runtime/StateTree/Source/StateTreeEditorModule/Private/Customizations/StateTreeEditorNodeDetails.cpp b/Engine/Plugins/Runtime/StateTree/Source/StateTreeEditorModule/Private/Customizations/StateTreeEditorNodeDetails.cpp
--- a/Engine/Plugins/Runtime/StateTree/Source/StateTreeEditorModule/Private/Customizations/StateTreeEditorNodeDetails.cpp
+++ b/Engine/Plugins/Runtime/StateTree/Source/StateTreeEditorModule/Private/Customizations/StateTreeEditorNodeDetails.cpp
@@ -481,7 +481,8 @@ void FStateTreeEditorNodeDetails::CustomizeHeader(TSharedRef<class IPropertyHand
 void FStateTreeEditorNodeDetails::OnCopyNode()
 {
 	FString Value;
-	if (StructProperty->GetValueAsFormattedString(Value) == FPropertyAccess::Success)
+	// Use PPF_Copy so that all properties get copied.
+	if (StructProperty->GetValueAsFormattedString(Value, PPF_Copy) == FPropertyAccess::Success)
 	{
 		FPlatformApplicationMisc::ClipboardCopy(*Value);
 	}
@@ -511,20 +512,29 @@ void FStateTreeEditorNodeDetails::OnPasteNode()
 
 				StructProperty->NotifyPreChange();
 
-				StructProperty->SetValueFromFormattedString(PastedText);
+				// Make sure we instantiate new objects when setting the value.
+				StructProperty->SetValueFromFormattedString(PastedText, EPropertyValueSetFlags::InstanceObjects);
 
 				// Reset GUIDs on paste
+				TArray<UObject*> OuterObjects; 
 				TArray<void*> RawNodeData;
+				StructProperty->GetOuterObjects(OuterObjects);
 				StructProperty->AccessRawData(RawNodeData);
-				for (void* Data : RawNodeData)
+				if (OuterObjects.Num() == RawNodeData.Num())
 				{
-					if (FStateTreeEditorNode* EditorNode = static_cast<FStateTreeEditorNode*>(Data))
+					for (int32 Index = 0; Index < RawNodeData.Num(); Index++)
 					{
-						if (FStateTreeNodeBase* Node = EditorNode->Node.GetMutablePtr<FStateTreeNodeBase>())
+						UObject* OuterObject = OuterObjects[Index];
+						FStateTreeEditorNode* EditorNode = static_cast<FStateTreeEditorNode*>(RawNodeData[Index]);
+						if (EditorNode && OuterObject)
 						{
-							Node->Name = FName(Node->Name.ToString() + TEXT(" Copy"));
+							if (FStateTreeNodeBase* Node = EditorNode->Node.GetMutablePtr<FStateTreeNodeBase>())
+							{
+								Node->Name = FName(Node->Name.ToString() + TEXT(" Copy"));
+							}
+
+							EditorNode->ID = FGuid::NewGuid();
 						}
-						EditorNode->ID = FGuid::NewGuid();
 					}
 				}
 
