StateTree
- restored deprecation (was comments out due to submitting code via multiple CLs)
- cleaned up logging
- skip property binding to not take place when task is not ticked

diff --git a/Engine/Plugins/Runtime/StateTree/Source/StateTreeModule/Private/StateTreeExecutionContext.cpp b/Engine/Plugins/Runtime/StateTree/Source/StateTreeModule/Private/StateTreeExecutionContext.cpp
--- a/Engine/Plugins/Runtime/StateTree/Source/StateTreeModule/Private/StateTreeExecutionContext.cpp
+++ b/Engine/Plugins/Runtime/StateTree/Source/StateTreeModule/Private/StateTreeExecutionContext.cpp
@@ -450,6 +450,8 @@ void FStateTreeExecutionContext::SendEvent(const FGameplayTag Tag, const FConstS
 		return;
 	}
 
+	STATETREE_LOG(Verbose, TEXT("Send Event '%s'"), *Tag.ToString());
+
 	FStateTreeEventQueue& EventQueue = InstanceData.GetEventQueue();
 	EventQueue.SendEvent(&Owner, Tag, Payload, Origin);
 }
@@ -593,7 +595,9 @@ EStateTreeRunStatus FStateTreeExecutionContext::EnterState(const FStateTreeTrans
 	check(Exec.FirstTaskStructIndex.IsValid() && Exec.FirstTaskObjectIndex.IsValid()); 
 	int32 InstanceStructIndex = Exec.FirstTaskStructIndex.Get();
 	int32 InstanceObjectIndex = Exec.FirstTaskObjectIndex.Get();
-	
+
+	STATETREE_LOG(Log, TEXT("Enter state '%s' (%d)"), *DebugGetStatePath(Transition.NextActiveStates), Exec.StateChangeCount);
+
 	for (int32 Index = 0; Index < Transition.NextActiveStates.Num() && Result != EStateTreeRunStatus::Failed; Index++)
 	{
 		const FStateTreeStateHandle CurrentHandle = Transition.NextActiveStates[Index];
@@ -624,7 +628,7 @@ EStateTreeRunStatus FStateTreeExecutionContext::EnterState(const FStateTreeTrans
 		CurrentTransition.CurrentState = CurrentHandle;
 		CurrentTransition.ChangeType = bWasActive ? EStateTreeStateChangeType::Sustained : EStateTreeStateChangeType::Changed;
 
-		STATETREE_CLOG(bIsEnteringState, Log, TEXT("%*sEnter state '%s' %s"), Index*UE::StateTree::DebugIndentSize, TEXT(""), *DebugGetStatePath(Transition.NextActiveStates, Index), *UEnum::GetValueAsString(CurrentTransition.ChangeType));
+		STATETREE_CLOG(bIsEnteringState, Log, TEXT("%*sState '%s' %s"), Index*UE::StateTree::DebugIndentSize, TEXT(""), *DebugGetStatePath(Transition.NextActiveStates, Index), *UEnum::GetDisplayValueAsText(CurrentTransition.ChangeType).ToString());
 		
 		// Activate tasks on current state.
 		for (int32 TaskIndex = State.TasksBegin; TaskIndex < (State.TasksBegin + State.TasksNum); TaskIndex++)
@@ -643,7 +647,7 @@ EStateTreeRunStatus FStateTreeExecutionContext::EnterState(const FStateTreeTrans
 
 			if (bIsEnteringState && bShouldCallStateChange)
 			{
-				STATETREE_LOG(Verbose, TEXT("%*s  Notify Task '%s'"), Index*UE::StateTree::DebugIndentSize, TEXT(""), *Task.Name.ToString());
+				STATETREE_LOG(Verbose, TEXT("%*s  Task '%s'"), Index*UE::StateTree::DebugIndentSize, TEXT(""), *Task.Name.ToString());
 				QUICK_SCOPE_CYCLE_COUNTER(StateTree_Task_EnterState);
 				CSV_SCOPED_TIMING_STAT_EXCLUSIVE(StateTree_Task_EnterState);
 				const EStateTreeRunStatus Status = Task.EnterState(*this, CurrentTransition);
@@ -737,6 +741,8 @@ void FStateTreeExecutionContext::ExitState(const FStateTreeTransitionResult& Tra
 	}
 
 	// Call in reverse order.
+	STATETREE_LOG(Log, TEXT("Exit state '%s' (%d)"), *DebugGetStatePath(Transition.CurrentActiveStates), Exec.StateChangeCount);
+
 	FStateTreeTransitionResult CurrentTransition = Transition;
 
 	for (int32 Index = NumExitedStates - 1; Index >= 0; Index--)
@@ -747,7 +753,7 @@ void FStateTreeExecutionContext::ExitState(const FStateTreeTransitionResult& Tra
 		CurrentTransition.CurrentState = CurrentHandle;
 		CurrentTransition.ChangeType = ExitedStateChangeType[Index];
 
-		STATETREE_LOG(Log, TEXT("%*sExit state '%s' %s"), Index*UE::StateTree::DebugIndentSize, TEXT(""), *DebugGetStatePath(Transition.CurrentActiveStates, ExitedStateActiveIndex[Index]), *UEnum::GetValueAsString(CurrentTransition.ChangeType));
+		STATETREE_LOG(Log, TEXT("%*sState '%s' %s"), Index*UE::StateTree::DebugIndentSize, TEXT(""), *DebugGetStatePath(Transition.CurrentActiveStates, ExitedStateActiveIndex[Index]), *UEnum::GetDisplayValueAsText(CurrentTransition.ChangeType).ToString());
 
 		// Tasks
 		for (int32 TaskIndex = (State.TasksBegin + State.TasksNum) - 1; TaskIndex >= State.TasksBegin; TaskIndex--)
@@ -764,7 +770,7 @@ void FStateTreeExecutionContext::ExitState(const FStateTreeTransitionResult& Tra
 
 				if (bShouldCallStateChange)
 				{
-					STATETREE_LOG(Verbose, TEXT("%*s  Notify Task '%s'"), Index*UE::StateTree::DebugIndentSize, TEXT(""), *Task.Name.ToString());
+					STATETREE_LOG(Verbose, TEXT("%*s  Task '%s'"), Index*UE::StateTree::DebugIndentSize, TEXT(""), *Task.Name.ToString());
 					{
 						QUICK_SCOPE_CYCLE_COUNTER(StateTree_Task_ExitState);
 						CSV_SCOPED_TIMING_STAT_EXCLUSIVE(StateTree_Task_ExitState);
@@ -787,6 +793,8 @@ void FStateTreeExecutionContext::StateCompleted()
 		return;
 	}
 
+	STATETREE_LOG(Verbose, TEXT("State Completed %s (%d)"), *UEnum::GetDisplayValueAsText(Exec.LastTickStatus).ToString(), Exec.StateChangeCount);
+
 	// Call from child towards root to allow to pass results back.
 	// Note: Completed is assumed to be called immediately after tick or enter state, so there's no property copying.
 	for (int32 Index = Exec.ActiveStates.Num() - 1; Index >= 0; Index--)
@@ -794,7 +802,7 @@ void FStateTreeExecutionContext::StateCompleted()
 		const FStateTreeStateHandle CurrentHandle = Exec.ActiveStates[Index];
 		const FCompactStateTreeState& State = StateTree.States[CurrentHandle.Index];
 
-		STATETREE_LOG(Verbose, TEXT("%*sState Completed '%s' %s"), Index*UE::StateTree::DebugIndentSize, TEXT(""), *DebugGetStatePath(Exec.ActiveStates, Index), *UEnum::GetValueAsString(Exec.LastTickStatus));
+		STATETREE_LOG(Verbose, TEXT("%*sState '%s'"), Index*UE::StateTree::DebugIndentSize, TEXT(""), *DebugGetStatePath(Exec.ActiveStates, Index));
 
 		// Notify Tasks
 		for (int32 TaskIndex = (State.TasksBegin + State.TasksNum) - 1; TaskIndex >= State.TasksBegin; TaskIndex--)
@@ -806,7 +814,7 @@ void FStateTreeExecutionContext::StateCompleted()
 			{
 				const FStateTreeTaskBase& Task = StateTree.Nodes[TaskIndex].Get<FStateTreeTaskBase>();
 
-				STATETREE_LOG(Verbose, TEXT("%*s  Notify Task '%s'"), Index*UE::StateTree::DebugIndentSize, TEXT(""), *Task.Name.ToString());
+				STATETREE_LOG(Verbose, TEXT("%*s  Task '%s'"), Index*UE::StateTree::DebugIndentSize, TEXT(""), *Task.Name.ToString());
 				Task.StateCompleted(*this, Exec.LastTickStatus, Exec.ActiveStates);
 			}
 		}
@@ -928,13 +936,15 @@ EStateTreeRunStatus FStateTreeExecutionContext::TickTasks(const float DeltaTime)
 
 	// Used to stop ticking tasks after one fails, but we still want to keep updating the data of them
 	bool bShouldTickTasks = true;
-	
+
+	STATETREE_CLOG(Exec.ActiveStates.Num() > 0, VeryVerbose, TEXT("Ticking Tasks"));
+
 	for (int32 Index = 0; Index < Exec.ActiveStates.Num(); Index++)
 	{
 		const FStateTreeStateHandle CurrentHandle = Exec.ActiveStates[Index];
 		const FCompactStateTreeState& State = StateTree.States[CurrentHandle.Index];
 
-		STATETREE_CLOG(State.TasksNum > 0, VeryVerbose, TEXT("%*sTicking Tasks of state '%s'"), Index*UE::StateTree::DebugIndentSize, TEXT(""), *DebugGetStatePath(Exec.ActiveStates, Index));
+		STATETREE_CLOG(State.TasksNum > 0, VeryVerbose, TEXT("%*sState '%s'"), Index*UE::StateTree::DebugIndentSize, TEXT(""), *DebugGetStatePath(Exec.ActiveStates, Index));
 
 		if (State.Type == EStateTreeStateType::Linked)
 		{
@@ -952,16 +962,18 @@ EStateTreeRunStatus FStateTreeExecutionContext::TickTasks(const float DeltaTime)
 			const FStateTreeTaskBase& Task = StateTree.Nodes[TaskIndex].Get<FStateTreeTaskBase>();
 			SetNodeDataView(Task, InstanceStructIndex, InstanceObjectIndex);
 
+			const bool bNeedsTick = bShouldTickTasks && (Task.bShouldCallTick || (bHasEvents && Task.bShouldCallTickOnlyOnEvents));
+			
 			// Copy bound properties.
-			if (Task.BindingsBatch.IsValid() && Task.bShouldCopyBoundPropertiesOnTick)
+			// Only copy properties when the task is actually ticked, and copy properties at tick is requested.
+			if (Task.BindingsBatch.IsValid() && bNeedsTick && Task.bShouldCopyBoundPropertiesOnTick)
 			{
 				StateTree.PropertyBindings.CopyTo(DataViews, Task.BindingsBatch, DataViews[Task.DataViewIndex.Get()]);
 			}
-			STATETREE_LOG(VeryVerbose, TEXT("%*s  Tick: '%s'"), Index*UE::StateTree::DebugIndentSize, TEXT(""), *Task.Name.ToString());
-
-			const bool bNeedTick = Task.bShouldCallTick || (bHasEvents && Task.bShouldCallTickOnlyOnEvents);
 			
-			if (bShouldTickTasks && bNeedTick)
+			STATETREE_LOG(VeryVerbose, TEXT("%*s  Tick: '%s' %s"), Index*UE::StateTree::DebugIndentSize, TEXT(""), *Task.Name.ToString(), !bNeedsTick ? TEXT("[not ticked]") : TEXT(""));
+
+			if (bNeedsTick)
 			{
 				QUICK_SCOPE_CYCLE_COUNTER(StateTree_Task_Tick);
 				CSV_SCOPED_TIMING_STAT_EXCLUSIVE(StateTree_Task_Tick);
@@ -1350,12 +1362,13 @@ FString FStateTreeExecutionContext::GetSafeStateName(const FStateTreeStateHandle
 FString FStateTreeExecutionContext::DebugGetStatePath(const FStateTreeActiveStates& ActiveStates, const int32 ActiveStateIndex) const
 {
 	FString StatePath;
-	if (!ensureMsgf(ActiveStates.IsValidIndex(ActiveStateIndex), TEXT("Provided index must be valid")))
+	const int32 Num = ActiveStateIndex == INDEX_NONE ? ActiveStates.Num() : (ActiveStateIndex + 1);  
+	if (!ensureMsgf(ActiveStates.IsValidIndex(Num - 1), TEXT("Provided index must be valid")))
 	{
 		return StatePath;
 	}
 
-	for (int32 i = 0; i <= ActiveStateIndex; i++)
+	for (int32 i = 0; i < Num; i++)
 	{
 		const FCompactStateTreeState& State = StateTree.States[ActiveStates[i].Index];
 		StatePath.Appendf(TEXT("%s%s"), i == 0 ? TEXT("") : TEXT("."), *State.Name.ToString());
@@ -1479,7 +1492,7 @@ void FStateTreeExecutionContext::DebugPrintInternalLayout()
 	DebugString += FString::Printf(TEXT("\nExternal Data(%d)\n  [ %-40s | %-8s | %5s ]\n"), StateTree.ExternalDataDescs.Num(), TEXT("Name"), TEXT("Optional"), TEXT("Index"));
 	for (const FStateTreeExternalDataDesc& Desc : StateTree.ExternalDataDescs)
 	{
-		DebugString += FString::Printf(TEXT("  | %-40s | %8s | %5d |\n"), Desc.Struct ? *Desc.Struct->GetName() : TEXT("null"), *UEnum::GetValueAsString(Desc.Requirement), Desc.Handle.DataViewIndex.Get());
+		DebugString += FString::Printf(TEXT("  | %-40s | %8s | %5d |\n"), Desc.Struct ? *Desc.Struct->GetName() : TEXT("null"), *UEnum::GetDisplayValueAsText(Desc.Requirement).ToString(), Desc.Handle.DataViewIndex.Get());
 	}
 
 	// Bindings
diff --git a/Engine/Plugins/Runtime/StateTree/Source/StateTreeModule/Public/StateTreeExecutionContext.h b/Engine/Plugins/Runtime/StateTree/Source/StateTreeModule/Public/StateTreeExecutionContext.h
--- a/Engine/Plugins/Runtime/StateTree/Source/StateTreeModule/Public/StateTreeExecutionContext.h
+++ b/Engine/Plugins/Runtime/StateTree/Source/StateTreeModule/Public/StateTreeExecutionContext.h
@@ -141,7 +141,7 @@ public:
 	TArray<FName> GetActiveStateNames() const;
 
 	/** Sends event for the StateTree. */
-	// @todo: remove temporarily - UE_DEPRECATED(5.2, "Use AddEvent() with individual parameters instead.")
+	UE_DEPRECATED(5.2, "Use AddEvent() with individual parameters instead.")
 	void SendEvent(const FStateTreeEvent& Event) const;
 
 	/** Sends event for the StateTree. */
@@ -364,7 +364,7 @@ protected:
 	FString GetSafeStateName(const FStateTreeStateHandle State) const;
 
 	/** @return String describing full path of an activate state for logging and debug. */
-	FString DebugGetStatePath(const FStateTreeActiveStates& ActiveStates, int32 ActiveStateIndex) const;
+	FString DebugGetStatePath(const FStateTreeActiveStates& ActiveStates, const int32 ActiveStateIndex = INDEX_NONE) const;
 
 	/** @return String describing all events that are currently being processed  for logging and debug. */
 	FString DebugGetEventsAsString() const;
diff --git a/Engine/Plugins/Runtime/StateTree/Source/StateTreeModule/Public/StateTreeTaskBase.h b/Engine/Plugins/Runtime/StateTree/Source/StateTreeModule/Public/StateTreeTaskBase.h
--- a/Engine/Plugins/Runtime/StateTree/Source/StateTreeModule/Public/StateTreeTaskBase.h
+++ b/Engine/Plugins/Runtime/StateTree/Source/StateTreeModule/Public/StateTreeTaskBase.h
@@ -78,9 +78,9 @@ struct STATETREEMODULE_API FStateTreeTaskBase : public FStateTreeNodeBase
 	 * Default value is true. */
 	uint8 bShouldStateChangeOnReselect : 1;
 
-	/** If set to true, Tick() is called. Default true. */
+	/** If set to true, Tick() is called. Not ticking implies no property copy. Default true. */
 	uint8 bShouldCallTick : 1;
-	/** If set to true, Tick() is called only when there are events. No effect if bShouldCallTickState is true. Default false. */
+	/** If set to true, Tick() is called only when there are events. No effect if bShouldCallTickState is true. Not ticking implies no property copy. Default false. */
 	uint8 bShouldCallTickOnlyOnEvents : 1;
 
 	/** If set to true, copy the values of bound properties before calling Tick(). Default true. */
