Fixed missing linkage for constants bound to references.

diff --git a/Engine/Plugins/Runtime/StateTree/Source/StateTreeEditorModule/Private/Customizations/StateTreeStateLinkDetails.h b/Engine/Plugins/Runtime/StateTree/Source/StateTreeEditorModule/Private/Customizations/StateTreeStateLinkDetails.h
--- a/Engine/Plugins/Runtime/StateTree/Source/StateTreeEditorModule/Private/Customizations/StateTreeStateLinkDetails.h
+++ b/Engine/Plugins/Runtime/StateTree/Source/StateTreeEditorModule/Private/Customizations/StateTreeStateLinkDetails.h
@@ -29,10 +29,10 @@ public:
 private:
 
 	// Special indices for the combo selection.
-	static const int ComboNotSet = -1;
-	static const int ComboSucceeded = -2;
-	static const int ComboFailed = -3;
-	static const int ComboNextState = -4;
+	inline static constexpr int ComboNotSet = -1;
+	inline static constexpr int ComboSucceeded = -2;
+	inline static constexpr int ComboFailed = -3;
+	inline static constexpr int ComboNextState = -4;
 
 	void CacheStates();
 	void CacheStates(const UStateTreeState* State);
diff --git a/Engine/Source/Developer/TraceInsights/Private/Insights/TimingProfilerManager.h b/Engine/Source/Developer/TraceInsights/Private/Insights/TimingProfilerManager.h
--- a/Engine/Source/Developer/TraceInsights/Private/Insights/TimingProfilerManager.h
+++ b/Engine/Source/Developer/TraceInsights/Private/Insights/TimingProfilerManager.h
@@ -37,7 +37,7 @@ class FTimingProfilerManager : public TSharedFromThis<FTimingProfilerManager>, p
 	friend class FTimingProfilerActionManager;
 
 public:
-	static const uint32 UnlimitedEventDepth = 1000;
+	inline static constexpr uint32 UnlimitedEventDepth = 1000;
 
 public:
 	/** Creates the Timing Profiler manager, only one instance can exist. */
