Add IgnoreForMemberInitializationTest to FGuid properties in StateTreeState.h

diff --git a/Engine/Plugins/Runtime/StateTree/Source/StateTreeEditorModule/Public/StateTreeState.h b/Engine/Plugins/Runtime/StateTree/Source/StateTreeEditorModule/Public/StateTreeState.h
--- a/Engine/Plugins/Runtime/StateTree/Source/StateTreeEditorModule/Public/StateTreeState.h
+++ b/Engine/Plugins/Runtime/StateTree/Source/StateTreeEditorModule/Public/StateTreeState.h
@@ -25,13 +25,13 @@ struct STATETREEEDITORMODULE_API FStateTreeStateLink
 
 	void Set(const EStateTreeTransitionType InType, const UStateTreeState* InState = nullptr);
 	void Set(const UStateTreeState* InState) { Set(EStateTreeTransitionType::GotoState, InState); }
-	
+
 	bool IsValid() const { return ID.IsValid(); }
 
 	UPROPERTY(EditDefaultsOnly, Category = Link)
 	FName Name;
-	
-	UPROPERTY(EditDefaultsOnly, Category = Link)
+
+	UPROPERTY(EditDefaultsOnly, Category = Link, meta = (IgnoreForMemberInitializationTest))
 	FGuid ID;
 
 	UPROPERTY(EditDefaultsOnly, Category = Link)
@@ -70,13 +70,13 @@ struct STATETREEEDITORMODULE_API FStateTreeTransition
 	UPROPERTY(EditDefaultsOnly, Category = Transition)
 	FGameplayTag EventTag;
 
-	UPROPERTY(EditDefaultsOnly, Category = Transition, meta=(DisplayName="Transition To"))
+	UPROPERTY(EditDefaultsOnly, Category = Transition, meta = (DisplayName="Transition To"))
 	FStateTreeStateLink State;
 
 	// Gate delay in seconds.
 	UPROPERTY(EditDefaultsOnly, Category = Transition, meta = (UIMin = "0", ClampMin = "0", UIMax = "25", ClampMax = "25"))
 	float GateDelay = 0.0f;
-	
+
 	UPROPERTY(EditDefaultsOnly, Category = Transition, meta = (BaseStruct = "/Script/StateTreeModule.StateTreeConditionBase", BaseClass = "/Script/StateTreeModule.StateTreeConditionBlueprintBase"))
 	TArray<FStateTreeEditorNode> Conditions;
 };
@@ -91,14 +91,14 @@ struct STATETREEEDITORMODULE_API FStateTreeStateParameters
 		Parameters.Reset();
 		bFixedLayout = false;
 	}
-	
+
 	UPROPERTY(EditDefaultsOnly, Category = Parameters)
 	FInstancedPropertyBag Parameters;
 
 	UPROPERTY(EditDefaultsOnly, Category = Parameters)
 	bool bFixedLayout = false;
 
-	UPROPERTY(EditDefaultsOnly, Category = Parameters)
+	UPROPERTY(EditDefaultsOnly, Category = Parameters, meta = (IgnoreForMemberInitializationTest))
 	FGuid ID;
 };
 
@@ -120,9 +120,9 @@ public:
 #endif
 
 	UStateTreeState* GetNextSiblingState() const;
-	
+
 	// StateTree Builder API
-	
+
 	/** Adds child state with specified name. */
 	UStateTreeState& AddChildState(const FName ChildName, const EStateTreeStateType StateType = EStateTreeStateType::State)
 	{
@@ -137,7 +137,7 @@ public:
 
 	/**
 	 * Adds enter condition of specified type.
-	 * @return reference to the new condition. 
+	 * @return reference to the new condition.
 	 */
 	template<typename T, typename... TArgs>
 	TStateTreeEditorNode<T>& AddEnterCondition(TArgs&&... InArgs)
@@ -155,7 +155,7 @@ public:
 
 	/**
 	 * Adds Task of specified type.
-	 * @return reference to the new Task. 
+	 * @return reference to the new Task.
 	 */
 	template<typename T, typename... TArgs>
 	TStateTreeEditorNode<T>& AddTask(TArgs&&... InArgs)
@@ -173,7 +173,7 @@ public:
 
 	/**
 	 * Adds Transition.
-	 * @return reference to the new Transition. 
+	 * @return reference to the new Transition.
 	 */
 	FStateTreeTransition& AddTransition(const EStateTreeTransitionTrigger InTrigger, const EStateTreeTransitionType InType, const UStateTreeState* InState = nullptr)
 	{
@@ -185,7 +185,7 @@ public:
 		return Transitions.Emplace_GetRef(InTrigger, InEventTag, InType, InState);
 	}
 
- 
+
 	// ~StateTree Builder API
 
 	UPROPERTY(EditDefaultsOnly, Category = "State")
@@ -200,7 +200,7 @@ public:
 	UPROPERTY(EditDefaultsOnly, Category = "State")
 	FStateTreeStateParameters Parameters;
 
-	UPROPERTY()
+	UPROPERTY(meta = (IgnoreForMemberInitializationTest))
 	FGuid ID;
 
 	UPROPERTY(EditDefaultsOnly, Category = "Enter Conditions", meta = (BaseStruct = "/Script/StateTreeModule.StateTreeConditionBase", BaseClass = "/Script/StateTreeModule.StateTreeConditionBlueprintBase"))
