StateTree: Fix BP task execution
- Added explicit bShouldCallTick and not make the tick based on of the event is implemented
- This fixes the cases where we are using the BP task class to implement UObject based tasks in C++

diff --git a/Engine/Plugins/Runtime/StateTree/Source/StateTreeModule/Private/Blueprint/StateTreeTaskBlueprintBase.cpp b/Engine/Plugins/Runtime/StateTree/Source/StateTreeModule/Private/Blueprint/StateTreeTaskBlueprintBase.cpp
--- a/Engine/Plugins/Runtime/StateTree/Source/StateTreeModule/Private/Blueprint/StateTreeTaskBlueprintBase.cpp
+++ b/Engine/Plugins/Runtime/StateTree/Source/StateTreeModule/Private/Blueprint/StateTreeTaskBlueprintBase.cpp
@@ -14,6 +14,7 @@
 UStateTreeTaskBlueprintBase::UStateTreeTaskBlueprintBase(const FObjectInitializer& ObjectInitializer)
 	: Super(ObjectInitializer)
 	, bShouldStateChangeOnReselect(true)
+	, bShouldCallTick(true)
 	, bShouldCallTickOnlyOnEvents(false)
 	, bShouldCopyBoundPropertiesOnTick(true)
 	, bShouldCopyBoundPropertiesOnExitState(true)
@@ -72,7 +73,7 @@ EDataValidationResult FStateTreeBlueprintTaskWrapper::Compile(FStateTreeDataView
 	
 	// Copy over ticking related options.
 	bShouldStateChangeOnReselect = InstanceData.bShouldStateChangeOnReselect;
-	bShouldCallTick = InstanceData.bHasTick;
+	bShouldCallTick = InstanceData.bShouldCallTick || InstanceData.bHasTick;
 	bShouldCallTickOnlyOnEvents = InstanceData.bShouldCallTickOnlyOnEvents;
 	bShouldCopyBoundPropertiesOnTick = InstanceData.bShouldCopyBoundPropertiesOnTick;
 	bShouldCopyBoundPropertiesOnExitState = InstanceData.bShouldCopyBoundPropertiesOnExitState;
diff --git a/Engine/Plugins/Runtime/StateTree/Source/StateTreeModule/Public/Blueprint/StateTreeTaskBlueprintBase.h b/Engine/Plugins/Runtime/StateTree/Source/StateTreeModule/Public/Blueprint/StateTreeTaskBlueprintBase.h
--- a/Engine/Plugins/Runtime/StateTree/Source/StateTreeModule/Public/Blueprint/StateTreeTaskBlueprintBase.h
+++ b/Engine/Plugins/Runtime/StateTree/Source/StateTreeModule/Public/Blueprint/StateTreeTaskBlueprintBase.h
@@ -57,6 +57,12 @@ protected:
 	UPROPERTY(EditDefaultsOnly, Category="Default")
 	uint8 bShouldStateChangeOnReselect : 1;
 
+	/**
+	 * If set to true, Tick() is called. Not ticking implies no property copy. Default true.
+	 * Note: this is intentionally not a property, should be only set by C++ derived classes when the tick should not be called.
+	 */
+	uint8 bShouldCallTick : 1;
+
 	/** If set to true, Tick() is called. Default false. */
 	UPROPERTY(EditDefaultsOnly, Category="Default")
 	uint8 bShouldCallTickOnlyOnEvents : 1;
