StateTree: Changed completion transitions to trigger at the state that was first completed.
- record the first completed state on EnterState() and TickTasks()
- returning succeeded on any node in EnterState() now behaves the same as returning succeeded on Tick()

diff --git a/Engine/Plugins/Runtime/StateTree/Source/StateTreeModule/Private/StateTreeExecutionContext.cpp b/Engine/Plugins/Runtime/StateTree/Source/StateTreeModule/Private/StateTreeExecutionContext.cpp
--- a/Engine/Plugins/Runtime/StateTree/Source/StateTreeModule/Private/StateTreeExecutionContext.cpp
+++ b/Engine/Plugins/Runtime/StateTree/Source/StateTreeModule/Private/StateTreeExecutionContext.cpp
@@ -578,6 +578,7 @@ EStateTreeRunStatus FStateTreeExecutionContext::EnterState(const FStateTreeTrans
 
 	FStateTreeExecutionState& Exec = GetExecState();
 	Exec.StateChangeCount++;
+	Exec.CompletedStateHandle = FStateTreeStateHandle::Invalid;
 
 	// On target branch means that the state is the target of current transition or child of it.
 	// States which were active before and will remain active, but are not on target branch will not get
@@ -651,12 +652,21 @@ EStateTreeRunStatus FStateTreeExecutionContext::EnterState(const FStateTreeTrans
 				QUICK_SCOPE_CYCLE_COUNTER(StateTree_Task_EnterState);
 				CSV_SCOPED_TIMING_STAT_EXCLUSIVE(StateTree_Task_EnterState);
 				const EStateTreeRunStatus Status = Task.EnterState(*this, CurrentTransition);
+
+				if (Status != EStateTreeRunStatus::Running)
+				{
+					// Store the first state that completed, will be used to decide where to trigger transitions.
+					if (!Exec.CompletedStateHandle.IsValid())
+					{
+						Exec.CompletedStateHandle = CurrentHandle;
+					}
+					Result = Status;
+				}
 				
 				if (Status == EStateTreeRunStatus::Failed)
 				{
-					// Store how far in the enter state we got. This will be used to match the ExitState() calls.
+					// Store how far in the enter state we got. This will be used to match the StateCompleted() and ExitState() calls.
 					Exec.EnterStateFailedTaskIndex = FStateTreeIndex16(TaskIndex);
-					Result = Status;
 					break;
 				}
 			}
@@ -918,7 +928,7 @@ EStateTreeRunStatus FStateTreeExecutionContext::TickTasks(const float DeltaTime)
 {
 	CSV_SCOPED_TIMING_STAT_EXCLUSIVE(StateTree_TickTasks);
 
-	const FStateTreeExecutionState& Exec = GetExecState();
+	FStateTreeExecutionState& Exec = GetExecState();
 
 	if (Exec.ActiveStates.IsEmpty())
 	{
@@ -934,6 +944,8 @@ EStateTreeRunStatus FStateTreeExecutionContext::TickTasks(const float DeltaTime)
 	int32 InstanceStructIndex = Exec.FirstTaskStructIndex.Get();
 	int32 InstanceObjectIndex = Exec.FirstTaskObjectIndex.Get();
 
+	Exec.CompletedStateHandle = FStateTreeStateHandle::Invalid;
+	
 	// Used to stop ticking tasks after one fails, but we still want to keep updating the data of them
 	bool bShouldTickTasks = true;
 
@@ -983,6 +995,11 @@ EStateTreeRunStatus FStateTreeExecutionContext::TickTasks(const float DeltaTime)
 				// TODO: Add more control over which states can control the failed/succeeded result.
 				if (TaskResult != EStateTreeRunStatus::Running)
 				{
+					// Store the first state that completed, will be used to decide where to trigger transitions.
+					if (!Exec.CompletedStateHandle.IsValid())
+					{
+						Exec.CompletedStateHandle = CurrentHandle;
+					}
 					Result = TaskResult;
 				}
 				
@@ -1101,14 +1118,25 @@ bool FStateTreeExecutionContext::TriggerTransitions(FStateTreeInstanceData& Shar
 	CSV_SCOPED_TIMING_STAT_EXCLUSIVE(StateTree_TriggerTransition);
 
 	FStateTreeExecutionState& Exec = GetExecState();
+
+	int32 StateStartIndex = Exec.ActiveStates.Num() - 1;
+
+	// On succeed/failed we will start looking for transition start at the first state that was completed.
+	auto GetLastCompletedStateIndex = [&Exec]() -> int32
+	{
+		const int32 Index = Exec.ActiveStates.IndexOfReverse(Exec.CompletedStateHandle);
+		return Index != INDEX_NONE ? Index : (Exec.ActiveStates.Num() - 1);
+	};
 	
 	EStateTreeTransitionTrigger CompletionTrigger = EStateTreeTransitionTrigger::None;
 	if (Exec.LastTickStatus == EStateTreeRunStatus::Succeeded)
 	{
+		StateStartIndex = GetLastCompletedStateIndex();
 		CompletionTrigger = EStateTreeTransitionTrigger::OnStateSucceeded;
 	}
 	else if (Exec.LastTickStatus == EStateTreeRunStatus::Failed)
 	{
+		StateStartIndex = GetLastCompletedStateIndex();
 		CompletionTrigger = EStateTreeTransitionTrigger::OnStateFailed;
 	}
 
@@ -1128,7 +1156,7 @@ bool FStateTreeExecutionContext::TriggerTransitions(FStateTreeInstanceData& Shar
 	};
 
 	// Walk towards root and check all transitions along the way.
-	for (int32 StateIndex = Exec.ActiveStates.Num() - 1; StateIndex >= 0; StateIndex--)
+	for (int32 StateIndex = StateStartIndex; StateIndex >= 0; StateIndex--)
 	{
 		const FCompactStateTreeState& State = StateTree.States[Exec.ActiveStates[StateIndex].Index];
 		
@@ -1658,4 +1686,3 @@ TArray<FName> FStateTreeExecutionContext::GetActiveStateNames() const
 
 #undef STATETREE_LOG
 #undef STATETREE_CLOG
-
diff --git a/Engine/Plugins/Runtime/StateTree/Source/StateTreeModule/Public/StateTreeExecutionContext.h b/Engine/Plugins/Runtime/StateTree/Source/StateTreeModule/Public/StateTreeExecutionContext.h
--- a/Engine/Plugins/Runtime/StateTree/Source/StateTreeModule/Public/StateTreeExecutionContext.h
+++ b/Engine/Plugins/Runtime/StateTree/Source/StateTreeModule/Public/StateTreeExecutionContext.h
@@ -40,6 +40,9 @@ struct STATETREEMODULE_API FStateTreeExecutionState
 	/** Delayed transition handle, if exists */
 	FStateTreeIndex16 GatedTransitionIndex = FStateTreeIndex16::Invalid;
 
+	/** Handle of the state that was first to report state completed (success or failure), used to trigger completion transitions. */
+	FStateTreeStateHandle CompletedStateHandle = FStateTreeStateHandle::Invalid;
+
 	/** Number of times a new state has been changed. */
 	uint16 StateChangeCount = 0;
 
diff --git a/Engine/Plugins/Runtime/StateTree/Source/StateTreeTestSuite/Private/StateTreeTest.cpp b/Engine/Plugins/Runtime/StateTree/Source/StateTreeTestSuite/Private/StateTreeTest.cpp
--- a/Engine/Plugins/Runtime/StateTree/Source/StateTreeTestSuite/Private/StateTreeTest.cpp
+++ b/Engine/Plugins/Runtime/StateTree/Source/StateTreeTestSuite/Private/StateTreeTest.cpp
@@ -17,6 +17,8 @@
 
 #define LOCTEXT_NAMESPACE "AITestSuite_StateTreeTest"
 
+UE_DISABLE_OPTIMIZATION_SHIP
+
 std::atomic<int32> FStateTreeTestConditionInstanceData::GlobalCounter = 0;
 
 namespace UE::StateTree::Tests
@@ -154,10 +156,10 @@ struct FStateTreeTest_Select : FAITestBase
 		UStateTreeState& State1A = State1.AddChildState(FName(TEXT("State1A")));
 
 		auto& TaskRoot = Root.AddTask<FTestTask_Stand>(FName(TEXT("TaskRoot")));
-		TaskRoot.GetNode().TicksToCompletion = 2;
+		TaskRoot.GetNode().TicksToCompletion = 3;  // let Task1A to complete first
 
 		auto& Task1 = State1.AddTask<FTestTask_Stand>(FName(TEXT("Task1")));
-		Task1.GetNode().TicksToCompletion = 2;
+		Task1.GetNode().TicksToCompletion = 3; // let Task1A to complete first
 
 		auto& Task1A = State1A.AddTask<FTestTask_Stand>(FName(TEXT("Task1A")));
 		Task1A.GetNode().TicksToCompletion = 2;
@@ -433,5 +435,138 @@ struct FStateTreeTest_SharedInstanceData : FAITestBase
 };
 IMPLEMENT_AI_INSTANT_TEST(FStateTreeTest_SharedInstanceData, "System.StateTree.SharedInstanceData");
 
+struct FStateTreeTest_TransitionPriority : FAITestBase
+{
+	virtual bool InstantTest() override
+	{
+		UStateTree& StateTree = UE::StateTree::Tests::NewStateTree(&GetWorld());
+		UStateTreeEditorData& EditorData = *Cast<UStateTreeEditorData>(StateTree.EditorData);
+		
+		UStateTreeState& Root = EditorData.AddSubTree(FName(TEXT("Root")));
+		UStateTreeState& State1 = Root.AddChildState(FName(TEXT("State1")));
+		UStateTreeState& State1A = State1.AddChildState(FName(TEXT("State1A")));
+		UStateTreeState& State1B = State1.AddChildState(FName(TEXT("State1B")));
+		UStateTreeState& State1C = State1.AddChildState(FName(TEXT("State1C")));
+
+		auto& Task1 = State1.AddTask<FTestTask_Stand>(FName(TEXT("Task1")));
+		Task1.GetNode().TicksToCompletion = 2;
+		State1.AddTransition(EStateTreeTransitionTrigger::OnStateCompleted, EStateTreeTransitionType::Succeeded);
+		
+		auto& Task1A = State1A.AddTask<FTestTask_Stand>(FName(TEXT("Task1A")));
+		Task1A.GetNode().TicksToCompletion = 1;
+		State1A.AddTransition(EStateTreeTransitionTrigger::OnStateCompleted, EStateTreeTransitionType::NextState);
+
+		auto& Task1B = State1B.AddTask<FTestTask_Stand>(FName(TEXT("Task1B")));
+		Task1B.GetNode().TicksToCompletion = 2;
+		State1B.AddTransition(EStateTreeTransitionTrigger::OnStateCompleted, EStateTreeTransitionType::NextState);
+
+		auto& Task1C = State1C.AddTask<FTestTask_Stand>(FName(TEXT("Task1C")));
+		Task1C.GetNode().TicksToCompletion = 2;
+		
+		FStateTreeCompilerLog Log;
+		FStateTreeCompiler Compiler(Log);
+		const bool bResult = Compiler.Compile(StateTree);
+		AITEST_TRUE("StateTree should get compiled", bResult);
+
+		EStateTreeRunStatus Status = EStateTreeRunStatus::Unset;
+		FStateTreeInstanceData InstanceData;
+		FTestStateTreeExecutionContext Exec(StateTree, StateTree, InstanceData);
+		const bool bInitSucceeded = Exec.IsValid();
+		AITEST_TRUE("StateTree should init", bInitSucceeded);
+
+		const FString TickStr(TEXT("Tick"));
+		const FString EnterStateStr(TEXT("EnterState"));
+		const FString ExitStateStr(TEXT("ExitState"));
+		const FString StateCompletedStr(TEXT("StateCompleted"));
+
+		// Start and enter state
+		Status = Exec.Start();
+		AITEST_TRUE("StateTree Task1 should enter state", Exec.Expect(Task1.GetName(), EnterStateStr));
+		AITEST_TRUE("StateTree Task1A should enter state", Exec.Expect(Task1A.GetName(), EnterStateStr));
+		Exec.LogClear();
+
+		// Transition from Task1A to Task1B
+		Status = Exec.Tick(0.1f);
+		AITEST_TRUE("StateTree Task1A should complete", Exec.Expect(Task1A.GetName(), StateCompletedStr));
+		AITEST_TRUE("StateTree Task1B should enter state", Exec.Expect(Task1B.GetName(), EnterStateStr));
+		Exec.LogClear();
+
+		// Task1 completes, and we should take State1 transition. 
+		Status = Exec.Tick(0.1f);
+		AITEST_TRUE("StateTree Task1 should complete", Exec.Expect(Task1.GetName(), StateCompletedStr));
+		AITEST_EQUAL("Tree execution should stop on success", Status, EStateTreeRunStatus::Succeeded);
+		Exec.LogClear();
+
+		return true;
+	}
+};
+IMPLEMENT_AI_INSTANT_TEST(FStateTreeTest_TransitionPriority, "System.StateTree.TransitionPriority");
+
+struct FStateTreeTest_TransitionPriorityEnterState : FAITestBase
+{
+	virtual bool InstantTest() override
+	{
+		UStateTree& StateTree = UE::StateTree::Tests::NewStateTree(&GetWorld());
+		UStateTreeEditorData& EditorData = *Cast<UStateTreeEditorData>(StateTree.EditorData);
+		
+		UStateTreeState& Root = EditorData.AddSubTree(FName(TEXT("Root")));
+		UStateTreeState& State0 = Root.AddChildState(FName(TEXT("State0")));
+		UStateTreeState& State1 = Root.AddChildState(FName(TEXT("State1")));
+		UStateTreeState& State1A = State1.AddChildState(FName(TEXT("State1A")));
+		UStateTreeState& State2 = Root.AddChildState(FName(TEXT("State2")));
+		UStateTreeState& State3 = Root.AddChildState(FName(TEXT("State3")));
+
+		auto& Task0 = State0.AddTask<FTestTask_Stand>(FName(TEXT("Task0")));
+		State0.AddTransition(EStateTreeTransitionTrigger::OnStateCompleted, EStateTreeTransitionType::GotoState, &State1);
+
+		auto& Task1 = State1.AddTask<FTestTask_Stand>(FName(TEXT("Task1")));
+		Task1.GetNode().EnterStateResult = EStateTreeRunStatus::Failed;
+		State1.AddTransition(EStateTreeTransitionTrigger::OnStateCompleted, EStateTreeTransitionType::GotoState, &State2);
+		
+		auto& Task1A = State1A.AddTask<FTestTask_Stand>(FName(TEXT("Task1A")));
+		State1A.AddTransition(EStateTreeTransitionTrigger::OnStateCompleted, EStateTreeTransitionType::GotoState, &State3);
+
+		auto& Task2 = State2.AddTask<FTestTask_Stand>(FName(TEXT("Task2")));
+		State2.AddTransition(EStateTreeTransitionTrigger::OnStateCompleted, EStateTreeTransitionType::Succeeded);
+
+		auto& Task3 = State3.AddTask<FTestTask_Stand>(FName(TEXT("Task3")));
+		State3.AddTransition(EStateTreeTransitionTrigger::OnStateCompleted, EStateTreeTransitionType::Succeeded);
+
+		
+		FStateTreeCompilerLog Log;
+		FStateTreeCompiler Compiler(Log);
+		const bool bResult = Compiler.Compile(StateTree);
+		AITEST_TRUE("StateTree should get compiled", bResult);
+
+		EStateTreeRunStatus Status = EStateTreeRunStatus::Unset;
+		FStateTreeInstanceData InstanceData;
+		FTestStateTreeExecutionContext Exec(StateTree, StateTree, InstanceData);
+		const bool bInitSucceeded = Exec.IsValid();
+		AITEST_TRUE("StateTree should init", bInitSucceeded);
+
+		const FString TickStr(TEXT("Tick"));
+		const FString EnterStateStr(TEXT("EnterState"));
+		const FString ExitStateStr(TEXT("ExitState"));
+		const FString StateCompletedStr(TEXT("StateCompleted"));
+
+		// Start and enter state
+		Status = Exec.Start();
+		AITEST_TRUE("StateTree Task0 should enter state", Exec.Expect(Task0.GetName(), EnterStateStr));
+		Exec.LogClear();
+
+		// Transition from State0 to State1, it should fail (Task1), and the transition on State1->State2 (and not State1A->State3)
+		Status = Exec.Tick(0.1f);
+		AITEST_TRUE("StateTree Task0 should complete", Exec.Expect(Task0.GetName(), StateCompletedStr));
+		AITEST_TRUE("StateTree Task2 should enter state", Exec.Expect(Task2.GetName(), EnterStateStr));
+		AITEST_FALSE("StateTree Task3 should not enter state", Exec.Expect(Task3.GetName(), EnterStateStr));
+		Exec.LogClear();
+
+		return true;
+	}
+};
+IMPLEMENT_AI_INSTANT_TEST(FStateTreeTest_TransitionPriorityEnterState, "System.StateTree.TransitionPriorityEnterState");
+
+UE_ENABLE_OPTIMIZATION_SHIP
+
 #undef LOCTEXT_NAMESPACE
 
