StateTree Fixes and Improvements
- Use hierarchical GamplayTag match when matchin transition event tags
- When returning succeed/failed fomr a subtree, handle it as if the linked state was completed
- Call Stop() before initializing the instance data on Start() to avoid the instance data to be overritten by Stop()
- Early out Stop() if the tree is not running
- Fixed failed transition not setting OutTransition.CurrentActiveStates correctly

diff --git a/Engine/Plugins/Runtime/StateTree/Source/StateTreeModule/Private/StateTreeExecutionContext.cpp b/Engine/Plugins/Runtime/StateTree/Source/StateTreeModule/Private/StateTreeExecutionContext.cpp
--- a/Engine/Plugins/Runtime/StateTree/Source/StateTreeModule/Private/StateTreeExecutionContext.cpp
+++ b/Engine/Plugins/Runtime/StateTree/Source/StateTreeModule/Private/StateTreeExecutionContext.cpp
@@ -161,6 +161,9 @@ EStateTreeRunStatus FStateTreeExecutionContext::Start()
 		return EStateTreeRunStatus::Failed;
 	}
 
+	// Stop if still running previous state.
+	Stop();
+
 	// Initialize instance data if needed.
 	if (!InstanceData.IsValid())
 	{
@@ -177,14 +180,6 @@ EStateTreeRunStatus FStateTreeExecutionContext::Start()
 	const TSharedPtr<FStateTreeInstanceData> SharedInstanceData = StateTree.GetSharedInstanceData();
 	check(SharedInstanceData.IsValid());
 
-	FStateTreeExecutionState* Exec = &GetExecState(); // Using pointer as we will need to reacquire the exec later.
-
-	// Stop if still running previous state.
-	if (Exec->TreeRunStatus == EStateTreeRunStatus::Running)
-	{
-		Stop();
-	}
-
 	// Call TreeStart on evaluators.
 	StartEvaluators();
 
@@ -192,6 +187,7 @@ EStateTreeRunStatus FStateTreeExecutionContext::Start()
 	TickEvaluators(0.0f);
 
 	// Initialize to unset running state.
+	FStateTreeExecutionState* Exec = &GetExecState(); // Using pointer as we will need to reacquire the exec later.
 	Exec->TreeRunStatus = EStateTreeRunStatus::Running;
 	Exec->ActiveStates.Reset();
 	Exec->LastTickStatus = EStateTreeRunStatus::Unset;
@@ -256,9 +252,14 @@ EStateTreeRunStatus FStateTreeExecutionContext::Stop()
 	{
 		return EStateTreeRunStatus::Failed;
 	}
-	
+
 	FStateTreeExecutionState& Exec = GetExecState();
 
+	if (Exec.TreeRunStatus != EStateTreeRunStatus::Running)
+	{
+		return Exec.TreeRunStatus;
+	}
+
 	// Capture events added between ticks.
 	FStateTreeEventQueue& EventQueue = InstanceData.GetEventQueue();
 	EventsToProcess = EventQueue.GetEvents();
@@ -1151,10 +1152,12 @@ bool FStateTreeExecutionContext::TriggerTransitions(FStateTreeInstanceData& Shar
 		
 		return EventsToProcess.ContainsByPredicate([QueriedTag](const FStateTreeEvent& Event)
 		{
-			return Event.Tag == QueriedTag;
+			return Event.Tag.MatchesTag(QueriedTag);
 		});
 	};
 
+	OutTransition.CurrentActiveStates = Exec.ActiveStates;
+
 	// Walk towards root and check all transitions along the way.
 	for (int32 StateIndex = StateStartIndex; StateIndex >= 0; StateIndex--)
 	{
@@ -1194,7 +1197,6 @@ bool FStateTreeExecutionContext::TriggerTransitions(FStateTreeInstanceData& Shar
 				
 				if (Transition.Type == EStateTreeTransitionType::GotoState || Transition.Type == EStateTreeTransitionType::NextState)
 				{
-					OutTransition.CurrentActiveStates = Exec.ActiveStates;
 					OutTransition.TargetState = Transition.State;
 					OutTransition.NextActiveStates.Reset();
 
@@ -1209,20 +1211,26 @@ bool FStateTreeExecutionContext::TriggerTransitions(FStateTreeInstanceData& Shar
 					// NotSet is no-operation, but can be used to mask a transition at parent state. Returning unset keeps updating current state.
 					return false;
 				}
-				else if (Transition.Type == EStateTreeTransitionType::Succeeded)
-				{
-					STATETREE_LOG(Verbose, TEXT("Stop tree execution from state '%s' (%s): Succeeded"), *GetSafeStateName(Exec.ActiveStates.Last()), *State.Name.ToString());
-					OutTransition.CurrentActiveStates = Exec.ActiveStates;
-					OutTransition.TargetState = FStateTreeStateHandle::Succeeded;
-					OutTransition.NextActiveStates = FStateTreeActiveStates(FStateTreeStateHandle::Succeeded);
-					return true;
-				}
 				else
 				{
-					STATETREE_LOG(Verbose, TEXT("Stop tree execution from state '%s' (%s): Failed"), *GetSafeStateName(Exec.ActiveStates.Last()), *State.Name.ToString());
-					OutTransition.CurrentActiveStates = Exec.ActiveStates;
-					OutTransition.TargetState = FStateTreeStateHandle::Failed;
-					OutTransition.NextActiveStates = FStateTreeActiveStates(FStateTreeStateHandle::Failed);
+					// Tree succeeded or failed.
+					// If we are on a linked state, find the first parent linked state and mark that as completed, continue to find completion transitions.
+					const FStateTreeStateHandle ParentLinkedState = GetParentLinkedStateHandle(Exec.ActiveStates, StateIndex);
+					if (ParentLinkedState.IsValid())
+					{
+						STATETREE_LOG(Verbose, TEXT("Compled subtree '%s' from state '%s' (%s): %s"),
+							*GetSafeStateName(ParentLinkedState), *GetSafeStateName(Exec.ActiveStates.Last()), *State.Name.ToString(), *UEnum::GetDisplayValueAsText(Transition.Type).ToString());
+						Exec.CompletedStateHandle = ParentLinkedState;
+						Exec.LastTickStatus = Transition.Type == EStateTreeTransitionType::Succeeded ? EStateTreeRunStatus::Succeeded : EStateTreeRunStatus::Failed;
+						return TriggerTransitions(SharedInstanceData, OutTransition);
+					}
+					
+					STATETREE_LOG(Verbose, TEXT("Stop tree execution from state '%s' (%s): %s"),
+						*GetSafeStateName(Exec.ActiveStates.Last()), *State.Name.ToString(), *UEnum::GetDisplayValueAsText(Transition.Type).ToString());
+
+					const FStateTreeStateHandle NextState = Transition.Type == EStateTreeTransitionType::Succeeded ? FStateTreeStateHandle::Succeeded : FStateTreeStateHandle::Failed; 
+					OutTransition.TargetState = NextState;
+					OutTransition.NextActiveStates = FStateTreeActiveStates(NextState);
 					return true;
 				}
 			}
@@ -1247,6 +1255,21 @@ bool FStateTreeExecutionContext::TriggerTransitions(FStateTreeInstanceData& Shar
 	return false;
 }
 
+FStateTreeStateHandle FStateTreeExecutionContext::GetParentLinkedStateHandle(const FStateTreeActiveStates& ActiveStates, const int32 StartStartIndex) const
+{
+	check(ActiveStates.IsValidIndex(StartStartIndex));
+	for (int32 StateIndex = StartStartIndex; StateIndex >= 0; StateIndex--)
+	{
+		const FCompactStateTreeState& State = StateTree.States[ActiveStates[StateIndex].Index];
+		if (State.LinkedState.IsValid())
+		{
+			return ActiveStates[StateIndex];
+		}
+	}
+
+	return FStateTreeStateHandle();
+}	
+
 bool FStateTreeExecutionContext::SelectState(FStateTreeInstanceData& SharedInstanceData, const FStateTreeStateHandle NextState, FStateTreeActiveStates& OutNewActiveState)
 {
 	const FStateTreeExecutionState& Exec = GetExecState();
diff --git a/Engine/Plugins/Runtime/StateTree/Source/StateTreeModule/Public/StateTreeExecutionContext.h b/Engine/Plugins/Runtime/StateTree/Source/StateTreeModule/Public/StateTreeExecutionContext.h
--- a/Engine/Plugins/Runtime/StateTree/Source/StateTreeModule/Public/StateTreeExecutionContext.h
+++ b/Engine/Plugins/Runtime/StateTree/Source/StateTreeModule/Public/StateTreeExecutionContext.h
@@ -325,6 +325,12 @@ protected:
 	 */
 	bool TriggerTransitions(FStateTreeInstanceData& SharedInstanceData, FStateTreeTransitionResult& OutTransition);
 
+	/**
+	 * Traverses the ActiveStates from StartStateIndex to 0 and returns first linked state.
+	 * @return Parent linked state, or invalid state if no linked state found. 
+	 */
+	FStateTreeStateHandle GetParentLinkedStateHandle(const FStateTreeActiveStates& ActiveStates, const int32 StartStateIndex) const;
+
 	/**
 	 * Runs state selection logic starting at the specified state, walking towards the leaf states.
 	 * If a state cannot be selected, false is returned. 
diff --git a/Engine/Plugins/Runtime/StateTree/Source/StateTreeTestSuite/Private/StateTreeTest.cpp b/Engine/Plugins/Runtime/StateTree/Source/StateTreeTestSuite/Private/StateTreeTest.cpp
--- a/Engine/Plugins/Runtime/StateTree/Source/StateTreeTestSuite/Private/StateTreeTest.cpp
+++ b/Engine/Plugins/Runtime/StateTree/Source/StateTreeTestSuite/Private/StateTreeTest.cpp
@@ -278,6 +278,59 @@ struct FStateTreeTest_FailEnterState : FAITestBase
 };
 IMPLEMENT_AI_INSTANT_TEST(FStateTreeTest_FailEnterState, "System.StateTree.FailEnterState");
 
+
+struct FStateTreeTest_Restart : FAITestBase
+{
+	virtual bool InstantTest() override
+	{
+		UStateTree& StateTree = UE::StateTree::Tests::NewStateTree(&GetWorld());
+		UStateTreeEditorData& EditorData = *Cast<UStateTreeEditorData>(StateTree.EditorData);
+		
+		UStateTreeState& Root = EditorData.AddSubTree(FName(TEXT("Root")));
+		UStateTreeState& State1 = Root.AddChildState(FName(TEXT("State1")));
+
+		auto& Task1 = State1.AddTask<FTestTask_Stand>(FName(TEXT("Task1")));
+		Task1.GetNode().TicksToCompletion = 2;
+
+		FStateTreeCompilerLog Log;
+		FStateTreeCompiler Compiler(Log);
+		const bool bResult = Compiler.Compile(StateTree);
+		AITEST_TRUE("StateTree should get compiled", bResult);
+
+		EStateTreeRunStatus Status = EStateTreeRunStatus::Unset;
+		FStateTreeInstanceData InstanceData;
+		FTestStateTreeExecutionContext Exec(StateTree, StateTree, InstanceData);
+		const bool bInitSucceeded = Exec.IsValid();
+		AITEST_TRUE("StateTree should init", bInitSucceeded);
+
+		const FString TickStr(TEXT("Tick"));
+		const FString EnterStateStr(TEXT("EnterState"));
+		const FString ExitStateStr(TEXT("ExitState"));
+		const FString StateCompletedStr(TEXT("StateCompleted"));
+
+		// Start and enter state
+		Status = Exec.Start();
+		AITEST_TRUE("StateTree Task1 should enter state", Exec.Expect(Task1.GetName(), EnterStateStr));
+		AITEST_TRUE("StateTree exec status should be running", Exec.GetLastTickStatus() == EStateTreeRunStatus::Running);
+		Exec.LogClear();
+
+		// Tick
+		Status = Exec.Tick(0.1f);
+		AITEST_TRUE("StateTree exec status should be running", Exec.GetLastTickStatus() == EStateTreeRunStatus::Running);
+		Exec.LogClear();
+
+		// Call Start again, should stop and start the tree.
+		Status = Exec.Start();
+		AITEST_TRUE("StateTree Task1 should exit state", Exec.Expect(Task1.GetName(), ExitStateStr));
+		AITEST_TRUE("StateTree Task1 should enter state", Exec.Expect(Task1.GetName(), EnterStateStr));
+		AITEST_TRUE("StateTree exec status should be running", Exec.GetLastTickStatus() == EStateTreeRunStatus::Running);
+		Exec.LogClear();
+
+		return true;
+	}
+};
+IMPLEMENT_AI_INSTANT_TEST(FStateTreeTest_Restart, "System.StateTree.Restart");
+
 struct FStateTreeTest_SubTree : FAITestBase
 {
 	virtual bool InstantTest() override
@@ -294,13 +347,16 @@ struct FStateTreeTest_SubTree : FAITestBase
 
 		State1.LinkedSubtree.Set(&State3);
 
+		State1.AddTransition(EStateTreeTransitionTrigger::OnStateCompleted, EStateTreeTransitionType::GotoState, &State2);
+
 		auto& Task2 = State2.AddTask<FTestTask_Stand>(FName(TEXT("Task2")));
+		State2.AddTransition(EStateTreeTransitionTrigger::OnStateCompleted, EStateTreeTransitionType::Succeeded);
 
 		auto& Task3A = State3A.AddTask<FTestTask_Stand>(FName(TEXT("Task3A")));
 		State3A.AddTransition(EStateTreeTransitionTrigger::OnStateCompleted, EStateTreeTransitionType::GotoState, &State3B);
 
 		auto& Task3B = State3B.AddTask<FTestTask_Stand>(FName(TEXT("Task3B")));
-
+		State3B.AddTransition(EStateTreeTransitionTrigger::OnStateCompleted, EStateTreeTransitionType::Succeeded);
 
 		FStateTreeCompilerLog Log;
 		FStateTreeCompiler Compiler(Log);
@@ -327,10 +383,24 @@ struct FStateTreeTest_SubTree : FAITestBase
 		AITEST_TRUE("StateTree should be running", Status == EStateTreeRunStatus::Running);
 		Exec.LogClear();
 
+		// Transition within subtree
 		Status = Exec.Tick(0.1f);
 		AITEST_TRUE("StateTree Active States should be in Root/State1/State3/State3B", Exec.ExpectInActiveStates(Root.Name, State1.Name, State3.Name, State3B.Name));
 		AITEST_TRUE("StateTree Task3B should enter state", Exec.Expect(Task3B.GetName(), EnterStateStr));
 		AITEST_TRUE("StateTree should be running", Status == EStateTreeRunStatus::Running);
+		Exec.LogClear();
+
+		// Complete subtree
+		Status = Exec.Tick(0.1f);
+		AITEST_TRUE("StateTree Active States should be in Root/State2", Exec.ExpectInActiveStates(Root.Name, State2.Name));
+		AITEST_TRUE("StateTree Task2 should enter state", Exec.Expect(Task2.GetName(), EnterStateStr));
+		AITEST_TRUE("StateTree should be running", Status == EStateTreeRunStatus::Running);
+		Exec.LogClear();
+
+		// Complete the whole tree
+		Status = Exec.Tick(0.1f);
+		AITEST_TRUE("StateTree should complete in succeeded", Status == EStateTreeRunStatus::Succeeded);
+		Exec.LogClear();
 		
 		return true;
 	}
