PR #9828: [StateTree] StateTreeEditor compile button with save options (Contributed by [REDACTED_PERSON])

diff --git a/Engine/Plugins/Runtime/StateTree/Source/StateTreeEditorModule/Private/StateTreeEditor.cpp b/Engine/Plugins/Runtime/StateTree/Source/StateTreeEditorModule/Private/StateTreeEditor.cpp
--- a/Engine/Plugins/Runtime/StateTree/Source/StateTreeEditorModule/Private/StateTreeEditor.cpp
+++ b/Engine/Plugins/Runtime/StateTree/Source/StateTreeEditorModule/Private/StateTreeEditor.cpp
@@ -16,6 +16,7 @@
 #include "StateTreeEditorCommands.h"
 #include "StateTreeEditorData.h"
 #include "StateTreeEditorModule.h"
+#include "StateTreeEditorSettings.h"
 #include "StateTreeObjectHash.h"
 #include "StateTreeState.h"
 #include "StateTreeTaskBase.h"
@@ -25,6 +26,7 @@
 #include "ToolMenus.h"
 #include "Widgets/Docking/SDockTab.h"
 #include "Widgets/Input/SMultiLineEditableTextBox.h"
+#include "FileHelpers.h"
 
 #define LOCTEXT_NAMESPACE "StateTreeEditor"
 
@@ -674,7 +676,52 @@ namespace UE::StateTree::Editor::Internal
 		});
 	}
 
-}
+	static void MakeSaveOnCompileSubMenu(UToolMenu* InMenu)
+	{
+		FToolMenuSection& Section = InMenu->AddSection("Section");
+		const FStateTreeEditorCommands& Commands = FStateTreeEditorCommands::Get();
+		Section.AddMenuEntry(Commands.SaveOnCompile_Never);
+		Section.AddMenuEntry(Commands.SaveOnCompile_SuccessOnly);
+		Section.AddMenuEntry(Commands.SaveOnCompile_Always);
+	}
+
+	static void GenerateCompileOptionsMenu(UToolMenu* InMenu)
+	{
+		FToolMenuSection& Section = InMenu->AddSection("Section");
+		const FStateTreeEditorCommands& Commands = FStateTreeEditorCommands::Get();
+
+		// @TODO: disable the menu and change up the tooltip when all sub items are disabled
+		Section.AddSubMenu(
+			"SaveOnCompile",
+			LOCTEXT("SaveOnCompileSubMenu", "Save on Compile"),
+			LOCTEXT("SaveOnCompileSubMenu_ToolTip", "Determines how the StateTree is saved whenever you compile it."),
+			FNewToolMenuDelegate::CreateStatic(&MakeSaveOnCompileSubMenu));
+	}
+
+	static void SetSaveOnCompileSetting(const EStateTreeSaveOnCompile NewSetting)
+	{
+		UStateTreeEditorSettings* Settings = GetMutableDefault<UStateTreeEditorSettings>();
+		Settings->SaveOnCompile = NewSetting;
+		Settings->SaveConfig();
+	}
+
+	static bool IsSaveOnCompileOptionSet(TWeakPtr<FStateTreeEditor> Editor, const EStateTreeSaveOnCompile Option)
+	{
+		const UStateTreeEditorSettings* Settings = GetDefault<UStateTreeEditorSettings>();
+
+		EStateTreeSaveOnCompile CurrentSetting = Settings->SaveOnCompile;
+		if (!Editor.IsValid() || !Editor.Pin()->IsSaveOnCompileEnabled())
+		{
+			// If save-on-compile is disabled for the StateTree, then we want to 
+			// show "Never" as being selected
+			// 
+			// @TODO: a tooltip explaining why would be nice too
+			CurrentSetting = EStateTreeSaveOnCompile::Never;
+		}
+
+		return (CurrentSetting == Option);
+	}
+} // UE::StateTree::Editor::Internal
 
 namespace UE::StateTree::Editor
 {
@@ -699,7 +746,7 @@ namespace UE::StateTree::Editor
 		return EditorDataHash;
 	}
 
-};
+}; // UE::StateTree::Editor
 
 void FStateTreeEditor::BindCommands()
 {
@@ -709,6 +756,31 @@ void FStateTreeEditor::BindCommands()
 		Commands.Compile,
 		FExecuteAction::CreateSP(this, &FStateTreeEditor::Compile),
 		FCanExecuteAction::CreateSP(this, &FStateTreeEditor::CanCompile));
+
+	TWeakPtr<FStateTreeEditor> WeakThisPtr = SharedThis(this);
+	ToolkitCommands->MapAction(
+		FStateTreeEditorCommands::Get().SaveOnCompile_Never,
+		FExecuteAction::CreateStatic(&UE::StateTree::Editor::Internal::SetSaveOnCompileSetting, EStateTreeSaveOnCompile::Never),
+		FCanExecuteAction::CreateSP(this, &FStateTreeEditor::IsSaveOnCompileEnabled),
+		FIsActionChecked::CreateStatic(&UE::StateTree::Editor::Internal::IsSaveOnCompileOptionSet, WeakThisPtr, EStateTreeSaveOnCompile::Never)
+	);
+	ToolkitCommands->MapAction(
+		FStateTreeEditorCommands::Get().SaveOnCompile_SuccessOnly,
+		FExecuteAction::CreateStatic(&UE::StateTree::Editor::Internal::SetSaveOnCompileSetting, EStateTreeSaveOnCompile::SuccessOnly),
+		FCanExecuteAction::CreateSP(this, &FStateTreeEditor::IsSaveOnCompileEnabled),
+		FIsActionChecked::CreateStatic(&UE::StateTree::Editor::Internal::IsSaveOnCompileOptionSet, WeakThisPtr, EStateTreeSaveOnCompile::SuccessOnly)
+	);
+	ToolkitCommands->MapAction(
+		FStateTreeEditorCommands::Get().SaveOnCompile_Always,
+		FExecuteAction::CreateStatic(&UE::StateTree::Editor::Internal::SetSaveOnCompileSetting, EStateTreeSaveOnCompile::Always),
+		FCanExecuteAction::CreateSP(this, &FStateTreeEditor::IsSaveOnCompileEnabled),
+		FIsActionChecked::CreateStatic(&UE::StateTree::Editor::Internal::IsSaveOnCompileOptionSet, WeakThisPtr, EStateTreeSaveOnCompile::Always)
+	);
+}
+
+bool FStateTreeEditor::IsSaveOnCompileEnabled() const
+{
+	return true;
 }
 
 void FStateTreeEditor::RegisterToolbar()
@@ -740,14 +812,21 @@ void FStateTreeEditor::RegisterToolbar()
 			{
 				const FStateTreeEditorCommands& Commands = FStateTreeEditorCommands::Get();
 
-				FToolMenuEntry CompileButton = FToolMenuEntry::InitToolBarButton(
+				FToolMenuEntry& CompileButton = InSection.AddEntry(FToolMenuEntry::InitToolBarButton(
 					Commands.Compile,
 					TAttribute<FText>(),
 					TAttribute<FText>(),
-					TAttribute<FSlateIcon>(StateTreeEditor.ToSharedRef(), &FStateTreeEditor::GetCompileStatusImage)
-					);
-
-				InSection.AddEntry(CompileButton);
+					TAttribute<FSlateIcon>(StateTreeEditor.ToSharedRef(), &FStateTreeEditor::GetCompileStatusImage)));
+				CompileButton.StyleNameOverride = "CalloutToolbar";
+
+				FToolMenuEntry& CompileOptions = InSection.AddEntry(FToolMenuEntry::InitComboButton(
+					"CompileComboButton",
+					FUIAction(),
+					FNewToolMenuDelegate::CreateStatic(&UE::StateTree::Editor::Internal::GenerateCompileOptionsMenu),
+					LOCTEXT("CompileOptions_ToolbarTooltip", "Options to customize how State Trees compile")
+				));
+				CompileOptions.StyleNameOverride = "CalloutToolbar";
+				CompileOptions.ToolBarData.bSimpleComboBox = true;
 			}
 		}
 	}));
@@ -795,6 +874,16 @@ void FStateTreeEditor::Compile()
 		// Show log
 		TabManager->TryInvokeTab(CompilerResultsTabId);
 	}
+
+	const UStateTreeEditorSettings* Settings = GetMutableDefault<UStateTreeEditorSettings>();
+	const bool bShouldSaveOnCompile = ((Settings->SaveOnCompile == EStateTreeSaveOnCompile::Always)
+									|| ((Settings->SaveOnCompile == EStateTreeSaveOnCompile::SuccessOnly) && bLastCompileSucceeded));
+
+	if (bShouldSaveOnCompile)
+	{
+		const TArray<UPackage*> PackagesToSave { StateTree->GetOutermost() };
+		FEditorFileUtils::PromptForCheckoutAndSave(PackagesToSave, /*bCheckDirty =*/true, /*bPromptToSave =*/false);
+	}
 }
 
 bool FStateTreeEditor::CanCompile() const
@@ -852,5 +941,4 @@ void FStateTreeEditor::UpdateAsset()
 	EditorDataHash = UE::StateTree::Editor::CalcAssetHash(*StateTree);
 }
 
-
 #undef LOCTEXT_NAMESPACE
diff --git a/Engine/Plugins/Runtime/StateTree/Source/StateTreeEditorModule/Private/StateTreeEditor.h b/Engine/Plugins/Runtime/StateTree/Source/StateTreeEditorModule/Private/StateTreeEditor.h
--- a/Engine/Plugins/Runtime/StateTree/Source/StateTreeEditorModule/Private/StateTreeEditor.h
+++ b/Engine/Plugins/Runtime/StateTree/Source/StateTreeEditorModule/Private/StateTreeEditor.h
@@ -33,6 +33,8 @@ public:
 
 	void InitEditor(const EToolkitMode::Type Mode, const TSharedPtr< class IToolkitHost >& InitToolkitHost, UStateTree* StateTree);
 
+	bool IsSaveOnCompileEnabled() const;
+
 	//~ Begin IToolkit Interface
 	virtual FName GetToolkitFName() const override;
 	virtual FText GetBaseToolkitName() const override;
diff --git a/Engine/Plugins/Runtime/StateTree/Source/StateTreeEditorModule/Private/StateTreeEditorCommands.cpp b/Engine/Plugins/Runtime/StateTree/Source/StateTreeEditorModule/Private/StateTreeEditorCommands.cpp
--- a/Engine/Plugins/Runtime/StateTree/Source/StateTreeEditorModule/Private/StateTreeEditorCommands.cpp
+++ b/Engine/Plugins/Runtime/StateTree/Source/StateTreeEditorModule/Private/StateTreeEditorCommands.cpp
@@ -18,6 +18,10 @@ FStateTreeEditorCommands::FStateTreeEditorCommands()
 void FStateTreeEditorCommands::RegisterCommands()
 {
 	UI_COMMAND(Compile, "Compile", "Compile the current StateTree.", EUserInterfaceActionType::Button, FInputChord(EKeys::F7));
+
+	UI_COMMAND(SaveOnCompile_Never, "Never", "Sets the save-on-compile option to 'Never', meaning that your StateTree will not be saved when they are compiled", EUserInterfaceActionType::RadioButton, FInputChord());
+	UI_COMMAND(SaveOnCompile_SuccessOnly, "On Success Only", "Sets the save-on-compile option to 'Success Only', meaning that your StateTree will be saved whenever they are successfully compiled", EUserInterfaceActionType::RadioButton, FInputChord());
+	UI_COMMAND(SaveOnCompile_Always, "Always", "Sets the save-on-compile option to 'Always', meaning that your StateTree will be saved whenever they are compiled (even if there were errors)", EUserInterfaceActionType::RadioButton, FInputChord());
 }
 
 #undef LOCTEXT_NAMESPACE
diff --git a/Engine/Plugins/Runtime/StateTree/Source/StateTreeEditorModule/Private/StateTreeEditorCommands.h b/Engine/Plugins/Runtime/StateTree/Source/StateTreeEditorModule/Private/StateTreeEditorCommands.h
--- a/Engine/Plugins/Runtime/StateTree/Source/StateTreeEditorModule/Private/StateTreeEditorCommands.h
+++ b/Engine/Plugins/Runtime/StateTree/Source/StateTreeEditorModule/Private/StateTreeEditorCommands.h
@@ -1,4 +1,4 @@
-﻿// Copyright Epic Games, Inc. All Rights Reserved.
+// Copyright Epic Games, Inc. All Rights Reserved.
 
 #pragma once
 
@@ -17,4 +17,7 @@ public:
 	virtual void RegisterCommands() override;
 
 	TSharedPtr<FUICommandInfo> Compile;
+	TSharedPtr<FUICommandInfo> SaveOnCompile_Never;
+	TSharedPtr<FUICommandInfo> SaveOnCompile_SuccessOnly;
+	TSharedPtr<FUICommandInfo> SaveOnCompile_Always;
 };
\ No newline at end of file
diff --git a/Engine/Plugins/Runtime/StateTree/Source/StateTreeEditorModule/Private/StateTreeEditorSettings.cpp b/Engine/Plugins/Runtime/StateTree/Source/StateTreeEditorModule/Private/StateTreeEditorSettings.cpp
new file mode 100644
--- /dev/null
+++ b/Engine/Plugins/Runtime/StateTree/Source/StateTreeEditorModule/Private/StateTreeEditorSettings.cpp
@@ -0,0 +1,8 @@
+// Copyright Epic Games, Inc. All Rights Reserved.
+
+#include "StateTreeEditorSettings.h"
+
+UStateTreeEditorSettings::UStateTreeEditorSettings(const FObjectInitializer& ObjectInitializer)
+	: Super(ObjectInitializer)
+{
+}
diff --git a/Engine/Plugins/Runtime/StateTree/Source/StateTreeEditorModule/Public/StateTreeEditorSettings.h b/Engine/Plugins/Runtime/StateTree/Source/StateTreeEditorModule/Public/StateTreeEditorSettings.h
new file mode 100644
--- /dev/null
+++ b/Engine/Plugins/Runtime/StateTree/Source/StateTreeEditorModule/Public/StateTreeEditorSettings.h
@@ -0,0 +1,25 @@
+// Copyright Epic Games, Inc. All Rights Reserved.
+
+#pragma once
+
+#include "CoreMinimal.h"
+#include "Engine/DeveloperSettings.h"
+#include "StateTreeEditorSettings.generated.h"
+
+UENUM()
+enum class EStateTreeSaveOnCompile : uint8
+{
+	Never UMETA(DisplayName = "Never"),
+	SuccessOnly UMETA(DisplayName = "On Success Only"),
+	Always UMETA(DisplayName = "Always"),
+};
+
+UCLASS(config = EditorPerProjectUserSettings)
+class STATETREEEDITORMODULE_API UStateTreeEditorSettings : public UDeveloperSettings
+{
+	GENERATED_UCLASS_BODY()
+public:
+	/** Determines when to save StateTrees post-compile */
+	UPROPERTY(EditAnywhere, config, Category = "Compiler")
+	EStateTreeSaveOnCompile SaveOnCompile = EStateTreeSaveOnCompile::Never;
+};
diff --git a/Engine/Plugins/Runtime/StateTree/Source/StateTreeEditorModule/StateTreeEditorModule.Build.cs b/Engine/Plugins/Runtime/StateTree/Source/StateTreeEditorModule/StateTreeEditorModule.Build.cs
--- a/Engine/Plugins/Runtime/StateTree/Source/StateTreeEditorModule/StateTreeEditorModule.Build.cs
+++ b/Engine/Plugins/Runtime/StateTree/Source/StateTreeEditorModule/StateTreeEditorModule.Build.cs
@@ -45,6 +45,7 @@ namespace UnrealBuildTool.Rules
 				"ToolMenus",
 				"ToolWidgets",
 				"ApplicationCore",
+				"DeveloperSettings"
 			}
 			);
 
