IncrementalCooking - Add DeclareConstructClasses to types that need it.
#rnx

diff --git a/Engine/Plugins/Animation/IKRig/Source/IKRig/Private/Retargeter/IKRetargeter.cpp b/Engine/Plugins/Animation/IKRig/Source/IKRig/Private/Retargeter/IKRetargeter.cpp
--- a/Engine/Plugins/Animation/IKRig/Source/IKRig/Private/Retargeter/IKRetargeter.cpp
+++ b/Engine/Plugins/Animation/IKRig/Source/IKRig/Private/Retargeter/IKRetargeter.cpp
@@ -123,6 +123,14 @@ void UIKRetargeter::PostLoad()
 	CleanAndInitialize();
 }
 
+#if WITH_EDITORONLY_DATA
+void UIKRetargeter::DeclareConstructClasses(TArray<FTopLevelAssetPath>& OutConstructClasses, const UClass* SpecificSubclass)
+{
+	Super::DeclareConstructClasses(OutConstructClasses, SpecificSubclass);
+	OutConstructClasses.Add(FTopLevelAssetPath(URetargetChainSettings::StaticClass()));
+}
+#endif
+
 void UIKRetargeter::CleanAndInitialize()
 {
 	// remove null settings
diff --git a/Engine/Plugins/Animation/IKRig/Source/IKRig/Public/Retargeter/IKRetargeter.h b/Engine/Plugins/Animation/IKRig/Source/IKRig/Public/Retargeter/IKRetargeter.h
--- a/Engine/Plugins/Animation/IKRig/Source/IKRig/Public/Retargeter/IKRetargeter.h
+++ b/Engine/Plugins/Animation/IKRig/Source/IKRig/Public/Retargeter/IKRetargeter.h
@@ -334,6 +334,10 @@ public:
 	/** UObject */
 	virtual void Serialize(FArchive& Ar) override;
 	virtual void PostLoad() override;
+#if WITH_EDITORONLY_DATA
+	static void DeclareConstructClasses(TArray<FTopLevelAssetPath>& OutConstructClasses, const UClass* SpecificSubclass);
+#endif
+
 	/** END UObject */
 
 #if WITH_EDITOR
diff --git a/Engine/Plugins/Experimental/PCG/Source/PCG/Private/PCGGraph.cpp b/Engine/Plugins/Experimental/PCG/Source/PCG/Private/PCGGraph.cpp
--- a/Engine/Plugins/Experimental/PCG/Source/PCG/Private/PCGGraph.cpp
+++ b/Engine/Plugins/Experimental/PCG/Source/PCG/Private/PCGGraph.cpp
@@ -131,6 +131,14 @@ void UPCGGraph::PostLoad()
 #endif
 }
 
+#if WITH_EDITORONLY_DATA
+void UPCGGraph::DeclareConstructClasses(TArray<FTopLevelAssetPath>& OutConstructClasses, const UClass* SpecificSubclass)
+{
+	Super::DeclareConstructClasses(OutConstructClasses, SpecificSubclass);
+	OutConstructClasses.Add(FTopLevelAssetPath(UPCGPin::StaticClass()));
+}
+#endif
+
 void UPCGGraph::BeginDestroy()
 {
 #if WITH_EDITOR
diff --git a/Engine/Plugins/Experimental/PCG/Source/PCG/Public/PCGGraph.h b/Engine/Plugins/Experimental/PCG/Source/PCG/Public/PCGGraph.h
--- a/Engine/Plugins/Experimental/PCG/Source/PCG/Public/PCGGraph.h
+++ b/Engine/Plugins/Experimental/PCG/Source/PCG/Public/PCGGraph.h
@@ -26,6 +26,10 @@ public:
 	UPCGGraph(const FObjectInitializer& ObjectInitializer);
 	/** ~Begin UObject interface */
 	virtual void PostLoad() override;
+#if WITH_EDITORONLY_DATA
+	static void DeclareConstructClasses(TArray<FTopLevelAssetPath>& OutConstructClasses, const UClass* SpecificSubclass);
+#endif
+
 	virtual void BeginDestroy() override;
 
 #if WITH_EDITOR
diff --git a/Engine/Plugins/Experimental/Water/Source/Runtime/Private/WaterZoneActor.cpp b/Engine/Plugins/Experimental/Water/Source/Runtime/Private/WaterZoneActor.cpp
--- a/Engine/Plugins/Experimental/Water/Source/Runtime/Private/WaterZoneActor.cpp
+++ b/Engine/Plugins/Experimental/Water/Source/Runtime/Private/WaterZoneActor.cpp
@@ -114,6 +114,14 @@ void AWaterZone::PostLoad()
 #endif // WITH_EDITORONLY_DATA
 }
 
+#if WITH_EDITORONLY_DATA
+void AWaterZone::DeclareConstructClasses(TArray<FTopLevelAssetPath>& OutConstructClasses, const UClass* SpecificSubclass)
+{
+	Super::DeclareConstructClasses(OutConstructClasses, SpecificSubclass);
+	OutConstructClasses.Add(FTopLevelAssetPath(UBoxComponent::StaticClass()));
+}
+#endif
+
 void AWaterZone::PostRegisterAllComponents()
 {
 	Super::PostRegisterAllComponents();
diff --git a/Engine/Plugins/Experimental/Water/Source/Runtime/Public/WaterZoneActor.h b/Engine/Plugins/Experimental/Water/Source/Runtime/Public/WaterZoneActor.h
--- a/Engine/Plugins/Experimental/Water/Source/Runtime/Public/WaterZoneActor.h
+++ b/Engine/Plugins/Experimental/Water/Source/Runtime/Public/WaterZoneActor.h
@@ -52,6 +52,9 @@ public:
 	virtual void BeginPlay() override;
 	virtual void PostLoadSubobjects(FObjectInstancingGraph* OuterInstanceGraph) override;
 	virtual void PostLoad() override;
+#if WITH_EDITORONLY_DATA
+	static void DeclareConstructClasses(TArray<FTopLevelAssetPath>& OutConstructClasses, const UClass* SpecificSubclass);
+#endif
 	virtual void PostRegisterAllComponents() override;
 
 	FVector2f GetWaterHeightExtents() const { return WaterHeightExtents; }
diff --git a/Engine/Plugins/Runtime/StateTree/Source/StateTreeModule/Private/StateTree.cpp b/Engine/Plugins/Runtime/StateTree/Source/StateTreeModule/Private/StateTree.cpp
--- a/Engine/Plugins/Runtime/StateTree/Source/StateTreeModule/Private/StateTree.cpp
+++ b/Engine/Plugins/Runtime/StateTree/Source/StateTreeModule/Private/StateTree.cpp
@@ -3,6 +3,7 @@
 #include "StateTree.h"
 #include "StateTreeLinker.h"
 #include "StateTreeNodeBase.h"
+#include "StateTreeSchema.h"
 #include "StateTreeTaskBase.h"
 #include "StateTreeEvaluatorBase.h"
 #include "AssetRegistry/AssetData.h"
@@ -161,6 +162,22 @@ void UStateTree::PostLoad()
 	}
 }
 
+#if WITH_EDITORONLY_DATA
+void UStateTree::DeclareConstructClasses(TArray<FTopLevelAssetPath>& OutConstructClasses, const UClass* SpecificSubclass)
+{
+	Super::DeclareConstructClasses(OutConstructClasses, SpecificSubclass);
+	TArray<UClass*> SchemaClasses;
+	GetDerivedClasses(UStateTreeSchema::StaticClass(), SchemaClasses);
+	for (UClass* SchemaClass : SchemaClasses)
+	{
+		if (!SchemaClass->HasAnyClassFlags(CLASS_Abstract | CLASS_Transient))
+		{
+			OutConstructClasses.Add(FTopLevelAssetPath(SchemaClass));
+		}
+	}
+}
+#endif
+
 void UStateTree::Serialize(FStructuredArchiveRecord Record)
 {
 	Super::Serialize(Record);
diff --git a/Engine/Plugins/Runtime/StateTree/Source/StateTreeModule/Public/StateTree.h b/Engine/Plugins/Runtime/StateTree/Source/StateTreeModule/Public/StateTree.h
--- a/Engine/Plugins/Runtime/StateTree/Source/StateTreeModule/Public/StateTree.h
+++ b/Engine/Plugins/Runtime/StateTree/Source/StateTreeModule/Public/StateTree.h
@@ -128,6 +128,9 @@ protected:
 	[[nodiscard]] bool Link();
 
 	virtual void PostLoad() override;
+#if WITH_EDITORONLY_DATA
+	static void DeclareConstructClasses(TArray<FTopLevelAssetPath>& OutConstructClasses, const UClass* SpecificSubclass);
+#endif
 	virtual void Serialize(FStructuredArchiveRecord Record) override;
 	
 #if WITH_EDITOR
diff --git a/Engine/Source/Runtime/Engine/Classes/Components/ShapeComponent.h b/Engine/Source/Runtime/Engine/Classes/Components/ShapeComponent.h
--- a/Engine/Source/Runtime/Engine/Classes/Components/ShapeComponent.h
+++ b/Engine/Source/Runtime/Engine/Classes/Components/ShapeComponent.h
@@ -110,6 +110,9 @@ public:
 
 	//~ Begin UObject Interface.
 	virtual void Serialize(FArchive& Ar) override;
+#if WITH_EDITORONLY_DATA
+	static void DeclareConstructClasses(TArray<FTopLevelAssetPath>& OutConstructClasses, const UClass* SpecificSubclass);
+#endif
 
 #if WITH_EDITOR
 	virtual void PostEditChangeProperty(FPropertyChangedEvent& PropertyChangedEvent) override;
diff --git a/Engine/Source/Runtime/Engine/Private/Components/ShapeComponent.cpp b/Engine/Source/Runtime/Engine/Private/Components/ShapeComponent.cpp
--- a/Engine/Source/Runtime/Engine/Private/Components/ShapeComponent.cpp
+++ b/Engine/Source/Runtime/Engine/Private/Components/ShapeComponent.cpp
@@ -122,6 +122,14 @@ void UShapeComponent::Serialize(FArchive& Ar)
 #endif // WITH_EDITOR
 }
 
+#if WITH_EDITORONLY_DATA
+void UShapeComponent::DeclareConstructClasses(TArray<FTopLevelAssetPath>& OutConstructClasses, const UClass* SpecificSubclass)
+{
+	Super::DeclareConstructClasses(OutConstructClasses, SpecificSubclass);
+	OutConstructClasses.Add(FTopLevelAssetPath(UBodySetup::StaticClass()));
+}
+#endif
+
 #if WITH_EDITOR
 void UShapeComponent::PostEditChangeProperty(FPropertyChangedEvent& PropertyChangedEvent)
 {
diff --git a/Engine/Source/Runtime/Engine/Private/WorldPartition/HLOD/HLODLayer.cpp b/Engine/Source/Runtime/Engine/Private/WorldPartition/HLOD/HLODLayer.cpp
--- a/Engine/Source/Runtime/Engine/Private/WorldPartition/HLOD/HLODLayer.cpp
+++ b/Engine/Source/Runtime/Engine/Private/WorldPartition/HLOD/HLODLayer.cpp
@@ -174,6 +174,17 @@ void UHLODLayer::PostLoad()
 	}
 }
 
+#if WITH_EDITORONLY_DATA
+void UHLODLayer::DeclareConstructClasses(TArray<FTopLevelAssetPath>& OutConstructClasses, const UClass* SpecificSubclass)
+{
+	Super::DeclareConstructClasses(OutConstructClasses, SpecificSubclass);
+	OutConstructClasses.Add(FTopLevelAssetPath(TEXT("/Script/WorldPartitionHLODUtilities.HLODBuilderInstancingSettings")));
+	OutConstructClasses.Add(FTopLevelAssetPath(TEXT("/Script/WorldPartitionHLODUtilities.HLODBuilderMeshMerge")));
+	OutConstructClasses.Add(FTopLevelAssetPath(TEXT("/Script/WorldPartitionHLODUtilities.HLODBuilderMeshSimplify")));
+	OutConstructClasses.Add(FTopLevelAssetPath(TEXT("/Script/WorldPartitionHLODUtilities.HLODBuilderMeshApproximate")));
+}
+#endif
+
 void UHLODLayer::PostEditChangeProperty(FPropertyChangedEvent& PropertyChangedEvent)
 {
 	Super::PostEditChangeProperty(PropertyChangedEvent);
diff --git a/Engine/Source/Runtime/Engine/Public/WorldPartition/HLOD/HLODLayer.h b/Engine/Source/Runtime/Engine/Public/WorldPartition/HLOD/HLODLayer.h
--- a/Engine/Source/Runtime/Engine/Public/WorldPartition/HLOD/HLODLayer.h
+++ b/Engine/Source/Runtime/Engine/Public/WorldPartition/HLOD/HLODLayer.h
@@ -71,6 +71,9 @@ public:
 private:
 	//~ Begin UObject Interface.
 	virtual void PostLoad() override;
+#if WITH_EDITORONLY_DATA
+	static void DeclareConstructClasses(TArray<FTopLevelAssetPath>& OutConstructClasses, const UClass* SpecificSubclass);
+#endif
 	virtual void PostEditChangeProperty(FPropertyChangedEvent& PropertyChangedEvent) override;
 	//~ End UObject Interface.
 #endif
diff --git a/Engine/Source/Runtime/Foliage/Private/InstancedFoliage.cpp b/Engine/Source/Runtime/Foliage/Private/InstancedFoliage.cpp
--- a/Engine/Source/Runtime/Foliage/Private/InstancedFoliage.cpp
+++ b/Engine/Source/Runtime/Foliage/Private/InstancedFoliage.cpp
@@ -4881,6 +4881,14 @@ void AInstancedFoliageActor::PostLoad()
 	}
 }
 
+#if WITH_EDITORONLY_DATA
+void AInstancedFoliageActor::DeclareConstructClasses(TArray<FTopLevelAssetPath>& OutConstructClasses, const UClass* SpecificSubclass)
+{
+	Super::DeclareConstructClasses(OutConstructClasses, SpecificSubclass);
+	OutConstructClasses.Add(FTopLevelAssetPath(UFoliageInstancedStaticMeshComponent::StaticClass()));
+}
+#endif
+
 #if WITH_EDITOR
 
 void AInstancedFoliageActor::RepairDuplicateIFA(AInstancedFoliageActor* DuplicateIFA)
diff --git a/Engine/Source/Runtime/Foliage/Public/InstancedFoliageActor.h b/Engine/Source/Runtime/Foliage/Public/InstancedFoliageActor.h
--- a/Engine/Source/Runtime/Foliage/Public/InstancedFoliageActor.h
+++ b/Engine/Source/Runtime/Foliage/Public/InstancedFoliageActor.h
@@ -52,6 +52,10 @@ public:
 	//~ Begin UObject Interface.
 	virtual void Serialize(FArchive& Ar) override;
 	virtual void PostLoad() override;
+#if WITH_EDITORONLY_DATA
+	FOLIAGE_API static void DeclareConstructClasses(TArray<FTopLevelAssetPath>& OutConstructClasses, const UClass* SpecificSubclass);
+#endif
+
 	static void AddReferencedObjects(UObject* InThis, FReferenceCollector& Collector);
 	//~ End UObject Interface. 
 
