PR #9823: Fix a problem in StateTreeConditions where GetWorld is null (Contributed by [REDACTED_PERSON])

diff --git a/Engine/Plugins/Runtime/StateTree/Source/StateTreeModule/Private/Blueprint/StateTreeConditionBlueprintBase.cpp b/Engine/Plugins/Runtime/StateTree/Source/StateTreeModule/Private/Blueprint/StateTreeConditionBlueprintBase.cpp
--- a/Engine/Plugins/Runtime/StateTree/Source/StateTreeModule/Private/Blueprint/StateTreeConditionBlueprintBase.cpp
+++ b/Engine/Plugins/Runtime/StateTree/Source/StateTreeModule/Private/Blueprint/StateTreeConditionBlueprintBase.cpp
@@ -21,7 +21,7 @@ bool UStateTreeConditionBlueprintBase::TestCondition(FStateTreeExecutionContext&
 {
 	if (bHasTestCondition)
 	{
-		FScopedCurrentContext(*this, Context);
+		FScopedCurrentContext ScopedContext(*this, Context);
 		return ReceiveTestCondition();
 	}
 	return false;
diff --git a/Engine/Plugins/Runtime/StateTree/Source/StateTreeModule/Private/Blueprint/StateTreeEvaluatorBlueprintBase.cpp b/Engine/Plugins/Runtime/StateTree/Source/StateTreeModule/Private/Blueprint/StateTreeEvaluatorBlueprintBase.cpp
--- a/Engine/Plugins/Runtime/StateTree/Source/StateTreeModule/Private/Blueprint/StateTreeEvaluatorBlueprintBase.cpp
+++ b/Engine/Plugins/Runtime/StateTree/Source/StateTreeModule/Private/Blueprint/StateTreeEvaluatorBlueprintBase.cpp
@@ -22,7 +22,7 @@ void UStateTreeEvaluatorBlueprintBase::TreeStart(FStateTreeExecutionContext& Con
 {
 	if (bHasTreeStart)
 	{
-		FScopedCurrentContext(*this, Context);
+		FScopedCurrentContext ScopedContext(*this, Context);
 		ReceiveTreeStart();
 	}
 }
@@ -31,7 +31,7 @@ void UStateTreeEvaluatorBlueprintBase::TreeStop(FStateTreeExecutionContext& Cont
 {
 	if (bHasTreeStop)
 	{
-		FScopedCurrentContext(*this, Context);
+		FScopedCurrentContext ScopedContext(*this, Context);
 		ReceiveTreeStop();
 	}
 }
@@ -40,7 +40,7 @@ void UStateTreeEvaluatorBlueprintBase::Tick(FStateTreeExecutionContext& Context,
 {
 	if (bHasTick)
 	{
-		FScopedCurrentContext(*this, Context);
+		FScopedCurrentContext ScopedContext(*this, Context);
 		ReceiveTick(DeltaTime);
 	}
 }
diff --git a/Engine/Plugins/Runtime/StateTree/Source/StateTreeModule/Private/Blueprint/StateTreeNodeBlueprintBase.cpp b/Engine/Plugins/Runtime/StateTree/Source/StateTreeModule/Private/Blueprint/StateTreeNodeBlueprintBase.cpp
--- a/Engine/Plugins/Runtime/StateTree/Source/StateTreeModule/Private/Blueprint/StateTreeNodeBlueprintBase.cpp
+++ b/Engine/Plugins/Runtime/StateTree/Source/StateTreeModule/Private/Blueprint/StateTreeNodeBlueprintBase.cpp
@@ -10,10 +10,17 @@
 
 UWorld* UStateTreeNodeBlueprintBase::GetWorld() const
 {
-	// The items are duplicated as the StateTreeExecution context as outer, so this should be essentially the same as GetWorld() on StateTree context.
+	// The items are duplicated as the State Tree execution context as outer, so this should be essentially the same as GetWorld() on StateTree context.
 	// The CDO is used by the BP editor to check for certain functionality, make it return nullptr so that the GetWorld() passes as overridden. 
 	if (!HasAnyFlags(RF_ClassDefaultObject))
 	{
+		if (CurrentContext != nullptr)
+		{
+			if (CurrentContext->GetOwner() != nullptr)
+			{
+				return CurrentContext->GetOwner()->GetWorld();
+			}
+		}
 		if (UObject* Outer = GetOuter())
 		{
 			return Outer->GetWorld();
diff --git a/Engine/Plugins/Runtime/StateTree/Source/StateTreeModule/Private/Blueprint/StateTreeTaskBlueprintBase.cpp b/Engine/Plugins/Runtime/StateTree/Source/StateTreeModule/Private/Blueprint/StateTreeTaskBlueprintBase.cpp
--- a/Engine/Plugins/Runtime/StateTree/Source/StateTreeModule/Private/Blueprint/StateTreeTaskBlueprintBase.cpp
+++ b/Engine/Plugins/Runtime/StateTree/Source/StateTreeModule/Private/Blueprint/StateTreeTaskBlueprintBase.cpp
@@ -29,7 +29,7 @@ EStateTreeRunStatus UStateTreeTaskBlueprintBase::EnterState(FStateTreeExecutionC
 {
 	if (bHasEnterState)
 	{
-		FScopedCurrentContext(*this, Context);
+		FScopedCurrentContext ScopedContext(*this, Context);
 		return ReceiveEnterState(Transition);
 	}
 	return EStateTreeRunStatus::Running;
@@ -39,7 +39,7 @@ void UStateTreeTaskBlueprintBase::ExitState(FStateTreeExecutionContext& Context,
 {
 	if (bHasExitState)
 	{
-		FScopedCurrentContext(*this, Context);
+		FScopedCurrentContext ScopedContext(*this, Context);
 		ReceiveExitState(Transition);
 	}
 }
@@ -48,7 +48,7 @@ void UStateTreeTaskBlueprintBase::StateCompleted(FStateTreeExecutionContext& Con
 {
 	if (bHasStateCompleted)
 	{
-		FScopedCurrentContext(*this, Context);
+		FScopedCurrentContext ScopedContext(*this, Context);
 		ReceiveStateCompleted(CompletionStatus, CompletedActiveStates);
 	}
 }
@@ -57,7 +57,7 @@ EStateTreeRunStatus UStateTreeTaskBlueprintBase::Tick(FStateTreeExecutionContext
 {
 	if (bHasTick)
 	{
-		FScopedCurrentContext(*this, Context);
+		FScopedCurrentContext ScopedContext(*this, Context);
 		return ReceiveTick(DeltaTime);
 	}
 	return EStateTreeRunStatus::Running;
