StateTree: Transition improvements
- Removed gated delay (transition will be delayed until it has been "pressed" for N seconds)
- Delay now triggers after the delay duration has passed (single trigger, not "pressed", works with events too)
- Added random variation for the delay
- Tick/Event transitions has precedence over completion transitions (that is, they are handled before completion transitions)
- Removed delay and blocking from completion transitions
- Small improvements for transitions UI

diff --git a/Engine/Plugins/AI/MassAI/Source/MassAIBehavior/Private/MassStateTreeExecutionContext.cpp b/Engine/Plugins/AI/MassAI/Source/MassAIBehavior/Private/MassStateTreeExecutionContext.cpp
--- a/Engine/Plugins/AI/MassAI/Source/MassAIBehavior/Private/MassStateTreeExecutionContext.cpp
+++ b/Engine/Plugins/AI/MassAI/Source/MassAIBehavior/Private/MassStateTreeExecutionContext.cpp
@@ -15,11 +15,11 @@ FMassStateTreeExecutionContext::FMassStateTreeExecutionContext(UObject& InOwner,
 {
 }
 
-void FMassStateTreeExecutionContext::BeginGatedTransition(const FStateTreeExecutionState& Exec)
+void FMassStateTreeExecutionContext::BeginDelayedTransition(const FStateTreeTransitionDelayedState& DelayedState)
 {
 	if (SignalSubsystem != nullptr && Entity.IsSet())
 	{
 		// Tick again after the games time has passed to see if the condition still holds true.
-		SignalSubsystem->DelaySignalEntity(UE::Mass::Signals::DelayedTransitionWakeup, Entity, Exec.GatedTransitionTime + KINDA_SMALL_NUMBER);
+		SignalSubsystem->DelaySignalEntity(UE::Mass::Signals::DelayedTransitionWakeup, Entity, DelayedState.TimeLeft + KINDA_SMALL_NUMBER);
 	}
 }
diff --git a/Engine/Plugins/AI/MassAI/Source/MassAIBehavior/Public/MassStateTreeExecutionContext.h b/Engine/Plugins/AI/MassAI/Source/MassAIBehavior/Public/MassStateTreeExecutionContext.h
--- a/Engine/Plugins/AI/MassAI/Source/MassAIBehavior/Public/MassStateTreeExecutionContext.h
+++ b/Engine/Plugins/AI/MassAI/Source/MassAIBehavior/Public/MassStateTreeExecutionContext.h
@@ -30,7 +30,7 @@ protected:
 	/** Prefix that will be used by STATETREE_LOG and STATETREE_CLOG, using Entity description. */
 	virtual FString GetInstanceDescription() const override { return FString::Printf(TEXT("Entity [%s]: "), *Entity.DebugGetDescription()); }
 
-	virtual void BeginGatedTransition(const FStateTreeExecutionState& Exec) override;
+	virtual void BeginDelayedTransition(const FStateTreeTransitionDelayedState& DelayedState) override;
 
 	FMassEntityManager* EntityManager;
 	UMassSignalSubsystem* SignalSubsystem = nullptr;
diff --git a/Engine/Plugins/Runtime/StateTree/Config/BaseStateTree.ini b/Engine/Plugins/Runtime/StateTree/Config/BaseStateTree.ini
--- a/Engine/Plugins/Runtime/StateTree/Config/BaseStateTree.ini
+++ b/Engine/Plugins/Runtime/StateTree/Config/BaseStateTree.ini
@@ -58,3 +58,5 @@
 +EnumRedirects=(OldName="EStateTreeBindableStructSource",ValueChanges=(("TreeData", "Context"), ("TreeParameter", "Parameter"), ("StartParameter", "State"))
 
 +PropertyRedirects=(OldName="/Script/StateTreeModule.StateTree.NamedExternalDataDescs",NewName="/Script/StateTreeModule.StateTree.ContextDataDescs")
+
++PropertyRedirects=(OldName="/Script/StateTreeModule.StateTreeTransition.GateDelay",NewName="/Script/StateTreeModule.StateTreeTransition.DelayDuration")
diff --git a/Engine/Plugins/Runtime/StateTree/Source/StateTreeEditorModule/Private/Customizations/StateTreeStateLinkDetails.cpp b/Engine/Plugins/Runtime/StateTree/Source/StateTreeEditorModule/Private/Customizations/StateTreeStateLinkDetails.cpp
--- a/Engine/Plugins/Runtime/StateTree/Source/StateTreeEditorModule/Private/Customizations/StateTreeStateLinkDetails.cpp
+++ b/Engine/Plugins/Runtime/StateTree/Source/StateTreeEditorModule/Private/Customizations/StateTreeStateLinkDetails.cpp
@@ -178,7 +178,7 @@ TSharedRef<SWidget> FStateTreeStateLinkDetails::OnGetStateContent() const
 
 	if (CachedNames.Num() > 0)
 	{
-		MenuBuilder.BeginSection(FName(), LOCTEXT("TransitionGotoState", "Goto State"));
+		MenuBuilder.BeginSection(FName(), LOCTEXT("TransitionGotoState", "Go to State"));
 
 		for (int32 Idx = 0; Idx < CachedNames.Num(); Idx++)
 		{
@@ -186,7 +186,7 @@ TSharedRef<SWidget> FStateTreeStateLinkDetails::OnGetStateContent() const
 			Args.Add(TEXT("StateName"), FText::FromName(CachedNames[Idx]));
 
 			FUIAction ItemAction(FExecuteAction::CreateSP(const_cast<FStateTreeStateLinkDetails*>(this), &FStateTreeStateLinkDetails::OnStateComboChange, Idx));
-			MenuBuilder.AddMenuEntry(FText::FromName(CachedNames[Idx]), FText::Format(LOCTEXT("TransitionGotoStateTooltip", "Goto State {StateName}."), Args), FSlateIcon(), ItemAction);
+			MenuBuilder.AddMenuEntry(FText::FromName(CachedNames[Idx]), FText::Format(LOCTEXT("TransitionGotoStateTooltip", "Go to State {StateName}."), Args), FSlateIcon(), ItemAction);
 		}
 
 		MenuBuilder.EndSection();
@@ -201,7 +201,7 @@ FText FStateTreeStateLinkDetails::GetCurrentStateDesc() const
 
 	if (TransitionType == EStateTreeTransitionType::NotSet)
 	{
-		return LOCTEXT("TransitionNotSet", "None");
+		return LOCTEXT("TransitionBlock", "Block Transition");
 	}
 	else if (TransitionType == EStateTreeTransitionType::NextState)
 	{
diff --git a/Engine/Plugins/Runtime/StateTree/Source/StateTreeEditorModule/Private/Customizations/StateTreeTransitionDetails.cpp b/Engine/Plugins/Runtime/StateTree/Source/StateTreeEditorModule/Private/Customizations/StateTreeTransitionDetails.cpp
--- a/Engine/Plugins/Runtime/StateTree/Source/StateTreeEditorModule/Private/Customizations/StateTreeTransitionDetails.cpp
+++ b/Engine/Plugins/Runtime/StateTree/Source/StateTreeEditorModule/Private/Customizations/StateTreeTransitionDetails.cpp
@@ -11,6 +11,7 @@
 #include "StateTreeState.h"
 #include "Widgets/Text/STextBlock.h"
 #include "Widgets/Layout/SBox.h"
+#include "StateTreePropertyHelpers.h"
 
 #define LOCTEXT_NAMESPACE "StateTreeEditor"
 
@@ -27,10 +28,13 @@ void FStateTreeTransitionDetails::CustomizeHeader(TSharedRef<class IPropertyHand
 	TriggerProperty = StructProperty->GetChildHandle(TEXT("Trigger"));
 	EventTagProperty = StructProperty->GetChildHandle(TEXT("EventTag"));
 	StateProperty = StructProperty->GetChildHandle(TEXT("State"));
-	GateDelayProperty = StructProperty->GetChildHandle(TEXT("GateDelay"));
+	DelayTransitionProperty = StructProperty->GetChildHandle(TEXT("bDelayTransition"));
+	DelayDurationProperty = StructProperty->GetChildHandle(TEXT("DelayDuration"));
+	DelayRandomVarianceProperty = StructProperty->GetChildHandle(TEXT("DelayRandomVariance"));
 	ConditionsProperty = StructProperty->GetChildHandle(TEXT("Conditions"));
 
 	HeaderRow
+		.RowTag(StructProperty->GetProperty()->GetFName())
 		.WholeRowContent()
 		.VAlign(VAlign_Center)
 		[
@@ -54,107 +58,138 @@ void FStateTreeTransitionDetails::CustomizeHeader(TSharedRef<class IPropertyHand
 
 void FStateTreeTransitionDetails::CustomizeChildren(TSharedRef<class IPropertyHandle> StructPropertyHandle, class IDetailChildrenBuilder& StructBuilder, IPropertyTypeCustomizationUtils& StructCustomizationUtils)
 {
-	if (TriggerProperty)
-	{
-		StructBuilder.AddProperty(TriggerProperty.ToSharedRef());
-	}
+	check(TriggerProperty);
+	check(EventTagProperty);
+	check(DelayTransitionProperty);
+	check(DelayDurationProperty);
+	check(DelayRandomVarianceProperty);
+	check(StateProperty);
+	check(ConditionsProperty);
 
-	if (EventTagProperty)
-	{
-		IDetailPropertyRow& EventTagRow = StructBuilder.AddProperty(EventTagProperty.ToSharedRef());
-		EventTagRow.Visibility(TAttribute<EVisibility>::Create(TAttribute<EVisibility>::FGetter::CreateLambda([this]()
+	StructBuilder.AddProperty(TriggerProperty.ToSharedRef());
+
+	// Show event only when the trigger is set to Event. 
+	StructBuilder.AddProperty(EventTagProperty.ToSharedRef())
+		.Visibility(TAttribute<EVisibility>::Create(TAttribute<EVisibility>::FGetter::CreateLambda([this]()
 		{
-			uint8 TriggerValue = 0; 
-			if (TriggerProperty.IsValid())
-			{
-				TriggerProperty->GetValue(TriggerValue);
-			}
-			return TriggerValue == (uint8)EStateTreeTransitionTrigger::OnEvent ? EVisibility::Visible : EVisibility::Collapsed;
+			return (GetTrigger() == EStateTreeTransitionTrigger::OnEvent) ? EVisibility::Visible : EVisibility::Collapsed;
 		})));
-	}
-	
-	if (GateDelayProperty)
-	{
-		StructBuilder.AddProperty(GateDelayProperty.ToSharedRef());
-	}
 
-	if (StateProperty)
+	// Delay
+	auto IsDelayVisible = [this]()
 	{
-		StructBuilder.AddProperty(StateProperty.ToSharedRef());
-	}
+		return !EnumHasAnyFlags(GetTrigger(), EStateTreeTransitionTrigger::OnStateCompleted) ? EVisibility::Visible : EVisibility::Collapsed;
+	};
 
-	if (ConditionsProperty)
-	{
-		// Show conditions always expanded, with simplified header (remove item count)
-		IDetailPropertyRow& Property = StructBuilder.AddProperty(ConditionsProperty.ToSharedRef());
-		Property.ShouldAutoExpand(true);
+	StructBuilder.AddProperty(DelayTransitionProperty.ToSharedRef())
+		.Visibility(TAttribute<EVisibility>::Create(TAttribute<EVisibility>::FGetter::CreateLambda(IsDelayVisible)));
+	StructBuilder.AddProperty(DelayDurationProperty.ToSharedRef())
+		.Visibility(TAttribute<EVisibility>::Create(TAttribute<EVisibility>::FGetter::CreateLambda(IsDelayVisible)));
+	StructBuilder.AddProperty(DelayRandomVarianceProperty.ToSharedRef())
+		.Visibility(TAttribute<EVisibility>::Create(TAttribute<EVisibility>::FGetter::CreateLambda(IsDelayVisible)));
 
-		static const bool bShowChildren = true;
-		Property.CustomWidget(bShowChildren)
-			.NameContent()
+	StructBuilder.AddProperty(StateProperty.ToSharedRef());
+
+	// Show conditions always expanded, with simplified header (remove item count)
+	IDetailPropertyRow& ConditionsRow = StructBuilder.AddProperty(ConditionsProperty.ToSharedRef());
+	ConditionsRow.ShouldAutoExpand(true);
+
+	constexpr bool bShowChildren = true;
+	ConditionsRow.CustomWidget(bShowChildren)
+		.RowTag(ConditionsProperty->GetProperty()->GetFName())
+		.NameContent()
+		.VAlign(VAlign_Center)
+		[
+			SNew(SHorizontalBox)
+			+ SHorizontalBox::Slot()
+			.AutoWidth()
 			.VAlign(VAlign_Center)
+			.Padding(FMargin(0.0f, 2.0f))
 			[
-				SNew(SHorizontalBox)
-				+ SHorizontalBox::Slot()
-				.AutoWidth()
-				.VAlign(VAlign_Center)
-				.Padding(FMargin(0.0f, 2.0f))
-				[
-					SNew(STextBlock)
-					.Text(ConditionsProperty->GetPropertyDisplayName())
-					.Font(IDetailLayoutBuilder::GetDetailFontBold())
-				]
+				SNew(STextBlock)
+				.Text(ConditionsProperty->GetPropertyDisplayName())
+				.Font(IDetailLayoutBuilder::GetDetailFontBold())
 			]
-			.ValueContent()
-			.VAlign(VAlign_Center)
-			.HAlign(HAlign_Right)
-			[
-				SNew(SBox) // Empty, suppress noisy array details.
-			];
+		]
+		.ValueContent()
+		.VAlign(VAlign_Center)
+		.HAlign(HAlign_Right)
+		[
+			SNew(SBox) // Empty, suppress noisy array details.
+		];
+}
+
+EStateTreeTransitionTrigger FStateTreeTransitionDetails::GetTrigger() const
+{
+	check(TriggerProperty);
+	EStateTreeTransitionTrigger TriggerValue = EStateTreeTransitionTrigger::None;
+	if (TriggerProperty.IsValid())
+	{
+		TriggerProperty->GetValue((uint8&)TriggerValue);
+	}
+	return TriggerValue;
+}
+
+bool FStateTreeTransitionDetails::GetDelayTransition() const
+{
+	check(DelayTransitionProperty);
+	bool bDelayTransition = false;
+	if (DelayTransitionProperty.IsValid())
+	{
+		DelayTransitionProperty->GetValue(bDelayTransition);
 	}
+	return bDelayTransition;
 }
 
 FText FStateTreeTransitionDetails::GetDescription() const
 {
-	if (StateProperty)
+	check(StateProperty);
+	if (StateProperty->GetNumPerObjectValues() != 1)
 	{
-		TArray<void*> RawData;
-		StateProperty->AccessRawData(RawData);
-		if (RawData.Num() == 1)
+		return LOCTEXT("MultipleSelected", "Multiple Selected");
+	}
+
+	EStateTreeTransitionTrigger Trigger = GetTrigger();
+	FText TriggerText = UEnum::GetDisplayValueAsText(Trigger);
+
+	if (Trigger == EStateTreeTransitionTrigger::OnEvent)
+	{
+		FGameplayTag EventTag;
+		UE::StateTree::PropertyHelpers::GetStructValue<FGameplayTag>(EventTagProperty, EventTag);
+		TriggerText = FText::Format(LOCTEXT("TransitionOnEvent", "On Event {0}"), FText::FromName(EventTag.GetTagName()));
+	}
+	
+	FText TargetText;
+	TArray<void*> RawData;
+	StateProperty->AccessRawData(RawData);
+	check(RawData.Num() > 0);
+	
+	const FStateTreeStateLink* State = static_cast<FStateTreeStateLink*>(RawData[0]);
+	if (State != nullptr)
+	{
+		switch (State->Type)
 		{
-			FStateTreeStateLink* State = static_cast<FStateTreeStateLink*>(RawData[0]);
-			if (State != nullptr)
+		case EStateTreeTransitionType::NotSet:
+			TargetText = LOCTEXT("TransitionBlock", "Block Transition");
+			break;
+		case EStateTreeTransitionType::Succeeded:
+			TargetText = LOCTEXT("TransitionTreeSucceeded", "Tree Succeeded");
+			break;
+		case EStateTreeTransitionType::Failed:
+			TargetText = LOCTEXT("TransitionTreeFailed", "Tree Failed");
+			break;
+		case EStateTreeTransitionType::NextState:
+			TargetText = LOCTEXT("TransitionNextState", "Next State");
+			break;
+		case EStateTreeTransitionType::GotoState:
 			{
-				switch (State->Type)
-				{
-				case EStateTreeTransitionType::NotSet:
-					return LOCTEXT("TransitionNotSet", "None");
-					break;
-				case EStateTreeTransitionType::Succeeded:
-					return LOCTEXT("TransitionTreeSucceeded", "Tree Succeeded");
-					break;
-				case EStateTreeTransitionType::Failed:
-					return LOCTEXT("TransitionTreeFailed", "Tree Failed");
-					break;
-				case EStateTreeTransitionType::NextState:
-					return LOCTEXT("TransitionNextState", "Next State");
-					break;
-				case EStateTreeTransitionType::GotoState:
-					{
-						FFormatNamedArguments Args;
-						Args.Add(TEXT("State"), FText::FromName(State->Name));
-						return FText::Format(LOCTEXT("TransitionActionGotoState", "Go to State {State}"), Args);
-					}
-					break;
-				}
+				TargetText = FText::Format(LOCTEXT("TransitionGotoState", "Go to State {0}"), FText::FromName(State->Name));
 			}
-		}
-		else
-		{
-			return LOCTEXT("MultipleSelected", "Multiple Selected");
+			break;
 		}
 	}
-	return FText::GetEmpty();
+
+	return FText::Format(LOCTEXT("TransitionDesc", "{0} {1}"), TriggerText, TargetText);
 }
 
 #undef LOCTEXT_NAMESPACE
diff --git a/Engine/Plugins/Runtime/StateTree/Source/StateTreeEditorModule/Private/Customizations/StateTreeTransitionDetails.h b/Engine/Plugins/Runtime/StateTree/Source/StateTreeEditorModule/Private/Customizations/StateTreeTransitionDetails.h
--- a/Engine/Plugins/Runtime/StateTree/Source/StateTreeEditorModule/Private/Customizations/StateTreeTransitionDetails.h
+++ b/Engine/Plugins/Runtime/StateTree/Source/StateTreeEditorModule/Private/Customizations/StateTreeTransitionDetails.h
@@ -9,6 +9,7 @@
 class IPropertyHandle;
 class UStateTree;
 class UStateTreeState;
+enum class EStateTreeTransitionTrigger : uint8;
 
 /**
  * Type customization for FStateTreeTransition.
@@ -28,10 +29,15 @@ private:
 
 	FText GetDescription() const;
 
+	EStateTreeTransitionTrigger GetTrigger() const;
+	bool GetDelayTransition() const;
+	
 	TSharedPtr<IPropertyHandle> TriggerProperty;
 	TSharedPtr<IPropertyHandle> EventTagProperty;
 	TSharedPtr<IPropertyHandle> StateProperty;
-	TSharedPtr<IPropertyHandle> GateDelayProperty;
+	TSharedPtr<IPropertyHandle> DelayTransitionProperty;
+	TSharedPtr<IPropertyHandle> DelayDurationProperty;
+	TSharedPtr<IPropertyHandle> DelayRandomVarianceProperty;
 	TSharedPtr<IPropertyHandle> ConditionsProperty;
 
 	class IPropertyUtilities* PropUtils;
diff --git a/Engine/Plugins/Runtime/StateTree/Source/StateTreeEditorModule/Private/StateTreeCompiler.cpp b/Engine/Plugins/Runtime/StateTree/Source/StateTreeEditorModule/Private/StateTreeCompiler.cpp
--- a/Engine/Plugins/Runtime/StateTree/Source/StateTreeEditorModule/Private/StateTreeCompiler.cpp
+++ b/Engine/Plugins/Runtime/StateTree/Source/StateTreeEditorModule/Private/StateTreeCompiler.cpp
@@ -560,13 +560,31 @@ bool FStateTreeCompiler::CreateStateTransitions()
 			CompactTransition.Trigger = Transition.Trigger;
 			CompactTransition.EventTag = Transition.EventTag;
 			CompactTransition.Type = Transition.State.Type;
-			CompactTransition.GateDelay = (uint8)FMath::Clamp(FMath::CeilToInt(Transition.GateDelay * 10.0f), 0, 255);
+			if (Transition.bDelayTransition)
+			{
+				CompactTransition.Delay.Set(Transition.DelayDuration, Transition.DelayRandomVariance);
+			}
+
+			if (EnumHasAnyFlags(Transition.Trigger, EStateTreeTransitionTrigger::OnStateCompleted))
+			{
+				// Completion transitions cannot have delay.
+				CompactTransition.Delay.Reset();
+
+				// Completion transitions must have valid target state.
+				if (CompactTransition.Type == EStateTreeTransitionType::NotSet)
+				{
+					Log.Reportf(EMessageSeverity::Error,
+						TEXT("State completion transition to '%s' must have transition to valid state, 'None' not accepted."),
+						*Transition.State.Name.ToString());
+				}
+			}
+			
 			CompactTransition.State = FStateTreeStateHandle::Invalid;
 			if (!ResolveTransitionState(*SourceState, Transition.State, CompactTransition.State))
 			{
 				return false;
 			}
-			// Note: Unset transition is allowed here. It can be used to mask a transition at parent.
+			
 			const int32 ConditionsBegin = Nodes.Num();
 			if (const auto Validation = UE::StateTree::Compiler::IsValidCount16(ConditionsBegin); Validation.DidFail())
 			{
diff --git a/Engine/Plugins/Runtime/StateTree/Source/StateTreeEditorModule/Private/StateTreeState.cpp b/Engine/Plugins/Runtime/StateTree/Source/StateTreeEditorModule/Private/StateTreeState.cpp
--- a/Engine/Plugins/Runtime/StateTree/Source/StateTreeEditorModule/Private/StateTreeState.cpp
+++ b/Engine/Plugins/Runtime/StateTree/Source/StateTreeEditorModule/Private/StateTreeState.cpp
@@ -194,6 +194,38 @@ void UStateTreeState::PostEditChangeChainProperty(FPropertyChangedChainEvent& Pr
 					}
 				}
 			}
+			else if (PropertyChangedEvent.ChangeType == EPropertyChangeType::ArrayAdd)
+			{
+				if (MemberProperty->GetFName() == GET_MEMBER_NAME_CHECKED(UStateTreeState, Transitions))
+				{
+					const int32 TransitionsIndex = PropertyChangedEvent.GetArrayIndex(MemberProperty->GetFName().ToString());
+					if (Transitions.IsValidIndex(TransitionsIndex))
+					{
+						// Set default transition on newly created states.
+						FStateTreeTransition& Transition = Transitions[TransitionsIndex];
+						Transition.Trigger = EStateTreeTransitionTrigger::OnStateCompleted;
+						const UStateTreeState* RootState = GetRootState();
+						Transition.State.Set(RootState ? RootState : this);
+					}
+				}
+			}
+			else
+			{
+				if (MemberProperty->GetFName() == GET_MEMBER_NAME_CHECKED(UStateTreeState, Transitions))
+				{
+					const int32 TransitionsIndex = PropertyChangedEvent.GetArrayIndex(MemberProperty->GetFName().ToString());
+					if (Transitions.IsValidIndex(TransitionsIndex))
+					{
+						FStateTreeTransition& Transition = Transitions[TransitionsIndex];
+
+						if (EnumHasAnyFlags(Transition.Trigger, EStateTreeTransitionTrigger::OnStateCompleted))
+						{
+							// Reset delay on completion transitions
+							Transition.bDelayTransition = false;
+						}
+					}
+				}
+			}
 		}
 	}
 }
@@ -236,7 +268,17 @@ void UStateTreeState::UpdateParametersFromLinkedSubtree()
 
 #endif
 
-UStateTreeState* UStateTreeState::GetNextSiblingState() const
+const UStateTreeState* UStateTreeState::GetRootState() const
+{
+	const UStateTreeState* RootState = this;
+	while (RootState->Parent != nullptr)
+	{
+		RootState = RootState->Parent;
+	}
+	return RootState;
+}
+
+const UStateTreeState* UStateTreeState::GetNextSiblingState() const
 {
 	if (!Parent)
 	{
diff --git a/Engine/Plugins/Runtime/StateTree/Source/StateTreeEditorModule/Public/StateTreeState.h b/Engine/Plugins/Runtime/StateTree/Source/StateTreeEditorModule/Public/StateTreeState.h
--- a/Engine/Plugins/Runtime/StateTree/Source/StateTreeEditorModule/Public/StateTreeState.h
+++ b/Engine/Plugins/Runtime/StateTree/Source/StateTreeEditorModule/Public/StateTreeState.h
@@ -49,7 +49,7 @@ struct STATETREEEDITORMODULE_API FStateTreeTransition
 	FStateTreeTransition() = default;
 	FStateTreeTransition(const EStateTreeTransitionTrigger InTrigger, const EStateTreeTransitionType InType, const UStateTreeState* InState = nullptr);
 	FStateTreeTransition(const EStateTreeTransitionTrigger InTrigger, const FGameplayTag InEventTag, const EStateTreeTransitionType InType, const UStateTreeState* InState = nullptr);
-
+	
 	template<typename T, typename... TArgs>
 	TStateTreeEditorNode<T>& AddCondition(TArgs&&... InArgs)
 	{
@@ -64,20 +64,32 @@ struct STATETREEEDITORMODULE_API FStateTreeTransition
 		return static_cast<TStateTreeEditorNode<T>&>(CondNode);
 	}
 
-	UPROPERTY(EditDefaultsOnly, Category = Transition)
+	/** When to try trigger the transition. */
+	UPROPERTY(EditDefaultsOnly, Category = "Transition")
 	EStateTreeTransitionTrigger Trigger = EStateTreeTransitionTrigger::OnStateCompleted;
 
-	UPROPERTY(EditDefaultsOnly, Category = Transition)
+	/** Tag of the State Tree event that triggers the transition. */
+	UPROPERTY(EditDefaultsOnly, Category = "Transition")
 	FGameplayTag EventTag;
 
-	UPROPERTY(EditDefaultsOnly, Category = Transition, meta = (DisplayName="Transition To"))
+	/** Transition target state. */
+	UPROPERTY(EditDefaultsOnly, Category = "Transition", meta=(DisplayName="Transition To"))
 	FStateTreeStateLink State;
 
-	// Gate delay in seconds.
-	UPROPERTY(EditDefaultsOnly, Category = Transition, meta = (UIMin = "0", ClampMin = "0", UIMax = "25", ClampMax = "25"))
-	float GateDelay = 0.0f;
+	/** Delay the triggering of the transition. */
+	UPROPERTY(EditDefaultsOnly, Category = "Transition")
+	bool bDelayTransition = false;
+
+	/** Transition delay duration in seconds. */
+	UPROPERTY(EditDefaultsOnly, Category = "Transition", meta = (EditCondition = "bDelayTransition", UIMin = "0", ClampMin = "0", UIMax = "25", ClampMax = "25", ForceUnits="s"))
+	float DelayDuration = 0.0f;
+
+	/** Transition delay random variance in seconds. */
+	UPROPERTY(EditDefaultsOnly, Category = "Transition", meta = (EditCondition = "bDelayTransition", UIMin = "0", ClampMin = "0", UIMax = "25", ClampMax = "25", ForceUnits="s"))
+	float DelayRandomVariance = 0.0f;
 
-	UPROPERTY(EditDefaultsOnly, Category = Transition, meta = (BaseStruct = "/Script/StateTreeModule.StateTreeConditionBase", BaseClass = "/Script/StateTreeModule.StateTreeConditionBlueprintBase"))
+	/** Conditions that must pass so that the transition can be triggered. */
+	UPROPERTY(EditDefaultsOnly, Category = "Transition", meta = (BaseStruct = "/Script/StateTreeModule.StateTreeConditionBase", BaseClass = "/Script/StateTreeModule.StateTreeConditionBlueprintBase"))
 	TArray<FStateTreeEditorNode> Conditions;
 };
 
@@ -119,7 +131,8 @@ public:
 	void UpdateParametersFromLinkedSubtree();
 #endif
 
-	UStateTreeState* GetNextSiblingState() const;
+	const UStateTreeState* GetRootState() const;
+	const UStateTreeState* GetNextSiblingState() const;
 
 	// StateTree Builder API
 
diff --git a/Engine/Plugins/Runtime/StateTree/Source/StateTreeModule/Private/StateTreeExecutionContext.cpp b/Engine/Plugins/Runtime/StateTree/Source/StateTreeModule/Private/StateTreeExecutionContext.cpp
--- a/Engine/Plugins/Runtime/StateTree/Source/StateTreeModule/Private/StateTreeExecutionContext.cpp
+++ b/Engine/Plugins/Runtime/StateTree/Source/StateTreeModule/Private/StateTreeExecutionContext.cpp
@@ -109,8 +109,8 @@ bool FStateTreeExecutionContext::AreExternalDataViewsValid() const
 
 void FStateTreeExecutionContext::UpdateLinkedStateParameters(const FCompactStateTreeState& State, const uint16 ParameterInstanceIndex)
 {
-	const FStateTreeDataView StateParamsInstance = InstanceData.GetMutableStruct(ParameterInstanceIndex);
-	const FCompactStateTreeParameters& StateParams = StateParamsInstance.GetMutable<FCompactStateTreeParameters>();
+	FStateTreeDataView StateParamsInstance = InstanceData.GetMutableStruct(ParameterInstanceIndex);
+	FCompactStateTreeParameters& StateParams = StateParamsInstance.GetMutable<FCompactStateTreeParameters>();
 
 	// Update parameters if the state has any.
 	if (StateParams.Parameters.IsValid())
@@ -146,7 +146,7 @@ void FStateTreeExecutionContext::UpdateSubtreeStateParameters(const FCompactStat
 		// Set view to default parameters.
 		const FConstStructView ParamInstance = StateTree.DefaultInstanceData.GetStruct(State.ParameterInstanceIndex.Get()); // These are used as const, so get them from the tree initial values.
 		const FStateTreeDataView ParamInstanceView(ParamInstance.GetScriptStruct(), const_cast<uint8*>(ParamInstance.GetMemory())); 
-		const FCompactStateTreeParameters& Params = ParamInstanceView.GetMutable<FCompactStateTreeParameters>();
+		FCompactStateTreeParameters& Params = ParamInstanceView.GetMutable<FCompactStateTreeParameters>();
 		DataViews[State.ParameterDataViewIndex.Get()] = FStateTreeDataView(Params.Parameters.GetMutableValue());
 	}
 }
@@ -339,10 +339,10 @@ EStateTreeRunStatus FStateTreeExecutionContext::Tick(const float DeltaTime)
 		return Exec->TreeRunStatus;
 	}
 
-	// Update the gated transition time.
-	if (Exec->GatedTransitionIndex.IsValid())
+	// Update the delayed transitions.
+	for (FStateTreeTransitionDelayedState& DelayedState : Exec->DelayedTransitions)
 	{
-		Exec->GatedTransitionTime -= DeltaTime;
+		DelayedState.TimeLeft -= DeltaTime;
 	}
 	
 	// Tick global evaluators.
@@ -686,10 +686,7 @@ void FStateTreeExecutionContext::ExitState(const FStateTreeTransitionResult& Tra
 		return;
 	}
 
-	// Reset transition delay
 	FStateTreeExecutionState& Exec = GetExecState();
-	Exec.GatedTransitionIndex = FStateTreeIndex16::Invalid;
-	Exec.GatedTransitionTime = 0.0f;
 
 	// On target branch means that the state is the target of current transition or child of it.
 	// States which were active before and will remain active, but are not on target branch will not get
@@ -760,6 +757,13 @@ void FStateTreeExecutionContext::ExitState(const FStateTreeTransitionResult& Tra
 	{
 		const FStateTreeStateHandle CurrentHandle = ExitedStates[Index];
 		const FCompactStateTreeState& State = StateTree.States[CurrentHandle.Index];
+
+		// Remove any delayed transitions that belong to this state.
+		Exec.DelayedTransitions.RemoveAllSwap(
+			[Begin = State.TransitionsBegin, End = State.TransitionsBegin + State.TransitionsNum](const FStateTreeTransitionDelayedState& DelayedState)
+			{
+				return DelayedState.TransitionIndex.Get() >= Begin && DelayedState.TransitionIndex.Get() < End;
+			});
 		
 		CurrentTransition.CurrentState = CurrentHandle;
 		CurrentTransition.ChangeType = ExitedStateChangeType[Index];
@@ -1120,35 +1124,58 @@ FString FStateTreeExecutionContext::DebugGetEventsAsString() const
 	return Result;
 }
 
-bool FStateTreeExecutionContext::TriggerTransitions(FStateTreeInstanceData& SharedInstanceData, FStateTreeTransitionResult& OutTransition)
+bool FStateTreeExecutionContext::SelectTransition(FStateTreeInstanceData& SharedInstanceData, FStateTreeExecutionState& Exec,
+		const int16 StateIndex, const FCompactStateTreeState& State, const FCompactStateTransition& Transition, FStateTreeTransitionResult& OutTransition)
 {
-	CSV_SCOPED_TIMING_STAT_EXCLUSIVE(StateTree_TriggerTransition);
-
-	FStateTreeExecutionState& Exec = GetExecState();
-
-	int32 StateStartIndex = Exec.ActiveStates.Num() - 1;
-
-	// On succeed/failed we will start looking for transition start at the first state that was completed.
-	auto GetLastCompletedStateIndex = [&Exec]() -> int32
+	if (Transition.Type == EStateTreeTransitionType::GotoState || Transition.Type == EStateTreeTransitionType::NextState)
 	{
-		const int32 Index = Exec.ActiveStates.IndexOfReverse(Exec.CompletedStateHandle);
-		return Index != INDEX_NONE ? Index : (Exec.ActiveStates.Num() - 1);
-	};
-	
-	EStateTreeTransitionTrigger CompletionTrigger = EStateTreeTransitionTrigger::None;
-	if (Exec.LastTickStatus == EStateTreeRunStatus::Succeeded)
-	{
-		StateStartIndex = GetLastCompletedStateIndex();
-		CompletionTrigger = EStateTreeTransitionTrigger::OnStateSucceeded;
+		OutTransition.TargetState = Transition.State;
+		OutTransition.NextActiveStates.Reset();
+
+		if (SelectState(SharedInstanceData, Transition.State, OutTransition.NextActiveStates))
+		{
+			STATETREE_LOG(Verbose, TEXT("Transition on state '%s' (%s) -[%s]-> state '%s'"),
+				*GetSafeStateName(Exec.ActiveStates.Last()), *State.Name.ToString(), *GetSafeStateName(Transition.State), *GetSafeStateName(OutTransition.NextActiveStates.Last()));
+			return true;
+		}
+		
+		return  false;
 	}
-	else if (Exec.LastTickStatus == EStateTreeRunStatus::Failed)
+	
+	if (Transition.Type == EStateTreeTransitionType::Succeeded || Transition.Type == EStateTreeTransitionType::Failed)
 	{
-		StateStartIndex = GetLastCompletedStateIndex();
-		CompletionTrigger = EStateTreeTransitionTrigger::OnStateFailed;
+		// Tree succeeded or failed.
+		// If we are on a linked state, find the first parent linked state and mark that as completed, continue to find completion transitions.
+		const FStateTreeStateHandle ParentLinkedState = GetParentLinkedStateHandle(Exec.ActiveStates, StateIndex);
+		if (ParentLinkedState.IsValid())
+		{
+			STATETREE_LOG(Verbose, TEXT("Compled subtree '%s' from state '%s' (%s): %s"),
+				*GetSafeStateName(ParentLinkedState), *GetSafeStateName(Exec.ActiveStates.Last()), *State.Name.ToString(), *UEnum::GetDisplayValueAsText(Transition.Type).ToString());
+			Exec.CompletedStateHandle = ParentLinkedState;
+			Exec.LastTickStatus = Transition.Type == EStateTreeTransitionType::Succeeded ? EStateTreeRunStatus::Succeeded : EStateTreeRunStatus::Failed;
+			return TriggerTransitions(SharedInstanceData, OutTransition);
+		}
+		
+		STATETREE_LOG(Verbose, TEXT("Stop tree execution from state '%s' (%s): %s"),
+			*GetSafeStateName(Exec.ActiveStates.Last()), *State.Name.ToString(), *UEnum::GetDisplayValueAsText(Transition.Type).ToString());
+
+		const FStateTreeStateHandle NextState = Transition.Type == EStateTreeTransitionType::Succeeded ? FStateTreeStateHandle::Succeeded : FStateTreeStateHandle::Failed; 
+		OutTransition.TargetState = NextState;
+		OutTransition.NextActiveStates = FStateTreeActiveStates(NextState);
+		return true;
 	}
 
+	return false;
+}
+
+bool FStateTreeExecutionContext::TriggerTransitions(FStateTreeInstanceData& SharedInstanceData, FStateTreeTransitionResult& OutTransition)
+{
+	CSV_SCOPED_TIMING_STAT_EXCLUSIVE(StateTree_TriggerTransition);
+
+	FStateTreeExecutionState& Exec = GetExecState();
+
 	STATETREE_CLOG(!EventsToProcess.IsEmpty(), Verbose, TEXT("Trigger transitions with events [%s]"), *DebugGetEventsAsString());
-	
+
 	auto HasEvent = [this](const FGameplayTag QueriedTag)
 	{
 		if (EventsToProcess.IsEmpty())
@@ -1164,99 +1191,120 @@ bool FStateTreeExecutionContext::TriggerTransitions(FStateTreeInstanceData& Shar
 
 	OutTransition.CurrentActiveStates = Exec.ActiveStates;
 
-	// Walk towards root and check all transitions along the way.
-	for (int32 StateIndex = StateStartIndex; StateIndex >= 0; StateIndex--)
+	//
+	// Check tick and event based transitions first.
+	//
+	for (int32 StateIndex = Exec.ActiveStates.Num() - 1; StateIndex >= 0; StateIndex--)
 	{
 		const FCompactStateTreeState& State = StateTree.States[Exec.ActiveStates[StateIndex].Index];
-		
 		for (uint8 i = 0; i < State.TransitionsNum; i++)
 		{
 			// All transition conditions must pass
 			const int16 TransitionIndex = State.TransitionsBegin + i;
 			const FCompactStateTransition& Transition = StateTree.Transitions[TransitionIndex];
 
-			const bool bShouldCheck = EnumHasAnyFlags(Transition.Trigger, CompletionTrigger)
-									|| Transition.Trigger == EStateTreeTransitionTrigger::OnTick
-									|| (Transition.Trigger == EStateTreeTransitionTrigger::OnEvent && HasEvent(Transition.EventTag));
-			
-			if (bShouldCheck && TestAllConditions(SharedInstanceData, Transition.ConditionsBegin, Transition.ConditionsNum))
+			if (Transition.Trigger == EStateTreeTransitionTrigger::OnTick
+				|| (Transition.Trigger == EStateTreeTransitionTrigger::OnEvent && HasEvent(Transition.EventTag)))
 			{
-				// If a transition has delay, we stop testing other transitions, but the transition will not pass the condition until the delay time passes.
-				if (Transition.GateDelay > 0)
+				// If delay has passed, allow to trigger the transition even if conditions do not pass.
+				FStateTreeTransitionDelayedState* DelayedState = Transition.HasDelay() ? Exec.FindDelayedTransition(FStateTreeIndex16(TransitionIndex)) : nullptr;
+				const bool bDelayCompleted = DelayedState != nullptr && DelayedState->TimeLeft <= 0.0f;
+
+				if (bDelayCompleted || TestAllConditions(SharedInstanceData, Transition.ConditionsBegin, Transition.ConditionsNum))
 				{
-					if ((int32)Exec.GatedTransitionIndex.Get() != TransitionIndex)
+					// If the transitions is delayed, set up the delay. 
+					if (Transition.HasDelay())
 					{
-						Exec.GatedTransitionIndex = FStateTreeIndex16(TransitionIndex);
-						Exec.GatedTransitionTime = FMath::RandRange(0.0f, Transition.GateDelay * 0.1f); // TODO: we need variance too.
-						BeginGatedTransition(Exec);
-						STATETREE_LOG(Verbose, TEXT("Gated transition triggered from '%s' (%s) -> '%s' %.1fs"), *GetSafeStateName(Exec.ActiveStates.Last()), *State.Name.ToString(), *GetSafeStateName(Transition.State), Exec.GatedTransitionTime);
+						if (DelayedState == nullptr)
+						{
+							DelayedState = &Exec.DelayedTransitions.AddDefaulted_GetRef();
+							
+							DelayedState->TransitionIndex = FStateTreeIndex16(TransitionIndex);
+							DelayedState->TimeLeft = Transition.Delay.GetRandomDuration();
+							BeginDelayedTransition(*DelayedState);
+							STATETREE_LOG(Verbose, TEXT("Delayed transition triggered from '%s' (%s) -> '%s' %.1fs"),
+								*GetSafeStateName(Exec.ActiveStates.Last()), *State.Name.ToString(), *GetSafeStateName(Transition.State), DelayedState->TimeLeft);
+							// Fall through to handle 0 delay.
+						}
+						check(DelayedState);
+
+						// We get here if the transitions re-triggers, or delay has completed.
+						// In case of re-trigger, we will just ignore it during the delay. 
+						if (DelayedState->TimeLeft > 0.0f)
+						{
+							continue;
+						}
+
+						STATETREE_LOG(Verbose, TEXT("Passed delayed transition from '%s' (%s) -> '%s'"),
+							*GetSafeStateName(Exec.ActiveStates.Last()), *State.Name.ToString(), *GetSafeStateName(Transition.State));
+
+						// The transition passed the delay, and remove it from the queue, and try trigger it.
+						Exec.DelayedTransitions.RemoveAllSwap([TransitionIndex](const FStateTreeTransitionDelayedState& DelayedState)
+							{
+								return DelayedState.TransitionIndex.Get() == TransitionIndex;
+							});
 					}
 
-					// Keep on updating current state, until we have tried to trigger
-					if (Exec.GatedTransitionTime > 0.0f)
+					if (Transition.Type == EStateTreeTransitionType::NotSet)
 					{
+						// NotSet is no-operation, but can be used to mask a transition at parent state. Returning unset keeps updating current state.
 						return false;
 					}
 
-					STATETREE_LOG(Verbose, TEXT("Passed gated transition from '%s' (%s) -> '%s'"), *GetSafeStateName(Exec.ActiveStates.Last()), *State.Name.ToString(), *GetSafeStateName(Transition.State));
-				}
-				
-				if (Transition.Type == EStateTreeTransitionType::GotoState || Transition.Type == EStateTreeTransitionType::NextState)
-				{
-					OutTransition.TargetState = Transition.State;
-					OutTransition.NextActiveStates.Reset();
-
-					if (SelectState(SharedInstanceData, Transition.State, OutTransition.NextActiveStates))
+					if (SelectTransition(SharedInstanceData, Exec, StateIndex, State, Transition, OutTransition))
 					{
-						STATETREE_LOG(Verbose, TEXT("Transition on state '%s' (%s) -[%s]-> state '%s'"), *GetSafeStateName(Exec.ActiveStates.Last()), *State.Name.ToString(), *GetSafeStateName(Transition.State), *GetSafeStateName(OutTransition.NextActiveStates.Last()));
 						return true;
 					}
 				}
-				else if (Transition.Type == EStateTreeTransitionType::NotSet)
-				{
-					// NotSet is no-operation, but can be used to mask a transition at parent state. Returning unset keeps updating current state.
-					return false;
-				}
-				else
+			}
+		}
+	}
+
+	//
+	// Check state completion transitions.
+	//
+	if (Exec.LastTickStatus != EStateTreeRunStatus::Running)
+	{
+		// Start from the last completed state.
+		const int32 StateStartIndex = Exec.CompletedStateHandle.IsValid() ? Exec.ActiveStates.IndexOfReverse(Exec.CompletedStateHandle) : (Exec.ActiveStates.Num() - 1);
+		const EStateTreeTransitionTrigger CompletionTrigger = Exec.LastTickStatus == EStateTreeRunStatus::Succeeded ? EStateTreeTransitionTrigger::OnStateSucceeded : EStateTreeTransitionTrigger::OnStateFailed;
+
+		check(StateStartIndex >= 0 && StateStartIndex < Exec.ActiveStates.Num());
+		
+		// Check completion transitions
+		for (int32 StateIndex = StateStartIndex; StateIndex >= 0; StateIndex--)
+		{
+			const FCompactStateTreeState& State = StateTree.States[Exec.ActiveStates[StateIndex].Index];
+			for (uint8 i = 0; i < State.TransitionsNum; i++)
+			{
+				// All transition conditions must pass
+				const int16 TransitionIndex = State.TransitionsBegin + i;
+				const FCompactStateTransition& Transition = StateTree.Transitions[TransitionIndex];
+
+				if (EnumHasAnyFlags(Transition.Trigger, CompletionTrigger))
 				{
-					// Tree succeeded or failed.
-					// If we are on a linked state, find the first parent linked state and mark that as completed, continue to find completion transitions.
-					const FStateTreeStateHandle ParentLinkedState = GetParentLinkedStateHandle(Exec.ActiveStates, StateIndex);
-					if (ParentLinkedState.IsValid())
+					if (TestAllConditions(SharedInstanceData, Transition.ConditionsBegin, Transition.ConditionsNum))
 					{
-						STATETREE_LOG(Verbose, TEXT("Compled subtree '%s' from state '%s' (%s): %s"),
-							*GetSafeStateName(ParentLinkedState), *GetSafeStateName(Exec.ActiveStates.Last()), *State.Name.ToString(), *UEnum::GetDisplayValueAsText(Transition.Type).ToString());
-						Exec.CompletedStateHandle = ParentLinkedState;
-						Exec.LastTickStatus = Transition.Type == EStateTreeTransitionType::Succeeded ? EStateTreeRunStatus::Succeeded : EStateTreeRunStatus::Failed;
-						return TriggerTransitions(SharedInstanceData, OutTransition);
+						// No delay allowed on completion conditions. 
+						
+						if (SelectTransition(SharedInstanceData, Exec, StateIndex, State, Transition, OutTransition))
+						{
+							return true;
+						}
 					}
-					
-					STATETREE_LOG(Verbose, TEXT("Stop tree execution from state '%s' (%s): %s"),
-						*GetSafeStateName(Exec.ActiveStates.Last()), *State.Name.ToString(), *UEnum::GetDisplayValueAsText(Transition.Type).ToString());
-
-					const FStateTreeStateHandle NextState = Transition.Type == EStateTreeTransitionType::Succeeded ? FStateTreeStateHandle::Succeeded : FStateTreeStateHandle::Failed; 
-					OutTransition.TargetState = NextState;
-					OutTransition.NextActiveStates = FStateTreeActiveStates(NextState);
-					return true;
 				}
 			}
-			else if ((int32)Exec.GatedTransitionIndex.Get() == TransitionIndex)
-			{
-				// If the current transition was gated transition, reset it if the condition failed.
-				Exec.GatedTransitionIndex = FStateTreeIndex16::Invalid;
-				Exec.GatedTransitionTime = 0.0f;
-			}
 		}
-	}
 
-	if (Exec.LastTickStatus != EStateTreeRunStatus::Running)
-	{
-		// Could not trigger completion transition, jump back to start.
-		static const FStateTreeStateHandle RootState = FStateTreeStateHandle(0);
-		OutTransition.TargetState = RootState;
-		return SelectState(SharedInstanceData, RootState, OutTransition.NextActiveStates);
+		if (Exec.LastTickStatus != EStateTreeRunStatus::Running)
+		{
+			// Could not trigger completion transition, jump back to start.
+			static const FStateTreeStateHandle RootState = FStateTreeStateHandle(0);
+			OutTransition.TargetState = RootState;
+			return SelectState(SharedInstanceData, RootState, OutTransition.NextActiveStates);
+		}
 	}
-
+	
 	// No transition triggered, keep on updating current state.
 	return false;
 }
diff --git a/Engine/Plugins/Runtime/StateTree/Source/StateTreeModule/Public/StateTree.h b/Engine/Plugins/Runtime/StateTree/Source/StateTreeModule/Public/StateTree.h
--- a/Engine/Plugins/Runtime/StateTree/Source/StateTreeModule/Public/StateTree.h
+++ b/Engine/Plugins/Runtime/StateTree/Source/StateTreeModule/Public/StateTree.h
@@ -30,8 +30,10 @@ struct STATETREEMODULE_API FStateTreeCustomVersion
 		IndexTypes,
 		// Added events.
 		AddedEvents,
-		// Added events.
+		// Testing mishap
 		AddedFoo,
+		// Changed transition delay
+		TransitionDelay,
 
 		// -----<new versions can be added above this line>-------------------------------------------------
 		VersionPlusOne,
diff --git a/Engine/Plugins/Runtime/StateTree/Source/StateTreeModule/Public/StateTreeExecutionContext.h b/Engine/Plugins/Runtime/StateTree/Source/StateTreeModule/Public/StateTreeExecutionContext.h
--- a/Engine/Plugins/Runtime/StateTree/Source/StateTreeModule/Public/StateTreeExecutionContext.h
+++ b/Engine/Plugins/Runtime/StateTree/Source/StateTreeModule/Public/StateTreeExecutionContext.h
@@ -14,11 +14,23 @@ struct FStateTreeTaskBase;
 struct FStateTreeConditionBase;
 struct FStateTreeEvent;
 
+struct STATETREEMODULE_API FStateTreeTransitionDelayedState
+{
+	FStateTreeIndex16 TransitionIndex = FStateTreeIndex16::Invalid;
+	float TimeLeft = 0.0f;
+};
+
 USTRUCT()
 struct STATETREEMODULE_API FStateTreeExecutionState
 {
 	GENERATED_BODY()
 
+	/** @returns Delayed transition state for a specific transition, or nullptr if it does not exists. */
+	FStateTreeTransitionDelayedState* FindDelayedTransition(const FStateTreeIndex16 TransitionIndex)
+	{
+		return DelayedTransitions.FindByPredicate([TransitionIndex](const FStateTreeTransitionDelayedState& State){ return State.TransitionIndex == TransitionIndex; });
+	}
+
 	/** Currently active states */
 	FStateTreeActiveStates ActiveStates;
 
@@ -37,9 +49,6 @@ struct STATETREEMODULE_API FStateTreeExecutionState
 	/** Running status of the instance */
 	EStateTreeRunStatus TreeRunStatus = EStateTreeRunStatus::Unset;
 
-	/** Delayed transition handle, if exists */
-	FStateTreeIndex16 GatedTransitionIndex = FStateTreeIndex16::Invalid;
-
 	/** Handle of the state that was first to report state completed (success or failure), used to trigger completion transitions. */
 	FStateTreeStateHandle CompletedStateHandle = FStateTreeStateHandle::Invalid;
 
@@ -47,7 +56,7 @@ struct STATETREEMODULE_API FStateTreeExecutionState
 	uint16 StateChangeCount = 0;
 
 	/** Running time of the delayed transition */
-	float GatedTransitionTime = 0.0f;
+	TArray<FStateTreeTransitionDelayedState> DelayedTransitions;
 };
 
 /**
@@ -273,8 +282,11 @@ protected:
 	/** @return Prefix that will be used by STATETREE_LOG and STATETREE_CLOG, empty by default. */
 	virtual FString GetInstanceDescription() const;
 
-	/** Callback when gated transition is triggered. Contexts that are event based can use this to trigger a future event. */
-	virtual void BeginGatedTransition(const FStateTreeExecutionState& Exec) {};
+	UE_DEPRECATED(5.2, "Use BeginDelayedTransition() instead.")
+	virtual void BeginGatedTransition(const FStateTreeExecutionState& Exec) final {};
+	
+	/** Callback when delayed transition is triggered. Contexts that are event based can use this to trigger a future event. */
+	virtual void BeginDelayedTransition(const FStateTreeTransitionDelayedState& DelayedState) {};
 
 	void UpdateInstanceData(const FStateTreeActiveStates& CurrentActiveStates, const FStateTreeActiveStates& NextActiveStates);
 
@@ -319,6 +331,13 @@ protected:
 	 */
 	bool TestAllConditions(FStateTreeInstanceData& SharedInstanceData, const int32 ConditionsOffset, const int32 ConditionsNum);
 
+	/**
+	 * Tries to select a state pointed by the specified transition. The OutTransition is filled on successful selection.
+	 * @return true of the state selection succeeds.
+	 */
+	bool SelectTransition(FStateTreeInstanceData& SharedInstanceData, FStateTreeExecutionState& Exec,
+			const int16 StateIndex, const FCompactStateTreeState& State, const FCompactStateTransition& Transition, FStateTreeTransitionResult& OutTransition);
+
 	/**
 	 * Triggers transitions based on current run status. CurrentStatus is used to select which transitions events are triggered.
 	 * If CurrentStatus is "Running", "Conditional" transitions pass, "Completed/Failed" will trigger "OnCompleted/OnSucceeded/OnFailed" transitions.
diff --git a/Engine/Plugins/Runtime/StateTree/Source/StateTreeModule/Public/StateTreeTypes.h b/Engine/Plugins/Runtime/StateTree/Source/StateTreeModule/Public/StateTreeTypes.h
--- a/Engine/Plugins/Runtime/StateTree/Source/StateTreeModule/Public/StateTreeTypes.h
+++ b/Engine/Plugins/Runtime/StateTree/Source/StateTreeModule/Public/StateTreeTypes.h
@@ -118,10 +118,10 @@ enum class EStateTreeTransitionTrigger : uint8
 	/** Try trigger transition when a state failed. */
     OnStateFailed = 0x2,
 
-	/** Try trigger transition each StateTree tick. */
+	/** Try trigger transition each State Tree tick. */
     OnTick = 0x4,
 	
-	/** Try trigger transition on specific event. */
+	/** Try trigger transition on specific State Tree event. */
 	OnEvent = 0x8,
 
 	MAX
@@ -401,6 +401,68 @@ struct STATETREEMODULE_API FStateTreeActiveStates
 	uint8 NumStates = 0;
 };
 
+/**
+ * Time duration with random variance. Stored compactly as two uint16s, which gives time range of about 650 seconds.
+ * The variance is symmetric (+-) around the specified duration.
+ */
+USTRUCT()
+struct STATETREEMODULE_API FStateTreeRandomTimeDuration
+{
+	GENERATED_BODY()
+
+	/** Reset duration to empty. */
+	void Reset()
+	{
+		Duration = 0;
+		RandomVariance = 0;
+	}
+
+	/** Sets the time duration with random variance. */
+	void Set(const float InDuration, const float InRandomVariance)
+	{
+		Duration = Quantize(InDuration);
+    	RandomVariance = Quantize(InRandomVariance);
+	}
+
+	/** @return the fixed duration. */
+	float GetDuration() const
+    {
+		return Duration / Scale;
+    }
+
+	/** @return the maximum random variance. */
+	float GetRandomVariance() const
+	{
+		return Duration / Scale;
+	}
+
+	/** @return True of the duration is empty (always returns 0). */
+	bool IsEmpty() const { return Duration == 0 && RandomVariance == 0; }
+	
+	/** @return Returns random duration around Duration, varied by +-RandomVariation. */
+	float GetRandomDuration() const
+	{
+		const int32 MinVal = FMath::Max(0, static_cast<int32>(Duration) - static_cast<int32>(RandomVariance));
+		const int32 MaxVal = static_cast<int32>(Duration) + static_cast<int32>(RandomVariance);
+		return FMath::RandRange(MinVal, MaxVal) / Scale;
+	}
+	
+protected:
+
+	static constexpr float Scale = 100.0f;
+
+	uint16 Quantize(const float Value) const
+	{
+		return (uint16)FMath::Clamp(FMath::RoundToInt32(Value * Scale), 0, (int32)MAX_uint16);
+	}
+	
+	UPROPERTY(EditDefaultsOnly, Category = Default)
+	uint16 Duration = 0;
+
+	UPROPERTY(EditDefaultsOnly, Category = Default)
+	uint16 RandomVariance = 0;
+};
+
 /**
  * Describes a state tree transition. Source is the state where the transition started, Target describes the state where the transition pointed at,
  * and Next describes the selected state. The reason Transition and Next are different is that Transition state can be a selector state,
@@ -447,6 +509,12 @@ struct STATETREEMODULE_API FCompactStateTransition
 {
 	GENERATED_BODY()
 
+	/** @return True if the transition has delay. */
+	bool HasDelay() const
+	{
+		return !Delay.IsEmpty();
+	}
+	
 	/** Transition event tag, used when trigger type is event. */
 	UPROPERTY()
 	FGameplayTag EventTag;
@@ -459,6 +527,10 @@ struct STATETREEMODULE_API FCompactStateTransition
 	UPROPERTY()
 	FStateTreeStateHandle State = FStateTreeStateHandle::Invalid;
 
+	/** Transition delay. */
+	UPROPERTY()
+	FStateTreeRandomTimeDuration Delay;
+	
 	/** Type of the transition. */
 	UPROPERTY()
 	EStateTreeTransitionType Type = EStateTreeTransitionType::NotSet;
@@ -467,10 +539,6 @@ struct STATETREEMODULE_API FCompactStateTransition
 	UPROPERTY()
 	EStateTreeTransitionTrigger Trigger = EStateTreeTransitionTrigger::None;
 
-	/** The time the conditions need to hold true for the transition to become active, in tenths of a seconds. */
-	UPROPERTY()
-	uint8 GateDelay = 0;
-
 	/** Number of conditions to test. */
 	UPROPERTY()
 	uint8 ConditionsNum = 0;
