StateTree: Be more resilient when loading out of date BP data
- Try to convert to latest version when duplicating data
- Convert property binding source/target classes to latest version on link
- Handle BP reference replace in instanced structs

diff --git a/Engine/Plugins/Experimental/StructUtils/Source/StructUtils/Private/InstancedStruct.cpp b/Engine/Plugins/Experimental/StructUtils/Source/StructUtils/Private/InstancedStruct.cpp
--- a/Engine/Plugins/Experimental/StructUtils/Source/StructUtils/Private/InstancedStruct.cpp
+++ b/Engine/Plugins/Experimental/StructUtils/Source/StructUtils/Private/InstancedStruct.cpp
@@ -220,6 +220,17 @@ bool FInstancedStruct::Serialize(FArchive& Ar)
 		Ar << SerialSize;
 		Ar.Seek(FinalOffset);	// Reset archive to its position
 	}
+	else if (Ar.IsCountingMemory() || Ar.IsModifyingWeakAndStrongReferences())
+	{
+		// Report type
+		Ar << NonConstStruct;
+	
+		// Report value
+		if (NonConstStruct != nullptr && ensureMsgf(GetMutableMemory() != nullptr, TEXT("A valid script struct should always have allocated memory")))
+		{
+			NonConstStruct->SerializeItem(Ar, GetMutableMemory(), /* Defaults */ nullptr);
+		}
+	}
 
 	return true;
 }
diff --git a/Engine/Plugins/Experimental/StructUtils/Source/StructUtils/Private/InstancedStructArray.cpp b/Engine/Plugins/Experimental/StructUtils/Source/StructUtils/Private/InstancedStructArray.cpp
--- a/Engine/Plugins/Experimental/StructUtils/Source/StructUtils/Private/InstancedStructArray.cpp
+++ b/Engine/Plugins/Experimental/StructUtils/Source/StructUtils/Private/InstancedStructArray.cpp
@@ -539,6 +539,27 @@ bool FInstancedStructArray::Serialize(FArchive& Ar)
 				Ar.Seek(FinalOffset);	// Reset archive to its position
 			}
 		}
+		else if (Ar.IsCountingMemory() || Ar.IsModifyingWeakAndStrongReferences())
+		{
+			// Report item types
+			for (int32 Index = 0; Index < NumItems; Index++)
+			{
+				const FItem& Item = GetItem(Index);
+				UScriptStruct* NonConstStruct = const_cast<UScriptStruct*>(Item.ScriptStruct);
+				Ar << NonConstStruct;
+			}
+
+			// Report item values
+			for (int32 Index = 0; Index < NumItemsSerialized; Index++)
+			{
+				const FItem& Item = GetItem(Index);
+				UScriptStruct* NonConstStruct = const_cast<UScriptStruct*>(Item.ScriptStruct);
+				if (NonConstStruct != nullptr)
+				{
+					NonConstStruct->SerializeItem(Ar, Memory + Item.Offset, /* Defaults */ nullptr);
+				}
+			}
+		}
 	}
 
 	return true;
diff --git a/Engine/Plugins/Runtime/StateTree/Source/StateTreeModule/Private/StateTree.cpp b/Engine/Plugins/Runtime/StateTree/Source/StateTreeModule/Private/StateTree.cpp
--- a/Engine/Plugins/Runtime/StateTree/Source/StateTreeModule/Private/StateTree.cpp
+++ b/Engine/Plugins/Runtime/StateTree/Source/StateTreeModule/Private/StateTree.cpp
@@ -220,11 +220,35 @@ bool UStateTree::Link()
 			UE_LOG(LogStateTree, Error, TEXT("%s: StartTree does not have instance data. Please recompile the StateTree asset."), *GetName());
 			return false;
 		}
-		
+
 		// Update property bag structs before resolving binding.
 		const TArrayView<FStateTreeBindableStructDesc> SourceStructs = PropertyBindings.GetSourceStructs();
 		const TArrayView<FStateTreePropCopyBatch> CopyBatches = PropertyBindings.GetCopyBatches();
 
+		
+		// Reconcile out of date classes.
+		for (FStateTreeBindableStructDesc& SourceStruct : SourceStructs)
+		{
+			if (const UClass* SourceClass = Cast<UClass>(SourceStruct.Struct))
+			{
+				if (SourceClass->HasAnyClassFlags(CLASS_NewerVersionExists))
+				{
+					SourceStruct.Struct = SourceClass->GetAuthoritativeClass();
+				}
+			}
+		}
+		for (FStateTreePropCopyBatch& CopyBatch : CopyBatches)
+		{
+			if (const UClass* TargetClass = Cast<UClass>(CopyBatch.TargetStruct.Struct))
+			{
+				if (TargetClass->HasAnyClassFlags(CLASS_NewerVersionExists))
+				{
+					CopyBatch.TargetStruct.Struct = TargetClass->GetAuthoritativeClass();
+				}
+			}
+		}
+
+
 		if (ParametersDataViewIndex.IsValid() && SourceStructs.IsValidIndex(ParametersDataViewIndex.Get()))
 		{
 			SourceStructs[ParametersDataViewIndex.Get()].Struct = Parameters.GetPropertyBagStruct();
diff --git a/Engine/Plugins/Runtime/StateTree/Source/StateTreeModule/Private/StateTreeInstanceData.cpp b/Engine/Plugins/Runtime/StateTree/Source/StateTreeModule/Private/StateTreeInstanceData.cpp
--- a/Engine/Plugins/Runtime/StateTree/Source/StateTreeModule/Private/StateTreeInstanceData.cpp
+++ b/Engine/Plugins/Runtime/StateTree/Source/StateTreeModule/Private/StateTreeInstanceData.cpp
@@ -2,9 +2,44 @@
 
 #include "StateTreeInstanceData.h"
 #include "UObject/UnrealType.h"
+#include "StateTreeTypes.h"
 
 #include UE_INLINE_GENERATED_CPP_BY_NAME(StateTreeInstanceData)
 
+namespace UE::StateTree
+{
+
+	/**
+	 * Duplicates object, and tries to covert old BP classes (REINST_*) to their newer version.
+	 */
+	UObject* DuplicateNodeInstance(const UObject& Instance, UObject& InOwner)
+	{
+		const UClass* InstanceClass = Instance.GetClass();
+		if (InstanceClass->HasAnyClassFlags(CLASS_NewerVersionExists))
+		{
+			const UClass* AuthoritativeClass = InstanceClass->GetAuthoritativeClass();
+			UObject* NewInstance = NewObject<UObject>(&InOwner, AuthoritativeClass);
+
+			// Try to copy the values over using serialization
+			TArray<uint8> Data;
+			FMemoryWriter Writer(Data);
+			UObject& NonConstInstance = const_cast<UObject&>(Instance);
+			NonConstInstance.Serialize(Writer);
+
+			FMemoryReader Reader(Data);
+			NewInstance->Serialize(Reader);
+				
+			UE_LOG(LogStateTree, Warning, TEXT("FStateTreeInstanceData: Duplicating '%s' with old class '%s' as '%s', potential data loss."),
+				*GetFullNameSafe(&Instance), *GetNameSafe(InstanceClass), *GetNameSafe(AuthoritativeClass));
+
+			return NewInstance;
+		}
+
+		return DuplicateObject(&Instance, &InOwner);
+	}
+
+} // UE::StateTree
+
 //----------------------------------------------------------------//
 //  FStateTreeInstanceData
 //----------------------------------------------------------------//
@@ -195,8 +230,7 @@ void FStateTreeInstanceData::CopyFrom(UObject& InOwner, const FStateTreeInstance
 	{
 		if (ensure(Instance != nullptr))
 		{
-			ensure(Instance->GetClass()->HasAnyClassFlags(CLASS_NewerVersionExists) == false);
-			Storage.InstanceObjects.Add(DuplicateObject(Instance, &InOwner));
+			Storage.InstanceObjects.Add(UE::StateTree::DuplicateNodeInstance(*Instance, InOwner));
 		}
 	}
 }
@@ -224,8 +258,7 @@ void FStateTreeInstanceData::Append(UObject& InOwner, TConstArrayView<FInstanced
 	{
 		if (ensure(Instance != nullptr))
 		{
-			ensure(Instance->GetClass()->HasAnyClassFlags(CLASS_NewerVersionExists) == false);
-			Storage.InstanceObjects.Add(DuplicateObject(Instance, &InOwner));
+			Storage.InstanceObjects.Add(UE::StateTree::DuplicateNodeInstance(*Instance, InOwner));
 		}
 	}
 }
@@ -241,8 +274,7 @@ void FStateTreeInstanceData::Append(UObject& InOwner, TConstArrayView<FConstStru
 	{
 		if (ensure(Instance != nullptr))
 		{
-			ensure(Instance->GetClass()->HasAnyClassFlags(CLASS_NewerVersionExists) == false);
-			Storage.InstanceObjects.Add(DuplicateObject(Instance, &InOwner));
+			Storage.InstanceObjects.Add(UE::StateTree::DuplicateNodeInstance(*Instance, InOwner));
 		}
 	}
 }
@@ -271,4 +303,3 @@ void FStateTreeInstanceData::Reset()
 	Storage.InstanceObjects.Reset();
 	Storage.EventQueue.Reset();
 }
-
