StateTree: Fix missing include.

diff --git a/Engine/Plugins/Runtime/StateTree/Source/StateTreeModule/Private/StateTreeInstanceData.cpp b/Engine/Plugins/Runtime/StateTree/Source/StateTreeModule/Private/StateTreeInstanceData.cpp
--- a/Engine/Plugins/Runtime/StateTree/Source/StateTreeModule/Private/StateTreeInstanceData.cpp
+++ b/Engine/Plugins/Runtime/StateTree/Source/StateTreeModule/Private/StateTreeInstanceData.cpp
@@ -3,6 +3,8 @@
 #include "StateTreeInstanceData.h"
 #include "UObject/UnrealType.h"
 #include "StateTreeTypes.h"
+#include "Serialization/MemoryReader.h"
+#include "Serialization/MemoryWriter.h"
 
 #include UE_INLINE_GENERATED_CPP_BY_NAME(StateTreeInstanceData)
 
