StateTree: Fix localization key duplicate.

diff --git a/Engine/Plugins/Runtime/StateTree/Source/StateTreeEditorModule/Private/Customizations/StateTreeStateLinkDetails.cpp b/Engine/Plugins/Runtime/StateTree/Source/StateTreeEditorModule/Private/Customizations/StateTreeStateLinkDetails.cpp
--- a/Engine/Plugins/Runtime/StateTree/Source/StateTreeEditorModule/Private/Customizations/StateTreeStateLinkDetails.cpp
+++ b/Engine/Plugins/Runtime/StateTree/Source/StateTreeEditorModule/Private/Customizations/StateTreeStateLinkDetails.cpp
@@ -178,15 +178,12 @@ TSharedRef<SWidget> FStateTreeStateLinkDetails::OnGetStateContent() const
 
 	if (CachedNames.Num() > 0)
 	{
-		MenuBuilder.BeginSection(FName(), LOCTEXT("TransitionGotoState", "Go to State"));
+		MenuBuilder.BeginSection(FName(), LOCTEXT("TransitionGotoStateSection", "Go to State"));
 
 		for (int32 Idx = 0; Idx < CachedNames.Num(); Idx++)
 		{
-			FFormatNamedArguments Args;
-			Args.Add(TEXT("StateName"), FText::FromName(CachedNames[Idx]));
-
 			FUIAction ItemAction(FExecuteAction::CreateSP(const_cast<FStateTreeStateLinkDetails*>(this), &FStateTreeStateLinkDetails::OnStateComboChange, Idx));
-			MenuBuilder.AddMenuEntry(FText::FromName(CachedNames[Idx]), FText::Format(LOCTEXT("TransitionGotoStateTooltip", "Go to State {StateName}."), Args), FSlateIcon(), ItemAction);
+			MenuBuilder.AddMenuEntry(FText::FromName(CachedNames[Idx]), FText::Format(LOCTEXT("TransitionGotoStateTooltip", "Go to State {0}."), FText::FromName(CachedNames[Idx])), FSlateIcon(), ItemAction);
 		}
 
 		MenuBuilder.EndSection();
