StateTree: extend when BP nodes can send events
- cache event queue for the whole active duration of the nodes
- this allows events to be sent outside tick (e.g. from delegate callbacks)

diff --git a/Engine/Plugins/Runtime/StateTree/Source/StateTreeModule/Private/Blueprint/StateTreeConditionBlueprintBase.cpp b/Engine/Plugins/Runtime/StateTree/Source/StateTreeModule/Private/Blueprint/StateTreeConditionBlueprintBase.cpp
--- a/Engine/Plugins/Runtime/StateTree/Source/StateTreeModule/Private/Blueprint/StateTreeConditionBlueprintBase.cpp
+++ b/Engine/Plugins/Runtime/StateTree/Source/StateTreeModule/Private/Blueprint/StateTreeConditionBlueprintBase.cpp
@@ -21,8 +21,14 @@ bool UStateTreeConditionBlueprintBase::TestCondition(FStateTreeExecutionContext&
 {
 	if (bHasTestCondition)
 	{
-		FScopedCurrentContext ScopedContext(*this, Context);
-		return ReceiveTestCondition();
+		// Cache the owner and event queue for the duration the condition is evaluated.
+		SetCachedEventQueueFromContext(Context);
+
+		const bool bResult = ReceiveTestCondition();
+
+		ClearCachedEventQueue();
+
+		return bResult;
 	}
 	return false;
 }
diff --git a/Engine/Plugins/Runtime/StateTree/Source/StateTreeModule/Private/Blueprint/StateTreeEvaluatorBlueprintBase.cpp b/Engine/Plugins/Runtime/StateTree/Source/StateTreeModule/Private/Blueprint/StateTreeEvaluatorBlueprintBase.cpp
--- a/Engine/Plugins/Runtime/StateTree/Source/StateTreeModule/Private/Blueprint/StateTreeEvaluatorBlueprintBase.cpp
+++ b/Engine/Plugins/Runtime/StateTree/Source/StateTreeModule/Private/Blueprint/StateTreeEvaluatorBlueprintBase.cpp
@@ -20,9 +20,11 @@ UStateTreeEvaluatorBlueprintBase::UStateTreeEvaluatorBlueprintBase(const FObject
 
 void UStateTreeEvaluatorBlueprintBase::TreeStart(FStateTreeExecutionContext& Context)
 {
+	// Evaluator became active, cache event queue and owner.
+	SetCachedEventQueueFromContext(Context);
+
 	if (bHasTreeStart)
 	{
-		FScopedCurrentContext ScopedContext(*this, Context);
 		ReceiveTreeStart();
 	}
 }
@@ -31,16 +33,17 @@ void UStateTreeEvaluatorBlueprintBase::TreeStop(FStateTreeExecutionContext& Cont
 {
 	if (bHasTreeStop)
 	{
-		FScopedCurrentContext ScopedContext(*this, Context);
 		ReceiveTreeStop();
 	}
+
+	// Evaluator became inactive, clear cached event queue and owner.
+	ClearCachedEventQueue();
 }
 
 void UStateTreeEvaluatorBlueprintBase::Tick(FStateTreeExecutionContext& Context, const float DeltaTime)
 {
 	if (bHasTick)
 	{
-		FScopedCurrentContext ScopedContext(*this, Context);
 		ReceiveTick(DeltaTime);
 	}
 }
diff --git a/Engine/Plugins/Runtime/StateTree/Source/StateTreeModule/Private/Blueprint/StateTreeNodeBlueprintBase.cpp b/Engine/Plugins/Runtime/StateTree/Source/StateTreeModule/Private/Blueprint/StateTreeNodeBlueprintBase.cpp
--- a/Engine/Plugins/Runtime/StateTree/Source/StateTreeModule/Private/Blueprint/StateTreeNodeBlueprintBase.cpp
+++ b/Engine/Plugins/Runtime/StateTree/Source/StateTreeModule/Private/Blueprint/StateTreeNodeBlueprintBase.cpp
@@ -14,12 +14,9 @@ UWorld* UStateTreeNodeBlueprintBase::GetWorld() const
 	// The CDO is used by the BP editor to check for certain functionality, make it return nullptr so that the GetWorld() passes as overridden. 
 	if (!HasAnyFlags(RF_ClassDefaultObject))
 	{
-		if (CurrentContext != nullptr)
+		if (CachedOwner != nullptr)
 		{
-			if (CurrentContext->GetOwner() != nullptr)
-			{
-				return CurrentContext->GetOwner()->GetWorld();
-			}
+			return CachedOwner->GetWorld();
 		}
 		if (UObject* Outer = GetOuter())
 		{
@@ -40,13 +37,25 @@ AActor* UStateTreeNodeBlueprintBase::GetOwnerActor(const FStateTreeExecutionCont
 	return Cast<AActor>(Context.GetOwner());
 }
 
+void UStateTreeNodeBlueprintBase::SetCachedEventQueueFromContext(const FStateTreeExecutionContext& Context) const
+{
+	CachedEventQueue = &Context.GetMutableEventQueue();
+	CachedOwner = Context.GetOwner();
+}
+
+void UStateTreeNodeBlueprintBase::ClearCachedEventQueue() const
+{
+	CachedEventQueue = nullptr;
+	CachedOwner = nullptr;
+}
+
 void UStateTreeNodeBlueprintBase::SendEvent(const FStateTreeEvent& Event)
 {
-	if (CurrentContext == nullptr)
+	if (CachedEventQueue == nullptr || CachedOwner == nullptr)
 	{
-		UE_VLOG_UELOG(this, LogStateTree, Error, TEXT("Trying to call SendEvent() outside StateTree tick. Use SendEvent() on UStateTreeComponent instead for sending signals externally."));
+		UE_VLOG_UELOG(this, LogStateTree, Error, TEXT("Trying to call SendEvent() while node is not active. Use SendEvent() on UStateTreeComponent instead for sending signals externally."));
 		return;
 	}
-	CurrentContext->SendEvent(Event.Tag, Event.Payload, Event.Origin);
+	CachedEventQueue->SendEvent(CachedOwner, Event.Tag, Event.Payload, Event.Origin);
 }
 
diff --git a/Engine/Plugins/Runtime/StateTree/Source/StateTreeModule/Private/Blueprint/StateTreeTaskBlueprintBase.cpp b/Engine/Plugins/Runtime/StateTree/Source/StateTreeModule/Private/Blueprint/StateTreeTaskBlueprintBase.cpp
--- a/Engine/Plugins/Runtime/StateTree/Source/StateTreeModule/Private/Blueprint/StateTreeTaskBlueprintBase.cpp
+++ b/Engine/Plugins/Runtime/StateTree/Source/StateTreeModule/Private/Blueprint/StateTreeTaskBlueprintBase.cpp
@@ -27,9 +27,11 @@ UStateTreeTaskBlueprintBase::UStateTreeTaskBlueprintBase(const FObjectInitialize
 
 EStateTreeRunStatus UStateTreeTaskBlueprintBase::EnterState(FStateTreeExecutionContext& Context, const FStateTreeTransitionResult& Transition)
 {
+	// Task became active, cache event queue and owner.
+	SetCachedEventQueueFromContext(Context);
+	
 	if (bHasEnterState)
 	{
-		FScopedCurrentContext ScopedContext(*this, Context);
 		return ReceiveEnterState(Transition);
 	}
 	return EStateTreeRunStatus::Running;
@@ -39,16 +41,17 @@ void UStateTreeTaskBlueprintBase::ExitState(FStateTreeExecutionContext& Context,
 {
 	if (bHasExitState)
 	{
-		FScopedCurrentContext ScopedContext(*this, Context);
 		ReceiveExitState(Transition);
 	}
+
+	// Task became inactive, clear cached event queue and owner.
+	ClearCachedEventQueue();
 }
 
 void UStateTreeTaskBlueprintBase::StateCompleted(FStateTreeExecutionContext& Context, const EStateTreeRunStatus CompletionStatus, const FStateTreeActiveStates& CompletedActiveStates)
 {
 	if (bHasStateCompleted)
 	{
-		FScopedCurrentContext ScopedContext(*this, Context);
 		ReceiveStateCompleted(CompletionStatus, CompletedActiveStates);
 	}
 }
@@ -57,7 +60,6 @@ EStateTreeRunStatus UStateTreeTaskBlueprintBase::Tick(FStateTreeExecutionContext
 {
 	if (bHasTick)
 	{
-		FScopedCurrentContext ScopedContext(*this, Context);
 		return ReceiveTick(DeltaTime);
 	}
 	return EStateTreeRunStatus::Running;
diff --git a/Engine/Plugins/Runtime/StateTree/Source/StateTreeModule/Public/Blueprint/StateTreeNodeBlueprintBase.h b/Engine/Plugins/Runtime/StateTree/Source/StateTreeModule/Public/Blueprint/StateTreeNodeBlueprintBase.h
--- a/Engine/Plugins/Runtime/StateTree/Source/StateTreeModule/Public/Blueprint/StateTreeNodeBlueprintBase.h
+++ b/Engine/Plugins/Runtime/StateTree/Source/StateTreeModule/Public/Blueprint/StateTreeNodeBlueprintBase.h
@@ -4,6 +4,7 @@
 
 #include "StateTreeTypes.h"
 #include "StateTreeEvents.h"
+#include "StateTreeExecutionContext.h"
 #include "StateTreeNodeBlueprintBase.generated.h"
 
 struct FStateTreeLinker;
@@ -42,21 +43,15 @@ protected:
 	virtual UWorld* GetWorld() const override;
 	AActor* GetOwnerActor(const FStateTreeExecutionContext& Context) const;
 
-	struct FScopedCurrentContext
-	{
-		FScopedCurrentContext(const UStateTreeNodeBlueprintBase& InNode, FStateTreeExecutionContext& Context) :
-			Node(InNode)
-		{
-			Node.CurrentContext = &Context;
-		}
-
-		~FScopedCurrentContext()
-		{
-			Node.CurrentContext = nullptr;
-		}
-		const UStateTreeNodeBlueprintBase& Node;
-	};
+	/** These methods are const as they set mutable variables and need to be called from a const method. */
+	void SetCachedEventQueueFromContext(const FStateTreeExecutionContext& Context) const;
+	void ClearCachedEventQueue() const;
 	
-	/** Cached execution context during Tick() and other functions. */
-	mutable FStateTreeExecutionContext* CurrentContext = nullptr;
+private:
+	/** Cached instance data while the node is active. */
+	mutable FStateTreeEventQueue* CachedEventQueue = nullptr;
+
+	/** Cached owner while the node is active. */
+	UPROPERTY()
+	mutable TObjectPtr<UObject> CachedOwner = nullptr; 
 };
