StateTree: small editor fixes
- updated type tooltip
- use stable sort to keep item order in same category intact

diff --git a/Engine/Plugins/Runtime/StateTree/Source/StateTreeEditorModule/Private/Customizations/StateTreeEditorNodeDetails.cpp b/Engine/Plugins/Runtime/StateTree/Source/StateTreeEditorModule/Private/Customizations/StateTreeEditorNodeDetails.cpp
--- a/Engine/Plugins/Runtime/StateTree/Source/StateTreeEditorModule/Private/Customizations/StateTreeEditorNodeDetails.cpp
+++ b/Engine/Plugins/Runtime/StateTree/Source/StateTreeEditorModule/Private/Customizations/StateTreeEditorNodeDetails.cpp
@@ -445,6 +445,7 @@ void FStateTreeEditorNodeDetails::CustomizeHeader(TSharedRef<class IPropertyHand
 			[
 				SAssignNew(ComboButton, SComboButton)
 				.OnGetMenuContent(this, &FStateTreeEditorNodeDetails::GeneratePicker)
+				.ToolTipText(this, &FStateTreeEditorNodeDetails::GetDisplayValueString)
 				.ContentPadding(0.f)
 				.ButtonContent()
 				[
@@ -475,7 +476,6 @@ void FStateTreeEditorNodeDetails::CustomizeHeader(TSharedRef<class IPropertyHand
 		.OverrideResetToDefault(ResetOverride)
 		.CopyAction(FUIAction(FExecuteAction::CreateSP(this, &FStateTreeEditorNodeDetails::OnCopyNode)))
 		.PasteAction(FUIAction(FExecuteAction::CreateSP(this, &FStateTreeEditorNodeDetails::OnPasteNode)));
-
 }
 
 void FStateTreeEditorNodeDetails::OnCopyNode()
@@ -652,7 +652,7 @@ void FStateTreeEditorNodeDetails::CustomizeChildren(TSharedRef<class IPropertyHa
 			}
 		}
 
-		SortedChildren.Sort([](const FSortedChild& LHS, const FSortedChild& RHS) { return LHS.Usage < RHS.Usage; });
+		SortedChildren.StableSort([](const FSortedChild& LHS, const FSortedChild& RHS) { return LHS.Usage < RHS.Usage; });
 
 		for (FSortedChild& Child : SortedChildren)
 		{
@@ -705,7 +705,7 @@ void FStateTreeEditorNodeDetails::OnBindingChanged(const FStateTreeEditorPropert
 
 	for (int32 i = 0; i < OuterObjects.Num(); i++)
 	{
-		const FStateTreeEditorNode* Node = static_cast<FStateTreeEditorNode*>(RawNodeData[i]);
+		FStateTreeEditorNode* Node = static_cast<FStateTreeEditorNode*>(RawNodeData[i]);
 		UObject* OuterObject = OuterObjects[i]; // Immediate outer, i.e StateTreeState
 		if (Node != nullptr && EditorData != nullptr && Node->Node.IsValid() && Node->Instance.IsValid())
 		{
@@ -1413,8 +1413,6 @@ void FStateTreeEditorNodeDetails::OnClassPicked(UClass* InClass) const
 
 				Node->ID = FGuid::NewGuid();
 			}
-
-			
 		}
 	}
 
