StateTree: Compiler was using wrong array when validating conditions.

diff --git a/Engine/Plugins/Runtime/StateTree/Source/StateTreeEditorModule/Private/StateTreeCompiler.cpp b/Engine/Plugins/Runtime/StateTree/Source/StateTreeEditorModule/Private/StateTreeCompiler.cpp
--- a/Engine/Plugins/Runtime/StateTree/Source/StateTreeEditorModule/Private/StateTreeCompiler.cpp
+++ b/Engine/Plugins/Runtime/StateTree/Source/StateTreeEditorModule/Private/StateTreeCompiler.cpp
@@ -705,7 +705,7 @@ bool FStateTreeCompiler::CreateCondition(UStateTreeState& State, const FStateTre
 		Cond.InstanceIndex = FStateTreeIndex16(InstanceIndex);
 		Cond.bInstanceIsObject = false;
 
-		if (!CompileAndValidateNode(StructDesc, Node.GetMutable<FStateTreeNodeBase>(),  FStateTreeDataView(InstanceStructs[InstanceIndex])))
+		if (!CompileAndValidateNode(StructDesc, Node.GetMutable<FStateTreeNodeBase>(), FStateTreeDataView(SharedInstanceStructs[InstanceIndex])))
 		{
 			return false;
 		}
@@ -730,7 +730,7 @@ bool FStateTreeCompiler::CreateCondition(UStateTreeState& State, const FStateTre
 		Cond.InstanceIndex = FStateTreeIndex16(InstanceIndex);
 		Cond.bInstanceIsObject = true;
 
-		if (!CompileAndValidateNode(StructDesc, Node.GetMutable<FStateTreeNodeBase>(),  FStateTreeDataView(Instance)))
+		if (!CompileAndValidateNode(StructDesc, Node.GetMutable<FStateTreeNodeBase>(), FStateTreeDataView(Instance)))
 		{
 			return false;
 		}
