StateTree: Fixed enum property bindings
- consolidate naming of invalid/empty to none in any enum details
- fix UStateTreeEditorData::GetStructByID (was generally returning first context data)
- when binding to an enum property, prune potential source properties based on enum type

diff --git a/Engine/Plugins/Runtime/StateTree/Source/StateTreeEditorModule/Private/Customizations/StateTreeAnyEnumDetails.cpp b/Engine/Plugins/Runtime/StateTree/Source/StateTreeEditorModule/Private/Customizations/StateTreeAnyEnumDetails.cpp
--- a/Engine/Plugins/Runtime/StateTree/Source/StateTreeEditorModule/Private/Customizations/StateTreeAnyEnumDetails.cpp
+++ b/Engine/Plugins/Runtime/StateTree/Source/StateTreeEditorModule/Private/Customizations/StateTreeAnyEnumDetails.cpp
@@ -71,7 +71,7 @@ FText FStateTreeAnyEnumDetails::GetDescription() const
 		return LOCTEXT("MultipleSelected", "Multiple Selected");
 	}
 
-	return LOCTEXT("Invalid", "Invalid");
+	return LOCTEXT("None", "None");
 }
 
 TSharedRef<SWidget> FStateTreeAnyEnumDetails::OnGetComboContent() const
@@ -105,7 +105,7 @@ TSharedRef<SWidget> FStateTreeAnyEnumDetails::OnGetComboContent() const
 
 	if (!bSuccess)
 	{
-		MenuBuilder.AddMenuEntry(LOCTEXT("Empty", "Empty"), TAttribute<FText>(), FSlateIcon(), FUIAction());
+		MenuBuilder.AddMenuEntry(LOCTEXT("None", "None"), TAttribute<FText>(), FSlateIcon(), FUIAction());
 	}
 
 	return MenuBuilder.MakeWidget();
diff --git a/Engine/Plugins/Runtime/StateTree/Source/StateTreeEditorModule/Private/Customizations/StateTreeBindingExtension.cpp b/Engine/Plugins/Runtime/StateTree/Source/StateTreeEditorModule/Private/Customizations/StateTreeBindingExtension.cpp
--- a/Engine/Plugins/Runtime/StateTree/Source/StateTreeEditorModule/Private/Customizations/StateTreeBindingExtension.cpp
+++ b/Engine/Plugins/Runtime/StateTree/Source/StateTreeEditorModule/Private/Customizations/StateTreeBindingExtension.cpp
@@ -153,25 +153,25 @@ void FStateTreeBindingExtension::ExtendWidgetRow(FDetailWidgetRow& InWidgetRow,
 
 	FProperty* Property = InPropertyHandle->GetProperty();
 
-	bool bIsDataRef = false;
+	bool bIsStructRef = false;
 	bool bIsAnyEnum = false;
 	if (FStructProperty* StructProperty = CastField<FStructProperty>(Property))
 	{
 		bIsAnyEnum = StructProperty->Struct == FStateTreeAnyEnum::StaticStruct();
-		bIsDataRef = StructProperty->Struct == FStateTreeStructRef::StaticStruct();
+		bIsStructRef = StructProperty->Struct == FStateTreeStructRef::StaticStruct();
 	}
 
-	const UScriptStruct* DataRefBaseStruct = nullptr;
-	if (bIsDataRef)
+	const UScriptStruct* StructRefBaseStruct = nullptr;
+	if (bIsStructRef)
 	{
 		FString BaseStructName;
-		DataRefBaseStruct = UE::StateTree::Compiler::GetBaseStructFromMetaData(Property, BaseStructName);
+		StructRefBaseStruct = UE::StateTree::Compiler::GetBaseStructFromMetaData(Property, BaseStructName);
 	}
 	
 	FPropertyBindingWidgetArgs Args;
 	Args.Property = InPropertyHandle->GetProperty();
 
-	Args.OnCanBindProperty = FOnCanBindProperty::CreateLambda([EditorBindings, OwnerObject, InPropertyHandle, bIsAnyEnum, bIsDataRef, DataRefBaseStruct](FProperty* InProperty)
+	Args.OnCanBindProperty = FOnCanBindProperty::CreateLambda([EditorBindings, OwnerObject, InPropertyHandle, bIsAnyEnum, bIsStructRef, StructRefBaseStruct](FProperty* InProperty)
 		{
 			if (!EditorBindings || !OwnerObject)
 			{
@@ -211,7 +211,7 @@ void FStateTreeBindingExtension::ExtendWidgetRow(FDetailWidgetRow& InWidgetRow,
 					bCanBind = bAllowAnyBinding || AnyEnum.Enum == EnumProperty->GetEnum();
 				}
 			}
-			else if (bIsDataRef && DataRefBaseStruct != nullptr)
+			else if (bIsStructRef && StructRefBaseStruct != nullptr)
 			{
 				if (const FStructProperty* SourceStructProperty = CastField<FStructProperty>(InProperty))
 				{
@@ -219,11 +219,11 @@ void FStateTreeBindingExtension::ExtendWidgetRow(FDetailWidgetRow& InWidgetRow,
 					{
 						FString SourceBaseStructName;
 						const UScriptStruct* SourceDataRefBaseStruct = UE::StateTree::Compiler::GetBaseStructFromMetaData(SourceStructProperty, SourceBaseStructName);
-						bCanBind = SourceDataRefBaseStruct && SourceDataRefBaseStruct->IsChildOf(DataRefBaseStruct);
+						bCanBind = SourceDataRefBaseStruct && SourceDataRefBaseStruct->IsChildOf(StructRefBaseStruct);
 					}
 					else
 					{
-						bCanBind = SourceStructProperty->Struct && SourceStructProperty->Struct->IsChildOf(DataRefBaseStruct);
+						bCanBind = SourceStructProperty->Struct && SourceStructProperty->Struct->IsChildOf(StructRefBaseStruct);
 					}
 				}
 			}
diff --git a/Engine/Plugins/Runtime/StateTree/Source/StateTreeEditorModule/Private/StateTreeEditorData.cpp b/Engine/Plugins/Runtime/StateTree/Source/StateTreeEditorModule/Private/StateTreeEditorData.cpp
--- a/Engine/Plugins/Runtime/StateTree/Source/StateTreeEditorModule/Private/StateTreeEditorData.cpp
+++ b/Engine/Plugins/Runtime/StateTree/Source/StateTreeEditorModule/Private/StateTreeEditorData.cpp
@@ -259,12 +259,15 @@ bool UStateTreeEditorData::GetStructByID(const FGuid StructID, FStateTreeBindabl
 	{
 		for (const FStateTreeExternalDataDesc& Desc : Schema->GetContextDataDescs())
 		{
-			OutStructDesc.Struct = Desc.Struct;
-			OutStructDesc.Name = Desc.Name;
-			OutStructDesc.ID = Desc.ID;
-			bResult = true;
-			break;
-		}	
+			if (StructID == Desc.ID)
+			{
+				OutStructDesc.Struct = Desc.Struct;
+				OutStructDesc.Name = Desc.Name;
+				OutStructDesc.ID = Desc.ID;
+				bResult = true;
+				break;
+			}
+		}
 	}
 
 	// Evaluators
@@ -274,11 +277,14 @@ bool UStateTreeEditorData::GetStructByID(const FGuid StructID, FStateTreeBindabl
 		{
 			if (const FStateTreeEvaluatorBase* Evaluator = Node.Node.GetPtr<FStateTreeEvaluatorBase>())
 			{
-				OutStructDesc.Struct = Evaluator->GetInstanceDataType();
-				OutStructDesc.Name = Evaluator->Name;
-				OutStructDesc.ID = Node.ID;
-				bResult = true;
-				break;
+				if (StructID == Node.ID)
+				{
+					OutStructDesc.Struct = Evaluator->GetInstanceDataType();
+					OutStructDesc.Name = Evaluator->Name;
+					OutStructDesc.ID = Node.ID;
+					bResult = true;
+					break;
+				}
 			}
 		}
 	}
diff --git a/Engine/Plugins/Runtime/StateTree/Source/StateTreeEditorModule/Private/StateTreePropertyBindingCompiler.cpp b/Engine/Plugins/Runtime/StateTree/Source/StateTreeEditorModule/Private/StateTreePropertyBindingCompiler.cpp
--- a/Engine/Plugins/Runtime/StateTree/Source/StateTreeEditorModule/Private/StateTreePropertyBindingCompiler.cpp
+++ b/Engine/Plugins/Runtime/StateTree/Source/StateTreeEditorModule/Private/StateTreePropertyBindingCompiler.cpp
@@ -590,15 +590,31 @@ EPropertyAccessCompatibility FStateTreePropertyBindingCompiler::GetPropertyCompa
 		return (SourceProperty->PropertyClass->IsChildOf(TargetProperty->PropertyClass)) ? EPropertyAccessCompatibility::Compatible : EPropertyAccessCompatibility::Incompatible;
 	}
 
-	// Extract underlying types for enums
-	if (const FEnumProperty* EnumPropertyA = CastField<const FEnumProperty>(FromProperty))
+
+	// When copying to an enum property, expect FromProperty to be the same enum.
+	auto GetPropertyEnum = [](const FProperty* Property) -> const UEnum*
 	{
-		FromProperty = EnumPropertyA->GetUnderlyingProperty();
+		if (const FByteProperty* ByteProperty = CastField<FByteProperty>(Property))
+		{
+			return ByteProperty->GetIntPropertyEnum();
+		}
+		if (const FEnumProperty* EnumProperty = CastField<FEnumProperty>(Property))
+		{
+			return EnumProperty->GetEnum();
+		}
+		return nullptr;
+	};
+	
+	if (const UEnum* ToPropertyEnum = GetPropertyEnum(ToProperty))
+	{
+		const UEnum* FromPropertyEnum = GetPropertyEnum(FromProperty);
+		return (ToPropertyEnum == FromPropertyEnum) ? EPropertyAccessCompatibility::Compatible : EPropertyAccessCompatibility::Incompatible;
 	}
-
-	if (const FEnumProperty* EnumPropertyB = CastField<const FEnumProperty>(ToProperty))
+	
+	// Allow source enums to be promoted to numbers.
+	if (const FEnumProperty* EnumPropertyA = CastField<const FEnumProperty>(FromProperty))
 	{
-		ToProperty = EnumPropertyB->GetUnderlyingProperty();
+		FromProperty = EnumPropertyA->GetUnderlyingProperty();
 	}
 
 	if (FromProperty->SameType(ToProperty))
