StateTree: Added global tasks
- Added global tasks which are run at tree level similar to evals (will replace evals)
	- allows initial state selection to rely in calculated data (e.g. world conditions)
	- allows better reuse of existing logic
- Call evaluator stop in reverse order
- Use some uses of mutable instanced structs
- Copy bindings along with ST nodes
- Remove bindings when removing a ST node, so that undo can restore them

diff --git a/Engine/Plugins/Runtime/StateTree/Source/StateTreeEditorModule/Private/Customizations/StateTreeBindingExtension.cpp b/Engine/Plugins/Runtime/StateTree/Source/StateTreeEditorModule/Private/Customizations/StateTreeBindingExtension.cpp
--- a/Engine/Plugins/Runtime/StateTree/Source/StateTreeEditorModule/Private/Customizations/StateTreeBindingExtension.cpp
+++ b/Engine/Plugins/Runtime/StateTree/Source/StateTreeEditorModule/Private/Customizations/StateTreeBindingExtension.cpp
@@ -78,6 +78,8 @@ FText GetSectionNameFromDataSource(const EStateTreeBindableStructSource Source)
 		return LOCTEXT("Parameters", "Parameters");
 	case EStateTreeBindableStructSource::Evaluator:
 		return LOCTEXT("Evaluators", "Evaluators");
+	case EStateTreeBindableStructSource::GlobalTask:
+		return LOCTEXT("StateGlobalTasks", "Global Tasks");
 	case EStateTreeBindableStructSource::State:
 		return LOCTEXT("StateParameters", "State");
 	case EStateTreeBindableStructSource::Task:
diff --git a/Engine/Plugins/Runtime/StateTree/Source/StateTreeEditorModule/Private/Customizations/StateTreeEditorDataDetails.cpp b/Engine/Plugins/Runtime/StateTree/Source/StateTreeEditorModule/Private/Customizations/StateTreeEditorDataDetails.cpp
--- a/Engine/Plugins/Runtime/StateTree/Source/StateTreeEditorModule/Private/Customizations/StateTreeEditorDataDetails.cpp
+++ b/Engine/Plugins/Runtime/StateTree/Source/StateTreeEditorModule/Private/Customizations/StateTreeEditorDataDetails.cpp
@@ -176,6 +176,12 @@ void FStateTreeEditorDataDetails::CustomizeDetails(IDetailLayoutBuilder& DetailB
 		DetailBuilder.EditCategory(EvalCategoryName).SetCategoryVisibility(false);
 	}
 
+	// Global Tasks category
+	TSharedPtr<IPropertyHandle> GlobalTasksProperty = DetailBuilder.GetProperty(TEXT("GlobalTasks"));
+	check(GlobalTasksProperty.IsValid());
+	const FName GlobalTasksCategoryName(TEXT("GlobalTasks"));
+	MakeArrayCategory(DetailBuilder, GlobalTasksCategoryName, LOCTEXT("EditorDataDetailsGlobalTasks", "Global Tasks"), /*SortOrder*/4, GlobalTasksProperty);
+
 	// Refresh the UI when the Schema changes.	
 	TSharedPtr<IPropertyHandle> SchemaProperty = DetailBuilder.GetProperty(TEXT("Schema"));
 	check(SchemaProperty.IsValid());
diff --git a/Engine/Plugins/Runtime/StateTree/Source/StateTreeEditorModule/Private/Customizations/StateTreeEditorNodeDetails.cpp b/Engine/Plugins/Runtime/StateTree/Source/StateTreeEditorModule/Private/Customizations/StateTreeEditorNodeDetails.cpp
--- a/Engine/Plugins/Runtime/StateTree/Source/StateTreeEditorModule/Private/Customizations/StateTreeEditorNodeDetails.cpp
+++ b/Engine/Plugins/Runtime/StateTree/Source/StateTreeEditorModule/Private/Customizations/StateTreeEditorNodeDetails.cpp
@@ -533,7 +533,18 @@ void FStateTreeEditorNodeDetails::OnPasteNode()
 								Node->Name = FName(Node->Name.ToString() + TEXT(" Copy"));
 							}
 
+							const FGuid OldStructID = EditorNode->ID; 
 							EditorNode->ID = FGuid::NewGuid();
+
+							// Copy bindings from the copied node.
+							if (OldStructID.IsValid() && EditorData)
+							{
+								if (FStateTreeEditorPropertyBindings* Bindings = EditorData->GetPropertyEditorBindings())
+								{
+									Bindings->CopyBindings(OldStructID, EditorNode->ID);
+								}
+							}
+
 						}
 					}
 				}
diff --git a/Engine/Plugins/Runtime/StateTree/Source/StateTreeEditorModule/Private/StateTreeCompiler.cpp b/Engine/Plugins/Runtime/StateTree/Source/StateTreeEditorModule/Private/StateTreeCompiler.cpp
--- a/Engine/Plugins/Runtime/StateTree/Source/StateTreeEditorModule/Private/StateTreeCompiler.cpp
+++ b/Engine/Plugins/Runtime/StateTree/Source/StateTreeEditorModule/Private/StateTreeCompiler.cpp
@@ -149,6 +149,13 @@ bool FStateTreeCompiler::Compile(UStateTree& InStateTree)
 		return false;
 	}
 
+	
+	if (!CreateGlobalTasks())
+	{
+		StateTree->ResetCompiled();
+		return false;
+	}
+
 	if (!CreateStateTasksAndParameters())
 	{
 		StateTree->ResetCompiled();
@@ -336,6 +343,35 @@ bool FStateTreeCompiler::CreateEvaluators()
 	return true;
 }
 
+bool FStateTreeCompiler::CreateGlobalTasks()
+{
+	const int32 GlobalTasksBegin = Nodes.Num();
+	if (const auto Validation = UE::StateTree::Compiler::IsValidCount16(GlobalTasksBegin); Validation.DidFail())
+	{
+		Validation.Log(Log, TEXT("GlobalTasksBegin"));
+		return false;
+	}
+	StateTree->GlobalTasksBegin = uint16(GlobalTasksBegin);
+
+	for (FStateTreeEditorNode& TaskNode : EditorData->GlobalTasks)
+	{
+		if (!CreateTask(TaskNode))
+		{
+			return false;
+		}
+	}
+	
+	const int32 GlobalTasksNum = Nodes.Num() - GlobalTasksBegin;
+	if (const auto Validation = UE::StateTree::Compiler::IsValidCount16(GlobalTasksNum); Validation.DidFail())
+	{
+		Validation.Log(Log, TEXT("GlobalTasksNum"));
+		return false;
+	}
+	StateTree->GlobalTasksNum = uint16(GlobalTasksNum);
+
+	return true;
+}
+
 bool FStateTreeCompiler::CreateStateTasksAndParameters()
 {
 	for (int32 i = 0; i < StateTree->States.Num(); i++)
@@ -437,13 +473,13 @@ bool FStateTreeCompiler::CreateStateTasksAndParameters()
 
 		for (FStateTreeEditorNode& TaskNode : SourceState->Tasks)
 		{
-			if (!CreateTask(*SourceState, TaskNode))
+			if (!CreateTask(TaskNode))
 			{
 				return false;
 			}
 		}
 
-		if (!CreateTask(*SourceState, SourceState->SingleTask))
+		if (!CreateTask(SourceState->SingleTask))
 		{
 			return false;
 		}
@@ -682,7 +718,7 @@ bool FStateTreeCompiler::CreateCondition(UStateTreeState& State, const FStateTre
 	}
 
 	// Copy the condition
-	const FInstancedStruct& Node = Nodes.Add_GetRef(CondNode.Node);
+	FInstancedStruct& Node = Nodes.Add_GetRef(CondNode.Node);
 	FStateTreeConditionBase& Cond = Node.GetMutable<FStateTreeConditionBase>();
 
 	Cond.Operand = Operand;
@@ -793,7 +829,7 @@ bool FStateTreeCompiler::CompileAndValidateNode(const FStateTreeBindableStructDe
 	return Result != EDataValidationResult::Invalid;
 }
 
-bool FStateTreeCompiler::CreateTask(UStateTreeState& State, const FStateTreeEditorNode& TaskNode)
+bool FStateTreeCompiler::CreateTask(const FStateTreeEditorNode& TaskNode)
 {
 	// Silently ignore empty nodes.
 	if (!TaskNode.Node.IsValid())
@@ -816,7 +852,7 @@ bool FStateTreeCompiler::CreateTask(UStateTreeState& State, const FStateTreeEdit
 	}
 
 	// Copy the task
-	const FInstancedStruct& Node = Nodes.Add_GetRef(TaskNode.Node);
+	FInstancedStruct& Node = Nodes.Add_GetRef(TaskNode.Node);
 	FStateTreeTaskBase& Task = Node.GetMutable<FStateTreeTaskBase>();
 
 	if (TaskNode.Instance.IsValid())
@@ -924,7 +960,7 @@ bool FStateTreeCompiler::CreateEvaluator(const FStateTreeEditorNode& EvalNode)
     }
 
 	// Copy the evaluator
-	const FInstancedStruct& Node = Nodes.Add_GetRef(EvalNode.Node);
+	FInstancedStruct& Node = Nodes.Add_GetRef(EvalNode.Node);
 	FStateTreeEvaluatorBase& Eval = Node.GetMutable<FStateTreeEvaluatorBase>();
 
 	if (EvalNode.Instance.IsValid())
diff --git a/Engine/Plugins/Runtime/StateTree/Source/StateTreeEditorModule/Private/StateTreeEditorData.cpp b/Engine/Plugins/Runtime/StateTree/Source/StateTreeEditorModule/Private/StateTreeEditorData.cpp
--- a/Engine/Plugins/Runtime/StateTree/Source/StateTreeEditorModule/Private/StateTreeEditorData.cpp
+++ b/Engine/Plugins/Runtime/StateTree/Source/StateTreeEditorModule/Private/StateTreeEditorData.cpp
@@ -55,9 +55,37 @@ void UStateTreeEditorData::PostEditChangeChainProperty(FPropertyChangedChainEven
 					{
 						Eval->Name = FName(Eval->Name.ToString() + TEXT(" Duplicate"));
 					}
+					
+					const FGuid OldStructID = Evaluators[ArrayIndex].ID;
 					Evaluators[ArrayIndex].ID = FGuid::NewGuid();
+					EditorBindings.CopyBindings(OldStructID, Evaluators[ArrayIndex].ID);
 				}
 			}
+			else if (MemberName == GET_MEMBER_NAME_CHECKED(UStateTreeEditorData, GlobalTasks))
+			{
+				const int32 ArrayIndex = PropertyChangedEvent.GetArrayIndex(MemberProperty->GetFName().ToString());
+				if (GlobalTasks.IsValidIndex(ArrayIndex))
+				{
+					if (FStateTreeTaskBase* Task = GlobalTasks[ArrayIndex].Node.GetMutablePtr<FStateTreeTaskBase>())
+					{
+						Task->Name = FName(Task->Name.ToString() + TEXT(" Duplicate"));
+					}
+					
+					const FGuid OldStructID = GlobalTasks[ArrayIndex].ID;
+					GlobalTasks[ArrayIndex].ID = FGuid::NewGuid();
+					EditorBindings.CopyBindings(OldStructID, GlobalTasks[ArrayIndex].ID);
+				}
+			}
+		}
+		else if (PropertyChangedEvent.ChangeType == EPropertyChangeType::ArrayRemove)
+		{
+			if (MemberProperty->GetFName() == GET_MEMBER_NAME_CHECKED(UStateTreeEditorData, Evaluators)
+				|| MemberProperty->GetFName() == GET_MEMBER_NAME_CHECKED(UStateTreeEditorData, GlobalTasks))
+			{
+				TMap<FGuid, const UStruct*> AllStructIDs;
+				GetAllStructIDs(AllStructIDs);
+				EditorBindings.RemoveUnusedBindings(AllStructIDs);
+			}
 		}
 	}
 }
@@ -138,6 +166,28 @@ void UStateTreeEditorData::GetAccessibleStructs(const TConstArrayView<const USta
 		}
 	}
 
+	// Global Tasks
+	// Global tasks can access evaluators and other global tasks that come before them.
+	for (const FStateTreeEditorNode& Node : GlobalTasks)
+	{
+		if (const FStateTreeTaskBase* Task = Node.Node.GetPtr<FStateTreeTaskBase>())
+		{
+			// Stop iterating as soon as we find the target node.
+			if (Node.ID == TargetStructID)
+			{
+				bFoundTarget = true;
+				break;
+			}
+
+			// Collect global tasks accessible so far.
+			FStateTreeBindableStructDesc& Desc = OutStructDescs.AddDefaulted_GetRef();
+			Desc.Struct = Task->GetInstanceDataType();
+			Desc.Name = Task->Name;
+			Desc.ID = Node.ID;
+			Desc.DataSource = EStateTreeBindableStructSource::GlobalTask;
+		}
+	}
+
 	// Conditions and Tasks
 	// Visit the tree in execution order. Conditions and tasks can access tasks that are executed before them.
 	if (!bFoundTarget)
@@ -289,6 +339,25 @@ bool UStateTreeEditorData::GetStructByID(const FGuid StructID, FStateTreeBindabl
 		}
 	}
 
+	// Global tasks
+	if (!bResult)
+	{
+		for (const FStateTreeEditorNode& Node : GlobalTasks)
+		{
+			if (const FStateTreeTaskBase* Task = Node.Node.GetPtr<FStateTreeTaskBase>())
+			{
+				if (StructID == Node.ID)
+				{
+					OutStructDesc.Struct = Task->GetInstanceDataType();
+					OutStructDesc.Name = Task->Name;
+					OutStructDesc.ID = Node.ID;
+					bResult = true;
+					break;
+				}
+			}
+		}
+	}
+
 	if (!bResult)
 	{
 		VisitHierarchyNodes([&bResult, &OutStructDesc, StructID](const UStateTreeState* State, const FGuid& ID, const FName& Name, const EStateTreeNodeType NodeType, const UScriptStruct* NodeStruct, const UStruct* InstanceStruct)
@@ -371,6 +440,15 @@ void UStateTreeEditorData::GetAllStructIDs(TMap<FGuid, const UStruct*>& AllStruc
 		}
 	}
 
+	// Global tasks
+	for (const FStateTreeEditorNode& Node : GlobalTasks)
+	{
+		if (const FStateTreeTaskBase* Task = Node.Node.GetPtr<FStateTreeTaskBase>())
+		{
+			AllStructs.Emplace(Node.ID, Task->GetInstanceDataType());
+		}
+	}
+
 	VisitHierarchyNodes([&AllStructs](const UStateTreeState* State, const FGuid& ID, const FName& Name, const EStateTreeNodeType NodeType, const UScriptStruct* NodeStruct, const UStruct* InstanceStruct)
 		{
 			AllStructs.Add(ID, InstanceStruct);
@@ -506,5 +584,3 @@ void UStateTreeEditorData::VisitHierarchyNodes(TFunctionRef<EStateTreeVisitor(co
 		return VisitStateNodes(State, InFunc);
 	});
 }
-
-
diff --git a/Engine/Plugins/Runtime/StateTree/Source/StateTreeEditorModule/Private/StateTreeEditorPropertyBindings.cpp b/Engine/Plugins/Runtime/StateTree/Source/StateTreeEditorModule/Private/StateTreeEditorPropertyBindings.cpp
--- a/Engine/Plugins/Runtime/StateTree/Source/StateTreeEditorModule/Private/StateTreeEditorPropertyBindings.cpp
+++ b/Engine/Plugins/Runtime/StateTree/Source/StateTreeEditorModule/Private/StateTreeEditorPropertyBindings.cpp
@@ -28,6 +28,25 @@ void FStateTreeEditorPropertyBindings::RemovePropertyBindings(const FStateTreeEd
 		});
 }
 
+void FStateTreeEditorPropertyBindings::CopyBindings(const FGuid FromStructID, const FGuid ToStructID)
+{
+	// Copy all bindings that target "FromStructID" and retarget them to "ToStructID".
+	TArray<FStateTreeEditorPropertyBinding> NewBindings;
+	for (const FStateTreeEditorPropertyBinding& Binding : PropertyBindings)
+	{
+		if (Binding.TargetPath.StructID == FromStructID)
+		{
+			FStateTreeEditorPropertyBinding& NewBinding = NewBindings.Add_GetRef(Binding);
+			NewBinding.TargetPath.StructID = ToStructID;
+		}
+	}
+
+	for (const FStateTreeEditorPropertyBinding& NewBinding : NewBindings)
+	{
+		AddPropertyBinding(NewBinding.SourcePath, NewBinding.TargetPath);
+	}
+}
+
 bool FStateTreeEditorPropertyBindings::HasPropertyBinding(const FStateTreeEditorPropertyPath& TargetPath) const
 {
 	return PropertyBindings.ContainsByPredicate([TargetPath](const FStateTreeEditorPropertyBinding& Binding)
diff --git a/Engine/Plugins/Runtime/StateTree/Source/StateTreeEditorModule/Private/StateTreeState.cpp b/Engine/Plugins/Runtime/StateTree/Source/StateTreeEditorModule/Private/StateTreeState.cpp
--- a/Engine/Plugins/Runtime/StateTree/Source/StateTreeEditorModule/Private/StateTreeState.cpp
+++ b/Engine/Plugins/Runtime/StateTree/Source/StateTreeEditorModule/Private/StateTreeState.cpp
@@ -77,6 +77,21 @@ void UStateTreeState::PostEditChangeChainProperty(FPropertyChangedChainEvent& Pr
 		}
 		return INDEX_NONE;
 	};
+
+	auto CopyBindings = [this](const FGuid FromStructID, const FGuid ToStructID)
+	{
+		if (UStateTreeEditorData* EditorData = GetTypedOuter<UStateTreeEditorData>())
+		{
+		    if (FromStructID.IsValid())
+            {
+                if (FStateTreeEditorPropertyBindings* Bindings = EditorData->GetPropertyEditorBindings())
+                {
+                	Bindings->CopyBindings(FromStructID, ToStructID);
+                }
+            }
+		}
+	};
+
 	
 	if (Property)
 	{
@@ -160,7 +175,9 @@ void UStateTreeState::PostEditChangeChainProperty(FPropertyChangedChainEvent& Pr
 						{
 							Task->Name = FName(Task->Name.ToString() + TEXT(" Duplicate"));
 						}
+						const FGuid OldStructID = Tasks[ArrayIndex].ID; 
 						Tasks[ArrayIndex].ID = FGuid::NewGuid();
+						CopyBindings(OldStructID, Tasks[ArrayIndex].ID);
 					}
 				}
 				if (MemberProperty->GetFName() == GET_MEMBER_NAME_CHECKED(UStateTreeState, EnterConditions))
@@ -172,7 +189,9 @@ void UStateTreeState::PostEditChangeChainProperty(FPropertyChangedChainEvent& Pr
 						{
 							Condition->Name = FName(Condition->Name.ToString() + TEXT(" Duplicate"));
 						}
+						const FGuid OldStructID = EnterConditions[ArrayIndex].ID; 
 						EnterConditions[ArrayIndex].ID = FGuid::NewGuid();
+						CopyBindings(OldStructID, EnterConditions[ArrayIndex].ID);
 					}
 				}
 				if (MemberProperty->GetFName() == GET_MEMBER_NAME_CHECKED(UStateTreeState, Transitions))
@@ -189,7 +208,9 @@ void UStateTreeState::PostEditChangeChainProperty(FPropertyChangedChainEvent& Pr
 							{
 								Condition->Name = FName(Condition->Name.ToString() + TEXT(" Duplicate"));
 							}
+							const FGuid OldStructID = Transition.Conditions[ConditionsIndex].ID;
 							Transition.Conditions[ConditionsIndex].ID = FGuid::NewGuid();
+							CopyBindings(OldStructID, Transition.Conditions[ConditionsIndex].ID);
 						}
 					}
 				}
@@ -209,6 +230,23 @@ void UStateTreeState::PostEditChangeChainProperty(FPropertyChangedChainEvent& Pr
 					}
 				}
 			}
+			else if (PropertyChangedEvent.ChangeType == EPropertyChangeType::ArrayRemove)
+			{
+				if (MemberProperty->GetFName() == GET_MEMBER_NAME_CHECKED(UStateTreeState, Tasks)
+					|| MemberProperty->GetFName() == GET_MEMBER_NAME_CHECKED(UStateTreeState, EnterConditions)
+					|| MemberProperty->GetFName() == GET_MEMBER_NAME_CHECKED(UStateTreeState, Transitions))
+				{
+					if (UStateTreeEditorData* TreeData = GetTypedOuter<UStateTreeEditorData>())
+					{
+						TreeData->Modify();
+						FStateTreeEditorPropertyBindings* Bindings = TreeData->GetPropertyEditorBindings();
+						check(Bindings);
+						TMap<FGuid, const UStruct*> AllStructIDs;
+						TreeData->GetAllStructIDs(AllStructIDs);
+						Bindings->RemoveUnusedBindings(AllStructIDs);
+					}
+				}
+			}
 			else
 			{
 				if (MemberProperty->GetFName() == GET_MEMBER_NAME_CHECKED(UStateTreeState, Transitions))
diff --git a/Engine/Plugins/Runtime/StateTree/Source/StateTreeEditorModule/Public/StateTreeCompiler.h b/Engine/Plugins/Runtime/StateTree/Source/StateTreeEditorModule/Public/StateTreeCompiler.h
--- a/Engine/Plugins/Runtime/StateTree/Source/StateTreeEditorModule/Public/StateTreeCompiler.h
+++ b/Engine/Plugins/Runtime/StateTree/Source/StateTreeEditorModule/Public/StateTreeCompiler.h
@@ -39,12 +39,13 @@ private:
 	bool CreateStateRecursive(UStateTreeState& State, const FStateTreeStateHandle Parent);
 	
 	bool CreateEvaluators();
+	bool CreateGlobalTasks();
 	bool CreateStateTasksAndParameters();
 	bool CreateStateTransitions();
 	
 	bool CreateConditions(UStateTreeState& State, TConstArrayView<FStateTreeEditorNode> Conditions);
 	bool CreateCondition(UStateTreeState& State, const FStateTreeEditorNode& CondNode, const EStateTreeConditionOperand Operand, const int8 DeltaIndent);
-	bool CreateTask(UStateTreeState& State, const FStateTreeEditorNode& TaskNode);
+	bool CreateTask(const FStateTreeEditorNode& TaskNode);
 	bool CreateEvaluator(const FStateTreeEditorNode& EvalNode);
 	bool GetAndValidateBindings(const FStateTreeBindableStructDesc& TargetStruct, TArray<FStateTreeEditorPropertyBinding>& OutBindings) const;
 	bool IsPropertyAnyEnum(const FStateTreeBindableStructDesc& Struct, FStateTreeEditorPropertyPath Path) const;
diff --git a/Engine/Plugins/Runtime/StateTree/Source/StateTreeEditorModule/Public/StateTreeEditorData.h b/Engine/Plugins/Runtime/StateTree/Source/StateTreeEditorModule/Public/StateTreeEditorData.h
--- a/Engine/Plugins/Runtime/StateTree/Source/StateTreeEditorModule/Public/StateTreeEditorData.h
+++ b/Engine/Plugins/Runtime/StateTree/Source/StateTreeEditorModule/Public/StateTreeEditorData.h
@@ -109,15 +109,33 @@ public:
 	template<typename T, typename... TArgs>
 	TStateTreeEditorNode<T>& AddEvaluator(TArgs&&... InArgs)
 	{
-		FStateTreeEditorNode& EvalItem = Evaluators.AddDefaulted_GetRef();
-		EvalItem.ID = FGuid::NewGuid();
-		EvalItem.Node.InitializeAs<T>(Forward<TArgs>(InArgs)...);
-		T& Eval = EvalItem.Node.GetMutable<T>();
+		FStateTreeEditorNode& EditorNode = Evaluators.AddDefaulted_GetRef();
+		EditorNode.ID = FGuid::NewGuid();
+		EditorNode.Node.InitializeAs<T>(Forward<TArgs>(InArgs)...);
+		T& Eval = EditorNode.Node.GetMutable<T>();
 		if (const UScriptStruct* InstanceType = Cast<const UScriptStruct>(Eval.GetInstanceDataType()))
 		{
-			EvalItem.Instance.InitializeAs(InstanceType);
+			EditorNode.Instance.InitializeAs(InstanceType);
 		}
-		return static_cast<TStateTreeEditorNode<T>&>(EvalItem);
+		return static_cast<TStateTreeEditorNode<T>&>(EditorNode);
+	}
+
+	/**
+	 * Adds Global Task of specified type.
+	 * @return reference to the new task. 
+	 */
+	template<typename T, typename... TArgs>
+	TStateTreeEditorNode<T>& AddGlobalTask(TArgs&&... InArgs)
+	{
+		FStateTreeEditorNode& EditorNode = GlobalTasks.AddDefaulted_GetRef();
+		EditorNode.ID = FGuid::NewGuid();
+		EditorNode.Node.InitializeAs<T>(Forward<TArgs>(InArgs)...);
+		T& Task = EditorNode.Node.GetMutable<T>();
+		if (const UScriptStruct* InstanceType = Cast<const UScriptStruct>(Task.GetInstanceDataType()))
+		{
+			EditorNode.Instance.InitializeAs(InstanceType);
+		}
+		return static_cast<TStateTreeEditorNode<T>&>(EditorNode);
 	}
 
 	/**
@@ -140,6 +158,9 @@ public:
 	UPROPERTY(EditDefaultsOnly, Category = "Evaluators", meta = (BaseStruct = "/Script/StateTreeModule.StateTreeEvaluatorBase", BaseClass = "/Script/StateTreeModule.StateTreeEvaluatorBlueprintBase"))
 	TArray<FStateTreeEditorNode> Evaluators;
 
+	UPROPERTY(EditDefaultsOnly, Category = "Global Tasks", meta = (BaseStruct = "/Script/StateTreeModule.StateTreeTaskBase", BaseClass = "/Script/StateTreeModule.StateTreeTaskBlueprintBase"))
+	TArray<FStateTreeEditorNode> GlobalTasks;
+
 	UPROPERTY(meta = (ExcludeFromHash))
 	FStateTreeEditorPropertyBindings EditorBindings;
 
diff --git a/Engine/Plugins/Runtime/StateTree/Source/StateTreeEditorModule/Public/StateTreeEditorPropertyBindings.h b/Engine/Plugins/Runtime/StateTree/Source/StateTreeEditorModule/Public/StateTreeEditorPropertyBindings.h
--- a/Engine/Plugins/Runtime/StateTree/Source/StateTreeEditorModule/Public/StateTreeEditorPropertyBindings.h
+++ b/Engine/Plugins/Runtime/StateTree/Source/StateTreeEditorModule/Public/StateTreeEditorPropertyBindings.h
@@ -57,6 +57,14 @@ struct STATETREEEDITORMODULE_API FStateTreeEditorPropertyBindings
 	 * @return True of the target path has any bindings.
 	 */
 	bool HasPropertyBinding(const FStateTreeEditorPropertyPath& TargetPath) const;
+
+	/**
+	 * Copies property bindings from an existing struct to another.
+	 * Overrides a binding to a specific property if it already exists in ToStructID.
+	 * @param FromStructID ID of the struct to copy from.
+	 * @param ToStructID ID of the struct to copy to.
+	 */
+	void CopyBindings(const FGuid FromStructID, const FGuid ToStructID);
 	
 	/**
 	 * @return Source path for given target path, or null if binding does not exists.
diff --git a/Engine/Plugins/Runtime/StateTree/Source/StateTreeModule/Private/StateTreeExecutionContext.cpp b/Engine/Plugins/Runtime/StateTree/Source/StateTreeModule/Private/StateTreeExecutionContext.cpp
--- a/Engine/Plugins/Runtime/StateTree/Source/StateTreeModule/Private/StateTreeExecutionContext.cpp
+++ b/Engine/Plugins/Runtime/StateTree/Source/StateTreeModule/Private/StateTreeExecutionContext.cpp
@@ -47,7 +47,9 @@ void FStateTreeExecutionContext::SetDefaultParameters()
 {
 	if (DataViews.IsValidIndex(StateTree.ParametersDataViewIndex.Get()))
 	{
-		DataViews[StateTree.ParametersDataViewIndex.Get()] = FStateTreeDataView(StateTree.GetDefaultParameters().GetMutableValue());	
+		// @todo: Handle constness correctly.
+		const FConstStructView ConstParameters = StateTree.GetDefaultParameters().GetValue();
+		DataViews[StateTree.ParametersDataViewIndex.Get()] = FStateTreeDataView(ConstParameters.GetScriptStruct(), const_cast<uint8*>(ConstParameters.GetMemory()));	
 	}
 }
 
@@ -57,7 +59,9 @@ void FStateTreeExecutionContext::SetParameters(const FInstancedPropertyBag& Para
 		TEXT("Parameters must be of the same struct type. Make sure to migrate the provided parameters to the same type as the StateTree default parameters."))
 		&& DataViews.IsValidIndex(StateTree.ParametersDataViewIndex.Get()))
 	{
-		DataViews[StateTree.ParametersDataViewIndex.Get()] = FStateTreeDataView(Parameters.GetMutableValue());
+		// @todo: Handle constness correctly.
+		const FConstStructView ConstParameters = Parameters.GetValue();
+		DataViews[StateTree.ParametersDataViewIndex.Get()] = FStateTreeDataView(ConstParameters.GetScriptStruct(), const_cast<uint8*>(ConstParameters.GetMemory()));	
 	}
 }
 
@@ -165,27 +169,33 @@ EStateTreeRunStatus FStateTreeExecutionContext::Start()
 	// Stop if still running previous state.
 	Stop();
 
-	// Initialize instance data if needed.
+	// Initialize instance data. No active states yet, so we'll initialize the evals and global tasks.
+	InstanceData.Reset();
+	const FStateTreeActiveStates Empty;
+	UpdateInstanceData(Empty, Empty);
 	if (!InstanceData.IsValid())
 	{
-		const FStateTreeActiveStates Empty;
-		UpdateInstanceData(Empty, Empty);
-		if (!InstanceData.IsValid())
-		{
-			STATETREE_LOG(Warning, TEXT("%s: Failed to initialize instance data on '%s' using StateTree '%s'. Try to recompile the StateTree asset."),
-				ANSI_TO_TCHAR(__FUNCTION__), *GetNameSafe(&Owner), *GetFullNameSafe(&StateTree));
-			return EStateTreeRunStatus::Failed;
-		}
+		STATETREE_LOG(Warning, TEXT("%s: Failed to initialize instance data on '%s' using StateTree '%s'. Try to recompile the StateTree asset."),
+			ANSI_TO_TCHAR(__FUNCTION__), *GetNameSafe(&Owner), *GetFullNameSafe(&StateTree));
+		return EStateTreeRunStatus::Failed;
 	}
 
 	const TSharedPtr<FStateTreeInstanceData> SharedInstanceData = StateTree.GetSharedInstanceData();
 	check(SharedInstanceData.IsValid());
 
-	// Call TreeStart on evaluators.
-	StartEvaluators();
+	// Start evaluators and global tasks. Fail the execution if any global task fails.
+	FStateTreeIndex16 LastInitializedTaskIndex;
+	if (!StartEvaluatorsAndGlobalTasks(LastInitializedTaskIndex))
+	{
+		StopEvaluatorsAndGlobalTasks(LastInitializedTaskIndex);
+		STATETREE_LOG(Warning, TEXT("%s: Failed to start evaluators and global tasks '%s' using StateTree '%s'. Try to recompile the StateTree asset."),
+			ANSI_TO_TCHAR(__FUNCTION__), *GetNameSafe(&Owner), *GetFullNameSafe(&StateTree));
+		return EStateTreeRunStatus::Failed;
+	}
 
-	// First tick
-	TickEvaluators(0.0f);
+	// First tick.
+	// Tasks are not ticked here, since their behavior is that EnterState() (called above) is treated as a tick.  
+	TickEvaluatorsAndGlobalTasks(0.0f, /*bTickGlobalTasks*/false);
 
 	// Initialize to unset running state.
 	FStateTreeExecutionState* Exec = &GetExecState(); // Using pointer as we will need to reacquire the exec later.
@@ -255,19 +265,13 @@ EStateTreeRunStatus FStateTreeExecutionContext::Stop()
 	}
 
 	FStateTreeExecutionState& Exec = GetExecState();
-
-	if (Exec.TreeRunStatus != EStateTreeRunStatus::Running)
-	{
-		return Exec.TreeRunStatus;
-	}
-
+	EStateTreeRunStatus Result = Exec.TreeRunStatus;
+	
 	// Capture events added between ticks.
 	FStateTreeEventQueue& EventQueue = InstanceData.GetMutableEventQueue();
 	EventsToProcess = EventQueue.GetEvents();
 	EventQueue.Reset();
 
-	TickEvaluators(0.0f);
-
 	// Exit states if still in some valid state.
 	if (Exec.TreeRunStatus == EStateTreeRunStatus::Running && !Exec.ActiveStates.IsEmpty())
 	{
@@ -280,23 +284,11 @@ EStateTreeRunStatus FStateTreeExecutionContext::Stop()
 
 		ExitState(Transition);
 
-		Exec.TreeRunStatus = EStateTreeRunStatus::Succeeded;
-		Exec.ActiveStates.Reset();
-	}
-	else
-	{
-		Exec.TreeRunStatus = Exec.ActiveStates.Last() == FStateTreeStateHandle::Succeeded ? EStateTreeRunStatus::Succeeded : EStateTreeRunStatus::Failed;
+		Result = EStateTreeRunStatus::Succeeded;
 	}
 
-	// Call TreeStop on evaluators.
-	StopEvaluators();
-
-	Exec.ActiveStates.Reset();
-	Exec.LastTickStatus = EStateTreeRunStatus::Unset;
-	Exec.FirstTaskStructIndex = FStateTreeIndex16::Invalid;
-	Exec.FirstTaskObjectIndex = FStateTreeIndex16::Invalid;
-
-	const EStateTreeRunStatus Result = Exec.TreeRunStatus; 
+	// Stop evaluators and global tasks.
+	StopEvaluatorsAndGlobalTasks();
 
 	// Destruct all allocated instance data (does not shrink the buffer). This will invalidate Exec too.
 	InstanceData.Reset();
@@ -345,8 +337,12 @@ EStateTreeRunStatus FStateTreeExecutionContext::Tick(const float DeltaTime)
 		DelayedState.TimeLeft -= DeltaTime;
 	}
 	
-	// Tick global evaluators.
-	TickEvaluators(DeltaTime);
+	// Tick global evaluators and tasks.
+	const EStateTreeRunStatus EvalAndGlobalTaskStatus = TickEvaluatorsAndGlobalTasks(DeltaTime);
+	if (EvalAndGlobalTaskStatus != EStateTreeRunStatus::Running)
+	{
+		return Stop();
+	}
 
 	if (Exec->LastTickStatus == EStateTreeRunStatus::Running)
 	{
@@ -515,6 +511,20 @@ void FStateTreeExecutionContext::UpdateInstanceData(const FStateTreeActiveStates
 		}
 	}
 
+	// Global tasks
+	for (int32 TaskIndex = StateTree.GlobalTasksBegin; TaskIndex < (StateTree.GlobalTasksBegin + StateTree.GlobalTasksNum); TaskIndex++)
+	{
+		const FStateTreeTaskBase& Task =  StateTree.Nodes[TaskIndex].Get<FStateTreeTaskBase>();
+		if (Task.bInstanceIsObject)
+		{
+			InstanceObjects.Add(StateTree.DefaultInstanceData.GetObject(Task.InstanceIndex.Get()));
+		}
+		else
+		{
+			InstanceStructs.Add(StateTree.DefaultInstanceData.GetStruct(Task.InstanceIndex.Get()));
+		}
+	}
+
 	// Expect initialized instance data to contain the common instances.
 	if (InstanceData.IsValid())
 	{
@@ -856,7 +866,7 @@ void FStateTreeExecutionContext::StateCompleted()
 	}
 }
 
-void FStateTreeExecutionContext::TickEvaluators(const float DeltaTime)
+EStateTreeRunStatus FStateTreeExecutionContext::TickEvaluatorsAndGlobalTasks(const float DeltaTime, bool bTickGlobalTasks)
 {
 	CSV_SCOPED_TIMING_STAT_EXCLUSIVE(StateTree_TickEvaluators);
 
@@ -882,31 +892,75 @@ void FStateTreeExecutionContext::TickEvaluators(const float DeltaTime)
 			Eval.Tick(*this, DeltaTime);
 		}
 	}
+
+
+	EStateTreeRunStatus Result = EStateTreeRunStatus::Running;
+
+	if (bTickGlobalTasks)
+	{
+		// Used to stop ticking tasks after one fails, but we still want to keep updating the data views so that property binding works properly.
+		bool bShouldTickTasks = true;
+		const bool bHasEvents = !EventsToProcess.IsEmpty();
+
+		for (int32 TaskIndex = StateTree.GlobalTasksBegin; TaskIndex < (StateTree.GlobalTasksBegin + StateTree.GlobalTasksNum); TaskIndex++)
+		{
+			const FStateTreeTaskBase& Task =  StateTree.Nodes[TaskIndex].Get<FStateTreeTaskBase>();
+			SetNodeDataView(Task, InstanceStructIndex, InstanceObjectIndex);
+
+			const bool bNeedsTick = bShouldTickTasks && (Task.bShouldCallTick || (bHasEvents && Task.bShouldCallTickOnlyOnEvents));
+			STATETREE_LOG(VeryVerbose, TEXT("  Tick: '%s' %s"), *Task.Name.ToString(), !bNeedsTick ? TEXT("[not ticked]") : TEXT(""));
+			if (!bNeedsTick)
+			{
+				continue;
+			}
+
+			// Copy bound properties.
+			// Only copy properties when the task is actually ticked, and copy properties at tick is requested.
+			if (Task.BindingsBatch.IsValid() && Task.bShouldCopyBoundPropertiesOnTick)
+			{
+				StateTree.PropertyBindings.CopyTo(DataViews, Task.BindingsBatch, DataViews[Task.DataViewIndex.Get()]);
+			}
+
+			{
+				QUICK_SCOPE_CYCLE_COUNTER(StateTree_Task_Tick);
+				CSV_SCOPED_TIMING_STAT_EXCLUSIVE(StateTree_Task_Tick);
+
+				const EStateTreeRunStatus TaskResult = Task.Tick(*this, DeltaTime);
+
+				// If a global task succeeds or fails, it will stop the whole tree.
+				if (TaskResult != EStateTreeRunStatus::Running)
+				{
+					Result = TaskResult;
+				}
+					
+				if (TaskResult == EStateTreeRunStatus::Failed)
+				{
+					bShouldTickTasks = false;
+				}
+			}
+		}
+	}
+
+	return Result;
 }
 
-void FStateTreeExecutionContext::StartEvaluators()
+bool FStateTreeExecutionContext::StartEvaluatorsAndGlobalTasks(FStateTreeIndex16& OutLastInitializedTaskIndex)
 {
 	CSV_SCOPED_TIMING_STAT_EXCLUSIVE(StateTree_StartEvaluators);
 
-	STATETREE_CLOG(StateTree.EvaluatorsNum > 0, Verbose, TEXT("Start Evaluators"));
+	STATETREE_CLOG(StateTree.EvaluatorsNum > 0 || StateTree.GlobalTasksNum > 0, Verbose, TEXT("Start Evaluators & Global tasks"));
 
-	// Tick evaluators
+	OutLastInitializedTaskIndex = FStateTreeIndex16();
+	bool bResult = true;
+	
+	// Start evaluators
 	int32 InstanceStructIndex = 1; // Exec is at index 0
 	int32 InstanceObjectIndex = 0;
 	
 	for (int32 EvalIndex = StateTree.EvaluatorsBegin; EvalIndex < (StateTree.EvaluatorsBegin + StateTree.EvaluatorsNum); EvalIndex++)
 	{
 		const FStateTreeEvaluatorBase& Eval = StateTree.Nodes[EvalIndex].Get<FStateTreeEvaluatorBase>();
-		if (Eval.bInstanceIsObject)
-		{
-			DataViews[Eval.DataViewIndex.Get()] = InstanceData.GetMutableObject(InstanceObjectIndex);
-			InstanceObjectIndex++;
-		}
-		else
-		{
-			DataViews[Eval.DataViewIndex.Get()] = InstanceData.GetMutableStruct(InstanceStructIndex);
-			InstanceStructIndex++;
-		}
+		SetNodeDataView(Eval, InstanceStructIndex, InstanceObjectIndex);
 
 		// Copy bound properties.
 		if (Eval.BindingsBatch.IsValid())
@@ -919,15 +973,43 @@ void FStateTreeExecutionContext::StartEvaluators()
 			Eval.TreeStart(*this);
 		}
 	}
+
+	// Start Global tasks
+	// Even if we call Enter/ExitState() on global tasks, they do not enter any specific state.
+	const FStateTreeTransitionResult Transition = {}; // Empty transition
+	
+	for (int32 TaskIndex = StateTree.GlobalTasksBegin; TaskIndex < (StateTree.GlobalTasksBegin + StateTree.GlobalTasksNum); TaskIndex++)
+	{
+		const FStateTreeTaskBase& Task =  StateTree.Nodes[TaskIndex].Get<FStateTreeTaskBase>();
+		SetNodeDataView(Task, InstanceStructIndex, InstanceObjectIndex);
+
+		// Copy bound properties.
+		if (Task.BindingsBatch.IsValid())
+		{
+			StateTree.PropertyBindings.CopyTo(DataViews, Task.BindingsBatch, DataViews[Task.DataViewIndex.Get()]);
+		}
+		STATETREE_LOG(Verbose, TEXT("  Start: '%s'"), *Task.Name.ToString());
+		{
+			QUICK_SCOPE_CYCLE_COUNTER(StateTree_Task_TreeStart);
+			if (Task.EnterState(*this, Transition) != EStateTreeRunStatus::Running)
+			{
+				OutLastInitializedTaskIndex = FStateTreeIndex16(TaskIndex);
+				bResult = false;
+				break;
+			}
+		}
+	}
+
+	return bResult;
 }
 
-void FStateTreeExecutionContext::StopEvaluators()
+void FStateTreeExecutionContext::StopEvaluatorsAndGlobalTasks(const FStateTreeIndex16 LastInitializedTaskIndex)
 {
 	CSV_SCOPED_TIMING_STAT_EXCLUSIVE(StateTree_StopEvaluators);
 
-	STATETREE_CLOG(StateTree.EvaluatorsNum > 0, Verbose, TEXT("Stop Evaluators"));
+	STATETREE_CLOG(StateTree.EvaluatorsNum > 0, Verbose, TEXT("Stop Evaluators & Global Tasks"));
 
-	// Tick evaluators
+	// Stop evaluators
 	int32 InstanceStructIndex = 1; // Exec is at index 0
 	int32 InstanceObjectIndex = 0;
 	
@@ -941,6 +1023,45 @@ void FStateTreeExecutionContext::StopEvaluators()
 		{
 			StateTree.PropertyBindings.CopyTo(DataViews, Eval.BindingsBatch, DataViews[Eval.DataViewIndex.Get()]);
 		}
+	}
+
+	// Stop Global tasks
+	for (int32 TaskIndex = StateTree.GlobalTasksBegin; TaskIndex < (StateTree.GlobalTasksBegin + StateTree.GlobalTasksNum); TaskIndex++)
+	{
+		const FStateTreeTaskBase& Task =  StateTree.Nodes[TaskIndex].Get<FStateTreeTaskBase>();
+		SetNodeDataView(Task, InstanceStructIndex, InstanceObjectIndex);
+
+		// Copy bound properties.
+		if (Task.BindingsBatch.IsValid() && Task.bShouldCopyBoundPropertiesOnExitState)
+		{
+			StateTree.PropertyBindings.CopyTo(DataViews, Task.BindingsBatch, DataViews[Task.DataViewIndex.Get()]);
+		}
+	}
+
+
+	// Call in reverse order.	
+	const FStateTreeTransitionResult Transition = {}; // Empty transition
+
+	for (int32 TaskIndex = (StateTree.GlobalTasksBegin + StateTree.GlobalTasksNum) - 1;  TaskIndex >= StateTree.GlobalTasksBegin ; TaskIndex--)
+	{
+		const FStateTreeTaskBase& Task =  StateTree.Nodes[TaskIndex].Get<FStateTreeTaskBase>();
+
+		UE_LOG(LogTemp, Error, TEXT("**** %s: Exit Global Task %s exec=%s"), *GetNameSafe(&Owner), *Task.Name.ToString(), *LexToString(TaskIndex <= LastInitializedTaskIndex.Get()));
+
+		// Relying here that invalid value of LastInitializedTaskIndex == MAX_uint16.
+		if (TaskIndex <= LastInitializedTaskIndex.Get())
+		{
+			STATETREE_LOG(Verbose, TEXT("  Stop: '%s'"), *Task.Name.ToString());
+			{
+				QUICK_SCOPE_CYCLE_COUNTER(StateTree_Task_TreeStop);
+				Task.ExitState(*this, Transition);
+			}
+		}
+	}
+
+	for (int32 EvalIndex = (StateTree.EvaluatorsBegin + StateTree.EvaluatorsNum) - 1; EvalIndex >= StateTree.EvaluatorsBegin; EvalIndex--)
+	{
+		const FStateTreeEvaluatorBase& Eval = StateTree.Nodes[EvalIndex].Get<FStateTreeEvaluatorBase>();
 		STATETREE_LOG(Verbose, TEXT("  Stop: '%s'"), *Eval.Name.ToString());
 		{
 			QUICK_SCOPE_CYCLE_COUNTER(StateTree_Eval_TreeStop);
@@ -971,7 +1092,7 @@ EStateTreeRunStatus FStateTreeExecutionContext::TickTasks(const float DeltaTime)
 
 	Exec.CompletedStateHandle = FStateTreeStateHandle::Invalid;
 	
-	// Used to stop ticking tasks after one fails, but we still want to keep updating the data of them
+	// Used to stop ticking tasks after one fails, but we still want to keep updating the data views so that property binding works properly.
 	bool bShouldTickTasks = true;
 
 	STATETREE_CLOG(Exec.ActiveStates.Num() > 0, VeryVerbose, TEXT("Ticking Tasks"));
@@ -1000,17 +1121,19 @@ EStateTreeRunStatus FStateTreeExecutionContext::TickTasks(const float DeltaTime)
 			SetNodeDataView(Task, InstanceStructIndex, InstanceObjectIndex);
 
 			const bool bNeedsTick = bShouldTickTasks && (Task.bShouldCallTick || (bHasEvents && Task.bShouldCallTickOnlyOnEvents));
+			STATETREE_LOG(VeryVerbose, TEXT("%*s  Tick: '%s' %s"), Index*UE::StateTree::DebugIndentSize, TEXT(""), *Task.Name.ToString(), !bNeedsTick ? TEXT("[not ticked]") : TEXT(""));
+			if (!bNeedsTick)
+			{
+				continue;
+			}
 			
 			// Copy bound properties.
 			// Only copy properties when the task is actually ticked, and copy properties at tick is requested.
-			if (Task.BindingsBatch.IsValid() && bNeedsTick && Task.bShouldCopyBoundPropertiesOnTick)
+			if (Task.BindingsBatch.IsValid() && Task.bShouldCopyBoundPropertiesOnTick)
 			{
 				StateTree.PropertyBindings.CopyTo(DataViews, Task.BindingsBatch, DataViews[Task.DataViewIndex.Get()]);
 			}
 			
-			STATETREE_LOG(VeryVerbose, TEXT("%*s  Tick: '%s' %s"), Index*UE::StateTree::DebugIndentSize, TEXT(""), *Task.Name.ToString(), !bNeedsTick ? TEXT("[not ticked]") : TEXT(""));
-
-			if (bNeedsTick)
 			{
 				QUICK_SCOPE_CYCLE_COUNTER(StateTree_Task_Tick);
 				CSV_SCOPED_TIMING_STAT_EXCLUSIVE(StateTree_Task_Tick);
diff --git a/Engine/Plugins/Runtime/StateTree/Source/StateTreeModule/Public/StateTree.h b/Engine/Plugins/Runtime/StateTree/Source/StateTreeModule/Public/StateTree.h
--- a/Engine/Plugins/Runtime/StateTree/Source/StateTreeModule/Public/StateTree.h
+++ b/Engine/Plugins/Runtime/StateTree/Source/StateTreeModule/Public/StateTree.h
@@ -206,6 +206,14 @@ private:
 	UPROPERTY()
 	uint16 EvaluatorsNum = 0;
 
+	/** Index of first global task in Nodes. */
+	UPROPERTY()
+	uint16 GlobalTasksBegin = 0;
+
+	/** Number of global tasks. */
+	UPROPERTY()
+	uint16 GlobalTasksNum = 0;
+
 	// Data created during linking.
 	
 	/** List of external data required by the state tree, created during linking. */
diff --git a/Engine/Plugins/Runtime/StateTree/Source/StateTreeModule/Public/StateTreeExecutionContext.h b/Engine/Plugins/Runtime/StateTree/Source/StateTreeModule/Public/StateTreeExecutionContext.h
--- a/Engine/Plugins/Runtime/StateTree/Source/StateTreeModule/Public/StateTreeExecutionContext.h
+++ b/Engine/Plugins/Runtime/StateTree/Source/StateTreeModule/Public/StateTreeExecutionContext.h
@@ -318,13 +318,20 @@ protected:
 	void StateCompleted();
 
 	/**
-	 * Ticks global evaluators by delta time.
+	 * Tick evaluators and global tasks by delta time.
 	 */
-	void TickEvaluators(const float DeltaTime);
+	EStateTreeRunStatus TickEvaluatorsAndGlobalTasks(const float DeltaTime, bool bTickGlobalTasks = true);
 
-	void StartEvaluators();
+	/**
+	 * Starts evaluators and global tasks.
+	 * @return true if all evaluators and tasks were started, or false if any failed.
+	 */
+	bool StartEvaluatorsAndGlobalTasks(FStateTreeIndex16& OutLastInitializedTaskIndex);
 
-	void StopEvaluators();
+	/**
+	 * Stops evaluators and global tasks.
+	 */
+	void StopEvaluatorsAndGlobalTasks(const FStateTreeIndex16 LastInitializedTaskIndex = FStateTreeIndex16());
 
 	/**
 	 * Ticks tasks of all active states starting from current state by delta time.
diff --git a/Engine/Plugins/Runtime/StateTree/Source/StateTreeModule/Public/StateTreePropertyBindings.h b/Engine/Plugins/Runtime/StateTree/Source/StateTreeModule/Public/StateTreePropertyBindings.h
--- a/Engine/Plugins/Runtime/StateTree/Source/StateTreeModule/Public/StateTreePropertyBindings.h
+++ b/Engine/Plugins/Runtime/StateTree/Source/StateTreeModule/Public/StateTreePropertyBindings.h
@@ -19,6 +19,8 @@ enum class EStateTreeBindableStructSource : uint8
 	Parameter,
 	/** Source is StateTree evaluator */
 	Evaluator,
+	/** Source is StateTree global task */
+	GlobalTask,
 	/** Source is State parameter */
 	State,
 	/** Source is State task */
diff --git a/Engine/Plugins/Runtime/StateTree/Source/StateTreeModule/Public/StateTreeTypes.h b/Engine/Plugins/Runtime/StateTree/Source/StateTreeModule/Public/StateTreeTypes.h
--- a/Engine/Plugins/Runtime/StateTree/Source/StateTreeModule/Public/StateTreeTypes.h
+++ b/Engine/Plugins/Runtime/StateTree/Source/StateTreeModule/Public/StateTreeTypes.h
@@ -784,14 +784,14 @@ struct STATETREEMODULE_API FStateTreeStructRef
 
 	/** Returns mutable reference to the struct, this getter assumes that all data is valid. */
 	template <typename T>
-    T& GetMutable() const
+    T& GetMutable()
 	{
 		return Data.template GetMutable<T>();
 	}
 
 	/** Returns mutable pointer to the struct, or nullptr if cast is not valid. */
 	template <typename T>
-    T* GetMutablePtr() const
+    T* GetMutablePtr()
 	{
 		return Data.template GetMutablePtr<T>();
 	}
