StateTree: bump external data handly type to 16bits to not limit max nodes to 254.

diff --git a/Engine/Plugins/Runtime/StateTree/Source/StateTreeEditorModule/Private/StateTreeCompiler.cpp b/Engine/Plugins/Runtime/StateTree/Source/StateTreeEditorModule/Private/StateTreeCompiler.cpp
--- a/Engine/Plugins/Runtime/StateTree/Source/StateTreeEditorModule/Private/StateTreeCompiler.cpp
+++ b/Engine/Plugins/Runtime/StateTree/Source/StateTreeEditorModule/Private/StateTreeCompiler.cpp
@@ -128,12 +128,12 @@ bool FStateTreeCompiler::Compile(UStateTree& InStateTree)
 					Desc.ID
 				};
 			const int32 ExternalStructIndex = BindingsCompiler.AddSourceStruct(ExtDataDesc);
-			if (const auto Validation = UE::StateTree::Compiler::IsValidIndex8(ExternalStructIndex); Validation.DidFail())
+			if (const auto Validation = UE::StateTree::Compiler::IsValidIndex16(ExternalStructIndex); Validation.DidFail())
 			{
 				Validation.Log(Log, TEXT("ExternalStructIndex"), ParametersDesc);
 				return false;
 			}
-			Desc.Handle.DataViewIndex = FStateTreeIndex8(ExternalStructIndex); 
+			Desc.Handle.DataViewIndex = FStateTreeIndex16(ExternalStructIndex); 
 		} 
 	}
 	
diff --git a/Engine/Plugins/Runtime/StateTree/Source/StateTreeModule/Private/StateTreeTypes.cpp b/Engine/Plugins/Runtime/StateTree/Source/StateTreeModule/Private/StateTreeTypes.cpp
--- a/Engine/Plugins/Runtime/StateTree/Source/StateTreeModule/Private/StateTreeTypes.cpp
+++ b/Engine/Plugins/Runtime/StateTree/Source/StateTreeModule/Private/StateTreeTypes.cpp
@@ -14,3 +14,50 @@ const FStateTreeIndex16 FStateTreeIndex16::Invalid = FStateTreeIndex16();
 const FStateTreeIndex8 FStateTreeIndex8::Invalid = FStateTreeIndex8();
 
 const FStateTreeExternalDataHandle FStateTreeExternalDataHandle::Invalid = FStateTreeExternalDataHandle();
+
+
+bool FStateTreeIndex16::SerializeFromMismatchedTag(const FPropertyTag& Tag, FStructuredArchive::FSlot Slot)
+{
+	// Support loading from Index8.
+	if (Tag.Type == NAME_StructProperty
+		&& Tag.StructName == FStateTreeIndex8::StaticStruct()->GetFName())
+	{
+		FStateTreeIndex8 OldValue;
+		FStateTreeIndex8::StaticStruct()->SerializeItem(Slot, &OldValue, nullptr);
+
+		int32 NewValue = OldValue.AsInt32();
+		if (!IsValidIndex(NewValue))
+		{
+			NewValue = INDEX_NONE;
+		}
+		
+		*this = FStateTreeIndex16(NewValue);
+		
+		return true;
+	}
+	
+	return false;
+}
+
+bool FStateTreeIndex8::SerializeFromMismatchedTag(const FPropertyTag& Tag, FStructuredArchive::FSlot Slot)
+{
+	// Support loading from Index16.
+	if (Tag.Type == NAME_StructProperty
+		&& Tag.StructName == FStateTreeIndex16::StaticStruct()->GetFName())
+	{
+		FStateTreeIndex16 OldValue;
+		FStateTreeIndex16::StaticStruct()->SerializeItem(Slot, &OldValue, nullptr);
+
+		int32 NewValue = OldValue.AsInt32();
+		if (!IsValidIndex(NewValue))
+		{
+			NewValue = INDEX_NONE;
+		}
+		
+		*this = FStateTreeIndex8(NewValue);
+		
+		return true;
+	}
+	
+	return false;
+}
diff --git a/Engine/Plugins/Runtime/StateTree/Source/StateTreeModule/Public/StateTreeLinker.h b/Engine/Plugins/Runtime/StateTree/Source/StateTreeModule/Public/StateTreeLinker.h
--- a/Engine/Plugins/Runtime/StateTree/Source/StateTreeModule/Public/StateTreeLinker.h
+++ b/Engine/Plugins/Runtime/StateTree/Source/StateTreeModule/Public/StateTreeLinker.h
@@ -80,9 +80,9 @@ struct FStateTreeLinker
 		{
 			Index = ExternalDataDescs.Add(Desc);
 			check(FStateTreeExternalDataHandle::IsValidIndex(Index + ExternalDataBaseIndex));
-			ExternalDataDescs[Index].Handle.DataViewIndex = FStateTreeIndex8(Index + ExternalDataBaseIndex);
+			ExternalDataDescs[Index].Handle.DataViewIndex = FStateTreeIndex16(Index + ExternalDataBaseIndex);
 		}
-		Handle.DataViewIndex = FStateTreeIndex8(Index + ExternalDataBaseIndex);
+		Handle.DataViewIndex = FStateTreeIndex16(Index + ExternalDataBaseIndex);
 	}
 
 	/** @return linked external data descriptors. */
diff --git a/Engine/Plugins/Runtime/StateTree/Source/StateTreeModule/Public/StateTreeTypes.h b/Engine/Plugins/Runtime/StateTree/Source/StateTreeModule/Public/StateTreeTypes.h
--- a/Engine/Plugins/Runtime/StateTree/Source/StateTreeModule/Public/StateTreeTypes.h
+++ b/Engine/Plugins/Runtime/StateTree/Source/StateTreeModule/Public/StateTreeTypes.h
@@ -200,11 +200,21 @@ struct STATETREEMODULE_API FStateTreeIndex16
 	bool operator==(const FStateTreeIndex16& RHS) const { return Value == RHS.Value; }
 	bool operator!=(const FStateTreeIndex16& RHS) const { return Value != RHS.Value; }
 
+	bool SerializeFromMismatchedTag(const FPropertyTag& Tag, FStructuredArchive::FSlot Slot);
+
 protected:
 	UPROPERTY()
 	uint16 Value = InvalidValue;
 };
-	
+
+template<>
+struct TStructOpsTypeTraits<FStateTreeIndex16> : public TStructOpsTypeTraitsBase2<FStateTreeIndex16>
+{
+	enum
+	{
+		WithStructuredSerializeFromMismatchedTag = true,
+	};
+};
 
 /** uint8 index that can be invalid. */
 USTRUCT(BlueprintType)
@@ -241,11 +251,22 @@ struct STATETREEMODULE_API FStateTreeIndex8
 	bool operator==(const FStateTreeIndex8& RHS) const { return Value == RHS.Value; }
 	bool operator!=(const FStateTreeIndex8& RHS) const { return Value != RHS.Value; }
 
+	bool SerializeFromMismatchedTag(const FPropertyTag& Tag, FStructuredArchive::FSlot Slot);
+	
 protected:
 	UPROPERTY()
 	uint8 Value = InvalidValue;
 };
 
+template<>
+struct TStructOpsTypeTraits<FStateTreeIndex8> : public TStructOpsTypeTraitsBase2<FStateTreeIndex8>
+{
+	enum
+	{
+		WithStructuredSerializeFromMismatchedTag = true,
+	};
+};
+
 
 /**
  * Describes an array of active states in a State Tree.
@@ -623,11 +644,11 @@ struct STATETREEMODULE_API FStateTreeExternalDataHandle
 
 	static const FStateTreeExternalDataHandle Invalid;
 	
-	static bool IsValidIndex(const int32 Index) { return FStateTreeIndex8::IsValidIndex(Index); }
+	static bool IsValidIndex(const int32 Index) { return FStateTreeIndex16::IsValidIndex(Index); }
 	bool IsValid() const { return DataViewIndex.IsValid(); }
 
 	UPROPERTY()
-	FStateTreeIndex8 DataViewIndex = FStateTreeIndex8::Invalid;
+	FStateTreeIndex16 DataViewIndex = FStateTreeIndex16::Invalid;
 };
 
 /**
