* Fixed non unity compile errors
* Added UE_INLINE_GENERATED_CPP_BY_NAME to a few files missing it
* Added HEADER_UNIT_SKIP comment to deprecated files

diff --git a/Engine/Plugins/Runtime/StateTree/Source/StateTreeEditorModule/Private/Customizations/StateTreeEditorNodeDetails.h b/Engine/Plugins/Runtime/StateTree/Source/StateTreeEditorModule/Private/Customizations/StateTreeEditorNodeDetails.h
--- a/Engine/Plugins/Runtime/StateTree/Source/StateTreeEditorModule/Private/Customizations/StateTreeEditorNodeDetails.h
+++ b/Engine/Plugins/Runtime/StateTree/Source/StateTreeEditorModule/Private/Customizations/StateTreeEditorNodeDetails.h
@@ -6,6 +6,7 @@
 #include "UnrealClient.h"
 #include "Types/SlateStructs.h"
 #include "IPropertyTypeCustomization.h"
+#include "UObject/ObjectKey.h"
 #include "Widgets/Views/STreeView.h"
 
 class IPropertyHandle;
diff --git a/Engine/Source/Runtime/Core/Public/Templates/AreTypesEqual.h b/Engine/Source/Runtime/Core/Public/Templates/AreTypesEqual.h
--- a/Engine/Source/Runtime/Core/Public/Templates/AreTypesEqual.h
+++ b/Engine/Source/Runtime/Core/Public/Templates/AreTypesEqual.h
@@ -2,6 +2,8 @@
 
 #pragma once
 
+// HEADER_UNIT_SKIP - Deprecated
+
 #include <type_traits>
 
 /** Tests whether two typenames refer to the same type. */
diff --git a/Engine/Source/Runtime/Core/Public/Traits/IsVoidType.h b/Engine/Source/Runtime/Core/Public/Traits/IsVoidType.h
--- a/Engine/Source/Runtime/Core/Public/Traits/IsVoidType.h
+++ b/Engine/Source/Runtime/Core/Public/Traits/IsVoidType.h
@@ -2,6 +2,8 @@
 
 #pragma once
 
+// HEADER_UNIT_SKIP - Deprecated
+
 template<typename T> struct TIsVoidTypeBase { enum { Value = false }; };
 template<> struct TIsVoidTypeBase<void> { enum { Value = true }; };
 template<> struct TIsVoidTypeBase<void const> { enum { Value = true }; };
