StateTree: More graceful handling of GetStateTreeRunStatus()
- Calling GetStateTreeRunStatus() on unintialized instance is not error anymore, returns failed
- Made UStateTreeComponent::GetStateTreeRunStatus() to return failed instead of unset to be consistent with FStateTreeExecutionContext
- Allow to access execution state directly from instance data without exec context
- Do no crete exec context in UStateTreeComponent when sending events

diff --git a/Engine/Plugins/Runtime/GameplayStateTree/Source/GameplayStateTreeModule/Private/Components/StateTreeComponent.cpp b/Engine/Plugins/Runtime/GameplayStateTree/Source/GameplayStateTreeModule/Private/Components/StateTreeComponent.cpp
--- a/Engine/Plugins/Runtime/GameplayStateTree/Source/GameplayStateTreeModule/Private/Components/StateTreeComponent.cpp
+++ b/Engine/Plugins/Runtime/GameplayStateTree/Source/GameplayStateTreeModule/Private/Components/StateTreeComponent.cpp
@@ -399,22 +399,17 @@ void UStateTreeComponent::SendStateTreeEvent(const FGameplayTag Tag, const FCons
 		return;
 	}
 
-	FStateTreeExecutionContext Context(*GetOwner(), *StateTreeRef.GetStateTree(), InstanceData);
-	if (Context.IsValid())
-	{
-		Context.SendEvent(Tag, Payload, Origin);
-	}
+	InstanceData.GetMutableEventQueue().SendEvent(this, Tag, Payload, Origin);
 }
 
 EStateTreeRunStatus UStateTreeComponent::GetStateTreeRunStatus() const
 {
-	FStateTreeExecutionContext Context(*GetOwner(), *StateTreeRef.GetStateTree(), const_cast<FStateTreeInstanceData&>(InstanceData));
-	if (Context.IsValid())
+	if (const FStateTreeExecutionState* Exec = InstanceData.GetExecutionState())
 	{
-		return Context.GetStateTreeRunStatus();
+		return Exec->TreeRunStatus;
 	}
 
-	return EStateTreeRunStatus::Unset;
+	return EStateTreeRunStatus::Failed;
 }
 
 #if WITH_GAMEPLAY_DEBUGGER
diff --git a/Engine/Plugins/Runtime/StateTree/Source/StateTreeModule/Private/StateTreeExecutionContext.cpp b/Engine/Plugins/Runtime/StateTree/Source/StateTreeModule/Private/StateTreeExecutionContext.cpp
--- a/Engine/Plugins/Runtime/StateTree/Source/StateTreeModule/Private/StateTreeExecutionContext.cpp
+++ b/Engine/Plugins/Runtime/StateTree/Source/StateTreeModule/Private/StateTreeExecutionContext.cpp
@@ -10,8 +10,6 @@
 #include "ProfilingDebugging/CsvProfiler.h"
 #include "Logging/LogScopedVerbosityOverride.h"
 
-#include UE_INLINE_GENERATED_CPP_BY_NAME(StateTreeExecutionContext)
-
 #define STATETREE_LOG(Verbosity, Format, ...) UE_VLOG_UELOG(GetOwner(), LogStateTree, Verbosity, TEXT("%s") Format, *GetInstanceDescription(), ##__VA_ARGS__)
 #define STATETREE_CLOG(Condition, Verbosity, Format, ...) UE_CVLOG_UELOG((Condition), GetOwner(), LogStateTree, Verbosity, TEXT("%s") Format, *GetInstanceDescription(), ##__VA_ARGS__)
 
@@ -434,15 +432,12 @@ EStateTreeRunStatus FStateTreeExecutionContext::GetStateTreeRunStatus() const
 		return EStateTreeRunStatus::Failed;
 	}
 
-	if (!InstanceData.IsValid())
+	if (const FStateTreeExecutionState* Exec = InstanceData.GetExecutionState())
 	{
-		STATETREE_LOG(Error, TEXT("%s: GetStateTreeRunStatus called on %s using StateTree %s with invalid instance data. Start() must be called before GetStateTreeRunStatus()."),
-			ANSI_TO_TCHAR(__FUNCTION__), *GetNameSafe(&Owner), *GetFullNameSafe(&StateTree));
-		return EStateTreeRunStatus::Failed;
+		return Exec->TreeRunStatus;
 	}
-
-	const FStateTreeExecutionState& Exec = GetExecState();
-	return Exec.TreeRunStatus;
+	
+	return EStateTreeRunStatus::Failed;
 }
 
 void FStateTreeExecutionContext::SendEvent(const FStateTreeEvent& Event) const
diff --git a/Engine/Plugins/Runtime/StateTree/Source/StateTreeModule/Private/StateTreeInstanceData.cpp b/Engine/Plugins/Runtime/StateTree/Source/StateTreeModule/Private/StateTreeInstanceData.cpp
--- a/Engine/Plugins/Runtime/StateTree/Source/StateTreeModule/Private/StateTreeInstanceData.cpp
+++ b/Engine/Plugins/Runtime/StateTree/Source/StateTreeModule/Private/StateTreeInstanceData.cpp
@@ -73,6 +73,16 @@ FStateTreeInstanceStorage& FStateTreeInstanceData::GetMutableStorage()
 	return *reinterpret_cast<FStateTreeInstanceStorage*>(InstanceStorage.GetMutableMemory());
 }
 
+const FStateTreeExecutionState* FStateTreeInstanceData::GetExecutionState() const
+{
+	if (!IsValid())
+	{
+		return nullptr;
+	}
+	const FConstStructView ExecView = GetStruct(0); // Execution state is fixed at index 0. 
+	return ExecView.GetPtr<FStateTreeExecutionState>();
+}
+
 TArray<FStateTreeEvent>& FStateTreeInstanceData::GetEvents() const
 {
 	return const_cast<FStateTreeInstanceData*>(this)->GetMutableStorage().EventQueue.GetEventsArray();
diff --git a/Engine/Plugins/Runtime/StateTree/Source/StateTreeModule/Public/StateTreeExecutionContext.h b/Engine/Plugins/Runtime/StateTree/Source/StateTreeModule/Public/StateTreeExecutionContext.h
--- a/Engine/Plugins/Runtime/StateTree/Source/StateTreeModule/Public/StateTreeExecutionContext.h
+++ b/Engine/Plugins/Runtime/StateTree/Source/StateTreeModule/Public/StateTreeExecutionContext.h
@@ -3,62 +3,15 @@
 #pragma once
 
 #include "StateTree.h"
-#include "StateTreePropertyBindings.h"
 #include "StateTreeInstanceData.h"
 #include "StateTreeNodeBase.h"
 #include "Experimental/ConcurrentLinearAllocator.h"
-#include "StateTreeExecutionContext.generated.h"
 
 struct FStateTreeEvaluatorBase;
 struct FStateTreeTaskBase;
 struct FStateTreeConditionBase;
 struct FStateTreeEvent;
 
-struct STATETREEMODULE_API FStateTreeTransitionDelayedState
-{
-	FStateTreeIndex16 TransitionIndex = FStateTreeIndex16::Invalid;
-	float TimeLeft = 0.0f;
-};
-
-USTRUCT()
-struct STATETREEMODULE_API FStateTreeExecutionState
-{
-	GENERATED_BODY()
-
-	/** @returns Delayed transition state for a specific transition, or nullptr if it does not exists. */
-	FStateTreeTransitionDelayedState* FindDelayedTransition(const FStateTreeIndex16 TransitionIndex)
-	{
-		return DelayedTransitions.FindByPredicate([TransitionIndex](const FStateTreeTransitionDelayedState& State){ return State.TransitionIndex == TransitionIndex; });
-	}
-
-	/** Currently active states */
-	FStateTreeActiveStates ActiveStates;
-
-	/** Index of the first task struct in the currently initialized instance data. */
-	FStateTreeIndex16 FirstTaskStructIndex = FStateTreeIndex16::Invalid;
-	
-	/** Index of the first task object in the currently initialized instance data. */
-	FStateTreeIndex16 FirstTaskObjectIndex = FStateTreeIndex16::Invalid;
-
-	/** The index of the task that failed during enter state. Exit state uses it to call ExitState() symmetrically. */
-	FStateTreeIndex16 EnterStateFailedTaskIndex = FStateTreeIndex16::Invalid;
-
-	/** Result of last tick */
-	EStateTreeRunStatus LastTickStatus = EStateTreeRunStatus::Failed;
-
-	/** Running status of the instance */
-	EStateTreeRunStatus TreeRunStatus = EStateTreeRunStatus::Unset;
-
-	/** Handle of the state that was first to report state completed (success or failure), used to trigger completion transitions. */
-	FStateTreeStateHandle CompletedStateHandle = FStateTreeStateHandle::Invalid;
-
-	/** Number of times a new state has been changed. */
-	uint16 StateChangeCount = 0;
-
-	/** Running time of the delayed transition */
-	TArray<FStateTreeTransitionDelayedState> DelayedTransitions;
-};
-
 /**
  * StateTree Execution Context is a helper that is used to update and access StateTree instance data.
  *
@@ -120,7 +73,6 @@ public:
 	UWorld* GetWorld() const { return Owner.GetWorld(); };
 
 	/** @return True of the the execution context is valid and initialized. */ 
-	//bool IsValid() const { return Owner != nullptr && StateTree != nullptr && InstanceData != nullptr; }
 	bool IsValid() const { return StateTree.IsReadyToRun(); }
 	
 	/** Start executing. */
diff --git a/Engine/Plugins/Runtime/StateTree/Source/StateTreeModule/Public/StateTreeInstanceData.h b/Engine/Plugins/Runtime/StateTree/Source/StateTreeModule/Public/StateTreeInstanceData.h
--- a/Engine/Plugins/Runtime/StateTree/Source/StateTreeModule/Public/StateTreeInstanceData.h
+++ b/Engine/Plugins/Runtime/StateTree/Source/StateTreeModule/Public/StateTreeInstanceData.h
@@ -8,6 +8,7 @@
 #include "StateTreeEvents.h"
 #include "StateTreeInstanceData.generated.h"
 
+struct FStateTreeExecutionState;
 
 /**
  * State Tree instance data is used to store the runtime state of a State Tree. It is used together with FStateTreeExecution context to tick the state tree.
@@ -129,6 +130,9 @@ struct STATETREEMODULE_API FStateTreeInstanceData
 	/** @return const pointer to an instance object   */
 	const UObject* GetObject(const int32 Index) const { return GetStorage().InstanceObjects[Index]; }
 
+	/** @return pointer to StateTree execution state, or null if the instance data is not initialized. */
+	const FStateTreeExecutionState* GetExecutionState() const;
+	
 	/** @return array to store unprocessed events. */
 	UE_DEPRECATED(5.2, "Use GetEventQueue() instead.")
 	TArray<FStateTreeEvent>& GetEvents() const;
diff --git a/Engine/Plugins/Runtime/StateTree/Source/StateTreeModule/Public/StateTreeTypes.h b/Engine/Plugins/Runtime/StateTree/Source/StateTreeModule/Public/StateTreeTypes.h
--- a/Engine/Plugins/Runtime/StateTree/Source/StateTreeModule/Public/StateTreeTypes.h
+++ b/Engine/Plugins/Runtime/StateTree/Source/StateTreeModule/Public/StateTreeTypes.h
@@ -960,3 +960,48 @@ protected:
 	/** Memory pointing at the class or struct */
 	uint8* Memory = nullptr;
 };
+
+struct STATETREEMODULE_API FStateTreeTransitionDelayedState
+{
+	FStateTreeIndex16 TransitionIndex = FStateTreeIndex16::Invalid;
+	float TimeLeft = 0.0f;
+};
+
+USTRUCT()
+struct STATETREEMODULE_API FStateTreeExecutionState
+{
+	GENERATED_BODY()
+
+	/** @returns Delayed transition state for a specific transition, or nullptr if it does not exists. */
+	FStateTreeTransitionDelayedState* FindDelayedTransition(const FStateTreeIndex16 TransitionIndex)
+	{
+		return DelayedTransitions.FindByPredicate([TransitionIndex](const FStateTreeTransitionDelayedState& State){ return State.TransitionIndex == TransitionIndex; });
+	}
+
+	/** Currently active states */
+	FStateTreeActiveStates ActiveStates;
+
+	/** Index of the first task struct in the currently initialized instance data. */
+	FStateTreeIndex16 FirstTaskStructIndex = FStateTreeIndex16::Invalid;
+	
+	/** Index of the first task object in the currently initialized instance data. */
+	FStateTreeIndex16 FirstTaskObjectIndex = FStateTreeIndex16::Invalid;
+
+	/** The index of the task that failed during enter state. Exit state uses it to call ExitState() symmetrically. */
+	FStateTreeIndex16 EnterStateFailedTaskIndex = FStateTreeIndex16::Invalid;
+
+	/** Result of last tick */
+	EStateTreeRunStatus LastTickStatus = EStateTreeRunStatus::Failed;
+
+	/** Running status of the instance */
+	EStateTreeRunStatus TreeRunStatus = EStateTreeRunStatus::Unset;
+
+	/** Handle of the state that was first to report state completed (success or failure), used to trigger completion transitions. */
+	FStateTreeStateHandle CompletedStateHandle = FStateTreeStateHandle::Invalid;
+
+	/** Number of times a new state has been changed. */
+	uint16 StateChangeCount = 0;
+
+	/** Running time of the delayed transition */
+	TArray<FStateTreeTransitionDelayedState> DelayedTransitions;
+};
\ No newline at end of file
