StateTree: Fix mutable uses of instanced structs and struct views.

diff --git a/Engine/Plugins/Runtime/StateTree/Source/StateTreeModule/Private/StateTree.cpp b/Engine/Plugins/Runtime/StateTree/Source/StateTreeModule/Private/StateTree.cpp
--- a/Engine/Plugins/Runtime/StateTree/Source/StateTreeModule/Private/StateTree.cpp
+++ b/Engine/Plugins/Runtime/StateTree/Source/StateTreeModule/Private/StateTree.cpp
@@ -81,6 +81,9 @@ void UStateTree::ResetCompiled()
 	EvaluatorsBegin = 0;
 	EvaluatorsNum = 0;
 
+	GlobalTasksBegin = 0;
+	GlobalTasksNum = 0;
+	
 	ResetLinked();
 }
 
@@ -313,7 +316,7 @@ bool UStateTree::Link()
 	
 	for (int32 Index = 0; Index < Nodes.Num(); Index++)
 	{
-		const FStructView Node = Nodes[Index];
+		FStructView Node = Nodes[Index];
 		if (FStateTreeNodeBase* NodePtr = Node.GetMutablePtr<FStateTreeNodeBase>())
 		{
 			Linker.SetCurrentInstanceDataType(NodePtr->GetInstanceDataType(), NodePtr->DataViewIndex.Get());
@@ -369,6 +372,7 @@ TArray<FStateTreeMemoryUsage> UStateTree::CalculateEstimatedMemoryUsage() const
 	const int32 TreeMemUsageIndex = MemoryUsages.Emplace(TEXT("State Tree Max"));
 	const int32 InstanceMemUsageIndex = MemoryUsages.Emplace(TEXT("Instance Overhead"));
 	const int32 EvalMemUsageIndex = MemoryUsages.Emplace(TEXT("Evaluators"));
+	const int32 GlobalTaskMemUsageIndex = MemoryUsages.Emplace(TEXT("GlobalTask"));
 	const int32 SharedMemUsageIndex = MemoryUsages.Emplace(TEXT("Shared Data"));
 
 	auto GetRootStateHandle = [this](const FStateTreeStateHandle InState) -> FStateTreeStateHandle
@@ -490,6 +494,22 @@ TArray<FStateTreeMemoryUsage> UStateTree::CalculateEstimatedMemoryUsage() const
 		EvalMemUsage.NodeCount++;
 	}
 
+	// Global Tasks
+	FStateTreeMemoryUsage& GlobalTaskMemUsage = MemoryUsages[GlobalTaskMemUsageIndex];
+	for (int32 TaskIndex = GlobalTasksBegin; TaskIndex < (GlobalTasksBegin + GlobalTasksNum); TaskIndex++)
+	{
+		const FStateTreeTaskBase& Task = Nodes[TaskIndex].Get<FStateTreeTaskBase>();
+		if (Task.bInstanceIsObject == false)
+		{
+			GlobalTaskMemUsage.AddUsage(DefaultInstanceData.GetStruct(Task.InstanceIndex.Get()));
+		}
+		else
+		{
+			GlobalTaskMemUsage.AddUsage(DefaultInstanceData.GetObject(Task.InstanceIndex.Get()));
+		}
+		GlobalTaskMemUsage.NodeCount++;
+	}
+
 	// Estimate highest combined usage.
 	FStateTreeMemoryUsage& TreeMemUsage = MemoryUsages[TreeMemUsageIndex];
 
@@ -500,6 +520,9 @@ TArray<FStateTreeMemoryUsage> UStateTree::CalculateEstimatedMemoryUsage() const
 	TreeMemUsage.EstimatedMemoryUsage += EvalMemUsage.EstimatedMemoryUsage;
 	TreeMemUsage.NodeCount += EvalMemUsage.NodeCount;
 
+	TreeMemUsage.EstimatedMemoryUsage += GlobalTaskMemUsage.EstimatedMemoryUsage;
+	TreeMemUsage.NodeCount += GlobalTaskMemUsage.NodeCount;
+
 	FStateTreeMemoryUsage& InstanceMemUsage = MemoryUsages[InstanceMemUsageIndex];
 	// FStateTreeInstanceData overhead.
 	InstanceMemUsage.EstimatedMemoryUsage += sizeof(FStateTreeInstanceData);
diff --git a/Engine/Plugins/Runtime/StateTree/Source/StateTreeModule/Public/StateTreeInstanceData.h b/Engine/Plugins/Runtime/StateTree/Source/StateTreeModule/Public/StateTreeInstanceData.h
--- a/Engine/Plugins/Runtime/StateTree/Source/StateTreeModule/Public/StateTreeInstanceData.h
+++ b/Engine/Plugins/Runtime/StateTree/Source/StateTreeModule/Public/StateTreeInstanceData.h
@@ -224,10 +224,10 @@ struct TStateTreeInstanceDataStructRef
 
 	bool IsValid() const { return Storage.IsValidStructIndex(StructIndex); }
 
-	T& operator*() const
+	T& operator*()
 	{
 		check(IsValid());
-		const FStructView Struct = Storage.GetMutableStruct(StructIndex);
+		FStructView Struct = Storage.GetMutableStruct(StructIndex);
 		check(Struct.GetScriptStruct() == TBaseStructure<T>::Get());
 		return *reinterpret_cast<T*>(Struct.GetMutableMemory());
 	}
