StateTree: Allow tasks to request transitions
- Added ability for tasks to request transitions
- Added buffered transition requests
- Added callback method to FStateTreeTaskBase which is called during transition triggering
- Allow FStateTreeStateLink to be added to tasks and conditions (resolved automatically during compile)

diff --git a/Engine/Plugins/Runtime/StateTree/Config/BaseStateTree.ini b/Engine/Plugins/Runtime/StateTree/Config/BaseStateTree.ini
--- a/Engine/Plugins/Runtime/StateTree/Config/BaseStateTree.ini
+++ b/Engine/Plugins/Runtime/StateTree/Config/BaseStateTree.ini
@@ -55,7 +55,7 @@
 +PropertyRedirects=(OldName="CompactStateTransition.Event",NewName="CompactStateTransition.Trigger")
 +ClassRedirects=(OldName="/Script/StateTreeModule.StateTreeItemBlueprintBase",NewName="/Script/StateTreeModule.StateTreeNodeBlueprintBase")
 
-+EnumRedirects=(OldName="EStateTreeBindableStructSource",ValueChanges=(("TreeData", "Context"), ("TreeParameter", "Parameter"), ("StartParameter", "State"))
++EnumRedirects=(OldName="/Script/StateTreeModule.EStateTreeBindableStructSource",ValueChanges=(("TreeData", "Context"), ("TreeParameter", "Parameter"), ("StartParameter", "State"), ("NotSet", "None"))
 
 +PropertyRedirects=(OldName="/Script/StateTreeModule.StateTree.NamedExternalDataDescs",NewName="/Script/StateTreeModule.StateTree.ContextDataDescs")
 
diff --git a/Engine/Plugins/Runtime/StateTree/Source/StateTreeEditorModule/Private/Customizations/StateTreeStateDetails.cpp b/Engine/Plugins/Runtime/StateTree/Source/StateTreeEditorModule/Private/Customizations/StateTreeStateDetails.cpp
--- a/Engine/Plugins/Runtime/StateTree/Source/StateTreeEditorModule/Private/Customizations/StateTreeStateDetails.cpp
+++ b/Engine/Plugins/Runtime/StateTree/Source/StateTreeEditorModule/Private/Customizations/StateTreeStateDetails.cpp
@@ -41,7 +41,7 @@ void FStateTreeStateDetails::CustomizeDetails(IDetailLayoutBuilder& DetailBuilde
 	TSharedPtr<IPropertyHandle> EnterConditionsProperty = DetailBuilder.GetProperty(TEXT("EnterConditions"));
 	TSharedPtr<IPropertyHandle> TransitionsProperty = DetailBuilder.GetProperty(TEXT("Transitions"));
 	TSharedPtr<IPropertyHandle> TypeProperty = DetailBuilder.GetProperty(TEXT("Type"));
-	TSharedPtr<IPropertyHandle> LinkedStateProperty = DetailBuilder.GetProperty(TEXT("LinkedState"));
+	TSharedPtr<IPropertyHandle> LinkedSubtreeProperty = DetailBuilder.GetProperty(TEXT("LinkedSubtree"));
 	TSharedPtr<IPropertyHandle> ParametersProperty = DetailBuilder.GetProperty(TEXT("Parameters"));
 
 	uint8 StateTypeValue = 0;
@@ -53,7 +53,7 @@ void FStateTreeStateDetails::CustomizeDetails(IDetailLayoutBuilder& DetailBuilde
 
 	if (StateType != EStateTreeStateType::Linked)
 	{
-		LinkedStateProperty->MarkHiddenByCustomization();
+		LinkedSubtreeProperty->MarkHiddenByCustomization();
 	}
 	
 	if (!(StateType == EStateTreeStateType::Subtree || StateType == EStateTreeStateType::Linked))
diff --git a/Engine/Plugins/Runtime/StateTree/Source/StateTreeEditorModule/Private/Customizations/StateTreeStateLinkDetails.cpp b/Engine/Plugins/Runtime/StateTree/Source/StateTreeEditorModule/Private/Customizations/StateTreeStateLinkDetails.cpp
--- a/Engine/Plugins/Runtime/StateTree/Source/StateTreeEditorModule/Private/Customizations/StateTreeStateLinkDetails.cpp
+++ b/Engine/Plugins/Runtime/StateTree/Source/StateTreeEditorModule/Private/Customizations/StateTreeStateLinkDetails.cpp
@@ -25,7 +25,7 @@ void FStateTreeStateLinkDetails::CustomizeHeader(TSharedRef<class IPropertyHandl
 
 	NameProperty = StructProperty->GetChildHandle(TEXT("Name"));
 	IDProperty = StructProperty->GetChildHandle(TEXT("ID"));
-	TypeProperty = StructProperty->GetChildHandle(TEXT("Type"));
+	LinkTypeProperty = StructProperty->GetChildHandle(TEXT("LinkType"));
 
 	if (const FProperty* MetaDataProperty = StructProperty->GetMetaDataProperty())
 	{
@@ -125,7 +125,7 @@ void FStateTreeStateLinkDetails::OnStateComboChange(int Idx)
 		FScopedTransaction Transaction(FText::Format(LOCTEXT("SetPropertyValue", "Set {0}"), StructProperty->GetPropertyDisplayName()));
 		if (Idx >= 0)
 		{
-			TypeProperty->SetValue((uint8)EStateTreeTransitionType::GotoState);
+			LinkTypeProperty->SetValue((uint8)EStateTreeTransitionType::GotoState);
 			NameProperty->SetValue(CachedNames[Idx], EPropertyValueSetFlags::NotTransactable);
 			UE::StateTree::PropertyHelpers::SetStructValue<FGuid>(IDProperty, CachedIDs[Idx], EPropertyValueSetFlags::NotTransactable);
 		}
@@ -134,17 +134,17 @@ void FStateTreeStateLinkDetails::OnStateComboChange(int Idx)
 			switch (Idx)
 			{
 			case ComboSucceeded:
-				TypeProperty->SetValue((uint8)EStateTreeTransitionType::Succeeded);
+				LinkTypeProperty->SetValue((uint8)EStateTreeTransitionType::Succeeded);
 				break;
 			case ComboFailed:
-				TypeProperty->SetValue((uint8)EStateTreeTransitionType::Failed);
+				LinkTypeProperty->SetValue((uint8)EStateTreeTransitionType::Failed);
 				break;
 			case ComboNextState:
-				TypeProperty->SetValue((uint8)EStateTreeTransitionType::NextState);
+				LinkTypeProperty->SetValue((uint8)EStateTreeTransitionType::NextState);
 				break;
 			case ComboNotSet:
 			default:
-				TypeProperty->SetValue((uint8)EStateTreeTransitionType::NotSet);
+				LinkTypeProperty->SetValue((uint8)EStateTreeTransitionType::None);
 				break;
 			}
 			// Clear name and id.
@@ -160,16 +160,16 @@ TSharedRef<SWidget> FStateTreeStateLinkDetails::OnGetStateContent() const
 
 	if (!bDirectStatesOnly)
 	{
-		FUIAction NextItemAction(FExecuteAction::CreateSP(const_cast<FStateTreeStateLinkDetails*>(this), &FStateTreeStateLinkDetails::OnStateComboChange, ComboNextState));
-		MenuBuilder.AddMenuEntry(LOCTEXT("TransitionNextState", "Next State"), LOCTEXT("TransitionNextTooltip", "Goto next sibling State."), FSlateIcon(), NextItemAction);
+		const FUIAction NotSetItemAction(FExecuteAction::CreateSP(const_cast<FStateTreeStateLinkDetails*>(this), &FStateTreeStateLinkDetails::OnStateComboChange, ComboNotSet));
+		MenuBuilder.AddMenuEntry(LOCTEXT("TransitionNone", "None"), LOCTEXT("TransitionNoneTooltip", "No transition."), FSlateIcon(), NotSetItemAction);
 
-		FUIAction NotSetItemAction(FExecuteAction::CreateSP(const_cast<FStateTreeStateLinkDetails*>(this), &FStateTreeStateLinkDetails::OnStateComboChange, ComboNotSet));
-		MenuBuilder.AddMenuEntry(LOCTEXT("TransitionBlock", "Block Transition"), LOCTEXT("TransitionBlockTooltip", "Will not transition to any state, but will block other transitions to trigger if the condition passes."), FSlateIcon(), NotSetItemAction);
+		const FUIAction NextItemAction(FExecuteAction::CreateSP(const_cast<FStateTreeStateLinkDetails*>(this), &FStateTreeStateLinkDetails::OnStateComboChange, ComboNextState));
+		MenuBuilder.AddMenuEntry(LOCTEXT("TransitionNextState", "Next State"), LOCTEXT("TransitionNextTooltip", "Goto next sibling State."), FSlateIcon(), NextItemAction);
 
-		FUIAction SucceededItemAction(FExecuteAction::CreateSP(const_cast<FStateTreeStateLinkDetails*>(this), &FStateTreeStateLinkDetails::OnStateComboChange, ComboSucceeded));
+		const FUIAction SucceededItemAction(FExecuteAction::CreateSP(const_cast<FStateTreeStateLinkDetails*>(this), &FStateTreeStateLinkDetails::OnStateComboChange, ComboSucceeded));
 		MenuBuilder.AddMenuEntry(LOCTEXT("TransitionTreeSucceeded", "Tree Succeeded"), LOCTEXT("TransitionTreeSuccessTooltip", "Complete tree with success."), FSlateIcon(), SucceededItemAction);
 
-		FUIAction FailedItemAction(FExecuteAction::CreateSP(const_cast<FStateTreeStateLinkDetails*>(this), &FStateTreeStateLinkDetails::OnStateComboChange, ComboFailed));
+		const FUIAction FailedItemAction(FExecuteAction::CreateSP(const_cast<FStateTreeStateLinkDetails*>(this), &FStateTreeStateLinkDetails::OnStateComboChange, ComboFailed));
 		MenuBuilder.AddMenuEntry(LOCTEXT("TransitionTreeFailed", "Tree Failed"), LOCTEXT("TransitionTreeFailedTooltip", "Complete tree with failure."), FSlateIcon(), FailedItemAction);
 	}
 
@@ -193,63 +193,68 @@ FText FStateTreeStateLinkDetails::GetCurrentStateDesc() const
 {
 	const EStateTreeTransitionType TransitionType = GetTransitionType().Get(EStateTreeTransitionType::Failed);
 
-	if (TransitionType == EStateTreeTransitionType::NotSet)
-	{
-		return LOCTEXT("TransitionBlock", "Block Transition");
-	}
-	else if (TransitionType == EStateTreeTransitionType::NextState)
-	{
-		return LOCTEXT("TransitionNextState", "Next State");
-	}
-	else if (TransitionType == EStateTreeTransitionType::Succeeded)
-	{
-		return LOCTEXT("TransitionTreeSucceeded", "Tree Succeeded");
-	}
-	else if (TransitionType == EStateTreeTransitionType::Failed)
-	{
-		return LOCTEXT("TransitionTreeFailed", "Tree Failed");
-	}
-	else if (TransitionType == EStateTreeTransitionType::GotoState)
+	FText Result;
+	switch (TransitionType)
 	{
-		FName OldName;
-		if (NameProperty && IDProperty)
+	case EStateTreeTransitionType::None:
+		Result = LOCTEXT("TransitionNone", "None");
+		break;
+	case EStateTreeTransitionType::NextState:
+		Result = LOCTEXT("TransitionNextState", "Next State");
+		break;
+	case EStateTreeTransitionType::Succeeded:
+		Result = LOCTEXT("TransitionTreeSucceeded", "Tree Succeeded");
+		break;
+	case EStateTreeTransitionType::Failed:
+		Result = LOCTEXT("TransitionTreeFailed", "Tree Failed");
+		break;
+	case EStateTreeTransitionType::GotoState:
 		{
-			NameProperty->GetValue(OldName);
-
-			FGuid StateID;
-			if (UE::StateTree::PropertyHelpers::GetStructValue<FGuid>(IDProperty, StateID) == FPropertyAccess::Success)
+			FName OldName;
+			if (NameProperty && IDProperty)
 			{
-				if (!StateID.IsValid())
-				{
-					return LOCTEXT("None", "None");
-				}
-				else
+				NameProperty->GetValue(OldName);
+
+				FGuid StateID;
+				if (UE::StateTree::PropertyHelpers::GetStructValue<FGuid>(IDProperty, StateID) == FPropertyAccess::Success)
 				{
-					for (int32 Idx = 0; Idx < CachedIDs.Num(); Idx++)
+					if (!StateID.IsValid())
 					{
-						if (CachedIDs[Idx] == StateID)
+						return LOCTEXT("None", "None");
+					}
+					else
+					{
+						for (int32 Idx = 0; Idx < CachedIDs.Num(); Idx++)
 						{
-							return FText::FromName(CachedNames[Idx]);
+							if (CachedIDs[Idx] == StateID)
+							{
+								return FText::FromName(CachedNames[Idx]);
+							}
 						}
 					}
 				}
 			}
-		}
 
-		FFormatNamedArguments Args;
-		Args.Add(TEXT("Identifier"), FText::FromName(OldName));
-		return FText::Format(LOCTEXT("InvalidReference", "Invalid Reference {Identifier}"), Args);
+			FFormatNamedArguments Args;
+			Args.Add(TEXT("Identifier"), FText::FromName(OldName));
+			Result = FText::Format(LOCTEXT("InvalidReference", "Invalid Reference {Identifier}"), Args);
+		}
+		break;
+	default:
+		ensureMsgf(false, TEXT("FStateTreeStateLinkDetails: Unhnandled enum %s"), *UEnum::GetValueAsString(TransitionType));
+		Result = LOCTEXT("TransitionInvalid", "Invalid");
+		break;
 	}
 
-	return LOCTEXT("TransitionInvalid", "Invalid");
+	return Result;
 }
 
 TOptional<EStateTreeTransitionType> FStateTreeStateLinkDetails::GetTransitionType() const
 {
-	if (TypeProperty)
+	if (LinkTypeProperty)
 	{
 		uint8 Value;
-		if (TypeProperty->GetValue(Value) == FPropertyAccess::Success)
+		if (LinkTypeProperty->GetValue(Value) == FPropertyAccess::Success)
 		{
 			return EStateTreeTransitionType(Value);
 		}
diff --git a/Engine/Plugins/Runtime/StateTree/Source/StateTreeEditorModule/Private/Customizations/StateTreeStateLinkDetails.h b/Engine/Plugins/Runtime/StateTree/Source/StateTreeEditorModule/Private/Customizations/StateTreeStateLinkDetails.h
--- a/Engine/Plugins/Runtime/StateTree/Source/StateTreeEditorModule/Private/Customizations/StateTreeStateLinkDetails.h
+++ b/Engine/Plugins/Runtime/StateTree/Source/StateTreeEditorModule/Private/Customizations/StateTreeStateLinkDetails.h
@@ -48,7 +48,7 @@ private:
 
 	TSharedPtr<IPropertyHandle> NameProperty;
 	TSharedPtr<IPropertyHandle> IDProperty;
-	TSharedPtr<IPropertyHandle> TypeProperty;
+	TSharedPtr<IPropertyHandle> LinkTypeProperty;
 
 	// If set, hide selecting meta states like Next or (tree) Succeeded.
 	bool bDirectStatesOnly = false;
diff --git a/Engine/Plugins/Runtime/StateTree/Source/StateTreeEditorModule/Private/Customizations/StateTreeTransitionDetails.cpp b/Engine/Plugins/Runtime/StateTree/Source/StateTreeEditorModule/Private/Customizations/StateTreeTransitionDetails.cpp
--- a/Engine/Plugins/Runtime/StateTree/Source/StateTreeEditorModule/Private/Customizations/StateTreeTransitionDetails.cpp
+++ b/Engine/Plugins/Runtime/StateTree/Source/StateTreeEditorModule/Private/Customizations/StateTreeTransitionDetails.cpp
@@ -25,6 +25,7 @@ void FStateTreeTransitionDetails::CustomizeHeader(TSharedRef<class IPropertyHand
 	PropUtils = StructCustomizationUtils.GetPropertyUtilities().Get();
 
 	TriggerProperty = StructProperty->GetChildHandle(TEXT("Trigger"));
+	PriorityProperty = StructProperty->GetChildHandle(TEXT("Priority"));
 	EventTagProperty = StructProperty->GetChildHandle(TEXT("EventTag"));
 	StateProperty = StructProperty->GetChildHandle(TEXT("State"));
 	DelayTransitionProperty = StructProperty->GetChildHandle(TEXT("bDelayTransition"));
@@ -54,7 +55,7 @@ void FStateTreeTransitionDetails::CustomizeHeader(TSharedRef<class IPropertyHand
 			]
 		];
 }
-
+ 
 void FStateTreeTransitionDetails::CustomizeChildren(TSharedRef<class IPropertyHandle> StructPropertyHandle, class IDetailChildrenBuilder& StructBuilder, IPropertyTypeCustomizationUtils& StructCustomizationUtils)
 {
 	check(TriggerProperty);
@@ -65,6 +66,12 @@ void FStateTreeTransitionDetails::CustomizeChildren(TSharedRef<class IPropertyHa
 	check(StateProperty);
 	check(ConditionsProperty);
 
+	auto IsTickOrEventTransition = [this]()
+	{
+		return !EnumHasAnyFlags(GetTrigger(), EStateTreeTransitionTrigger::OnStateCompleted) ? EVisibility::Visible : EVisibility::Collapsed;
+	};
+
+	// Trigger
 	StructBuilder.AddProperty(TriggerProperty.ToSharedRef());
 
 	// Show event only when the trigger is set to Event. 
@@ -74,20 +81,20 @@ void FStateTreeTransitionDetails::CustomizeChildren(TSharedRef<class IPropertyHa
 			return (GetTrigger() == EStateTreeTransitionTrigger::OnEvent) ? EVisibility::Visible : EVisibility::Collapsed;
 		})));
 
-	// Delay
-	auto IsDelayVisible = [this]()
-	{
-		return !EnumHasAnyFlags(GetTrigger(), EStateTreeTransitionTrigger::OnStateCompleted) ? EVisibility::Visible : EVisibility::Collapsed;
-	};
+	// State
+	StructBuilder.AddProperty(StateProperty.ToSharedRef());
+
+	// Priority
+	StructBuilder.AddProperty(PriorityProperty.ToSharedRef())
+		.Visibility(TAttribute<EVisibility>::Create(TAttribute<EVisibility>::FGetter::CreateLambda(IsTickOrEventTransition)));
 
+	// Delay
 	StructBuilder.AddProperty(DelayTransitionProperty.ToSharedRef())
-		.Visibility(TAttribute<EVisibility>::Create(TAttribute<EVisibility>::FGetter::CreateLambda(IsDelayVisible)));
+		.Visibility(TAttribute<EVisibility>::Create(TAttribute<EVisibility>::FGetter::CreateLambda(IsTickOrEventTransition)));
 	StructBuilder.AddProperty(DelayDurationProperty.ToSharedRef())
-		.Visibility(TAttribute<EVisibility>::Create(TAttribute<EVisibility>::FGetter::CreateLambda(IsDelayVisible)));
+		.Visibility(TAttribute<EVisibility>::Create(TAttribute<EVisibility>::FGetter::CreateLambda(IsTickOrEventTransition)));
 	StructBuilder.AddProperty(DelayRandomVarianceProperty.ToSharedRef())
-		.Visibility(TAttribute<EVisibility>::Create(TAttribute<EVisibility>::FGetter::CreateLambda(IsDelayVisible)));
-
-	StructBuilder.AddProperty(StateProperty.ToSharedRef());
+		.Visibility(TAttribute<EVisibility>::Create(TAttribute<EVisibility>::FGetter::CreateLambda(IsTickOrEventTransition)));
 
 	// Show conditions always expanded, with simplified header (remove item count)
 	IDetailPropertyRow& ConditionsRow = StructBuilder.AddProperty(ConditionsProperty.ToSharedRef());
@@ -166,10 +173,10 @@ FText FStateTreeTransitionDetails::GetDescription() const
 	const FStateTreeStateLink* State = static_cast<FStateTreeStateLink*>(RawData[0]);
 	if (State != nullptr)
 	{
-		switch (State->Type)
+		switch (State->LinkType)
 		{
-		case EStateTreeTransitionType::NotSet:
-			TargetText = LOCTEXT("TransitionBlock", "Block Transition");
+		case EStateTreeTransitionType::None:
+			TargetText = LOCTEXT("TransitionNone", "None");
 			break;
 		case EStateTreeTransitionType::Succeeded:
 			TargetText = LOCTEXT("TransitionTreeSucceeded", "Tree Succeeded");
diff --git a/Engine/Plugins/Runtime/StateTree/Source/StateTreeEditorModule/Private/Customizations/StateTreeTransitionDetails.h b/Engine/Plugins/Runtime/StateTree/Source/StateTreeEditorModule/Private/Customizations/StateTreeTransitionDetails.h
--- a/Engine/Plugins/Runtime/StateTree/Source/StateTreeEditorModule/Private/Customizations/StateTreeTransitionDetails.h
+++ b/Engine/Plugins/Runtime/StateTree/Source/StateTreeEditorModule/Private/Customizations/StateTreeTransitionDetails.h
@@ -31,6 +31,7 @@ private:
 	bool GetDelayTransition() const;
 	
 	TSharedPtr<IPropertyHandle> TriggerProperty;
+	TSharedPtr<IPropertyHandle> PriorityProperty;
 	TSharedPtr<IPropertyHandle> EventTagProperty;
 	TSharedPtr<IPropertyHandle> StateProperty;
 	TSharedPtr<IPropertyHandle> DelayTransitionProperty;
diff --git a/Engine/Plugins/Runtime/StateTree/Source/StateTreeEditorModule/Private/SStateTreeViewRow.cpp b/Engine/Plugins/Runtime/StateTree/Source/StateTreeEditorModule/Private/SStateTreeViewRow.cpp
--- a/Engine/Plugins/Runtime/StateTree/Source/StateTreeEditorModule/Private/SStateTreeViewRow.cpp
+++ b/Engine/Plugins/Runtime/StateTree/Source/StateTreeEditorModule/Private/SStateTreeViewRow.cpp
@@ -512,9 +512,9 @@ FText SStateTreeViewRow::GetTransitionsDesc(const UStateTreeState& State, const
 		const bool bMatch = bUseMask ? EnumHasAnyFlags(Transition.Trigger, Trigger) : Transition.Trigger == Trigger;
 		if (bMatch)
 		{
-			switch (Transition.State.Type)
+			switch (Transition.State.LinkType)
 			{
-			case EStateTreeTransitionType::NotSet:
+			case EStateTreeTransitionType::None:
 				DescItems.Add(LOCTEXT("TransitionNoneStyled", "[None]"));
 				break;
 			case EStateTreeTransitionType::Succeeded:
@@ -572,9 +572,9 @@ FText SStateTreeViewRow::GetTransitionsIcon(const UStateTreeState& State, const
 		const bool bMatch = bUseMask ? EnumHasAnyFlags(Transition.Trigger, Trigger) : Transition.Trigger == Trigger;
 		if (bMatch)
 		{
-			switch (Transition.State.Type)
+			switch (Transition.State.LinkType)
 			{
-			case EStateTreeTransitionType::NotSet:
+			case EStateTreeTransitionType::None:
 				IconType |= IconRightArrow;
 				break;
 			case EStateTreeTransitionType::Succeeded:
diff --git a/Engine/Plugins/Runtime/StateTree/Source/StateTreeEditorModule/Private/StateTreeCompiler.cpp b/Engine/Plugins/Runtime/StateTree/Source/StateTreeEditorModule/Private/StateTreeCompiler.cpp
--- a/Engine/Plugins/Runtime/StateTree/Source/StateTreeEditorModule/Private/StateTreeCompiler.cpp
+++ b/Engine/Plugins/Runtime/StateTree/Source/StateTreeEditorModule/Private/StateTreeCompiler.cpp
@@ -203,6 +203,8 @@ UStateTreeState* FStateTreeCompiler::GetState(const FGuid& StateID)
 
 bool FStateTreeCompiler::CreateStates()
 {
+	check(EditorData);
+	
 	// Create item for the runtime execution state
 	InstanceStructs.Add(FInstancedStruct::Make<FStateTreeExecutionState>());
 
@@ -244,6 +246,8 @@ bool FStateTreeCompiler::CreateStates()
 
 bool FStateTreeCompiler::CreateStateRecursive(UStateTreeState& State, const FStateTreeStateHandle Parent)
 {
+	check(StateTree);
+
 	FStateTreeCompilerLogStateScope LogStateScope(&State, Log);
 
 	const int32 StateIdx = StateTree->States.AddDefaulted();
@@ -312,6 +316,9 @@ bool FStateTreeCompiler::CreateConditions(UStateTreeState& State, TConstArrayVie
 
 bool FStateTreeCompiler::CreateEvaluators()
 {
+	check(EditorData);
+	check(StateTree);
+
 	const int32 EvaluatorsBegin = Nodes.Num();
 	if (const auto Validation = UE::StateTree::Compiler::IsValidCount16(EvaluatorsBegin); Validation.DidFail())
 	{
@@ -341,6 +348,9 @@ bool FStateTreeCompiler::CreateEvaluators()
 
 bool FStateTreeCompiler::CreateGlobalTasks()
 {
+	check(EditorData);
+	check(StateTree);
+
 	const int32 GlobalTasksBegin = Nodes.Num();
 	if (const auto Validation = UE::StateTree::Compiler::IsValidCount16(GlobalTasksBegin); Validation.DidFail())
 	{
@@ -349,9 +359,10 @@ bool FStateTreeCompiler::CreateGlobalTasks()
 	}
 	StateTree->GlobalTasksBegin = uint16(GlobalTasksBegin);
 
+	StateTree->bHasGlobalTransitionTasks = false;
 	for (FStateTreeEditorNode& TaskNode : EditorData->GlobalTasks)
 	{
-		if (!CreateTask(TaskNode))
+		if (!CreateTask(nullptr, TaskNode, StateTree->bHasGlobalTransitionTasks))
 		{
 			return false;
 		}
@@ -370,6 +381,8 @@ bool FStateTreeCompiler::CreateGlobalTasks()
 
 bool FStateTreeCompiler::CreateStateTasksAndParameters()
 {
+	check(StateTree);
+
 	for (int32 i = 0; i < StateTree->States.Num(); i++)
 	{
 		FCompactStateTreeState& CompactState = StateTree->States[i];
@@ -466,19 +479,24 @@ bool FStateTreeCompiler::CreateStateTasksAndParameters()
 			return false;
 		}
 		CompactState.TasksBegin = uint16(TasksBegin);
-
+		CompactState.bHasTransitionTasks = false;
+		
 		for (FStateTreeEditorNode& TaskNode : SourceState->Tasks)
 		{
-			if (!CreateTask(TaskNode))
+			bool bHasTransitionTasks = CompactState.bHasTransitionTasks; // bHasTransitionTasks is a bit flag.
+			if (!CreateTask(SourceState, TaskNode, bHasTransitionTasks))
 			{
 				return false;
 			}
+			CompactState.bHasTransitionTasks = bHasTransitionTasks; 
 		}
 
-		if (!CreateTask(SourceState->SingleTask))
+		bool bSingleStateHasTransitionTasks = CompactState.bHasTransitionTasks; // bHasTransitionTasks is a bit flag.
+		if (!CreateTask(SourceState, SourceState->SingleTask, bSingleStateHasTransitionTasks))
 		{
 			return false;
 		}
+		CompactState.bHasTransitionTasks = bSingleStateHasTransitionTasks;
 	
 		const int32 TasksNum = Nodes.Num() - TasksBegin;
 		if (const auto Validation = UE::StateTree::Compiler::IsValidCount8(TasksNum); Validation.DidFail())
@@ -494,6 +512,8 @@ bool FStateTreeCompiler::CreateStateTasksAndParameters()
 
 bool FStateTreeCompiler::CreateStateTransitions()
 {
+	check(StateTree);
+
 	for (int32 i = 0; i < StateTree->States.Num(); i++)
 	{
 		FCompactStateTreeState& CompactState = StateTree->States[i];
@@ -590,8 +610,8 @@ bool FStateTreeCompiler::CreateStateTransitions()
 		{
 			FCompactStateTransition& CompactTransition = StateTree->Transitions.AddDefaulted_GetRef();
 			CompactTransition.Trigger = Transition.Trigger;
+			CompactTransition.Priority = Transition.Priority;
 			CompactTransition.EventTag = Transition.EventTag;
-			CompactTransition.Type = Transition.State.Type;
 			if (Transition.bDelayTransition)
 			{
 				CompactTransition.Delay.Set(Transition.DelayDuration, Transition.DelayRandomVariance);
@@ -599,11 +619,14 @@ bool FStateTreeCompiler::CreateStateTransitions()
 
 			if (EnumHasAnyFlags(Transition.Trigger, EStateTreeTransitionTrigger::OnStateCompleted))
 			{
+				// Completion transitions dont have priority.
+				CompactTransition.Priority = EStateTreeTransitionPriority::None;
+				
 				// Completion transitions cannot have delay.
 				CompactTransition.Delay.Reset();
 
 				// Completion transitions must have valid target state.
-				if (CompactTransition.Type == EStateTreeTransitionType::NotSet)
+				if (Transition.State.LinkType == EStateTreeTransitionType::None)
 				{
 					Log.Reportf(EMessageSeverity::Error,
 						TEXT("State completion transition to '%s' must have transition to valid state, 'None' not accepted."),
@@ -612,7 +635,7 @@ bool FStateTreeCompiler::CreateStateTransitions()
 			}
 			
 			CompactTransition.State = FStateTreeStateHandle::Invalid;
-			if (!ResolveTransitionState(*SourceState, Transition.State, CompactTransition.State))
+			if (!ResolveTransitionState(SourceState, Transition.State, CompactTransition.State))
 			{
 				return false;
 			}
@@ -656,9 +679,9 @@ bool FStateTreeCompiler::CreateStateTransitions()
 	return true;
 }
 
-bool FStateTreeCompiler::ResolveTransitionState(const UStateTreeState& SourceState, const FStateTreeStateLink& Link, FStateTreeStateHandle& OutTransitionHandle) const 
+bool FStateTreeCompiler::ResolveTransitionState(const UStateTreeState* SourceState, const FStateTreeStateLink& Link, FStateTreeStateHandle& OutTransitionHandle) const 
 {
-	if (Link.Type == EStateTreeTransitionType::GotoState)
+	if (Link.LinkType == EStateTreeTransitionType::GotoState)
 	{
 		OutTransitionHandle = GetStateHandle(Link.ID);
 		if (!OutTransitionHandle.IsValid())
@@ -669,10 +692,10 @@ bool FStateTreeCompiler::ResolveTransitionState(const UStateTreeState& SourceSta
 			return false;
 		}
 	}
-	else if (Link.Type == EStateTreeTransitionType::NextState)
+	else if (Link.LinkType == EStateTreeTransitionType::NextState)
 	{
 		// Find next state.
-		const UStateTreeState* NextState = SourceState.GetNextSiblingState();
+		const UStateTreeState* NextState = SourceState ? SourceState->GetNextSiblingState() : nullptr;
 		if (NextState == nullptr)
 		{
 			Log.Reportf(EMessageSeverity::Error,
@@ -688,6 +711,21 @@ bool FStateTreeCompiler::ResolveTransitionState(const UStateTreeState& SourceSta
 			return false;
 		}
 	}
+	else if(Link.LinkType == EStateTreeTransitionType::Failed)
+	{
+		OutTransitionHandle = FStateTreeStateHandle::Failed;
+		return true;
+	}
+	else if(Link.LinkType == EStateTreeTransitionType::Succeeded)
+	{
+		OutTransitionHandle = FStateTreeStateHandle::Succeeded;
+		return true;
+	}
+	else if(Link.LinkType == EStateTreeTransitionType::None)
+	{
+		OutTransitionHandle = FStateTreeStateHandle::Invalid;
+		return true;
+	}
 	
 	return true;
 }
@@ -737,7 +775,7 @@ bool FStateTreeCompiler::CreateCondition(UStateTreeState& State, const FStateTre
 		Cond.InstanceIndex = FStateTreeIndex16(InstanceIndex);
 		Cond.bInstanceIsObject = false;
 
-		if (!CompileAndValidateNode(StructDesc, Node.GetMutable<FStateTreeNodeBase>(), FStateTreeDataView(SharedInstanceStructs[InstanceIndex])))
+		if (!CompileAndValidateNode(&State, StructDesc, Node, FStateTreeDataView(SharedInstanceStructs[InstanceIndex])))
 		{
 			return false;
 		}
@@ -762,7 +800,7 @@ bool FStateTreeCompiler::CreateCondition(UStateTreeState& State, const FStateTre
 		Cond.InstanceIndex = FStateTreeIndex16(InstanceIndex);
 		Cond.bInstanceIsObject = true;
 
-		if (!CompileAndValidateNode(StructDesc, Node.GetMutable<FStateTreeNodeBase>(), FStateTreeDataView(Instance)))
+		if (!CompileAndValidateNode(&State, StructDesc, Node, FStateTreeDataView(Instance)))
 		{
 			return false;
 		}
@@ -802,10 +840,44 @@ bool FStateTreeCompiler::CreateCondition(UStateTreeState& State, const FStateTre
 	return true;
 }
 
-bool FStateTreeCompiler::CompileAndValidateNode(const FStateTreeBindableStructDesc& NodeDesc, FStateTreeNodeBase& Node, const FStateTreeDataView InstanceData) const
+bool FStateTreeCompiler::CompileAndValidateNode(const UStateTreeState* SourceState, const FStateTreeBindableStructDesc& NodeDesc, FStructView NodeView, const FStateTreeDataView InstanceData) const
 {
+	if (!NodeView.IsValid())
+	{
+		return false;
+	}
+	
+	FStateTreeNodeBase& Node = NodeView.GetMutable<FStateTreeNodeBase>();
 	check(InstanceData.IsValid());
+
+	auto ValidateStateLinks = [this, SourceState](TPropertyValueIterator<FStructProperty> It)
+	{
+		for ( ; It; ++It)
+		{
+			if (It->Key->Struct == TBaseStructure<FStateTreeStateLink>::Get())
+			{
+				FStateTreeStateLink& StateLink = *static_cast<FStateTreeStateLink*>(const_cast<void*>(It->Value));
+
+				if (!ResolveTransitionState(SourceState, StateLink, StateLink.StateHandle))
+				{
+					return false;
+				}
+			}
+		}
+
+		return true;
+	};
 	
+	// Validate any state links.
+	if (!ValidateStateLinks(TPropertyValueIterator<FStructProperty>(InstanceData.GetStruct(), InstanceData.GetMutableMemory())))
+	{
+		return false;
+	}
+	if (!ValidateStateLinks(TPropertyValueIterator<FStructProperty>(NodeView.GetScriptStruct(), NodeView.GetMutableMemory())))
+	{
+		return false;
+	}
+
 	TArray<FText> ValidationErrors;
 	const EDataValidationResult Result = Node.Compile(InstanceData, ValidationErrors);
 
@@ -825,7 +897,7 @@ bool FStateTreeCompiler::CompileAndValidateNode(const FStateTreeBindableStructDe
 	return Result != EDataValidationResult::Invalid;
 }
 
-bool FStateTreeCompiler::CreateTask(const FStateTreeEditorNode& TaskNode)
+bool FStateTreeCompiler::CreateTask(UStateTreeState* State, const FStateTreeEditorNode& TaskNode, bool& bOutHasTransitionTasks)
 {
 	// Silently ignore empty nodes.
 	if (!TaskNode.Node.IsValid())
@@ -868,7 +940,7 @@ bool FStateTreeCompiler::CreateTask(const FStateTreeEditorNode& TaskNode)
 		Task.InstanceIndex = FStateTreeIndex16(InstanceIndex);
 		Task.bInstanceIsObject = false;
 
-		if (!CompileAndValidateNode(StructDesc, Node.GetMutable<FStateTreeNodeBase>(),  FStateTreeDataView(InstanceStructs[InstanceIndex])))
+		if (!CompileAndValidateNode(State, StructDesc, Node,  FStateTreeDataView(InstanceStructs[InstanceIndex])))
 		{
 			return false;
 		}
@@ -893,7 +965,7 @@ bool FStateTreeCompiler::CreateTask(const FStateTreeEditorNode& TaskNode)
 		Task.InstanceIndex = FStateTreeIndex16(InstanceIndex);
 		Task.bInstanceIsObject = true;
 
-		if (!CompileAndValidateNode(StructDesc, Node.GetMutable<FStateTreeNodeBase>(),  FStateTreeDataView(Instance)))
+		if (!CompileAndValidateNode(State, StructDesc, Node,  FStateTreeDataView(Instance)))
 		{
 			return false;
 		}
@@ -930,6 +1002,11 @@ bool FStateTreeCompiler::CreateTask(const FStateTreeEditorNode& TaskNode)
 	}
 	Task.DataViewIndex = FStateTreeIndex16(SourceStructIndex);
 
+	if (Task.bShouldAffectTransitions)
+	{
+		bOutHasTransitionTasks = true;
+	}
+	
 	return true;
 }
 
@@ -976,7 +1053,7 @@ bool FStateTreeCompiler::CreateEvaluator(const FStateTreeEditorNode& EvalNode)
 		Eval.InstanceIndex = FStateTreeIndex16(InstanceIndex);
 		Eval.bInstanceIsObject = false;
 
-		if (!CompileAndValidateNode(StructDesc, Node.GetMutable<FStateTreeNodeBase>(),  FStateTreeDataView(InstanceStructs[InstanceIndex])))
+		if (!CompileAndValidateNode(nullptr, StructDesc, Node,  FStateTreeDataView(InstanceStructs[InstanceIndex])))
 		{
 			return false;
 		}
@@ -1001,7 +1078,7 @@ bool FStateTreeCompiler::CreateEvaluator(const FStateTreeEditorNode& EvalNode)
 		Eval.InstanceIndex = FStateTreeIndex16(InstanceIndex);
 		Eval.bInstanceIsObject = true;
 
-		if (!CompileAndValidateNode(StructDesc, Node.GetMutable<FStateTreeNodeBase>(),  FStateTreeDataView(Instance)))
+		if (!CompileAndValidateNode(nullptr, StructDesc, Node,  FStateTreeDataView(Instance)))
 		{
 			return false;
 		}
@@ -1147,6 +1224,8 @@ bool FStateTreeCompiler::ValidateStructRef(const FStateTreeBindableStructDesc& S
 
 bool FStateTreeCompiler::GetAndValidateBindings(const FStateTreeBindableStructDesc& TargetStruct, TArray<FStateTreeEditorPropertyBinding>& OutBindings) const
 {
+	check(EditorData);
+	
 	if (TargetStruct.Struct == nullptr)
 	{
 		Log.Reportf(EMessageSeverity::Error, TargetStruct,
diff --git a/Engine/Plugins/Runtime/StateTree/Source/StateTreeEditorModule/Private/StateTreeState.cpp b/Engine/Plugins/Runtime/StateTree/Source/StateTreeEditorModule/Private/StateTreeState.cpp
--- a/Engine/Plugins/Runtime/StateTree/Source/StateTreeEditorModule/Private/StateTreeState.cpp
+++ b/Engine/Plugins/Runtime/StateTree/Source/StateTreeEditorModule/Private/StateTreeState.cpp
@@ -9,35 +9,20 @@
 
 #include UE_INLINE_GENERATED_CPP_BY_NAME(StateTreeState)
 
-//////////////////////////////////////////////////////////////////////////
-// FStateTreeStateLink
-
-void FStateTreeStateLink::Set(const EStateTreeTransitionType InType, const UStateTreeState* InState)
-{
-	Type = InType;
-	if (Type == EStateTreeTransitionType::GotoState)
-	{
-		check(InState);
-		Name = InState->Name;
-		ID = InState->ID;
-	}
-}
-
-
 //////////////////////////////////////////////////////////////////////////
 // FStateTreeTransition
 
 FStateTreeTransition::FStateTreeTransition(const EStateTreeTransitionTrigger InTrigger, const EStateTreeTransitionType InType, const UStateTreeState* InState)
 	: Trigger(InTrigger)
 {
-	State.Set(InType, InState);
+	State = InState ? InState->GetLinkToState() : FStateTreeStateLink(InType);
 }
 
 FStateTreeTransition::FStateTreeTransition(const EStateTreeTransitionTrigger InTrigger, const FGameplayTag InEventTag, const EStateTreeTransitionType InType, const UStateTreeState* InState)
 	: Trigger(InTrigger)
 	, EventTag(InEventTag)
 {
-	State.Set(InType, InState);
+	State = InState ? InState->GetLinkToState() : FStateTreeStateLink(InType);
 }
 
 
@@ -223,7 +208,7 @@ void UStateTreeState::PostEditChangeChainProperty(FPropertyChangedChainEvent& Pr
 						FStateTreeTransition& Transition = Transitions[TransitionsIndex];
 						Transition.Trigger = EStateTreeTransitionTrigger::OnStateCompleted;
 						const UStateTreeState* RootState = GetRootState();
-						Transition.State.Set(RootState ? RootState : this);
+						Transition.State = RootState->GetLinkToState();
 					}
 				}
 			}
@@ -334,3 +319,10 @@ const UStateTreeState* UStateTreeState::GetNextSiblingState() const
 	return nullptr;
 }
 
+FStateTreeStateLink UStateTreeState::GetLinkToState() const
+{
+	FStateTreeStateLink Link(EStateTreeTransitionType::GotoState);
+	Link.Name = Name;
+	Link.ID = ID;
+	return Link;
+}
diff --git a/Engine/Plugins/Runtime/StateTree/Source/StateTreeEditorModule/Public/StateTreeCompiler.h b/Engine/Plugins/Runtime/StateTree/Source/StateTreeEditorModule/Public/StateTreeCompiler.h
--- a/Engine/Plugins/Runtime/StateTree/Source/StateTreeEditorModule/Public/StateTreeCompiler.h
+++ b/Engine/Plugins/Runtime/StateTree/Source/StateTreeEditorModule/Public/StateTreeCompiler.h
@@ -31,7 +31,8 @@ public:
 	
 private:
 
-	bool ResolveTransitionState(const UStateTreeState& SourceState, const FStateTreeStateLink& Link, FStateTreeStateHandle& OutTransitionHandle) const;
+	/** Resolves the state a transition points to. SourceState is nullptr for global tasks. */
+	bool ResolveTransitionState(const UStateTreeState* SourceState, const FStateTreeStateLink& Link, FStateTreeStateHandle& OutTransitionHandle) const;
 	FStateTreeStateHandle GetStateHandle(const FGuid& StateID) const;
 	UStateTreeState* GetState(const FGuid& StateID);
 
@@ -45,13 +46,13 @@ private:
 	
 	bool CreateConditions(UStateTreeState& State, TConstArrayView<FStateTreeEditorNode> Conditions);
 	bool CreateCondition(UStateTreeState& State, const FStateTreeEditorNode& CondNode, const EStateTreeConditionOperand Operand, const int8 DeltaIndent);
-	bool CreateTask(const FStateTreeEditorNode& TaskNode);
+	bool CreateTask(UStateTreeState* State, const FStateTreeEditorNode& TaskNode, bool& bOutHasTransitionTasks);
 	bool CreateEvaluator(const FStateTreeEditorNode& EvalNode);
 	bool GetAndValidateBindings(const FStateTreeBindableStructDesc& TargetStruct, TArray<FStateTreeEditorPropertyBinding>& OutBindings) const;
 	bool IsPropertyAnyEnum(const FStateTreeBindableStructDesc& Struct, FStateTreeEditorPropertyPath Path) const;
 	bool ValidateStructRef(const FStateTreeBindableStructDesc& SourceStruct, FStateTreeEditorPropertyPath SourcePath,
 							const FStateTreeBindableStructDesc& TargetStruct, FStateTreeEditorPropertyPath TargetPath) const;
-	bool CompileAndValidateNode(const FStateTreeBindableStructDesc& NodeDesc, FStateTreeNodeBase& Node, const FStateTreeDataView InstanceData) const;
+	bool CompileAndValidateNode(const UStateTreeState* SourceState, const FStateTreeBindableStructDesc& NodeDesc, FStructView NodeView, const FStateTreeDataView InstanceData) const;
 
 	FStateTreeCompilerLog& Log;
 	UStateTree* StateTree = nullptr;
diff --git a/Engine/Plugins/Runtime/StateTree/Source/StateTreeEditorModule/Public/StateTreeState.h b/Engine/Plugins/Runtime/StateTree/Source/StateTreeEditorModule/Public/StateTreeState.h
--- a/Engine/Plugins/Runtime/StateTree/Source/StateTreeEditorModule/Public/StateTreeState.h
+++ b/Engine/Plugins/Runtime/StateTree/Source/StateTreeEditorModule/Public/StateTreeState.h
@@ -12,32 +12,6 @@
 
 class UStateTreeState;
 
-/**
- * Editor representation of a link to another state in StateTree
- */
-USTRUCT()
-struct STATETREEEDITORMODULE_API FStateTreeStateLink
-{
-	GENERATED_BODY()
-
-	FStateTreeStateLink() = default;
-	FStateTreeStateLink(const EStateTreeTransitionType InType) : Type(InType) {}
-
-	void Set(const EStateTreeTransitionType InType, const UStateTreeState* InState = nullptr);
-	void Set(const UStateTreeState* InState) { Set(EStateTreeTransitionType::GotoState, InState); }
-
-	bool IsValid() const { return ID.IsValid(); }
-
-	UPROPERTY(EditDefaultsOnly, Category = Link)
-	FName Name;
-
-	UPROPERTY(EditDefaultsOnly, Category = Link, meta = (IgnoreForMemberInitializationTest))
-	FGuid ID;
-
-	UPROPERTY(EditDefaultsOnly, Category = Link)
-	EStateTreeTransitionType Type = EStateTreeTransitionType::GotoState;
-};
-
 /**
  * Editor representation of a transition in StateTree
  */
@@ -76,6 +50,14 @@ struct STATETREEEDITORMODULE_API FStateTreeTransition
 	UPROPERTY(EditDefaultsOnly, Category = "Transition", meta=(DisplayName="Transition To"))
 	FStateTreeStateLink State;
 
+	/**
+	 * Transition priority when multiple transitions happen at the same time.
+	 * During transition handling, the transitions are visited from leaf to root.
+	 * The first visited transition, of highest priority, that leads to a state selection, will be activated.
+	 */
+	UPROPERTY(EditDefaultsOnly, Category = "Transition")
+	EStateTreeTransitionPriority Priority = EStateTreeTransitionPriority::Normal;
+
 	/** Delay the triggering of the transition. */
 	UPROPERTY(EditDefaultsOnly, Category = "Transition")
 	bool bDelayTransition = false;
@@ -135,7 +117,9 @@ public:
 	const UStateTreeState* GetNextSiblingState() const;
 
 	// StateTree Builder API
-
+	/** @return state link to this state. */
+	FStateTreeStateLink GetLinkToState() const;
+	
 	/** Adds child state with specified name. */
 	UStateTreeState& AddChildState(const FName ChildName, const EStateTreeStateType StateType = EStateTreeStateType::State)
 	{
diff --git a/Engine/Plugins/Runtime/StateTree/Source/StateTreeModule/Private/Blueprint/StateTreeConditionBlueprintBase.cpp b/Engine/Plugins/Runtime/StateTree/Source/StateTreeModule/Private/Blueprint/StateTreeConditionBlueprintBase.cpp
--- a/Engine/Plugins/Runtime/StateTree/Source/StateTreeModule/Private/Blueprint/StateTreeConditionBlueprintBase.cpp
+++ b/Engine/Plugins/Runtime/StateTree/Source/StateTreeModule/Private/Blueprint/StateTreeConditionBlueprintBase.cpp
@@ -22,11 +22,11 @@ bool UStateTreeConditionBlueprintBase::TestCondition(FStateTreeExecutionContext&
 	if (bHasTestCondition)
 	{
 		// Cache the owner and event queue for the duration the condition is evaluated.
-		SetCachedEventQueueFromContext(Context);
+		SetCachedInstanceDataFromContext(Context);
 
 		const bool bResult = ReceiveTestCondition();
 
-		ClearCachedEventQueue();
+		ClearCachedInstanceData();
 
 		return bResult;
 	}
diff --git a/Engine/Plugins/Runtime/StateTree/Source/StateTreeModule/Private/Blueprint/StateTreeEvaluatorBlueprintBase.cpp b/Engine/Plugins/Runtime/StateTree/Source/StateTreeModule/Private/Blueprint/StateTreeEvaluatorBlueprintBase.cpp
--- a/Engine/Plugins/Runtime/StateTree/Source/StateTreeModule/Private/Blueprint/StateTreeEvaluatorBlueprintBase.cpp
+++ b/Engine/Plugins/Runtime/StateTree/Source/StateTreeModule/Private/Blueprint/StateTreeEvaluatorBlueprintBase.cpp
@@ -22,7 +22,7 @@ UStateTreeEvaluatorBlueprintBase::UStateTreeEvaluatorBlueprintBase(const FObject
 void UStateTreeEvaluatorBlueprintBase::TreeStart(FStateTreeExecutionContext& Context)
 {
 	// Evaluator became active, cache event queue and owner.
-	SetCachedEventQueueFromContext(Context);
+	SetCachedInstanceDataFromContext(Context);
 
 	if (bHasTreeStart)
 	{
@@ -38,7 +38,7 @@ void UStateTreeEvaluatorBlueprintBase::TreeStop(FStateTreeExecutionContext& Cont
 	}
 
 	// Evaluator became inactive, clear cached event queue and owner.
-	ClearCachedEventQueue();
+	ClearCachedInstanceData();
 }
 
 void UStateTreeEvaluatorBlueprintBase::Tick(FStateTreeExecutionContext& Context, const float DeltaTime)
diff --git a/Engine/Plugins/Runtime/StateTree/Source/StateTreeModule/Private/Blueprint/StateTreeNodeBlueprintBase.cpp b/Engine/Plugins/Runtime/StateTree/Source/StateTreeModule/Private/Blueprint/StateTreeNodeBlueprintBase.cpp
--- a/Engine/Plugins/Runtime/StateTree/Source/StateTreeModule/Private/Blueprint/StateTreeNodeBlueprintBase.cpp
+++ b/Engine/Plugins/Runtime/StateTree/Source/StateTreeModule/Private/Blueprint/StateTreeNodeBlueprintBase.cpp
@@ -36,25 +36,45 @@ AActor* UStateTreeNodeBlueprintBase::GetOwnerActor(const FStateTreeExecutionCont
 	return Cast<AActor>(Context.GetOwner());
 }
 
-void UStateTreeNodeBlueprintBase::SetCachedEventQueueFromContext(const FStateTreeExecutionContext& Context) const
+void UStateTreeNodeBlueprintBase::SetCachedInstanceDataFromContext(const FStateTreeExecutionContext& Context) const
 {
-	CachedEventQueue = &Context.GetMutableEventQueue();
+	if (FStateTreeInstanceData* InstanceData = Context.GetMutableInstanceData())
+	{
+		InstanceStorage = &InstanceData->GetMutableStorage();
+	}
+	CachedState = Context.GetCurrentlyProcessedState();
 	CachedOwner = Context.GetOwner();
 }
 
-void UStateTreeNodeBlueprintBase::ClearCachedEventQueue() const
+void UStateTreeNodeBlueprintBase::ClearCachedInstanceData() const
 {
-	CachedEventQueue = nullptr;
+	InstanceStorage = nullptr;
+	CachedState = FStateTreeStateHandle::Invalid;
 	CachedOwner = nullptr;
 }
 
 void UStateTreeNodeBlueprintBase::SendEvent(const FStateTreeEvent& Event)
 {
-	if (CachedEventQueue == nullptr || CachedOwner == nullptr)
+	if (InstanceStorage == nullptr || CachedOwner == nullptr)
 	{
 		UE_VLOG_UELOG(this, LogStateTree, Error, TEXT("Trying to call SendEvent() while node is not active. Use SendEvent() on UStateTreeComponent instead for sending signals externally."));
 		return;
 	}
-	CachedEventQueue->SendEvent(CachedOwner, Event.Tag, Event.Payload, Event.Origin);
+	InstanceStorage->GetMutableEventQueue().SendEvent(CachedOwner, Event.Tag, Event.Payload, Event.Origin);
 }
 
+void UStateTreeNodeBlueprintBase::RequestTransition(const FStateTreeStateLink& TargetState, const EStateTreeTransitionPriority Priority)
+{
+	if (InstanceStorage == nullptr || CachedOwner == nullptr)
+	{
+		UE_VLOG_UELOG(this, LogStateTree, Error, TEXT("Trying to call SendEvent() while node is not active. Use SendEvent() on UStateTreeComponent instead for sending signals externally."));
+		return;
+	}
+
+	FStateTreeTransitionRequest Request;
+	Request.SourceState = CachedState;
+	Request.TargetState = TargetState.StateHandle;
+	Request.Priority = Priority;
+	
+	InstanceStorage->AddTransitionRequest(CachedOwner, Request);
+}
diff --git a/Engine/Plugins/Runtime/StateTree/Source/StateTreeModule/Private/Blueprint/StateTreeTaskBlueprintBase.cpp b/Engine/Plugins/Runtime/StateTree/Source/StateTreeModule/Private/Blueprint/StateTreeTaskBlueprintBase.cpp
--- a/Engine/Plugins/Runtime/StateTree/Source/StateTreeModule/Private/Blueprint/StateTreeTaskBlueprintBase.cpp
+++ b/Engine/Plugins/Runtime/StateTree/Source/StateTreeModule/Private/Blueprint/StateTreeTaskBlueprintBase.cpp
@@ -28,7 +28,7 @@ UStateTreeTaskBlueprintBase::UStateTreeTaskBlueprintBase(const FObjectInitialize
 EStateTreeRunStatus UStateTreeTaskBlueprintBase::EnterState(FStateTreeExecutionContext& Context, const FStateTreeTransitionResult& Transition)
 {
 	// Task became active, cache event queue and owner.
-	SetCachedEventQueueFromContext(Context);
+	SetCachedInstanceDataFromContext(Context);
 	
 	if (bHasEnterState)
 	{
@@ -45,7 +45,7 @@ void UStateTreeTaskBlueprintBase::ExitState(FStateTreeExecutionContext& Context,
 	}
 
 	// Task became inactive, clear cached event queue and owner.
-	ClearCachedEventQueue();
+	ClearCachedInstanceData();
 }
 
 void UStateTreeTaskBlueprintBase::StateCompleted(FStateTreeExecutionContext& Context, const EStateTreeRunStatus CompletionStatus, const FStateTreeActiveStates& CompletedActiveStates)
diff --git a/Engine/Plugins/Runtime/StateTree/Source/StateTreeModule/Private/StateTree.cpp b/Engine/Plugins/Runtime/StateTree/Source/StateTreeModule/Private/StateTree.cpp
--- a/Engine/Plugins/Runtime/StateTree/Source/StateTreeModule/Private/StateTree.cpp
+++ b/Engine/Plugins/Runtime/StateTree/Source/StateTreeModule/Private/StateTree.cpp
@@ -81,6 +81,7 @@ void UStateTree::ResetCompiled()
 
 	GlobalTasksBegin = 0;
 	GlobalTasksNum = 0;
+	bHasGlobalTransitionTasks = false;
 	
 	ResetLinked();
 }
diff --git a/Engine/Plugins/Runtime/StateTree/Source/StateTreeModule/Private/StateTreeExecutionContext.cpp b/Engine/Plugins/Runtime/StateTree/Source/StateTreeModule/Private/StateTreeExecutionContext.cpp
--- a/Engine/Plugins/Runtime/StateTree/Source/StateTreeModule/Private/StateTreeExecutionContext.cpp
+++ b/Engine/Plugins/Runtime/StateTree/Source/StateTreeModule/Private/StateTreeExecutionContext.cpp
@@ -28,6 +28,8 @@ FStateTreeExecutionContext::FStateTreeExecutionContext(UObject& InOwner, const U
 		DataViews.SetNum(StateTree.GetNumDataViews());
 		// Set data views associated to the parameters using the default values
 		SetDefaultParameters();
+
+		SharedInstanceData = StateTree.GetSharedInstanceData();
 	}
 	else
 	{
@@ -177,9 +179,6 @@ EStateTreeRunStatus FStateTreeExecutionContext::Start()
 		return EStateTreeRunStatus::Failed;
 	}
 
-	const TSharedPtr<FStateTreeInstanceData> SharedInstanceData = StateTree.GetSharedInstanceData();
-	check(SharedInstanceData.IsValid());
-
 	// Start evaluators and global tasks. Fail the execution if any global task fails.
 	FStateTreeIndex16 LastInitializedTaskIndex;
 	if (!StartEvaluatorsAndGlobalTasks(LastInitializedTaskIndex))
@@ -203,14 +202,14 @@ EStateTreeRunStatus FStateTreeExecutionContext::Start()
 	static const FStateTreeStateHandle RootState = FStateTreeStateHandle(0);
 
 	FStateTreeActiveStates NextActiveStates;
-	if (SelectState(*SharedInstanceData.Get(), RootState, NextActiveStates))
+	if (SelectState(RootState, NextActiveStates))
 	{
-		if (NextActiveStates.Last() == FStateTreeStateHandle::Succeeded || NextActiveStates.Last() == FStateTreeStateHandle::Failed)
+		if (NextActiveStates.Last().IsCompletionState())
 		{
 			// Transition to a terminal state (succeeded/failed), or default transition failed.
 			STATETREE_LOG(Warning, TEXT("%s: Tree %s at StateTree start on '%s' using StateTree '%s'."),
 				ANSI_TO_TCHAR(__FUNCTION__), NextActiveStates.Last() == FStateTreeStateHandle::Succeeded ? TEXT("succeeded") : TEXT("failed"), *GetNameSafe(&Owner), *GetFullNameSafe(&StateTree));
-			Exec->TreeRunStatus = NextActiveStates.Last() == FStateTreeStateHandle::Succeeded ? EStateTreeRunStatus::Succeeded : EStateTreeRunStatus::Failed;
+			Exec->TreeRunStatus = NextActiveStates.Last().ToCompletionStatus();
 		}
 		else
 		{
@@ -311,13 +310,9 @@ EStateTreeRunStatus FStateTreeExecutionContext::Tick(const float DeltaTime)
 		return EStateTreeRunStatus::Failed;
 	}
 
-	const TSharedPtr<FStateTreeInstanceData> SharedInstanceData = StateTree.GetSharedInstanceData();
-	check(SharedInstanceData.IsValid());
-	
 	FStateTreeEventQueue& EventQueue = InstanceData.GetMutableEventQueue();
-	
 	FStateTreeExecutionState* Exec = &GetExecState();
-
+	
 	// Capture events added between ticks.
 	EventsToProcess = EventQueue.GetEvents();
 	EventQueue.Reset();
@@ -333,7 +328,7 @@ EStateTreeRunStatus FStateTreeExecutionContext::Tick(const float DeltaTime)
 	{
 		DelayedState.TimeLeft -= DeltaTime;
 	}
-	
+
 	// Tick global evaluators and tasks.
 	const EStateTreeRunStatus EvalAndGlobalTaskStatus = TickEvaluatorsAndGlobalTasks(DeltaTime);
 	if (EvalAndGlobalTaskStatus != EStateTreeRunStatus::Running)
@@ -364,18 +359,18 @@ EStateTreeRunStatus FStateTreeExecutionContext::Tick(const float DeltaTime)
 		EventsToProcess.Append(EventQueue.GetEvents());
 
 		// Trigger conditional transitions or state succeed/failed transitions. First tick transition is handled here too.
-		FStateTreeTransitionResult Transition;
-		if (TriggerTransitions(*SharedInstanceData.Get(), Transition))
+		if (TriggerTransitions())
 		{
 			// We have committed to state change, consume events that were accumulated during the tick above.
 			EventQueue.Reset();
 
-			ExitState(Transition);
+			ExitState(NextTransition);
 
-			if (Transition.NextActiveStates.Last() == FStateTreeStateHandle::Succeeded || Transition.NextActiveStates.Last() == FStateTreeStateHandle::Failed)
+			// Tree succeeded or failed.
+			if (NextTransition.TargetState.IsCompletionState())
 			{
 				// Transition to a terminal state (succeeded/failed), or default transition failed.
-				Exec->TreeRunStatus = Transition.NextActiveStates.Last() == FStateTreeStateHandle::Succeeded ? EStateTreeRunStatus::Succeeded : EStateTreeRunStatus::Failed;
+				Exec->TreeRunStatus = NextTransition.TargetState.ToCompletionStatus();
 				Exec->ActiveStates.Reset();
 				return Exec->TreeRunStatus;
 			}
@@ -385,22 +380,24 @@ EStateTreeRunStatus FStateTreeExecutionContext::Tick(const float DeltaTime)
 			EventQueue.Reset();
 
 			// Enter state tasks can fail/succeed, treat it same as tick.
-			const EStateTreeRunStatus LastTickStatus = EnterState(Transition);
+			const EStateTreeRunStatus LastTickStatus = EnterState(NextTransition);
 
-			// Consider events so far processed. Events sent during EnterState went into EventQueue, and are processed in next iteration.
-			EventsToProcess.Reset();
+			NextTransition.Reset();
 
 			// Need to reacquire the exec state as EnterState may alter the allocation. 
 			Exec = &GetExecState();
 			Exec->LastTickStatus = LastTickStatus;
 
+			// Consider events so far processed. Events sent during EnterState went into EventQueue, and are processed in next iteration.
+			EventsToProcess.Reset();
+
 			// Report state completed immediately.
 			if (Exec->LastTickStatus != EStateTreeRunStatus::Running)
 			{
 				StateCompleted();
 			}
 		}
-		
+
 		// Stop as soon as have found a running state.
 		if (Exec->LastTickStatus == EStateTreeRunStatus::Running)
 		{
@@ -468,6 +465,38 @@ void FStateTreeExecutionContext::SendEvent(const FGameplayTag Tag, const FConstS
 	EventQueue.SendEvent(&Owner, Tag, Payload, Origin);
 }
 
+void FStateTreeExecutionContext::RequestTransition(const FStateTreeTransitionRequest& Request)
+{
+	CSV_SCOPED_TIMING_STAT_EXCLUSIVE(StateTree_RequestTransition);
+
+	if (!IsValid())
+	{
+		STATETREE_LOG(Warning, TEXT("%s: StateTree context is not initialized properly ('%s' using StateTree '%s')"),
+			ANSI_TO_TCHAR(__FUNCTION__), *GetNameSafe(&Owner), *GetFullNameSafe(&StateTree));
+		return;
+	}
+
+	if (!InstanceData.IsValid())
+	{
+		STATETREE_LOG(Error, TEXT("%s: RequestTransition called on %s using StateTree %s with invalid instance data. Start() must be called before requesting transition."),
+			ANSI_TO_TCHAR(__FUNCTION__), *GetNameSafe(&Owner), *GetFullNameSafe(&StateTree));
+		return;
+	}
+
+	STATETREE_LOG(Verbose, TEXT("Request transition to '%s' at priority %s"), *GetSafeStateName(Request.TargetState), *UEnum::GetDisplayValueAsText(Request.Priority).ToString());
+
+	if (bAllowDirectTransitions)
+	{
+		RequestTransition(Request.TargetState, Request.Priority);		
+	}
+	else
+	{
+		FStateTreeTransitionRequest RequestWithSource = Request;
+		RequestWithSource.SourceState = CurrentlyProcessedState;
+		InstanceData.AddTransitionRequest(&Owner, RequestWithSource);
+	}
+}
+
 void FStateTreeExecutionContext::UpdateInstanceData(const FStateTreeActiveStates& CurrentActiveStates, const FStateTreeActiveStates& NextActiveStates)
 {
 	// Find common section of states at start.
@@ -841,6 +870,8 @@ void FStateTreeExecutionContext::StateCompleted()
 		const FStateTreeStateHandle CurrentHandle = Exec.ActiveStates[Index];
 		const FCompactStateTreeState& State = StateTree.States[CurrentHandle.Index];
 
+		FCurrentlyProcessedStateScope StateScope(*this, CurrentHandle);
+		
 		STATETREE_LOG(Verbose, TEXT("%*sState '%s'"), Index*UE::StateTree::DebugIndentSize, TEXT(""), *DebugGetStatePath(Exec.ActiveStates, Index));
 
 		// Notify Tasks
@@ -1096,6 +1127,8 @@ EStateTreeRunStatus FStateTreeExecutionContext::TickTasks(const float DeltaTime)
 		const FStateTreeStateHandle CurrentHandle = Exec.ActiveStates[Index];
 		const FCompactStateTreeState& State = StateTree.States[CurrentHandle.Index];
 
+		FCurrentlyProcessedStateScope StateScope(*this, CurrentHandle);
+		
 		STATETREE_CLOG(State.TasksNum > 0, VeryVerbose, TEXT("%*sState '%s'"), Index*UE::StateTree::DebugIndentSize, TEXT(""), *DebugGetStatePath(Exec.ActiveStates, Index));
 
 		if (State.Type == EStateTreeStateType::Linked)
@@ -1163,11 +1196,11 @@ EStateTreeRunStatus FStateTreeExecutionContext::TickTasks(const float DeltaTime)
 	return Result;
 }
 
-bool FStateTreeExecutionContext::TestAllConditions(FStateTreeInstanceData& SharedInstanceData, const int32 ConditionsOffset, const int32 ConditionsNum)
+bool FStateTreeExecutionContext::TestAllConditions(const int32 ConditionsOffset, const int32 ConditionsNum)
 {
 	CSV_SCOPED_TIMING_STAT_EXCLUSIVE(StateTree_TestConditions);
 
-	if (ConditionsNum == 0)
+	if (ConditionsNum == 0 || !SharedInstanceData.IsValid())
 	{
 		return true;
 	}
@@ -1182,11 +1215,11 @@ bool FStateTreeExecutionContext::TestAllConditions(FStateTreeInstanceData& Share
 		const FStateTreeConditionBase& Cond = StateTree.Nodes[ConditionsOffset + Index].Get<FStateTreeConditionBase>();
 		if (Cond.bInstanceIsObject)
 		{
-			DataViews[Cond.DataViewIndex.Get()] = SharedInstanceData.GetMutableObject(Cond.InstanceIndex.Get());
+			DataViews[Cond.DataViewIndex.Get()] = SharedInstanceData->GetMutableObject(Cond.InstanceIndex.Get());
 		}
 		else
 		{
-			DataViews[Cond.DataViewIndex.Get()] = SharedInstanceData.GetMutableStruct(Cond.InstanceIndex.Get());
+			DataViews[Cond.DataViewIndex.Get()] = SharedInstanceData->GetMutableStruct(Cond.InstanceIndex.Get());
 		}
 
 		// Copy bound properties.
@@ -1261,93 +1294,127 @@ FString FStateTreeExecutionContext::DebugGetEventsAsString() const
 	return Result;
 }
 
-bool FStateTreeExecutionContext::SelectTransition(FStateTreeInstanceData& SharedInstanceData, FStateTreeExecutionState& Exec,
-		const int16 StateIndex, const FCompactStateTreeState& State, const FCompactStateTransition& Transition, FStateTreeTransitionResult& OutTransition)
+bool FStateTreeExecutionContext::RequestTransition(const FStateTreeStateHandle NextState, const EStateTreeTransitionPriority Priority)
 {
-	if (Transition.Type == EStateTreeTransitionType::GotoState || Transition.Type == EStateTreeTransitionType::NextState)
+	// Skip lower priority transitions.
+	if (NextTransition.Priority >= Priority)
 	{
-		OutTransition.TargetState = Transition.State;
-		OutTransition.NextActiveStates.Reset();
-
-		if (SelectState(SharedInstanceData, Transition.State, OutTransition.NextActiveStates))
-		{
-			STATETREE_LOG(Verbose, TEXT("Transition on state '%s' (%s) -[%s]-> state '%s'"),
-				*GetSafeStateName(Exec.ActiveStates.Last()), *State.Name.ToString(), *GetSafeStateName(Transition.State), *GetSafeStateName(OutTransition.NextActiveStates.Last()));
-			return true;
-		}
-		
-		return  false;
+		return false;
 	}
+
+	const FStateTreeExecutionState& Exec = GetExecState();
 	
-	if (Transition.Type == EStateTreeTransitionType::Succeeded || Transition.Type == EStateTreeTransitionType::Failed)
+	if (NextState.IsCompletionState())
 	{
-		// Tree succeeded or failed.
-		// If we are on a linked state, find the first parent linked state and mark that as completed, continue to find completion transitions.
-		const FStateTreeStateHandle ParentLinkedState = GetParentLinkedStateHandle(Exec.ActiveStates, StateIndex);
-		if (ParentLinkedState.IsValid())
-		{
-			STATETREE_LOG(Verbose, TEXT("Compled subtree '%s' from state '%s' (%s): %s"),
-				*GetSafeStateName(ParentLinkedState), *GetSafeStateName(Exec.ActiveStates.Last()), *State.Name.ToString(), *UEnum::GetDisplayValueAsText(Transition.Type).ToString());
-			Exec.CompletedStateHandle = ParentLinkedState;
-			Exec.LastTickStatus = Transition.Type == EStateTreeTransitionType::Succeeded ? EStateTreeRunStatus::Succeeded : EStateTreeRunStatus::Failed;
-			return TriggerTransitions(SharedInstanceData, OutTransition);
-		}
-		
-		STATETREE_LOG(Verbose, TEXT("Stop tree execution from state '%s' (%s): %s"),
-			*GetSafeStateName(Exec.ActiveStates.Last()), *State.Name.ToString(), *UEnum::GetDisplayValueAsText(Transition.Type).ToString());
+		NextTransition.CurrentActiveStates = Exec.ActiveStates;
+		NextTransition.CurrentRunStatus = Exec.TreeRunStatus;
+		NextTransition.SourceState = CurrentlyProcessedState;
+		NextTransition.TargetState = NextState;
+		NextTransition.NextActiveStates = FStateTreeActiveStates(NextState);
+		NextTransition.Priority = Priority;
+
+		STATETREE_LOG(Verbose, TEXT("Transition on state '%s' -[%s]-> state '%s'"),
+			*GetSafeStateName(NextTransition.CurrentActiveStates.Last()), *GetSafeStateName(NextState), *GetSafeStateName(NextTransition.NextActiveStates.Last()));
 
-		const FStateTreeStateHandle NextState = Transition.Type == EStateTreeTransitionType::Succeeded ? FStateTreeStateHandle::Succeeded : FStateTreeStateHandle::Failed; 
-		OutTransition.TargetState = NextState;
-		OutTransition.NextActiveStates = FStateTreeActiveStates(NextState);
 		return true;
 	}
+	if (!NextState.IsValid())
+	{
+		// NotSet is no-operation, but can be used to mask a transition at parent state. Returning unset keeps updating current state.
+		NextTransition.CurrentActiveStates = Exec.ActiveStates;
+		NextTransition.CurrentRunStatus = Exec.TreeRunStatus;
+		NextTransition.SourceState = CurrentlyProcessedState;
+		NextTransition.TargetState = FStateTreeStateHandle::Invalid;
+		NextTransition.NextActiveStates.Reset();
+		NextTransition.Priority = Priority;
+		return true;
+	}
+
+	FStateTreeActiveStates NewActiveState;
+	if (SelectState(NextState, NewActiveState))
+	{
+		NextTransition.CurrentActiveStates = Exec.ActiveStates;
+		NextTransition.CurrentRunStatus = Exec.TreeRunStatus;
+		NextTransition.SourceState = CurrentlyProcessedState;
+		NextTransition.TargetState = NextState;
+		NextTransition.NextActiveStates = NewActiveState;
+		NextTransition.Priority = Priority;
 
+		STATETREE_LOG(Verbose, TEXT("Transition on state '%s' -[%s]-> state '%s'"),
+			*GetSafeStateName(NextTransition.CurrentActiveStates.Last()), *GetSafeStateName(NextState), *GetSafeStateName(NextTransition.NextActiveStates.Last()));
+		
+		return true;
+	}
+		
 	return false;
 }
 
-bool FStateTreeExecutionContext::TriggerTransitions(FStateTreeInstanceData& SharedInstanceData, FStateTreeTransitionResult& OutTransition)
+bool FStateTreeExecutionContext::TriggerTransitions()
 {
 	CSV_SCOPED_TIMING_STAT_EXCLUSIVE(StateTree_TriggerTransition);
 
+	FAllowDirectTransitionsScope AllowDirectTransitionsScope(*this); // Set flag for the scope of this function to allow direct transitions without buffering.
 	FStateTreeExecutionState& Exec = GetExecState();
 
 	STATETREE_CLOG(!EventsToProcess.IsEmpty(), Verbose, TEXT("Trigger transitions with events [%s]"), *DebugGetEventsAsString());
 
-	auto HasEvent = [this](const FGameplayTag QueriedTag)
-	{
-		if (EventsToProcess.IsEmpty())
-		{
-			return false;
-		}
-		
-		return EventsToProcess.ContainsByPredicate([QueriedTag](const FStateTreeEvent& Event)
-		{
-			return Event.Tag.MatchesTag(QueriedTag);
-		});
-	};
-
-	OutTransition.CurrentActiveStates = Exec.ActiveStates;
+	NextTransition.Reset();
 
 	//
-	// Check tick and event based transitions first.
+	// Process transition requests
+	//
+	for (const FStateTreeTransitionRequest& Request : InstanceData.GetTransitionRequests())
+	{
+		RequestTransition(Request.TargetState, Request.Priority);
+	}
+	InstanceData.ResetTransitionRequests();
+	
+	//
+	// Check tick, event, and task based transitions first.
 	//
 	for (int32 StateIndex = Exec.ActiveStates.Num() - 1; StateIndex >= 0; StateIndex--)
 	{
-		const FCompactStateTreeState& State = StateTree.States[Exec.ActiveStates[StateIndex].Index];
+		const FStateTreeStateHandle StateHandle = Exec.ActiveStates[StateIndex];
+		const FCompactStateTreeState& State = StateTree.States[StateHandle.Index];
+
+		FCurrentlyProcessedStateScope StateScope(*this, StateHandle);
+
+		if (State.bHasTransitionTasks)
+		{
+			STATETREE_CLOG(State.TasksNum > 0, VeryVerbose, TEXT("%*sTrigger task transitions in state '%s'"), StateIndex*UE::StateTree::DebugIndentSize, TEXT(""), *DebugGetStatePath(Exec.ActiveStates, StateIndex));
+
+			for (int32 TaskIndex = (State.TasksBegin + State.TasksNum) - 1; TaskIndex >= State.TasksBegin; TaskIndex--)
+			{
+				const FStateTreeTaskBase& Task = StateTree.Nodes[TaskIndex].Get<FStateTreeTaskBase>();
+				if (Task.bShouldAffectTransitions)
+				{
+					STATETREE_LOG(VeryVerbose, TEXT("%*sTriggerTransitions: '%s'"), UE::StateTree::DebugIndentSize, TEXT(""), *Task.Name.ToString());
+					check(DataViews[Task.DataViewIndex.Get()].IsValid());
+					Task.TriggerTransitions(*this);
+				}
+			}
+		}
+		
 		for (uint8 i = 0; i < State.TransitionsNum; i++)
 		{
 			// All transition conditions must pass
 			const int16 TransitionIndex = State.TransitionsBegin + i;
 			const FCompactStateTransition& Transition = StateTree.Transitions[TransitionIndex];
 
+			// No need to test the transition if same or higher priority transition has already been processed.
+			if (Transition.Priority <= NextTransition.Priority)
+			{
+				continue;
+			}
+			
 			if (Transition.Trigger == EStateTreeTransitionTrigger::OnTick
-				|| (Transition.Trigger == EStateTreeTransitionTrigger::OnEvent && HasEvent(Transition.EventTag)))
+				|| (Transition.Trigger == EStateTreeTransitionTrigger::OnEvent && HasEventToProcess(Transition.EventTag)))
 			{
 				// If delay has passed, allow to trigger the transition even if conditions do not pass.
 				FStateTreeTransitionDelayedState* DelayedState = Transition.HasDelay() ? Exec.FindDelayedTransition(FStateTreeIndex16(TransitionIndex)) : nullptr;
 				const bool bDelayCompleted = DelayedState != nullptr && DelayedState->TimeLeft <= 0.0f;
 
-				if (bDelayCompleted || TestAllConditions(SharedInstanceData, Transition.ConditionsBegin, Transition.ConditionsNum))
+				if (bDelayCompleted || TestAllConditions(Transition.ConditionsBegin, Transition.ConditionsNum))
 				{
 					// If the transitions is delayed, set up the delay. 
 					if (Transition.HasDelay())
@@ -1382,25 +1449,32 @@ bool FStateTreeExecutionContext::TriggerTransitions(FStateTreeInstanceData& Shar
 							});
 					}
 
-					if (Transition.Type == EStateTreeTransitionType::NotSet)
-					{
-						// NotSet is no-operation, but can be used to mask a transition at parent state. Returning unset keeps updating current state.
-						return false;
-					}
-
-					if (SelectTransition(SharedInstanceData, Exec, StateIndex, State, Transition, OutTransition))
-					{
-						return true;
-					}
+					RequestTransition(Transition.State, Transition.Priority);
 				}
 			}
 		}
 	}
 
+	if (StateTree.bHasGlobalTransitionTasks)
+	{
+		STATETREE_LOG(VeryVerbose, TEXT("Trigger global task transitions"));
+		for (int32 TaskIndex = (StateTree.GlobalTasksBegin + StateTree.GlobalTasksNum) - 1; TaskIndex >= StateTree.GlobalTasksBegin; TaskIndex--)
+		{
+			const FStateTreeTaskBase& Task =  StateTree.Nodes[TaskIndex].Get<FStateTreeTaskBase>();
+			if (Task.bShouldAffectTransitions)
+			{
+				STATETREE_LOG(VeryVerbose, TEXT("%*sTriggerTransitions: '%s'"), UE::StateTree::DebugIndentSize, TEXT(""), *Task.Name.ToString());
+				check(DataViews[Task.DataViewIndex.Get()].IsValid());
+				Task.TriggerTransitions(*this);
+			}
+		}
+	}
+
 	//
 	// Check state completion transitions.
 	//
-	if (Exec.LastTickStatus != EStateTreeRunStatus::Running)
+	if (NextTransition.Priority == EStateTreeTransitionPriority::None
+		&& Exec.LastTickStatus != EStateTreeRunStatus::Running)
 	{
 		// Start from the last completed state.
 		const int32 StateStartIndex = Exec.CompletedStateHandle.IsValid() ? Exec.ActiveStates.IndexOfReverse(Exec.CompletedStateHandle) : (Exec.ActiveStates.Num() - 1);
@@ -1411,7 +1485,11 @@ bool FStateTreeExecutionContext::TriggerTransitions(FStateTreeInstanceData& Shar
 		// Check completion transitions
 		for (int32 StateIndex = StateStartIndex; StateIndex >= 0; StateIndex--)
 		{
-			const FCompactStateTreeState& State = StateTree.States[Exec.ActiveStates[StateIndex].Index];
+			const FStateTreeStateHandle StateHandle = Exec.ActiveStates[StateIndex];
+			const FCompactStateTreeState& State = StateTree.States[StateHandle.Index];
+
+			FCurrentlyProcessedStateScope StateScope(*this, StateHandle);
+			
 			for (uint8 i = 0; i < State.TransitionsNum; i++)
 			{
 				// All transition conditions must pass
@@ -1420,30 +1498,50 @@ bool FStateTreeExecutionContext::TriggerTransitions(FStateTreeInstanceData& Shar
 
 				if (EnumHasAnyFlags(Transition.Trigger, CompletionTrigger))
 				{
-					if (TestAllConditions(SharedInstanceData, Transition.ConditionsBegin, Transition.ConditionsNum))
+					if (TestAllConditions(Transition.ConditionsBegin, Transition.ConditionsNum))
 					{
-						// No delay allowed on completion conditions. 
-						
-						if (SelectTransition(SharedInstanceData, Exec, StateIndex, State, Transition, OutTransition))
+						// No delay allowed on completion conditions.
+						// No priority on completion transitions, use the priority to signal that state is selected.
+						if (RequestTransition(Transition.State, EStateTreeTransitionPriority::Normal))
 						{
-							return true;
+							break;
 						}
 					}
 				}
 			}
+
+			if (NextTransition.Priority != EStateTreeTransitionPriority::None)
+			{
+				break;
+			}
 		}
 
-		if (Exec.LastTickStatus != EStateTreeRunStatus::Running)
+		if (NextTransition.Priority == EStateTreeTransitionPriority::None)
 		{
 			// Could not trigger completion transition, jump back to start.
-			static const FStateTreeStateHandle RootState = FStateTreeStateHandle(0);
-			OutTransition.TargetState = RootState;
-			return SelectState(SharedInstanceData, RootState, OutTransition.NextActiveStates);
+			FCurrentlyProcessedStateScope StateScope(*this, FStateTreeStateHandle::Root);
+			RequestTransition(FStateTreeStateHandle::Root, EStateTreeTransitionPriority::Normal);
 		}
 	}
-	
-	// No transition triggered, keep on updating current state.
-	return false;
+
+	// Check if the transition was succeed/failed, if we're on a sub-tree, complete the subtree instead of transition.
+	if (NextTransition.TargetState.IsCompletionState())
+	{
+		const FStateTreeStateHandle ParentLinkedState = GetParentLinkedStateHandle(Exec.ActiveStates, NextTransition.SourceState);
+		if (ParentLinkedState.IsValid())
+		{
+			const EStateTreeRunStatus RunStatus = NextTransition.TargetState.ToCompletionStatus(); 
+			STATETREE_LOG(Verbose, TEXT("Compled subtree '%s' from state '%s' (%s): %s"),
+				*GetSafeStateName(ParentLinkedState), *GetSafeStateName(Exec.ActiveStates.Last()), *GetSafeStateName(NextTransition.SourceState), *UEnum::GetDisplayValueAsText(RunStatus).ToString());
+			Exec.CompletedStateHandle = ParentLinkedState;
+			Exec.LastTickStatus = RunStatus;
+
+			// Not a transition
+			return false;
+		}
+	}
+
+	return NextTransition.TargetState.IsValid();
 }
 
 FStateTreeStateHandle FStateTreeExecutionContext::GetParentLinkedStateHandle(const FStateTreeActiveStates& ActiveStates, const int32 StartStartIndex) const
@@ -1461,7 +1559,35 @@ FStateTreeStateHandle FStateTreeExecutionContext::GetParentLinkedStateHandle(con
 	return FStateTreeStateHandle();
 }	
 
-bool FStateTreeExecutionContext::SelectState(FStateTreeInstanceData& SharedInstanceData, const FStateTreeStateHandle NextState, FStateTreeActiveStates& OutNewActiveState)
+FStateTreeStateHandle FStateTreeExecutionContext::GetParentLinkedStateHandle(const FStateTreeActiveStates& ActiveStates, const FStateTreeStateHandle StartStateHandle) const
+{
+	// Find start state
+	int32 StateIndex = ActiveStates.Num() - 1;
+	while (StateIndex >= 0)
+	{
+		if (ActiveStates[StateIndex] == StartStateHandle)
+		{
+			break;
+		}
+		StateIndex--;
+	}
+
+	// Find parent linked state.
+	while (StateIndex >= 0)
+	{
+		const FCompactStateTreeState& State = StateTree.States[ActiveStates[StateIndex].Index];
+		if (State.LinkedState.IsValid())
+		{
+			return ActiveStates[StateIndex];
+		}
+		
+		StateIndex--;
+	}
+
+	return FStateTreeStateHandle();
+}
+
+bool FStateTreeExecutionContext::SelectState(const FStateTreeStateHandle NextState, FStateTreeActiveStates& OutNewActiveState)
 {
 	const FStateTreeExecutionState& Exec = GetExecState();
 
@@ -1517,10 +1643,10 @@ bool FStateTreeExecutionContext::SelectState(FStateTreeInstanceData& SharedInsta
 		return false;
 	}
 
-	return SelectStateInternal(SharedInstanceData, NextState, OutNewActiveState);
+	return SelectStateInternal(NextState, OutNewActiveState);
 }
 
-bool FStateTreeExecutionContext::SelectStateInternal(FStateTreeInstanceData& SharedInstanceData, const FStateTreeStateHandle NextState, FStateTreeActiveStates& OutNewActiveState)
+bool FStateTreeExecutionContext::SelectStateInternal(const FStateTreeStateHandle NextState, FStateTreeActiveStates& OutNewActiveState)
 {
 	CSV_SCOPED_TIMING_STAT_EXCLUSIVE(StateTree_SelectState);
 
@@ -1537,7 +1663,7 @@ bool FStateTreeExecutionContext::SelectStateInternal(FStateTreeInstanceData& Sha
 	const FCompactStateTreeState& State = StateTree.States[NextState.Index];
 
 	// Check that the state can be entered
-	if (TestAllConditions(SharedInstanceData, State.EnterConditionsBegin, State.EnterConditionsNum))
+	if (TestAllConditions(State.EnterConditionsBegin, State.EnterConditionsNum))
 	{
 		if (!OutNewActiveState.Push(NextState))
 		{
@@ -1549,7 +1675,7 @@ bool FStateTreeExecutionContext::SelectStateInternal(FStateTreeInstanceData& Sha
 		if (State.LinkedState.IsValid())
 		{
 			// If State is linked, proceed to the linked state.
-			if (SelectStateInternal(SharedInstanceData, State.LinkedState, OutNewActiveState))
+			if (SelectStateInternal(State.LinkedState, OutNewActiveState))
 			{
 				// Selection succeeded
 				return true;
@@ -1560,7 +1686,7 @@ bool FStateTreeExecutionContext::SelectStateInternal(FStateTreeInstanceData& Sha
 			// If the state has children, proceed to select children.
 			for (uint16 ChildState = State.ChildrenBegin; ChildState < State.ChildrenEnd; ChildState = StateTree.States[ChildState].GetNextSibling())
 			{
-				if (SelectStateInternal(SharedInstanceData, FStateTreeStateHandle(ChildState), OutNewActiveState))
+				if (SelectStateInternal(FStateTreeStateHandle(ChildState), OutNewActiveState))
 				{
 					// Selection succeeded
 					return true;
@@ -1741,13 +1867,12 @@ void FStateTreeExecutionContext::DebugPrintInternalLayout()
 	StateTree.PropertyBindings.DebugPrintInternalLayout(DebugString);
 
 	// Transitions
-	DebugString += FString::Printf(TEXT("\nTransitions(%d)\n  [ %-3s | %15s | %-20s | %-20s | %-40s | %-8s ]\n"), StateTree.Transitions.Num()
-		, TEXT("Idx"), TEXT("State"), TEXT("Transition Type"), TEXT("Transition Trigger"), TEXT("Transition Event Tag"), TEXT("Num Cond"));
+	DebugString += FString::Printf(TEXT("\nTransitions(%d)\n  [ %-3s | %15s | %-20s | %-40s | %-8s ]\n"), StateTree.Transitions.Num()
+		, TEXT("Idx"), TEXT("State"), TEXT("Transition Trigger"), TEXT("Transition Event Tag"), TEXT("Num Cond"));
 	for (const FCompactStateTransition& Transition : StateTree.Transitions)
 	{
-		DebugString += FString::Printf(TEXT("  | %3d | %15s | %-20s | %-20s | %-40s | %8d |\n"),
+		DebugString += FString::Printf(TEXT("  | %3d | %15s | %-20s | %-40s | %8d |\n"),
 									Transition.ConditionsBegin, *Transition.State.Describe(),
-									*UEnum::GetDisplayValueAsText(Transition.Type).ToString(),
 									*UEnum::GetDisplayValueAsText(Transition.Trigger).ToString(),
 									*Transition.EventTag.ToString(),
 									Transition.ConditionsNum);
diff --git a/Engine/Plugins/Runtime/StateTree/Source/StateTreeModule/Private/StateTreeInstanceData.cpp b/Engine/Plugins/Runtime/StateTree/Source/StateTreeModule/Private/StateTreeInstanceData.cpp
--- a/Engine/Plugins/Runtime/StateTree/Source/StateTreeModule/Private/StateTreeInstanceData.cpp
+++ b/Engine/Plugins/Runtime/StateTree/Source/StateTreeModule/Private/StateTreeInstanceData.cpp
@@ -4,6 +4,7 @@
 #include "StateTreeTypes.h"
 #include "Serialization/MemoryReader.h"
 #include "Serialization/MemoryWriter.h"
+#include "VisualLogger/VisualLogger.h"
 
 #include UE_INLINE_GENERATED_CPP_BY_NAME(StateTreeInstanceData)
 
@@ -41,8 +42,27 @@ namespace UE::StateTree
 
 } // UE::StateTree
 
+
+//----------------------------------------------------------------//
+// FStateTreeInstanceStorage
+//----------------------------------------------------------------//
+
+void FStateTreeInstanceStorage::AddTransitionRequest(const UObject* Owner, const FStateTreeTransitionRequest& Request)
+{
+	constexpr int32 MaxPendingTransitionRequests = 32;
+	
+	if (TransitionRequests.Num() >= MaxPendingTransitionRequests)
+	{
+		UE_VLOG_UELOG(Owner, LogStateTree, Error, TEXT("%s: Too many transition requests sent to '%s' (%d pending). Dropping request."), ANSI_TO_TCHAR(__FUNCTION__), *GetNameSafe(Owner), TransitionRequests.Num());
+		return;
+	}
+
+	TransitionRequests.Add(Request);
+}
+
+
 //----------------------------------------------------------------//
-//  FStateTreeInstanceData
+// FStateTreeInstanceData
 //----------------------------------------------------------------//
 
 FStateTreeInstanceData::FStateTreeInstanceData()
@@ -97,6 +117,21 @@ const FStateTreeEventQueue& FStateTreeInstanceData::GetEventQueue() const
 	return GetStorage().EventQueue;
 }
 
+void FStateTreeInstanceData::AddTransitionRequest(const UObject* Owner, const FStateTreeTransitionRequest& Request)
+{
+	GetMutableStorage().AddTransitionRequest(Owner, Request);
+}
+
+TConstArrayView<FStateTreeTransitionRequest> FStateTreeInstanceData::GetTransitionRequests() const
+{
+	return GetStorage().GetTransitionRequests();
+}
+
+void FStateTreeInstanceData::ResetTransitionRequests()
+{
+	GetMutableStorage().ResetTransitionRequests();
+}
+
 int32 FStateTreeInstanceData::GetEstimatedMemoryUsage() const
 {
 	const FStateTreeInstanceStorage& Storage = GetStorage();
diff --git a/Engine/Plugins/Runtime/StateTree/Source/StateTreeModule/Private/StateTreeTypes.cpp b/Engine/Plugins/Runtime/StateTree/Source/StateTreeModule/Private/StateTreeTypes.cpp
--- a/Engine/Plugins/Runtime/StateTree/Source/StateTreeModule/Private/StateTreeTypes.cpp
+++ b/Engine/Plugins/Runtime/StateTree/Source/StateTreeModule/Private/StateTreeTypes.cpp
@@ -1,6 +1,7 @@
 // Copyright Epic Games, Inc. All Rights Reserved.
 
 #include "StateTreeTypes.h"
+#include "StateTree.h"
 
 #include UE_INLINE_GENERATED_CPP_BY_NAME(StateTreeTypes)
 
@@ -9,6 +10,7 @@ DEFINE_LOG_CATEGORY(LogStateTree);
 const FStateTreeStateHandle FStateTreeStateHandle::Invalid = FStateTreeStateHandle();
 const FStateTreeStateHandle FStateTreeStateHandle::Succeeded = FStateTreeStateHandle(FStateTreeStateHandle::SucceededIndex);
 const FStateTreeStateHandle FStateTreeStateHandle::Failed = FStateTreeStateHandle(FStateTreeStateHandle::FailedIndex);
+const FStateTreeStateHandle FStateTreeStateHandle::Root = FStateTreeStateHandle(0);
 
 const FStateTreeIndex16 FStateTreeIndex16::Invalid = FStateTreeIndex16();
 const FStateTreeIndex8 FStateTreeIndex8::Invalid = FStateTreeIndex8();
@@ -61,3 +63,25 @@ bool FStateTreeIndex8::SerializeFromMismatchedTag(const FPropertyTag& Tag, FStru
 	
 	return false;
 }
+
+//////////////////////////////////////////////////////////////////////////
+// FStateTreeStateLink
+
+bool FStateTreeStateLink::Serialize(FStructuredArchive::FSlot Slot)
+{
+	Slot.GetUnderlyingArchive().UsingCustomVersion(FStateTreeCustomVersion::GUID);
+	return false; // Let the default serializer handle serializing.
+}
+
+void FStateTreeStateLink::PostSerialize(const FArchive& Ar)
+{
+#if WITH_EDITORONLY_DATA
+	const int32 CurrentVersion = Ar.CustomVer(FStateTreeCustomVersion::GUID);
+	if (CurrentVersion < FStateTreeCustomVersion::AddedExternalTransitions)
+	{
+PRAGMA_DISABLE_DEPRECATION_WARNINGS
+		LinkType = Type_DEPRECATED;
+PRAGMA_ENABLE_DEPRECATION_WARNINGS
+	}
+#endif // WITH_EDITORONLY_DAT
+}
diff --git a/Engine/Plugins/Runtime/StateTree/Source/StateTreeModule/Public/Blueprint/StateTreeNodeBlueprintBase.h b/Engine/Plugins/Runtime/StateTree/Source/StateTreeModule/Public/Blueprint/StateTreeNodeBlueprintBase.h
--- a/Engine/Plugins/Runtime/StateTree/Source/StateTreeModule/Public/Blueprint/StateTreeNodeBlueprintBase.h
+++ b/Engine/Plugins/Runtime/StateTree/Source/StateTreeModule/Public/Blueprint/StateTreeNodeBlueprintBase.h
@@ -35,23 +35,36 @@ class STATETREEMODULE_API UStateTreeNodeBlueprintBase : public UObject
 	GENERATED_BODY()
 
 public:
-	/** Sends event to the StateTree */
+	/** Sends event to the StateTree. */
 	UFUNCTION(BlueprintCallable, Category = "StateTree", meta = (HideSelfPin = "true", DisplayName = "StateTree Send Event"))
 	void SendEvent(const FStateTreeEvent& Event);
-	
+
+	/** Request state transition. */
+	UFUNCTION(BlueprintCallable, Category = "StateTree", meta = (HideSelfPin = "true", DisplayName = "StateTree Request Transition"))
+	void RequestTransition(const FStateTreeStateLink& TargetState, const EStateTreeTransitionPriority Priority = EStateTreeTransitionPriority::Normal);
+
 protected:
 	virtual UWorld* GetWorld() const override;
 	AActor* GetOwnerActor(const FStateTreeExecutionContext& Context) const;
 
 	/** These methods are const as they set mutable variables and need to be called from a const method. */
-	void SetCachedEventQueueFromContext(const FStateTreeExecutionContext& Context) const;
-	void ClearCachedEventQueue() const;
+	void SetCachedInstanceDataFromContext(const FStateTreeExecutionContext& Context) const;
+	void ClearCachedInstanceData() const;
+	
+	UE_DEPRECATED(5.2, "Use SetCachedInstanceDataFromContext() instead.")
+	void SetCachedEventQueueFromContext(const FStateTreeExecutionContext& Context) const { SetCachedInstanceDataFromContext(Context); }
+	UE_DEPRECATED(5.2, "Use ClearCachedInstanceData() instead.")
+	void ClearCachedEventQueue() const { ClearCachedInstanceData(); }
 	
 private:
+
+	/** Cached State where the node is processed on. */
+	mutable FStateTreeStateHandle CachedState;
+	
 	/** Cached instance data while the node is active. */
-	mutable FStateTreeEventQueue* CachedEventQueue = nullptr;
+	mutable FStateTreeInstanceStorage* InstanceStorage = nullptr;
 
 	/** Cached owner while the node is active. */
 	UPROPERTY()
-	mutable TObjectPtr<UObject> CachedOwner = nullptr; 
+	mutable TObjectPtr<UObject> CachedOwner = nullptr;
 };
diff --git a/Engine/Plugins/Runtime/StateTree/Source/StateTreeModule/Public/StateTree.h b/Engine/Plugins/Runtime/StateTree/Source/StateTreeModule/Public/StateTree.h
--- a/Engine/Plugins/Runtime/StateTree/Source/StateTreeModule/Public/StateTree.h
+++ b/Engine/Plugins/Runtime/StateTree/Source/StateTreeModule/Public/StateTree.h
@@ -34,6 +34,8 @@ struct STATETREEMODULE_API FStateTreeCustomVersion
 		AddedFoo,
 		// Changed transition delay
 		TransitionDelay,
+		// Added external transitions
+		AddedExternalTransitions,
 
 		// -----<new versions can be added above this line>-------------------------------------------------
 		VersionPlusOne,
@@ -214,6 +216,10 @@ private:
 	UPROPERTY()
 	uint16 GlobalTasksNum = 0;
 
+	/** True if any global task is a transition task. */
+	UPROPERTY()
+	bool bHasGlobalTransitionTasks = false;
+	
 	// Data created during linking.
 	
 	/** List of external data required by the state tree, created during linking. */
diff --git a/Engine/Plugins/Runtime/StateTree/Source/StateTreeModule/Public/StateTreeExecutionContext.h b/Engine/Plugins/Runtime/StateTree/Source/StateTreeModule/Public/StateTreeExecutionContext.h
--- a/Engine/Plugins/Runtime/StateTree/Source/StateTreeModule/Public/StateTreeExecutionContext.h
+++ b/Engine/Plugins/Runtime/StateTree/Source/StateTreeModule/Public/StateTreeExecutionContext.h
@@ -11,6 +11,7 @@ struct FStateTreeEvaluatorBase;
 struct FStateTreeTaskBase;
 struct FStateTreeConditionBase;
 struct FStateTreeEvent;
+struct FStateTreeTransitionRequest;
 
 /**
  * StateTree Execution Context is a helper that is used to update and access StateTree instance data.
@@ -137,6 +138,22 @@ public:
 	/** @return events to process this tick. */
 	TConstArrayView<FStateTreeEvent> GetEventsToProcess() const { return EventsToProcess; }
 
+	/** @return true if there is a pending event with specified tag. */
+	bool HasEventToProcess(const FGameplayTag Tag) const
+	{
+		if (EventsToProcess.IsEmpty())
+		{
+			return false;
+		}
+		
+		return EventsToProcess.ContainsByPredicate([Tag](const FStateTreeEvent& Event)
+		{
+			return Event.Tag.MatchesTag(Tag);
+		});
+	}
+
+	/** @return the currently processed state if applicable. */
+	FStateTreeStateHandle GetCurrentlyProcessedState() const { return CurrentlyProcessedState; }
 	
 	/** @return Pointer to a State or null if state not found */ 
 	const FCompactStateTreeState* GetStateFromHandle(const FStateTreeStateHandle StateHandle) const
@@ -235,7 +252,16 @@ public:
 		static_assert(TIsDerivedFrom<T, FStateTreeNodeBase>::IsDerived, "Expecting Node to derive from FStateTreeNodeBase.");
 		return TStateTreeInstanceDataStructRef<typename T::FInstanceDataType>(InstanceData, DataViews[Node.DataViewIndex.Get()].template GetMutable<typename T::FInstanceDataType>());
 	}
-	
+
+	/**
+	 * Requests transition to a state.
+	 * If called during during transition processing (e.g. from FStateTreeTaskBase::TriggerTransitions()) the transition
+	 * is attempted to be activate immediately (it can fail e.g. because of preconditions on a target state).
+	 * If called outside the transition handling, the request is buffered and handled at the beginning of next transition processing.
+	 * @param Request The state to transition to.
+	 */
+	void RequestTransition(const FStateTreeTransitionRequest& Request);
+
 protected:
 
 	/** @return Prefix that will be used by STATETREE_LOG and STATETREE_CLOG, empty by default. */
@@ -295,14 +321,12 @@ protected:
 	 * Checks all conditions at given range
 	 * @return True if all conditions pass.
 	 */
-	bool TestAllConditions(FStateTreeInstanceData& SharedInstanceData, const int32 ConditionsOffset, const int32 ConditionsNum);
+	bool TestAllConditions(const int32 ConditionsOffset, const int32 ConditionsNum);
 
 	/**
-	 * Tries to select a state pointed by the specified transition. The OutTransition is filled on successful selection.
-	 * @return true of the state selection succeeds.
+	 * Requests transition to a specified state with specified priority.
 	 */
-	bool SelectTransition(FStateTreeInstanceData& SharedInstanceData, FStateTreeExecutionState& Exec,
-			const int16 StateIndex, const FCompactStateTreeState& State, const FCompactStateTransition& Transition, FStateTreeTransitionResult& OutTransition);
+	bool RequestTransition(const FStateTreeStateHandle NextState, const EStateTreeTransitionPriority Priority);
 
 	/**
 	 * Triggers transitions based on current run status. CurrentStatus is used to select which transitions events are triggered.
@@ -311,7 +335,7 @@ protected:
 	 * the actual next state returned by the selector.
 	 * @return Transition result describing the source state, state transitioned to, and next selected state.
 	 */
-	bool TriggerTransitions(FStateTreeInstanceData& SharedInstanceData, FStateTreeTransitionResult& OutTransition);
+	bool TriggerTransitions();
 
 	/**
 	 * Traverses the ActiveStates from StartStateIndex to 0 and returns first linked state.
@@ -319,22 +343,23 @@ protected:
 	 */
 	FStateTreeStateHandle GetParentLinkedStateHandle(const FStateTreeActiveStates& ActiveStates, const int32 StartStateIndex) const;
 
+	FStateTreeStateHandle GetParentLinkedStateHandle(const FStateTreeActiveStates& ActiveStates, const FStateTreeStateHandle StartStateHandle) const;
+
 	/**
 	 * Runs state selection logic starting at the specified state, walking towards the leaf states.
 	 * If a state cannot be selected, false is returned. 
 	 * If NextState is a selector state, SelectStateInternal is called recursively (depth-first) to all child states (where NextState will be one of child states).
 	 * If NextState is a leaf state, the active states leading from root to the leaf are returned.
-	 * @param SharedInstanceData Reference to the instance data
 	 * @param NextState The state which we try to select next.
 	 * @param OutNewActiveStates Active states that got selected.
 	 * @return True if succeeded to select new active states.
 	 */
-	bool SelectState(FStateTreeInstanceData& SharedInstanceData, const FStateTreeStateHandle NextState, FStateTreeActiveStates& OutNewActiveStates);
+	bool SelectState(const FStateTreeStateHandle NextState, FStateTreeActiveStates& OutNewActiveStates);
 
 	/**
 	 * Used internally to do the recursive part of the SelectState().
 	 */
-	bool SelectStateInternal(FStateTreeInstanceData& SharedInstanceData, const FStateTreeStateHandle NextState, FStateTreeActiveStates& OutNewActiveStates);
+	bool SelectStateInternal(const FStateTreeStateHandle NextState, FStateTreeActiveStates& OutNewActiveStates);
 
 	/** @return StateTree execution state from the instance storage. */
 	FStateTreeExecutionState& GetExecState()
@@ -397,4 +422,50 @@ protected:
 
 	/** Events to process in current tick. */
 	TArray<FStateTreeEvent, TConcurrentLinearArrayAllocator<FDefaultBlockAllocationTag>> EventsToProcess;
+
+	/** Shared instance data for the duration of the context. */
+	TSharedPtr<FStateTreeInstanceData> SharedInstanceData;
+
+	/** Next transition, used by RequestTransition(). */
+	FStateTreeTransitionResult NextTransition;
+
+	/** Current state we're processing, or invalid if not applicable. */
+	FStateTreeStateHandle CurrentlyProcessedState;
+
+	/** True if transitions are allowed to be requested directly instead of buffering. */
+	bool bAllowDirectTransitions = false;
+
+	/** Helper struct to track when it is allowed to request transitions. */
+	struct FAllowDirectTransitionsScope
+	{
+		FAllowDirectTransitionsScope(FStateTreeExecutionContext& InContext)
+			: Context(InContext)
+		{
+			Context.bAllowDirectTransitions = true;
+		}
+
+		~FAllowDirectTransitionsScope()
+		{
+			Context.bAllowDirectTransitions = false;
+		}
+		
+		FStateTreeExecutionContext& Context;
+	};
+	
+	/** Helper struct to track currently processed state. */
+	struct FCurrentlyProcessedStateScope
+	{
+		FCurrentlyProcessedStateScope(FStateTreeExecutionContext& InContext, const FStateTreeStateHandle State)
+			: Context(InContext)
+		{
+			Context.CurrentlyProcessedState = State;
+		}
+
+		~FCurrentlyProcessedStateScope()
+		{
+			Context.CurrentlyProcessedState = FStateTreeStateHandle::Invalid;
+		}
+		
+		FStateTreeExecutionContext& Context;
+	};
 };
diff --git a/Engine/Plugins/Runtime/StateTree/Source/StateTreeModule/Public/StateTreeInstanceData.h b/Engine/Plugins/Runtime/StateTree/Source/StateTreeModule/Public/StateTreeInstanceData.h
--- a/Engine/Plugins/Runtime/StateTree/Source/StateTreeModule/Public/StateTreeInstanceData.h
+++ b/Engine/Plugins/Runtime/StateTree/Source/StateTreeModule/Public/StateTreeInstanceData.h
@@ -8,6 +8,7 @@
 #include "StateTreeEvents.h"
 #include "StateTreeInstanceData.generated.h"
 
+struct FStateTreeTransitionRequest;
 struct FStateTreeExecutionState;
 
 /**
@@ -56,6 +57,25 @@ struct STATETREEMODULE_API FStateTreeInstanceStorage
 
 	/** @return reference to the event queue. */
 	const FStateTreeEventQueue& GetEventQueue() const { return EventQueue; }
+
+	/** 
+	 * Buffers a transition request to be sent to the State Tree.
+	 * @param Owner Optional pointer to an owner UObject that is used for logging errors.
+	 * @param Request transition to request.
+	*/
+	void AddTransitionRequest(const UObject* Owner, const FStateTreeTransitionRequest& Request);
+
+	/** @return currently pending transition requests. */
+	TConstArrayView<FStateTreeTransitionRequest> GetTransitionRequests() const
+	{
+		return TransitionRequests;		
+	}
+
+	/** Reset all pending transition requests. */
+	void ResetTransitionRequests()
+	{
+		TransitionRequests.Reset();
+	}
 	
 protected:
 	/** Struct instances */
@@ -70,6 +90,10 @@ protected:
 	UPROPERTY()
 	FStateTreeEventQueue EventQueue;
 
+	/** Requested transitions */
+	UPROPERTY()
+	TArray<FStateTreeTransitionRequest> TransitionRequests;
+
 	friend struct FStateTreeInstanceData;
 };
 
@@ -141,6 +165,19 @@ struct STATETREEMODULE_API FStateTreeInstanceData
 	FStateTreeEventQueue& GetMutableEventQueue();
 	const FStateTreeEventQueue& GetEventQueue() const;
 
+	/** 
+	 * Buffers a transition request to be sent to the State Tree.
+	 * @param Owner Optional pointer to an owner UObject that is used for logging errors.
+	 * @param Request transition to request.
+	*/
+	void AddTransitionRequest(const UObject* Owner, const FStateTreeTransitionRequest& Request);
+
+	/** @return currently pending transition requests. */
+	TConstArrayView<FStateTreeTransitionRequest> GetTransitionRequests() const;
+
+	/** Reset all pending transition requests. */
+	void ResetTransitionRequests();
+
 	FStateTreeInstanceStorage& GetMutableStorage();
 	const FStateTreeInstanceStorage& GetStorage() const;
 
diff --git a/Engine/Plugins/Runtime/StateTree/Source/StateTreeModule/Public/StateTreeTaskBase.h b/Engine/Plugins/Runtime/StateTree/Source/StateTreeModule/Public/StateTreeTaskBase.h
--- a/Engine/Plugins/Runtime/StateTree/Source/StateTreeModule/Public/StateTreeTaskBase.h
+++ b/Engine/Plugins/Runtime/StateTree/Source/StateTreeModule/Public/StateTreeTaskBase.h
@@ -23,6 +23,7 @@ struct STATETREEMODULE_API FStateTreeTaskBase : public FStateTreeNodeBase
 		, bShouldCallTickOnlyOnEvents(false)
 		, bShouldCopyBoundPropertiesOnTick(true)
 		, bShouldCopyBoundPropertiesOnExitState(true)
+		, bShouldAffectTransitions(false)
 	{
 	}
 	
@@ -61,12 +62,21 @@ struct STATETREEMODULE_API FStateTreeTaskBase : public FStateTreeNodeBase
 
 	/**
 	 * Called during state tree tick when the task is on active state.
+	 * Note: The method is called only if bShouldCallTick or bShouldCallTickOnlyOnEvents is set.
 	 * @param Context Reference to current execution context.
 	 * @param DeltaTime Time since last StateTree tick.
 	 * @return Running status of the state: Running if still in progress, Succeeded if execution is done and succeeded, Failed if execution is done and failed.
 	 */
 	virtual EStateTreeRunStatus Tick(FStateTreeExecutionContext& Context, const float DeltaTime) const { return EStateTreeRunStatus::Running; };
 
+	/**
+	 * Called when state tree triggers transitions. This method is called during transition handling, before state's tick and event transitions are handled.
+	 * Note: the method is called only if bShouldAffectTransitions is set.
+	 * @param Context Reference to current execution context.
+	 */
+	virtual void TriggerTransitions(FStateTreeExecutionContext& Context) const {};
+
+
 #if WITH_GAMEPLAY_DEBUGGER
 	virtual void AppendDebugInfoString(FString& DebugString, const FStateTreeExecutionContext& Context) const;
 #endif
@@ -87,6 +97,9 @@ struct STATETREEMODULE_API FStateTreeTaskBase : public FStateTreeNodeBase
 	uint8 bShouldCopyBoundPropertiesOnTick : 1;
 	/** If set to true, copy the values of bound properties before calling ExitState(). Default true. */
 	uint8 bShouldCopyBoundPropertiesOnExitState : 1;
+
+	/** If set to true, TriggerTransitions() is called during transition handling. Default false. */
+	uint8 bShouldAffectTransitions : 1;
 };
 
 /**
diff --git a/Engine/Plugins/Runtime/StateTree/Source/StateTreeModule/Public/StateTreeTypes.h b/Engine/Plugins/Runtime/StateTree/Source/StateTreeModule/Public/StateTreeTypes.h
--- a/Engine/Plugins/Runtime/StateTree/Source/StateTreeModule/Public/StateTreeTypes.h
+++ b/Engine/Plugins/Runtime/StateTree/Source/StateTreeModule/Public/StateTreeTypes.h
@@ -5,6 +5,7 @@
 #include "StructView.h"
 #include "PropertyBag.h"
 #include "GameplayTagContainer.h"
+#include "Misc/Guid.h"
 #include "StateTreeTypes.generated.h"
 
 STATETREEMODULE_API DECLARE_LOG_CATEGORY_EXTERN(LogStateTree, Warning, All);
@@ -13,6 +14,14 @@ STATETREEMODULE_API DECLARE_LOG_CATEGORY_EXTERN(LogStateTree, Warning, All);
 #define WITH_STATETREE_DEBUG (!(UE_BUILD_SHIPPING || UE_BUILD_SHIPPING_WITH_EDITOR || UE_BUILD_TEST) && 1)
 #endif // WITH_STATETREE_DEBUG
 
+namespace UE::StateTree
+{
+	inline constexpr int32 MaxConditionIndent = 4;
+
+	inline const FName SchemaTag(TEXT("Schema"));
+}; // UE::StateTree
+
+
 /** Status describing current run status of a State Tree. */
 UENUM(BlueprintType)
 enum class EStateTreeRunStatus : uint8
@@ -48,20 +57,22 @@ enum class EStateTreeStateChangeType : uint8
 UENUM()
 enum class EStateTreeTransitionType : uint8
 {
-	/** Stop StateTree and mark execution succeeded. */
+	/** No transition will take place. */
+	None,
+
+	/** Stop State Tree or sub-tree and mark execution succeeded. */
 	Succeeded,
 	
-	/** Stop StateTree and mark execution failed. */
+	/** Stop State Tree or sub-tree and mark execution failed. */
 	Failed,
 	
-	/** Transition to specified state. */
+	/** Transition to the specified state. */
 	GotoState,
 	
-	/** No transition. */
-	NotSet,
-	
-	/** Goto next sibling state. */
+	/** Transition to the next sibling state. */
 	NextState,
+
+	NotSet UE_DEPRECATED(5.0, "Use None instead."),
 };
 
 /** Operand between conditions */
@@ -94,14 +105,6 @@ enum class EStateTreeStateType : uint8
 	Subtree,
 };
 
-namespace UE::StateTree
-{
-	inline constexpr int32 MaxConditionIndent = 4;
-
-	inline const FName SchemaTag(TEXT("Schema"));
-}; // UE::StateTree
-
-
 
 /** Transitions trigger. */
 UENUM()
@@ -128,6 +131,33 @@ enum class EStateTreeTransitionTrigger : uint8
 };
 ENUM_CLASS_FLAGS(EStateTreeTransitionTrigger)
 
+
+/** Transition priority. When multiple transitions trigger at the same time, the first transition of highest priority is selected. */
+UENUM(BlueprintType)
+enum class EStateTreeTransitionPriority : uint8
+{
+	None UMETA(Hidden),
+	
+	/** Normal priority. */
+	Normal,
+	
+	/** Medium priority. */
+	Medium,
+	
+	/** High priority. */
+	High,
+	
+	/** Critical priority. */
+	Critical,
+};
+
+inline bool operator<(const EStateTreeTransitionPriority Lhs, const EStateTreeTransitionPriority Rhs) { return static_cast<uint8>(Lhs) < static_cast<uint8>(Rhs); }
+inline bool operator>(const EStateTreeTransitionPriority Lhs, const EStateTreeTransitionPriority Rhs) { return static_cast<uint8>(Lhs) > static_cast<uint8>(Rhs); }
+inline bool operator<=(const EStateTreeTransitionPriority Lhs, const EStateTreeTransitionPriority Rhs) { return static_cast<uint8>(Lhs) <= static_cast<uint8>(Rhs); }
+inline bool operator>=(const EStateTreeTransitionPriority Lhs, const EStateTreeTransitionPriority Rhs) { return static_cast<uint8>(Lhs) >= static_cast<uint8>(Rhs); }
+inline bool operator==(const EStateTreeTransitionPriority Lhs, const EStateTreeTransitionPriority Rhs) { return static_cast<uint8>(Lhs) == static_cast<uint8>(Rhs); }
+inline bool operator!=(const EStateTreeTransitionPriority Lhs, const EStateTreeTransitionPriority Rhs) { return static_cast<uint8>(Lhs) != static_cast<uint8>(Rhs); }
+
 /** Handle to a StateTree state */
 USTRUCT(BlueprintType)
 struct STATETREEMODULE_API FStateTreeStateHandle
@@ -141,11 +171,25 @@ struct STATETREEMODULE_API FStateTreeStateHandle
 	static const FStateTreeStateHandle Invalid;
 	static const FStateTreeStateHandle Succeeded;
 	static const FStateTreeStateHandle Failed;
+	static const FStateTreeStateHandle Root;
 
 	FStateTreeStateHandle() = default;
 	explicit FStateTreeStateHandle(uint16 InIndex) : Index(InIndex) {}
 
 	bool IsValid() const { return Index != InvalidIndex; }
+	bool IsCompletionState() const { return Index == SucceededIndex || Index == FailedIndex; }
+	EStateTreeRunStatus ToCompletionStatus() const
+	{
+		if (Index == SucceededIndex)
+		{
+			return EStateTreeRunStatus::Succeeded;
+		}
+		else if (Index == FailedIndex)
+		{
+			return EStateTreeRunStatus::Failed;
+		}
+		return EStateTreeRunStatus::Unset;
+	}
 
 	bool operator==(const FStateTreeStateHandle& RHS) const { return Index == RHS.Index; }
 	bool operator!=(const FStateTreeStateHandle& RHS) const { return Index != RHS.Index; }
@@ -268,6 +312,32 @@ struct TStructOpsTypeTraits<FStateTreeIndex8> : public TStructOpsTypeTraitsBase2
 };
 
 
+/** Transition request */
+USTRUCT(BlueprintType)
+struct STATETREEMODULE_API FStateTreeTransitionRequest
+{
+	GENERATED_BODY()
+
+	FStateTreeTransitionRequest() = default;
+	FStateTreeTransitionRequest(const FStateTreeStateHandle InTargetState, const EStateTreeTransitionPriority InPriority = EStateTreeTransitionPriority::Normal)
+		: TargetState(InTargetState)
+		, Priority(InPriority)
+	{
+	}
+
+	/** Source state of the transition. Filled in by the StateTree execution context. */
+	UPROPERTY(EditDefaultsOnly, Category = "Default")
+	FStateTreeStateHandle SourceState;
+
+	/** Target state of the transition. */
+	UPROPERTY(EditDefaultsOnly, Category = "Default")
+	FStateTreeStateHandle TargetState;
+	
+	/** Priority of the transition. */
+	UPROPERTY(EditDefaultsOnly, Category = "Default")
+	EStateTreeTransitionPriority Priority = EStateTreeTransitionPriority::Normal;
+};
+
 /**
  * Describes an array of active states in a State Tree.
  */
@@ -496,29 +566,48 @@ struct STATETREEMODULE_API FStateTreeTransitionResult
 
 	FStateTreeTransitionResult() = default;
 
+	void Reset()
+	{
+		CurrentActiveStates.Reset();
+		CurrentRunStatus = EStateTreeRunStatus::Unset;
+		TargetState = FStateTreeStateHandle::Invalid;
+		NextActiveStates.Reset();
+		CurrentState = FStateTreeStateHandle::Invalid;
+		ChangeType = EStateTreeStateChangeType::Changed; 
+		Priority = EStateTreeTransitionPriority::None; 
+	}
+	
 	/** Current active states, where the transition started. */
-	UPROPERTY(EditDefaultsOnly, Category = Default, BlueprintReadOnly)
+	UPROPERTY(EditDefaultsOnly, Category = "Default", BlueprintReadOnly)
 	FStateTreeActiveStates CurrentActiveStates;
 
 	/** Current Run status. */
-	UPROPERTY(EditDefaultsOnly, Category = Default, BlueprintReadOnly)
+	UPROPERTY(EditDefaultsOnly, Category = "Default", BlueprintReadOnly)
 	EStateTreeRunStatus CurrentRunStatus = EStateTreeRunStatus::Unset;
 
+	/** Transition source state */
+	UPROPERTY(EditDefaultsOnly, Category = "Default", BlueprintReadOnly)
+	FStateTreeStateHandle SourceState = FStateTreeStateHandle::Invalid;
+
 	/** Transition target state */
-	UPROPERTY(EditDefaultsOnly, Category = Default, BlueprintReadOnly)
+	UPROPERTY(EditDefaultsOnly, Category = "Default", BlueprintReadOnly)
 	FStateTreeStateHandle TargetState = FStateTreeStateHandle::Invalid;
 
 	/** States selected as result of the transition. */
-	UPROPERTY(EditDefaultsOnly, Category = Default, BlueprintReadOnly)
+	UPROPERTY(EditDefaultsOnly, Category = "Default", BlueprintReadOnly)
 	FStateTreeActiveStates NextActiveStates;
 
-	/** The current state being executed. On enter/exit callbacks this is the state of the task or evaluator. */
-	UPROPERTY(EditDefaultsOnly, Category = Default, BlueprintReadOnly)
+	/** The current state being executed. On enter/exit callbacks this is the state of the task. */
+	UPROPERTY(EditDefaultsOnly, Category = "Default", BlueprintReadOnly)
 	FStateTreeStateHandle CurrentState = FStateTreeStateHandle::Invalid;
 
 	/** If the change type is Sustained, then the CurrentState was reselected, or if Changed then the state was just activated. */
-	UPROPERTY(EditDefaultsOnly, Category = Default, BlueprintReadOnly)
+	UPROPERTY(EditDefaultsOnly, Category = "Default", BlueprintReadOnly)
 	EStateTreeStateChangeType ChangeType = EStateTreeStateChangeType::Changed; 
+
+	/** Priority of the transition that caused the state change. */
+	UPROPERTY(EditDefaultsOnly, Category = "Default", BlueprintReadOnly)
+	EStateTreeTransitionPriority Priority = EStateTreeTransitionPriority::None; 
 };
 
 
@@ -552,14 +641,14 @@ struct STATETREEMODULE_API FCompactStateTransition
 	UPROPERTY()
 	FStateTreeRandomTimeDuration Delay;
 	
-	/** Type of the transition. */
-	UPROPERTY()
-	EStateTreeTransitionType Type = EStateTreeTransitionType::NotSet;
- 
 	/* Type of the transition trigger. */
 	UPROPERTY()
 	EStateTreeTransitionTrigger Trigger = EStateTreeTransitionTrigger::None;
 
+	/* Priority of the transition. */
+	UPROPERTY()
+	EStateTreeTransitionPriority Priority = EStateTreeTransitionPriority::Normal;
+
 	/** Number of conditions to test. */
 	UPROPERTY()
 	uint8 ConditionsNum = 0;
@@ -573,45 +662,76 @@ struct STATETREEMODULE_API FCompactStateTreeState
 {
 	GENERATED_BODY()
 
+	FCompactStateTreeState()
+		: bHasTransitionTasks(false)
+	{
+	}
+	
 	/** @return Index to the next sibling state. */
 	uint16 GetNextSibling() const { return ChildrenEnd; }
 
 	/** @return True if the state has any child states */
 	bool HasChildren() const { return ChildrenEnd > ChildrenBegin; }
 
+	/** Name of the State */
+	UPROPERTY()
+	FName Name;
+
+	/** Linked state handle if the state type is linked state. */
 	UPROPERTY()
-	FName Name;											// Name of the State
+	FStateTreeStateHandle LinkedState = FStateTreeStateHandle::Invalid; 
 
+	/** Parent state handle, invalid if root state. */
 	UPROPERTY()
-	FStateTreeStateHandle LinkedState = FStateTreeStateHandle::Invalid;	// Linked state 
+	FStateTreeStateHandle Parent = FStateTreeStateHandle::Invalid;
 
+	/** Index to first child state */
 	UPROPERTY()
-	FStateTreeStateHandle Parent = FStateTreeStateHandle::Invalid;		// Parent state
+	uint16 ChildrenBegin = 0;
+
+	/** Index one past the last child state. */
 	UPROPERTY()
-	uint16 ChildrenBegin = 0;							// Index to first child state
+	uint16 ChildrenEnd = 0;
+
+	/** Index to first state enter condition */
 	UPROPERTY()
-	uint16 ChildrenEnd = 0;								// Index one past the last child state
+	uint16 EnterConditionsBegin = 0;
 
+	/** Index to first transition */
 	UPROPERTY()
-	uint16 EnterConditionsBegin = 0;					// Index to first state enter condition
+	uint16 TransitionsBegin = 0;
+
+	/** Index to first task */
 	UPROPERTY()
-	uint16 TransitionsBegin = 0;						// Index to first transition
+	uint16 TasksBegin = 0;
+
+	/** Index to state instance data. */
 	UPROPERTY()
-	uint16 TasksBegin = 0;								// Index to first task
+	FStateTreeIndex16 ParameterInstanceIndex = FStateTreeIndex16::Invalid;
 
+	/** Data view index of the input parameters. */
 	UPROPERTY()
-	FStateTreeIndex16 ParameterInstanceIndex = FStateTreeIndex16::Invalid;	// Index to state instance data
+	FStateTreeIndex16 ParameterDataViewIndex = FStateTreeIndex16::Invalid;
+
+	/** Number of enter conditions */
 	UPROPERTY()
-	FStateTreeIndex16 ParameterDataViewIndex = FStateTreeIndex16::Invalid;	// Data view index of the input parameters
+	uint8 EnterConditionsNum = 0;
 
+	/** Number of transitions */
 	UPROPERTY()
-	uint8 EnterConditionsNum = 0;						// Number of enter conditions
+	uint8 TransitionsNum = 0;
+
+	/** Number of tasks */
 	UPROPERTY()
-	uint8 TransitionsNum = 0;							// Number of transitions
+	uint8 TasksNum = 0;
+
+	/** Type of the sate */
 	UPROPERTY()
-	uint8 TasksNum = 0;									// Number of tasks
+	EStateTreeStateType Type = EStateTreeStateType::State;
+
+	/** True if the state contains tasks that should be called during transition handling. */
 	UPROPERTY()
-	EStateTreeStateType Type = EStateTreeStateType::State; 
+	uint8 bHasTransitionTasks : 1;
 };
 
 USTRUCT()
@@ -1004,4 +1124,69 @@ struct STATETREEMODULE_API FStateTreeExecutionState
 
 	/** Running time of the delayed transition */
 	TArray<FStateTreeTransitionDelayedState> DelayedTransitions;
-};
\ No newline at end of file
+};
+
+/**
+ * Link to another state in StateTree
+ */
+USTRUCT(BlueprintType)
+struct STATETREEMODULE_API FStateTreeStateLink
+{
+	GENERATED_BODY()
+
+	FStateTreeStateLink() = default;
+
+#if WITH_EDITORONLY_DATA
+	FStateTreeStateLink(const EStateTreeTransitionType InType) : LinkType(InType) {}
+
+	UE_DEPRECATED(5.2, "Use UStateTreeState::GetLinkToState() instead.")
+	void Set(const EStateTreeTransitionType InType, const class UStateTreeState* InState = nullptr) {}
+	UE_DEPRECATED(5.2, "Use UStateTreeState::GetLinkToState() instead.")
+	void Set(const class UStateTreeState* InState) {}
+#endif // WITH_EDITORONLY_DATA
+	
+PRAGMA_DISABLE_DEPRECATION_WARNINGS
+	FStateTreeStateLink(const FStateTreeStateLink& Other) = default;
+	FStateTreeStateLink(FStateTreeStateLink&& Other) = default;
+	FStateTreeStateLink& operator=(FStateTreeStateLink const& Other) = default;
+	FStateTreeStateLink& operator=(FStateTreeStateLink&& Other) = default;
+PRAGMA_ENABLE_DEPRECATION_WARNINGS
+	
+	bool Serialize(FStructuredArchive::FSlot Slot);
+	void PostSerialize(const FArchive& Ar);
+
+#if WITH_EDITORONLY_DATA
+	UE_DEPRECATED(5.2, "Will be removed in later releases.")
+	bool IsValid() const { return ID.IsValid(); }
+	
+	/** Name of the state at the time of linking, used for error reporting. */
+	UPROPERTY(EditDefaultsOnly, Category = "Link")
+	FName Name;
+
+	/** ID of the state linked to. */
+	UPROPERTY(EditDefaultsOnly, Category = "Link")
+	FGuid ID;
+
+	/** Type of the transition, used at edit time to describe e.g. next state (which is calculated at compile time). */
+	UPROPERTY(EditDefaultsOnly, Category = "Link")
+	EStateTreeTransitionType LinkType = EStateTreeTransitionType::None;
+
+	UE_DEPRECATED(5.2, "Use LinkType instead.")
+	UPROPERTY()
+	EStateTreeTransitionType Type_DEPRECATED = EStateTreeTransitionType::GotoState;
+#endif // WITH_EDITORONLY_DATA
+
+	/** Handle of the linked state. */
+	UPROPERTY()
+	FStateTreeStateHandle StateHandle; 
+};
+
+template<>
+struct TStructOpsTypeTraits<FStateTreeStateLink> : public TStructOpsTypeTraitsBase2<FStateTreeStateLink>
+{
+	enum
+	{
+		WithStructuredSerializer = true,
+		WithPostSerialize = true,
+	};
+};
diff --git a/Engine/Plugins/Runtime/StateTree/Source/StateTreeTestSuite/Private/StateTreeTest.cpp b/Engine/Plugins/Runtime/StateTree/Source/StateTreeTestSuite/Private/StateTreeTest.cpp
--- a/Engine/Plugins/Runtime/StateTree/Source/StateTreeTestSuite/Private/StateTreeTest.cpp
+++ b/Engine/Plugins/Runtime/StateTree/Source/StateTreeTestSuite/Private/StateTreeTest.cpp
@@ -341,7 +341,7 @@ struct FStateTreeTest_SubTree : FAITestBase
 		UStateTreeState& State3A = State3.AddChildState(FName(TEXT("State3A")));
 		UStateTreeState& State3B = State3.AddChildState(FName(TEXT("State3B")));
 
-		State1.LinkedSubtree.Set(&State3);
+		State1.LinkedSubtree = State3.GetLinkToState();
 
 		State1.AddTransition(EStateTreeTransitionTrigger::OnStateCompleted, EStateTreeTransitionType::GotoState, &State2);
 
@@ -575,7 +575,7 @@ struct FStateTreeTest_TransitionPriorityEnterState : FAITestBase
 		UStateTree& StateTree = UE::StateTree::Tests::NewStateTree(&GetWorld());
 		UStateTreeEditorData& EditorData = *Cast<UStateTreeEditorData>(StateTree.EditorData);
 		
-		UStateTreeState& Root = EditorData.AddSubTree(FName(TEXT("Root")));
+		UStateTreeState& Root =	EditorData.AddSubTree(FName(TEXT("Root")));
 		UStateTreeState& State0 = Root.AddChildState(FName(TEXT("State0")));
 		UStateTreeState& State1 = Root.AddChildState(FName(TEXT("State1")));
 		UStateTreeState& State1A = State1.AddChildState(FName(TEXT("State1A")));
