#StateTree Downgrade warning to display for now

diff --git a/Engine/Plugins/Runtime/StateTree/Source/StateTreeModule/Private/StateTreeInstanceData.cpp b/Engine/Plugins/Runtime/StateTree/Source/StateTreeModule/Private/StateTreeInstanceData.cpp
--- a/Engine/Plugins/Runtime/StateTree/Source/StateTreeModule/Private/StateTreeInstanceData.cpp
+++ b/Engine/Plugins/Runtime/StateTree/Source/StateTreeModule/Private/StateTreeInstanceData.cpp
@@ -31,7 +31,7 @@ namespace UE::StateTree
 			FMemoryReader Reader(Data);
 			NewInstance->Serialize(Reader);
 				
-			UE_LOG(LogStateTree, Warning, TEXT("FStateTreeInstanceData: Duplicating '%s' with old class '%s' as '%s', potential data loss."),
+			UE_LOG(LogStateTree, Display, TEXT("FStateTreeInstanceData: Duplicating '%s' with old class '%s' as '%s', potential data loss."),
 				*GetFullNameSafe(&Instance), *GetNameSafe(InstanceClass), *GetNameSafe(AuthoritativeClass));
 
 			return NewInstance;
