StateTree: Editor fixed
- fix bug where expanded states were reset when new state was added
- fix resetting transition when new transition condition was added

diff --git a/Engine/Plugins/Runtime/StateTree/Source/StateTreeEditorModule/Private/SStateTreeView.cpp b/Engine/Plugins/Runtime/StateTree/Source/StateTreeEditorModule/Private/SStateTreeView.cpp
--- a/Engine/Plugins/Runtime/StateTree/Source/StateTreeEditorModule/Private/SStateTreeView.cpp
+++ b/Engine/Plugins/Runtime/StateTree/Source/StateTreeEditorModule/Private/SStateTreeView.cpp
@@ -65,6 +65,7 @@ void SStateTreeView::Construct(const FArguments& InArgs, TSharedRef<FStateTreeVi
 		.TreeItemsSource(&Subtrees)
 		.ItemHeight(32)
 		.OnSelectionChanged(this, &SStateTreeView::HandleTreeSelectionChanged)
+		.OnExpansionChanged(this, &SStateTreeView::HandleTreeExpansionChanged)
 		.OnContextMenuOpening(this, &SStateTreeView::HandleContextMenuOpening)
 		.AllowOverscroll(EAllowOverscroll::No)
 		.ExternalScrollbar(VerticalScrollBar);
@@ -224,6 +225,11 @@ void SStateTreeView::HandleModelStateAdded(UStateTreeState* ParentState, UStateT
 	bItemsDirty = true;
 
 	HandleRenameState(NewState);
+
+	if (StateTreeViewModel.IsValid())
+	{
+		StateTreeViewModel->SetSelection(NewState);
+	}
 }
 
 void SStateTreeView::HandleModelStatesChanged(const TSet<UStateTreeState*>& AffectedStates, const FPropertyChangedEvent& PropertyChangedEvent)
@@ -292,6 +298,16 @@ void SStateTreeView::HandleTreeSelectionChanged(TWeakObjectPtr<UStateTreeState>
 	bUpdatingSelection = false;
 }
 
+void SStateTreeView::HandleTreeExpansionChanged(TWeakObjectPtr<UStateTreeState> InSelectedItem, bool bExpanded)
+{
+	// Not calling Modify() on the state as we don't want the expansion to dirty the asset.
+	// @todo: this is temporary fix for a bug where adding a state will reset the expansion state. 
+	if (UStateTreeState* State = InSelectedItem.Get())
+	{
+		State->bExpanded = bExpanded;
+	}
+}
+
 TSharedPtr<SWidget> SStateTreeView::HandleContextMenuOpening()
 {
 	if (!StateTreeViewModel)
diff --git a/Engine/Plugins/Runtime/StateTree/Source/StateTreeEditorModule/Private/SStateTreeView.h b/Engine/Plugins/Runtime/StateTree/Source/StateTreeEditorModule/Private/SStateTreeView.h
--- a/Engine/Plugins/Runtime/StateTree/Source/StateTreeEditorModule/Private/SStateTreeView.h
+++ b/Engine/Plugins/Runtime/StateTree/Source/StateTreeEditorModule/Private/SStateTreeView.h
@@ -68,6 +68,8 @@ private:
 	TSharedRef<ITableRow> HandleGenerateRow(TWeakObjectPtr<UStateTreeState> InState, const TSharedRef<STableViewBase>& InOwnerTableView);
 	void HandleGetChildren(TWeakObjectPtr<UStateTreeState> InParent, TArray<TWeakObjectPtr<UStateTreeState>>& OutChildren);
 	void HandleTreeSelectionChanged(TWeakObjectPtr<UStateTreeState> InSelectedItem, ESelectInfo::Type SelectionType);
+	void HandleTreeExpansionChanged(TWeakObjectPtr<UStateTreeState> InSelectedItem, bool bExpanded);
+	
 	TSharedPtr<SWidget> HandleContextMenuOpening();
 
 	// Action handlers
diff --git a/Engine/Plugins/Runtime/StateTree/Source/StateTreeEditorModule/Private/StateTreePropertyHelpers.h b/Engine/Plugins/Runtime/StateTree/Source/StateTreeEditorModule/Private/StateTreePropertyHelpers.h
--- a/Engine/Plugins/Runtime/StateTree/Source/StateTreeEditorModule/Private/StateTreePropertyHelpers.h
+++ b/Engine/Plugins/Runtime/StateTree/Source/StateTreeEditorModule/Private/StateTreePropertyHelpers.h
@@ -4,6 +4,7 @@
 
 #include "Editor.h"
 #include "PropertyHandle.h"
+#include "Misc/EnumerateRange.h"
 
 struct FGuid;
 struct FSlateBrush;
@@ -119,4 +120,127 @@ FPropertyAccess::Result SetStructValue(const TSharedPtr<IPropertyHandle>& ValueP
 
 } // UE::StateTree::PropertyHelpers
 
+/**
+ * Helper class to deal with relative property paths in PostEditChangeChainProperty().
+ */
+struct FStateTreeEditPropertyPath
+{
+private:
+	struct FStateTreeEditPropertySegment
+	{
+		FStateTreeEditPropertySegment() = default;
+		FStateTreeEditPropertySegment(const FProperty* InProperty, const FName InPropertyName, const int32 InArrayIndex = INDEX_NONE)
+			: Property(InProperty)
+			, PropertyName(InPropertyName)
+			, ArrayIndex(InArrayIndex)
+		{
+		}
+	
+		const FProperty* Property = nullptr;
+		FName PropertyName = FName();
+		int32 ArrayIndex = INDEX_NONE;
+	};
+	
+public:
+	FStateTreeEditPropertyPath() = default;
+
+	/** Makes property path relative to BaseStruct. Checks if the path is not part of the type. */
+	explicit FStateTreeEditPropertyPath(const UStruct* BaseStruct, const FString& InPath)
+	{
+		TArray<FString> PathSegments;
+		InPath.ParseIntoArray(PathSegments, TEXT("."));
+
+		const UStruct* CurrBase = BaseStruct;
+		for (const FString& Segment : PathSegments)
+		{
+			const FName PropertyName(Segment);
+			if (const FProperty* Property = CurrBase->FindPropertyByName(PropertyName))
+			{
+				Path.Emplace(Property, PropertyName);
+
+				if (const FArrayProperty* ArrayProperty = CastField<FArrayProperty>(Property))
+				{
+					Property = ArrayProperty->Inner;
+				}
+
+				if (const FStructProperty* StructProperty = CastField<FStructProperty>(Property))
+				{
+					CurrBase = StructProperty->Struct;
+				}
+				else if (const FObjectProperty* ObjectProperty = CastField<FObjectProperty>(Property))
+				{
+					CurrBase = ObjectProperty->PropertyClass;
+				}
+			}
+			else
+			{
+				checkf(TEXT("Path %s id not part of type %s."), *InPath, *GetNameSafe(BaseStruct));
+				Path.Reset();
+				break;
+			}
+		}
+	}
+
+	/** Makes property path from property change event. */
+	explicit FStateTreeEditPropertyPath(const FPropertyChangedChainEvent& PropertyChangedEvent)
+	{
+		FEditPropertyChain::TDoubleLinkedListNode* PropertyNode = PropertyChangedEvent.PropertyChain.GetActiveMemberNode();
+		while (PropertyNode != nullptr)
+		{
+			if (FProperty* Property = PropertyNode->GetValue())
+			{
+				const FName PropertyName = Property->GetFName(); 
+				const int32 ArrayIndex = PropertyChangedEvent.GetArrayIndex(PropertyName.ToString());
+				Path.Emplace(Property, PropertyName, ArrayIndex);
+			}
+			PropertyNode = PropertyNode->GetNextNode();
+		}
+	}
+
+	/** @return true if the property path contains specified path. */
+	bool ContainsPath(const FStateTreeEditPropertyPath& InPath) const
+	{
+		if (InPath.Path.Num() > Path.Num())
+    	{
+    		return false;
+    	}
+
+    	for (TConstEnumerateRef<FStateTreeEditPropertySegment> Segment : EnumerateRange(InPath.Path))
+    	{
+    		if (Segment->PropertyName != Path[Segment.GetIndex()].PropertyName)
+    		{
+    			return false;
+    		}
+    	}
+    	return true;
+	}
+
+	/** @return true if the property path is exactly the specified path. */
+	bool IsPathExact(const FStateTreeEditPropertyPath& InPath) const
+	{
+		if (InPath.Path.Num() != Path.Num())
+		{
+			return false;
+		}
+
+		for (TConstEnumerateRef<FStateTreeEditPropertySegment> Segment : EnumerateRange(InPath.Path))
+		{
+			if (Segment->PropertyName != Path[Segment.GetIndex()].PropertyName)
+			{
+				return false;
+			}
+		}
+		return true;
+	}
+
+	/** @return array index at specified property, or INDEX_NONE, if the property is not array or property not found.  */
+	int32 GetPropertyArrayIndex(const FStateTreeEditPropertyPath& InPath) const
+	{
+		return ContainsPath(InPath) ? Path[InPath.Path.Num() - 1].ArrayIndex : INDEX_NONE;
+	}
+
+private:
+	TArray<FStateTreeEditPropertySegment> Path;
+};
+
 #undef LOCTEXT_NAMESPACE
diff --git a/Engine/Plugins/Runtime/StateTree/Source/StateTreeEditorModule/Private/StateTreeState.cpp b/Engine/Plugins/Runtime/StateTree/Source/StateTreeEditorModule/Private/StateTreeState.cpp
--- a/Engine/Plugins/Runtime/StateTree/Source/StateTreeEditorModule/Private/StateTreeState.cpp
+++ b/Engine/Plugins/Runtime/StateTree/Source/StateTreeEditorModule/Private/StateTreeState.cpp
@@ -6,6 +6,7 @@
 #include "StateTreeConditionBase.h"
 #include "StateTreeTaskBase.h"
 #include "StateTreeDelegates.h"
+#include "StateTreePropertyHelpers.h"
 
 #include UE_INLINE_GENERATED_CPP_BY_NAME(StateTreeState)
 
@@ -37,29 +38,13 @@ UStateTreeState::UStateTreeState(const FObjectInitializer& ObjectInitializer)
 }
 
 #if WITH_EDITOR
+
 void UStateTreeState::PostEditChangeChainProperty(FPropertyChangedChainEvent& PropertyChangedEvent)
 {
 	Super::PostEditChangeChainProperty(PropertyChangedEvent);
 
-	FProperty* Property = PropertyChangedEvent.Property;
-	FProperty* MemberProperty = nullptr;
-	if (PropertyChangedEvent.PropertyChain.GetActiveMemberNode())
-	{
-		MemberProperty = PropertyChangedEvent.PropertyChain.GetActiveMemberNode()->GetValue();
-	}
-
-	auto TryGetConditionIndex = [&PropertyChangedEvent]() -> int32
-	{
-		check (PropertyChangedEvent.PropertyChain.GetActiveMemberNode());
-		FEditPropertyChain::TDoubleLinkedListNode* ConditionPropertyNode = PropertyChangedEvent.PropertyChain.GetActiveMemberNode()->GetNextNode();
-		FProperty* ConditionsProperty = ConditionPropertyNode ? ConditionPropertyNode->GetValue() : nullptr;
-		if (ConditionsProperty && ConditionsProperty->GetFName() == GET_MEMBER_NAME_CHECKED(FStateTreeTransition, Conditions))
-		{
-			return PropertyChangedEvent.GetArrayIndex(ConditionsProperty->GetFName().ToString());
-		}
-		return INDEX_NONE;
-	};
-
+	const FStateTreeEditPropertyPath ChangePropertyPath(PropertyChangedEvent);
+	
 	auto CopyBindings = [this](const FGuid FromStructID, const FGuid ToStructID)
 	{
 		if (UStateTreeEditorData* EditorData = GetTypedOuter<UStateTreeEditorData>())
@@ -75,179 +60,188 @@ void UStateTreeState::PostEditChangeChainProperty(FPropertyChangedChainEvent& Pr
 	};
 
 	
-	if (Property)
+	static const FStateTreeEditPropertyPath StateNamePath(UStateTreeState::StaticClass(), TEXT("Name"));
+	static const FStateTreeEditPropertyPath StateTypePath(UStateTreeState::StaticClass(), TEXT("Type"));
+	static const FStateTreeEditPropertyPath StateLinkedSubtreePath(UStateTreeState::StaticClass(), TEXT("LinkedSubtree"));
+	static const FStateTreeEditPropertyPath StateParametersPath(UStateTreeState::StaticClass(), TEXT("Parameters"));
+	static const FStateTreeEditPropertyPath StateTasksPath(UStateTreeState::StaticClass(), TEXT("Tasks"));
+	static const FStateTreeEditPropertyPath StateEnterConditionsPath(UStateTreeState::StaticClass(), TEXT("EnterConditions"));
+	static const FStateTreeEditPropertyPath StateTransitionsPath(UStateTreeState::StaticClass(), TEXT("Transitions"));
+	static const FStateTreeEditPropertyPath StateTransitionsConditionsPath(UStateTreeState::StaticClass(), TEXT("Transitions.Conditions"));
+
+
+	// Broadcast name changes so that the UI can update.
+	if (ChangePropertyPath.IsPathExact(StateNamePath))
 	{
-		if (Property->GetOwnerClass() == UStateTreeState::StaticClass()
-			&& Property->GetFName() == GET_MEMBER_NAME_CHECKED(UStateTreeState, Name))
+		const UStateTree* StateTree = GetTypedOuter<UStateTree>();
+		if (ensure(StateTree))
 		{
-			UStateTree* StateTree = GetTypedOuter<UStateTree>();
-			if (ensure(StateTree))
-			{
-				UE::StateTree::Delegates::OnIdentifierChanged.Broadcast(*StateTree);
-			}
+			UE::StateTree::Delegates::OnIdentifierChanged.Broadcast(*StateTree);
 		}
+	}
 
-		if (Property->GetOwnerClass() == UStateTreeState::StaticClass()
-			&& Property->GetFName() == GET_MEMBER_NAME_CHECKED(UStateTreeState, Type))
+	if (ChangePropertyPath.IsPathExact(StateTypePath))
+	{
+		// Remove any tasks and evaluators when they are not used.
+		if (Type == EStateTreeStateType::Group || Type == EStateTreeStateType::Linked)
 		{
-			// Remove any tasks and evaluators when they are not used.
-			if (Type == EStateTreeStateType::Group || Type == EStateTreeStateType::Linked)
-			{
-				Tasks.Reset();
-			}
+			Tasks.Reset();
+		}
 
-			// If transitioning from linked state, reset the linked state.
-			if (Type != EStateTreeStateType::Linked)
-			{
-				LinkedSubtree = FStateTreeStateLink();
-			}
+		// If transitioning from linked state, reset the linked state.
+		if (Type != EStateTreeStateType::Linked)
+		{
+			LinkedSubtree = FStateTreeStateLink();
+		}
 
-			if (Type == EStateTreeStateType::Linked)
-			{
-				// Linked parameter layout is fixed, and copied from the linked target state.
-				Parameters.bFixedLayout = true;
-				UpdateParametersFromLinkedSubtree();
-			}
-			else if (Type == EStateTreeStateType::Subtree)
-			{
-				// Subtree parameter layout can be edited
-				Parameters.bFixedLayout = false;
-			}
-			else
+		if (Type == EStateTreeStateType::Linked)
+		{
+			// Linked parameter layout is fixed, and copied from the linked target state.
+			Parameters.bFixedLayout = true;
+			UpdateParametersFromLinkedSubtree();
+		}
+		else if (Type == EStateTreeStateType::Subtree)
+		{
+			// Subtree parameter layout can be edited
+			Parameters.bFixedLayout = false;
+		}
+		else
+		{
+			Parameters.Reset();
+		}
+	}
+
+	// When switching to new state, update the parameters.
+	if (ChangePropertyPath.IsPathExact(StateLinkedSubtreePath))
+	{
+		if (Type == EStateTreeStateType::Linked)
+		{
+			UpdateParametersFromLinkedSubtree();
+		}
+	}
+
+	// Broadcast subtree parameter layout edits so that the linked states can adapt.
+	if (ChangePropertyPath.IsPathExact(StateParametersPath))
+	{
+		if (Type == EStateTreeStateType::Subtree)
+		{
+			const UStateTree* StateTree = GetTypedOuter<UStateTree>();
+			if (ensure(StateTree))
 			{
-				Parameters.Reset();
+				UE::StateTree::Delegates::OnStateParametersChanged.Broadcast(*StateTree, ID);
 			}
 		}
+	}
 
-		if (Property->GetOwnerClass() == UStateTreeState::StaticClass()
-			&& Property->GetFName() == GET_MEMBER_NAME_CHECKED(UStateTreeState, LinkedSubtree))
+	// Reset delay on completion transitions
+	if (ChangePropertyPath.ContainsPath(StateTransitionsPath))
+	{
+		const int32 TransitionsIndex = ChangePropertyPath.GetPropertyArrayIndex(StateTransitionsPath);
+		if (Transitions.IsValidIndex(TransitionsIndex))
 		{
-			// When switching to new state, update the parameters.
-			if (Type == EStateTreeStateType::Linked)
+			FStateTreeTransition& Transition = Transitions[TransitionsIndex];
+
+			if (EnumHasAnyFlags(Transition.Trigger, EStateTreeTransitionTrigger::OnStateCompleted))
 			{
-				UpdateParametersFromLinkedSubtree();
+				Transition.bDelayTransition = false;
 			}
 		}
+	}
 
-		if (Property->GetOwnerClass() == UStateTreeState::StaticClass()
-			&& Property->GetFName() == GET_MEMBER_NAME_CHECKED(UStateTreeState, Parameters))
+	// Ensure unique ID on duplicated items.
+	if (PropertyChangedEvent.ChangeType == EPropertyChangeType::Duplicate)
+	{
+		
+		// Tasks
+		if (ChangePropertyPath.IsPathExact(StateTasksPath))
 		{
-			if (Type == EStateTreeStateType::Subtree)
+			const int32 ArrayIndex = ChangePropertyPath.GetPropertyArrayIndex(StateTasksPath);
+			if (Tasks.IsValidIndex(ArrayIndex))
 			{
-				// Broadcast subtree parameter edits so that the linked states can adapt.
-				const UStateTree* StateTree = GetTypedOuter<UStateTree>();
-				if (ensure(StateTree))
+ 				if (FStateTreeTaskBase* Task = Tasks[ArrayIndex].Node.GetMutablePtr<FStateTreeTaskBase>())
 				{
-					UE::StateTree::Delegates::OnStateParametersChanged.Broadcast(*StateTree, ID);
+					Task->Name = FName(Task->Name.ToString() + TEXT(" Duplicate"));
 				}
+				const FGuid OldStructID = Tasks[ArrayIndex].ID; 
+				Tasks[ArrayIndex].ID = FGuid::NewGuid();
+				CopyBindings(OldStructID, Tasks[ArrayIndex].ID);
 			}
 		}
 
-		if (MemberProperty)
+		// Enter conditions
+		if (ChangePropertyPath.IsPathExact(StateEnterConditionsPath))
 		{
-			// Ensure unique ID on duplicated items.
-			if (PropertyChangedEvent.ChangeType == EPropertyChangeType::Duplicate)
+			const int32 ArrayIndex = ChangePropertyPath.GetPropertyArrayIndex(StateEnterConditionsPath);
+			if (EnterConditions.IsValidIndex(ArrayIndex))
 			{
-				if (MemberProperty->GetFName() == GET_MEMBER_NAME_CHECKED(UStateTreeState, Tasks))
+				if (FStateTreeConditionBase* Condition = EnterConditions[ArrayIndex].Node.GetMutablePtr<FStateTreeConditionBase>())
 				{
-					const int32 ArrayIndex = PropertyChangedEvent.GetArrayIndex(MemberProperty->GetFName().ToString());
-					if (Tasks.IsValidIndex(ArrayIndex))
-					{
-						if (FStateTreeTaskBase* Task = Tasks[ArrayIndex].Node.GetMutablePtr<FStateTreeTaskBase>())
-						{
-							Task->Name = FName(Task->Name.ToString() + TEXT(" Duplicate"));
-						}
-						const FGuid OldStructID = Tasks[ArrayIndex].ID; 
-						Tasks[ArrayIndex].ID = FGuid::NewGuid();
-						CopyBindings(OldStructID, Tasks[ArrayIndex].ID);
-					}
-				}
-				if (MemberProperty->GetFName() == GET_MEMBER_NAME_CHECKED(UStateTreeState, EnterConditions))
-				{
-					const int32 ArrayIndex = PropertyChangedEvent.GetArrayIndex(MemberProperty->GetFName().ToString());
-					if (EnterConditions.IsValidIndex(ArrayIndex))
-					{
-						if (FStateTreeConditionBase* Condition = EnterConditions[ArrayIndex].Node.GetMutablePtr<FStateTreeConditionBase>())
-						{
-							Condition->Name = FName(Condition->Name.ToString() + TEXT(" Duplicate"));
-						}
-						const FGuid OldStructID = EnterConditions[ArrayIndex].ID; 
-						EnterConditions[ArrayIndex].ID = FGuid::NewGuid();
-						CopyBindings(OldStructID, EnterConditions[ArrayIndex].ID);
-					}
-				}
-				if (MemberProperty->GetFName() == GET_MEMBER_NAME_CHECKED(UStateTreeState, Transitions))
-				{
-					const int32 TransitionsIndex = PropertyChangedEvent.GetArrayIndex(MemberProperty->GetFName().ToString());
-					const int32 ConditionsIndex = TryGetConditionIndex();
-
-					if (Transitions.IsValidIndex(TransitionsIndex))
-					{
-						FStateTreeTransition& Transition = Transitions[TransitionsIndex];
-						if (Transition.Conditions.IsValidIndex(ConditionsIndex))
-						{
-							if (FStateTreeConditionBase* Condition = Transition.Conditions[ConditionsIndex].Node.GetMutablePtr<FStateTreeConditionBase>())
-							{
-								Condition->Name = FName(Condition->Name.ToString() + TEXT(" Duplicate"));
-							}
-							const FGuid OldStructID = Transition.Conditions[ConditionsIndex].ID;
-							Transition.Conditions[ConditionsIndex].ID = FGuid::NewGuid();
-							CopyBindings(OldStructID, Transition.Conditions[ConditionsIndex].ID);
-						}
-					}
+					Condition->Name = FName(Condition->Name.ToString() + TEXT(" Duplicate"));
 				}
+				const FGuid OldStructID = EnterConditions[ArrayIndex].ID; 
+				EnterConditions[ArrayIndex].ID = FGuid::NewGuid();
+				CopyBindings(OldStructID, EnterConditions[ArrayIndex].ID);
 			}
-			else if (PropertyChangedEvent.ChangeType == EPropertyChangeType::ArrayAdd)
+		}
+
+		// Transition conditions
+		if (ChangePropertyPath.IsPathExact(StateTransitionsConditionsPath))
+		{
+			const int32 TransitionsIndex = ChangePropertyPath.GetPropertyArrayIndex(StateTransitionsPath);
+			const int32 ConditionsIndex = ChangePropertyPath.GetPropertyArrayIndex(StateTransitionsConditionsPath);
+
+			if (Transitions.IsValidIndex(TransitionsIndex))
 			{
-				if (MemberProperty->GetFName() == GET_MEMBER_NAME_CHECKED(UStateTreeState, Transitions))
+				FStateTreeTransition& Transition = Transitions[TransitionsIndex];
+				if (Transition.Conditions.IsValidIndex(ConditionsIndex))
 				{
-					const int32 TransitionsIndex = PropertyChangedEvent.GetArrayIndex(MemberProperty->GetFName().ToString());
-					if (Transitions.IsValidIndex(TransitionsIndex))
+					if (FStateTreeConditionBase* Condition = Transition.Conditions[ConditionsIndex].Node.GetMutablePtr<FStateTreeConditionBase>())
 					{
-						// Set default transition on newly created states.
-						FStateTreeTransition& Transition = Transitions[TransitionsIndex];
-						Transition.Trigger = EStateTreeTransitionTrigger::OnStateCompleted;
-						const UStateTreeState* RootState = GetRootState();
-						Transition.State = RootState->GetLinkToState();
+						Condition->Name = FName(Condition->Name.ToString() + TEXT(" Duplicate"));
 					}
+					const FGuid OldStructID = Transition.Conditions[ConditionsIndex].ID;
+					Transition.Conditions[ConditionsIndex].ID = FGuid::NewGuid();
+					CopyBindings(OldStructID, Transition.Conditions[ConditionsIndex].ID);
 				}
 			}
-			else if (PropertyChangedEvent.ChangeType == EPropertyChangeType::ArrayRemove)
+		}
+	}
+	
+	// Set default state to root on new transitions.
+	if (PropertyChangedEvent.ChangeType == EPropertyChangeType::ArrayAdd)
+	{
+		if (ChangePropertyPath.IsPathExact(StateTransitionsPath))
+		{
+			const int32 TransitionsIndex = ChangePropertyPath.GetPropertyArrayIndex(StateTransitionsPath);
+			if (Transitions.IsValidIndex(TransitionsIndex))
 			{
-				if (MemberProperty->GetFName() == GET_MEMBER_NAME_CHECKED(UStateTreeState, Tasks)
-					|| MemberProperty->GetFName() == GET_MEMBER_NAME_CHECKED(UStateTreeState, EnterConditions)
-					|| MemberProperty->GetFName() == GET_MEMBER_NAME_CHECKED(UStateTreeState, Transitions))
-				{
-					if (UStateTreeEditorData* TreeData = GetTypedOuter<UStateTreeEditorData>())
-					{
-						TreeData->Modify();
-						FStateTreeEditorPropertyBindings* Bindings = TreeData->GetPropertyEditorBindings();
-						check(Bindings);
-						TMap<FGuid, const UStruct*> AllStructIDs;
-						TreeData->GetAllStructIDs(AllStructIDs);
-						Bindings->RemoveUnusedBindings(AllStructIDs);
-					}
-				}
+				FStateTreeTransition& Transition = Transitions[TransitionsIndex];
+				Transition.Trigger = EStateTreeTransitionTrigger::OnStateCompleted;
+				const UStateTreeState* RootState = GetRootState();
+				Transition.State = RootState->GetLinkToState();
 			}
-			else
-			{
-				if (MemberProperty->GetFName() == GET_MEMBER_NAME_CHECKED(UStateTreeState, Transitions))
-				{
-					const int32 TransitionsIndex = PropertyChangedEvent.GetArrayIndex(MemberProperty->GetFName().ToString());
-					if (Transitions.IsValidIndex(TransitionsIndex))
-					{
-						FStateTreeTransition& Transition = Transitions[TransitionsIndex];
+		}
+	}
 
-						if (EnumHasAnyFlags(Transition.Trigger, EStateTreeTransitionTrigger::OnStateCompleted))
-						{
-							// Reset delay on completion transitions
-							Transition.bDelayTransition = false;
-						}
-					}
-				}
+	// Remove bindings when bindable nodes are removed.
+	if (PropertyChangedEvent.ChangeType == EPropertyChangeType::ArrayRemove)
+	{
+		if (ChangePropertyPath.IsPathExact(StateTasksPath)
+			|| ChangePropertyPath.IsPathExact(StateEnterConditionsPath)
+			|| ChangePropertyPath.IsPathExact(StateTransitionsConditionsPath))
+		{
+			if (UStateTreeEditorData* TreeData = GetTypedOuter<UStateTreeEditorData>())
+			{
+				TreeData->Modify();
+				FStateTreeEditorPropertyBindings* Bindings = TreeData->GetPropertyEditorBindings();
+				check(Bindings);
+				TMap<FGuid, const UStruct*> AllStructIDs;
+				TreeData->GetAllStructIDs(AllStructIDs);
+				Bindings->RemoveUnusedBindings(AllStructIDs);
 			}
 		}
 	}
+
 }
 
 void UStateTreeState::PostLoad()
