Removing TaskGraphInterfaces.h include from CoreDelegates.h

diff --git a/Engine/Plugins/Online/OnlineFramework/Source/Hotfix/Private/UpdateManager.cpp b/Engine/Plugins/Online/OnlineFramework/Source/Hotfix/Private/UpdateManager.cpp
--- a/Engine/Plugins/Online/OnlineFramework/Source/Hotfix/Private/UpdateManager.cpp
+++ b/Engine/Plugins/Online/OnlineFramework/Source/Hotfix/Private/UpdateManager.cpp
@@ -1,6 +1,7 @@
 // Copyright Epic Games, Inc. All Rights Reserved.
 
 #include "UpdateManager.h"
+#include "HAL/IConsoleManager.h"
 #include "Misc/CommandLine.h"
 #include "Engine/GameInstance.h"
 #include "Misc/CoreDelegates.h"
diff --git a/Engine/Plugins/Runtime/StateTree/Source/StateTreeModule/Private/StateTreeInstanceData.cpp b/Engine/Plugins/Runtime/StateTree/Source/StateTreeModule/Private/StateTreeInstanceData.cpp
--- a/Engine/Plugins/Runtime/StateTree/Source/StateTreeModule/Private/StateTreeInstanceData.cpp
+++ b/Engine/Plugins/Runtime/StateTree/Source/StateTreeModule/Private/StateTreeInstanceData.cpp
@@ -60,6 +60,11 @@ void FStateTreeInstanceStorage::AddTransitionRequest(const UObject* Owner, const
 	TransitionRequests.Add(Request);
 }
 
+void FStateTreeInstanceStorage::ResetTransitionRequests()
+{
+	TransitionRequests.Reset();
+}
+
 
 //----------------------------------------------------------------//
 // FStateTreeInstanceData
diff --git a/Engine/Plugins/Runtime/StateTree/Source/StateTreeModule/Public/StateTreeInstanceData.h b/Engine/Plugins/Runtime/StateTree/Source/StateTreeModule/Public/StateTreeInstanceData.h
--- a/Engine/Plugins/Runtime/StateTree/Source/StateTreeModule/Public/StateTreeInstanceData.h
+++ b/Engine/Plugins/Runtime/StateTree/Source/StateTreeModule/Public/StateTreeInstanceData.h
@@ -1,4 +1,4 @@
-﻿// Copyright Epic Games, Inc. All Rights Reserved.
+// Copyright Epic Games, Inc. All Rights Reserved.
 
 #pragma once
 
@@ -72,10 +72,7 @@ struct STATETREEMODULE_API FStateTreeInstanceStorage
 	}
 
 	/** Reset all pending transition requests. */
-	void ResetTransitionRequests()
-	{
-		TransitionRequests.Reset();
-	}
+	void ResetTransitionRequests();
 	
 protected:
 	/** Struct instances */
diff --git a/Engine/Source/Runtime/AIModule/Private/BehaviorTree/BehaviorTreeComponent.cpp b/Engine/Source/Runtime/AIModule/Private/BehaviorTree/BehaviorTreeComponent.cpp
--- a/Engine/Source/Runtime/AIModule/Private/BehaviorTree/BehaviorTreeComponent.cpp
+++ b/Engine/Source/Runtime/AIModule/Private/BehaviorTree/BehaviorTreeComponent.cpp
@@ -14,6 +14,7 @@
 #include "BehaviorTree/Tasks/BTTask_RunBehaviorDynamic.h"
 #include "ProfilingDebugging/ScopedTimers.h"
 #include "ProfilingDebugging/CsvProfiler.h"
+#include "HAL/IConsoleManager.h"
 #include "Misc/CoreDelegates.h"
 #include "Misc/ConfigCacheIni.h"
 
diff --git a/Engine/Source/Runtime/Core/Private/ProfilingDebugging/CsvProfiler.cpp b/Engine/Source/Runtime/Core/Private/ProfilingDebugging/CsvProfiler.cpp
--- a/Engine/Source/Runtime/Core/Private/ProfilingDebugging/CsvProfiler.cpp
+++ b/Engine/Source/Runtime/Core/Private/ProfilingDebugging/CsvProfiler.cpp
@@ -7,6 +7,7 @@
 */
 
 #include "ProfilingDebugging/CsvProfiler.h"
+#include "Async/TaskGraphInterfaces.h"
 #include "CoreGlobals.h"
 #include "HAL/RunnableThread.h"
 #include "HAL/ThreadManager.h"
diff --git a/Engine/Source/Runtime/Core/Private/ProfilingDebugging/TraceAuxiliary.cpp b/Engine/Source/Runtime/Core/Private/ProfilingDebugging/TraceAuxiliary.cpp
--- a/Engine/Source/Runtime/Core/Private/ProfilingDebugging/TraceAuxiliary.cpp
+++ b/Engine/Source/Runtime/Core/Private/ProfilingDebugging/TraceAuxiliary.cpp
@@ -48,6 +48,7 @@
 #include "ProfilingDebugging/CsvProfiler.h"
 #include "ProfilingDebugging/MemoryTrace.h"
 #include "ProfilingDebugging/PlatformEvents.h"
+#include "Stats/Stats.h"
 #include "String/ParseTokens.h"
 #include "Templates/Invoke.h"
 #include "Templates/UnrealTemplate.h"
diff --git a/Engine/Source/Runtime/Core/Public/Misc/CoreDelegates.h b/Engine/Source/Runtime/Core/Public/Misc/CoreDelegates.h
--- a/Engine/Source/Runtime/Core/Public/Misc/CoreDelegates.h
+++ b/Engine/Source/Runtime/Core/Public/Misc/CoreDelegates.h
@@ -2,7 +2,10 @@
 
 #pragma once
 
+#if UE_ENABLE_INCLUDE_ORDER_DEPRECATED_IN_5_2
 #include "Async/TaskGraphInterfaces.h"
+#endif
+#include "Async/TaskGraphFwd.h"
 #include "Containers/Array.h"
 #include "Containers/ArrayView.h"
 #include "Containers/Map.h"
diff --git a/Engine/Source/Runtime/Engine/Private/Camera/CameraPhotography.cpp b/Engine/Source/Runtime/Engine/Private/Camera/CameraPhotography.cpp
--- a/Engine/Source/Runtime/Engine/Private/Camera/CameraPhotography.cpp
+++ b/Engine/Source/Runtime/Engine/Private/Camera/CameraPhotography.cpp
@@ -2,6 +2,7 @@
 
 #include "Camera/CameraPhotography.h"
 
+#include "HAL/IConsoleManager.h"
 #include "Misc/CoreDelegates.h"
 #include "CameraPhotographyModule.h"
 #include "Engine/World.h"
diff --git a/Engine/Source/Runtime/Engine/Private/WorldPartition/RuntimeSpatialHash/RuntimeSpatialHashGridHelper.cpp b/Engine/Source/Runtime/Engine/Private/WorldPartition/RuntimeSpatialHash/RuntimeSpatialHashGridHelper.cpp
--- a/Engine/Source/Runtime/Engine/Private/WorldPartition/RuntimeSpatialHash/RuntimeSpatialHashGridHelper.cpp
+++ b/Engine/Source/Runtime/Engine/Private/WorldPartition/RuntimeSpatialHash/RuntimeSpatialHashGridHelper.cpp
@@ -2,6 +2,7 @@
 
 #include "WorldPartition/RuntimeSpatialHash/RuntimeSpatialHashGridHelper.h"
 #include "ActorPartition/PartitionActor.h"
+#include "HAL/IConsoleManager.h"
 #include "ProfilingDebugging/ScopedTimers.h"
 #include "WorldPartition/WorldPartitionLog.h"
 #include "WorldPartition/WorldPartitionStreamingGeneration.h"
diff --git a/Engine/Source/Runtime/Renderer/Private/RendererOnScreenNotification.h b/Engine/Source/Runtime/Renderer/Private/RendererOnScreenNotification.h
--- a/Engine/Source/Runtime/Renderer/Private/RendererOnScreenNotification.h
+++ b/Engine/Source/Runtime/Renderer/Private/RendererOnScreenNotification.h
@@ -2,6 +2,7 @@
 
 #pragma once
 
+#include "Internationalization/Text.h"
 #include "Misc/CoreDelegates.h"
 #include "Misc/LazySingleton.h"
 
diff --git a/Engine/Source/Runtime/TypedElementFramework/Private/Elements/Framework/TypedElementRegistry.cpp b/Engine/Source/Runtime/TypedElementFramework/Private/Elements/Framework/TypedElementRegistry.cpp
--- a/Engine/Source/Runtime/TypedElementFramework/Private/Elements/Framework/TypedElementRegistry.cpp
+++ b/Engine/Source/Runtime/TypedElementFramework/Private/Elements/Framework/TypedElementRegistry.cpp
@@ -2,6 +2,7 @@
 
 #include "Elements/Framework/TypedElementRegistry.h"
 
+#include "HAL/IConsoleManager.h"
 #include "Misc/CoreDelegates.h"
 #include "Misc/ScopeLock.h"
 #include "UObject/StrongObjectPtr.h"
