StateTree: Handle objects in instance data migrationvis serialization

diff --git a/Engine/Plugins/Runtime/StateTree/Source/StateTreeModule/Private/StateTreeInstanceData.cpp b/Engine/Plugins/Runtime/StateTree/Source/StateTreeModule/Private/StateTreeInstanceData.cpp
--- a/Engine/Plugins/Runtime/StateTree/Source/StateTreeModule/Private/StateTreeInstanceData.cpp
+++ b/Engine/Plugins/Runtime/StateTree/Source/StateTreeModule/Private/StateTreeInstanceData.cpp
@@ -4,6 +4,8 @@
 #include "Serialization/MemoryReader.h"
 #include "Serialization/MemoryWriter.h"
 #include "VisualLogger/VisualLogger.h"
+#include "StateTree.h"
+#include "Serialization/ObjectAndNameAsStringProxyArchive.h"
 
 #include UE_INLINE_GENERATED_CPP_BY_NAME(StateTreeInstanceData)
 
@@ -22,16 +24,21 @@ namespace UE::StateTree
 			UObject* NewInstance = NewObject<UObject>(&InOwner, AuthoritativeClass);
 
 			// Try to copy the values over using serialization
+			// FObjectAndNameAsStringProxyArchive is used to store and restore names and objects as memory writer does not support UObject references at all.
 			TArray<uint8> Data;
 			FMemoryWriter Writer(Data);
+			FObjectAndNameAsStringProxyArchive WriterProxy(Writer, /*bInLoadIfFindFails*/true);
 			UObject& NonConstInstance = const_cast<UObject&>(Instance);
-			NonConstInstance.Serialize(Writer);
+			NonConstInstance.Serialize(WriterProxy);
 
 			FMemoryReader Reader(Data);
-			NewInstance->Serialize(Reader);
-				
-			UE_LOG(LogStateTree, Display, TEXT("FStateTreeInstanceData: Duplicating '%s' with old class '%s' as '%s', potential data loss."),
-				*GetFullNameSafe(&Instance), *GetNameSafe(InstanceClass), *GetNameSafe(AuthoritativeClass));
+			FObjectAndNameAsStringProxyArchive ReaderProxy(Reader, /*bInLoadIfFindFails*/true);
+			NewInstance->Serialize(ReaderProxy);
+
+			const UStateTree* OuterStateTree = Instance.GetTypedOuter<UStateTree>();
+
+			UE_LOG(LogStateTree, Display, TEXT("FStateTreeInstanceData: Duplicating '%s' with old class '%s' as '%s', potential data loss. Please resave State Tree asset %s."),
+				*GetFullNameSafe(&Instance), *GetNameSafe(InstanceClass), *GetNameSafe(AuthoritativeClass), *GetFullNameSafe(OuterStateTree));
 
 			return NewInstance;
 		}
