StateTree: fix eval & global task dataviews after UpdateInstanceData.

diff --git a/Engine/Plugins/Runtime/StateTree/Source/StateTreeModule/Private/StateTreeExecutionContext.cpp b/Engine/Plugins/Runtime/StateTree/Source/StateTreeModule/Private/StateTreeExecutionContext.cpp
--- a/Engine/Plugins/Runtime/StateTree/Source/StateTreeModule/Private/StateTreeExecutionContext.cpp
+++ b/Engine/Plugins/Runtime/StateTree/Source/StateTreeModule/Private/StateTreeExecutionContext.cpp
@@ -616,6 +616,9 @@ void FStateTreeExecutionContext::UpdateInstanceData(const FStateTreeActiveStates
 	FStateTreeExecutionState& Exec = GetExecState();
 	Exec.FirstTaskStructIndex = FStateTreeIndex16(FirstTaskStructIndex);
 	Exec.FirstTaskObjectIndex = FStateTreeIndex16(FirstTaskObjectIndex);
+
+	// Clear dataviews so that we dont have anything point to potentially stale data (e.g. the pointers for instanced structs may have changed).
+	DataViews.Init({}, DataViews.Num());
 }
 
 EStateTreeRunStatus FStateTreeExecutionContext::EnterState(const FStateTreeTransitionResult& Transition)
@@ -633,24 +636,33 @@ EStateTreeRunStatus FStateTreeExecutionContext::EnterState(const FStateTreeTrans
 	FStateTreeExecutionState& Exec = GetExecState();
 	Exec.StateChangeCount++;
 	Exec.CompletedStateHandle = FStateTreeStateHandle::Invalid;
+	Exec.EnterStateFailedTaskIndex = FStateTreeIndex16::Invalid; // This will make all tasks to be accepted.
+	Exec.ActiveStates.Reset();
 
 	// On target branch means that the state is the target of current transition or child of it.
 	// States which were active before and will remain active, but are not on target branch will not get
 	// EnterState called. That is, a transition is handled as "replan from this state".
 	bool bOnTargetBranch = false;
-
 	FStateTreeTransitionResult CurrentTransition = Transition;
-	
 	EStateTreeRunStatus Result = EStateTreeRunStatus::Running;
+	int32 InstanceStructIndex = 1; // Exec is at index 0
+	int32 InstanceObjectIndex = 0;
 
-	Exec.EnterStateFailedTaskIndex = FStateTreeIndex16::Invalid; // This will make all tasks to be accepted.
-	Exec.ActiveStates.Reset();
-	
-	// Do property copy on all states, propagating the results from last tick.
-	check(Exec.FirstTaskStructIndex.IsValid() && Exec.FirstTaskObjectIndex.IsValid()); 
-	int32 InstanceStructIndex = Exec.FirstTaskStructIndex.Get();
-	int32 InstanceObjectIndex = Exec.FirstTaskObjectIndex.Get();
+	// Update data views for evaluators and global tasks as UpdateInstanceData() might have changed the location of the instance data.
+	// Evaluators
+	for (int32 EvalIndex = StateTree.EvaluatorsBegin; EvalIndex < (StateTree.EvaluatorsBegin + StateTree.EvaluatorsNum); EvalIndex++)
+	{
+		const FStateTreeEvaluatorBase& Eval = StateTree.Nodes[EvalIndex].Get<FStateTreeEvaluatorBase>();
+		SetNodeDataView(Eval, InstanceStructIndex, InstanceObjectIndex);
+	}
 
+	// Global tasks
+	for (int32 TaskIndex = StateTree.GlobalTasksBegin; TaskIndex < (StateTree.GlobalTasksBegin + StateTree.GlobalTasksNum); TaskIndex++)
+	{
+		const FStateTreeTaskBase& Task =  StateTree.Nodes[TaskIndex].Get<FStateTreeTaskBase>();
+		SetNodeDataView(Task, InstanceStructIndex, InstanceObjectIndex);
+	}
+	
 	STATETREE_LOG(Log, TEXT("Enter state '%s' (%d)"), *DebugGetStatePath(Transition.NextActiveStates), Exec.StateChangeCount);
 
 	for (int32 Index = 0; Index < Transition.NextActiveStates.Num() && Result != EStateTreeRunStatus::Failed; Index++)
diff --git a/Engine/Plugins/Runtime/StateTree/Source/StateTreeTestSuite/Private/StateTreeTest.cpp b/Engine/Plugins/Runtime/StateTree/Source/StateTreeTestSuite/Private/StateTreeTest.cpp
--- a/Engine/Plugins/Runtime/StateTree/Source/StateTreeTestSuite/Private/StateTreeTest.cpp
+++ b/Engine/Plugins/Runtime/StateTree/Source/StateTreeTestSuite/Private/StateTreeTest.cpp
@@ -567,7 +567,7 @@ struct FStateTreeTest_TransitionPriority : FAITestBase
 		return true;
 	}
 };
-IMPLEMENT_AI_INSTANT_TEST(FStateTreeTest_TransitionPriority, "System.StateTree.TransitionPriority");
+IMPLEMENT_AI_INSTANT_TEST(FStateTreeTest_TransitionPriority, "System.StateTree.Transition.Priority");
 
 struct FStateTreeTest_TransitionPriorityEnterState : FAITestBase
 {
@@ -631,7 +631,66 @@ struct FStateTreeTest_TransitionPriorityEnterState : FAITestBase
 		return true;
 	}
 };
-IMPLEMENT_AI_INSTANT_TEST(FStateTreeTest_TransitionPriorityEnterState, "System.StateTree.TransitionPriorityEnterState");
+IMPLEMENT_AI_INSTANT_TEST(FStateTreeTest_TransitionPriorityEnterState, "System.StateTree.Transition.PriorityEnterState");
+
+struct FStateTreeTest_TransitionGlobalDataView : FAITestBase
+{
+	// Tests that the global eval and task dataviews are kept up to date when transitioning from  
+	virtual bool InstantTest() override
+	{
+		UStateTree& StateTree = UE::StateTree::Tests::NewStateTree(&GetWorld());
+		UStateTreeEditorData& EditorData = *Cast<UStateTreeEditorData>(StateTree.EditorData);
+
+		UStateTreeState& Root = EditorData.AddSubTree(FName(TEXT("Root")));
+		UStateTreeState& StateA = Root.AddChildState(FName(TEXT("A")));
+		UStateTreeState& StateB = Root.AddChildState(FName(TEXT("B")));
+
+		auto& EvalA = EditorData.AddEvaluator<FTestEval_A>(FName(TEXT("Eval")));
+		EvalA.GetInstanceData().IntA = 42;
+		auto& GlobalTask = EditorData.AddGlobalTask<FTestTask_PrintValue>(FName(TEXT("Global")));
+		GlobalTask.GetInstanceData().Value = 123;
+		
+		// State A
+		auto& Task0 = StateA.AddTask<FTestTask_Stand>(FName(TEXT("Task0")));
+		StateA.AddTransition(EStateTreeTransitionTrigger::OnStateCompleted, EStateTreeTransitionType::GotoState, &StateB);
+
+		// State B
+		auto& Task1 = StateB.AddTask<FTestTask_PrintValue>(FName(TEXT("Task1")));
+		EditorData.AddPropertyBinding(FStateTreeEditorPropertyPath(EvalA.ID, TEXT("IntA")), FStateTreeEditorPropertyPath(Task1.ID, TEXT("Value")));
+		auto& Task2 = StateB.AddTask<FTestTask_PrintValue>(FName(TEXT("Task2")));
+		EditorData.AddPropertyBinding(FStateTreeEditorPropertyPath(GlobalTask.ID, TEXT("Value")), FStateTreeEditorPropertyPath(Task2.ID, TEXT("Value")));
+
+		FStateTreeCompilerLog Log;
+		FStateTreeCompiler Compiler(Log);
+		const bool bResult = Compiler.Compile(StateTree);
+
+		AITEST_TRUE("StateTree should get compiled", bResult);
+
+		EStateTreeRunStatus Status = EStateTreeRunStatus::Unset;
+		FStateTreeInstanceData InstanceData;
+		FTestStateTreeExecutionContext Exec(StateTree, StateTree, InstanceData);
+		const bool bInitSucceeded = Exec.IsValid();
+		AITEST_TRUE("StateTree should init", bInitSucceeded);
+
+		const FString EnterStateStr(TEXT("EnterState"));
+		const FString EnterState42Str(TEXT("EnterState42"));
+		const FString EnterState123Str(TEXT("EnterState123"));
+
+		// Start and enter state
+		Status = Exec.Start();
+		AITEST_TRUE("StateTree Task0 should enter state", Exec.Expect(Task0.GetName(), EnterStateStr));
+		Exec.LogClear();
+
+		// Transition from StateA to StateB, Task0 should enter state with evaluator value copied.
+		Status = Exec.Tick(0.1f);
+		AITEST_TRUE("StateTree Task0 should enter state with value 42", Exec.Expect(Task1.GetName(), EnterState42Str));
+		AITEST_TRUE("StateTree Task1 should enter state with value 123", Exec.Expect(Task2.GetName(), EnterState123Str));
+		Exec.LogClear();
+
+		return true;
+	}
+};
+IMPLEMENT_AI_INSTANT_TEST(FStateTreeTest_TransitionGlobalDataView, "System.StateTree.Transition.GlobalDataView");
 
 UE_ENABLE_OPTIMIZATION_SHIP
 
diff --git a/Engine/Plugins/Runtime/StateTree/Source/StateTreeTestSuite/Private/StateTreeTestTypes.h b/Engine/Plugins/Runtime/StateTree/Source/StateTreeTestSuite/Private/StateTreeTestTypes.h
--- a/Engine/Plugins/Runtime/StateTree/Source/StateTreeTestSuite/Private/StateTreeTestTypes.h
+++ b/Engine/Plugins/Runtime/StateTree/Source/StateTreeTestSuite/Private/StateTreeTestTypes.h
@@ -117,6 +117,7 @@ struct FTestEval_A : public FStateTreeEvaluatorBase
 	using FInstanceDataType = FTestEval_AInstanceData;
 
 	FTestEval_A() = default;
+	FTestEval_A(const FName InName) { Name = InName; }
 	virtual ~FTestEval_A() override {}
 
 	virtual const UStruct* GetInstanceDataType() const override { return FInstanceDataType::StaticStruct(); }
@@ -150,6 +151,54 @@ struct FTestTask_B : public FStateTreeTaskBase
 	virtual const UStruct* GetInstanceDataType() const override { return FInstanceDataType::StaticStruct(); }
 };
 
+USTRUCT()
+struct FTestTask_PrintValueInstanceData
+{
+	GENERATED_BODY()
+
+	UPROPERTY(EditAnywhere, Category = "Parameter")
+	int32 Value = 0;
+};
+
+USTRUCT()
+struct FTestTask_PrintValue : public FStateTreeTaskBase
+{
+	GENERATED_BODY()
+
+	using FInstanceDataType = FTestTask_PrintValueInstanceData;
+
+	FTestTask_PrintValue() = default;
+	FTestTask_PrintValue(const FName InName) { Name = InName; }
+	virtual ~FTestTask_PrintValue() override {}
+	
+	virtual const UStruct* GetInstanceDataType() const override { return FInstanceDataType::StaticStruct(); }
+
+	virtual EStateTreeRunStatus EnterState(FStateTreeExecutionContext& Context, const FStateTreeTransitionResult& Transition) const override
+	{
+		FTestStateTreeExecutionContext& TestContext = static_cast<FTestStateTreeExecutionContext&>(Context);
+		const FInstanceDataType& InstanceData = Context.GetInstanceData(*this);
+		TestContext.Log(Name,  FString::Printf(TEXT("EnterState%d"), InstanceData.Value));
+		return EStateTreeRunStatus::Running;
+	}
+
+	virtual void ExitState(FStateTreeExecutionContext& Context, const FStateTreeTransitionResult& Transition) const override
+	{
+		FTestStateTreeExecutionContext& TestContext = static_cast<FTestStateTreeExecutionContext&>(Context);
+		const FInstanceDataType& InstanceData = Context.GetInstanceData(*this);
+		TestContext.Log(Name,  FString::Printf(TEXT("ExitState%d"), InstanceData.Value));
+	}
+
+	virtual EStateTreeRunStatus Tick(FStateTreeExecutionContext& Context, const float DeltaTime) const override
+	{
+		FTestStateTreeExecutionContext& TestContext = static_cast<FTestStateTreeExecutionContext&>(Context);
+		const FInstanceDataType& InstanceData = Context.GetInstanceData(*this);
+		TestContext.Log(Name,  FString::Printf(TEXT("Tick%d"), InstanceData.Value));
+		
+		return EStateTreeRunStatus::Running;
+	};
+};
+
+
 USTRUCT()
 struct FTestTask_StandInstanceData
 {
