StateTree: fixed bug of a subtrees being compiled in twice
- added check to omit subtrees when traversing the base state level (the code was not expecting subtrees at that level)
- added error logging when same binding source is added multiple times

diff --git a/Engine/Plugins/Runtime/StateTree/Source/StateTreeEditorModule/Private/StateTreeCompiler.cpp b/Engine/Plugins/Runtime/StateTree/Source/StateTreeEditorModule/Private/StateTreeCompiler.cpp
--- a/Engine/Plugins/Runtime/StateTree/Source/StateTreeEditorModule/Private/StateTreeCompiler.cpp
+++ b/Engine/Plugins/Runtime/StateTree/Source/StateTreeEditorModule/Private/StateTreeCompiler.cpp
@@ -212,7 +212,8 @@ bool FStateTreeCompiler::CreateStates()
 	// Create main tree (omit subtrees)
 	for (UStateTreeState* SubTree : EditorData->SubTrees)
 	{
-		if (SubTree != nullptr)
+		if (SubTree != nullptr
+			&& SubTree->Type != EStateTreeStateType::Subtree)
 		{
 			if (!CreateStateRecursive(*SubTree, FStateTreeStateHandle::Invalid))
 			{
diff --git a/Engine/Plugins/Runtime/StateTree/Source/StateTreeEditorModule/Private/StateTreePropertyBindingCompiler.cpp b/Engine/Plugins/Runtime/StateTree/Source/StateTreeEditorModule/Private/StateTreePropertyBindingCompiler.cpp
--- a/Engine/Plugins/Runtime/StateTree/Source/StateTreeEditorModule/Private/StateTreePropertyBindingCompiler.cpp
+++ b/Engine/Plugins/Runtime/StateTree/Source/StateTreeEditorModule/Private/StateTreePropertyBindingCompiler.cpp
@@ -152,6 +152,14 @@ void FStateTreePropertyBindingCompiler::Finalize()
 
 int32 FStateTreePropertyBindingCompiler::AddSourceStruct(const FStateTreeBindableStructDesc& SourceStruct)
 {
+	const int32 ExistingIndex = SourceStructs.IndexOfByPredicate([&SourceStruct](const FStateTreeBindableStructDesc& Struct) { return (Struct.ID == SourceStruct.ID); });
+	if (ExistingIndex != INDEX_NONE)
+	{
+		const FStateTreeBindableStructDesc& ExistingStruct = SourceStructs[ExistingIndex];
+		UE_LOG(LogStateTree, Error, TEXT("Struct '%s' already exists as '%s'"),
+			*SourceStruct.Name.ToString(), *SourceStruct.ID.ToString(), *ExistingStruct.Name.ToString());
+	}
+	
 	SourceStructs.Add(SourceStruct);
 	return SourceStructs.Num() - 1;
 }
