StateTree: Fix state tree reference when placed inside an instanced struct
- when inside an instanced struct, we have to use the property callbacks to update the data, because the underlying value is cached

diff --git a/Engine/Plugins/Runtime/StateTree/Source/StateTreeEditorModule/Private/Customizations/StateTreeReferenceDetails.cpp b/Engine/Plugins/Runtime/StateTree/Source/StateTreeEditorModule/Private/Customizations/StateTreeReferenceDetails.cpp
--- a/Engine/Plugins/Runtime/StateTree/Source/StateTreeEditorModule/Private/Customizations/StateTreeReferenceDetails.cpp
+++ b/Engine/Plugins/Runtime/StateTree/Source/StateTreeEditorModule/Private/Customizations/StateTreeReferenceDetails.cpp
@@ -5,6 +5,7 @@
 #include "StateTree.h"
 #include "StateTreeDelegates.h"
 #include "IDetailChildrenBuilder.h"
+#include "IPropertyUtilities.h"
 #include "PropertyCustomizationHelpers.h"
 #include "StateTreeReference.h"
 
@@ -21,7 +22,11 @@ void FStateTreeReferenceDetails::CustomizeHeader(const TSharedRef<IPropertyHandl
 
 	// Make sure parameters are synced when displaying StateTreeReference.
 	// Associated StateTree asset might have been recompiled after the StateTreeReference was loaded.
-	SyncParameters();
+	// Note: SyncParameters() will create an undoable transaction, do not try to sync when called during Undo/redo as it would overwrite the undo.   
+	if (!GIsTransacting)
+	{
+		SyncParameters();
+	}
 
 	const TSharedPtr<IPropertyHandle> StateTreeProperty = StructPropertyHandle->GetChildHandle(FName(TEXT("StateTree")));
 	check(StateTreeProperty.IsValid());
@@ -48,49 +53,7 @@ void FStateTreeReferenceDetails::CustomizeHeader(const TSharedRef<IPropertyHandl
 	
 	StateTreeProperty->SetOnPropertyValueChanged(FSimpleDelegate::CreateLambda([this]()
 		{
-			TArray<void*> RawData;
-			StructPropertyHandle->AccessRawData(RawData);
-
-			const TSharedPtr<IPropertyHandle> ParametersProperty = StructPropertyHandle->GetChildHandle(FName(TEXT("Parameters")));
-
-			if (ParametersProperty.IsValid() && !RawData.IsEmpty())
-			{
-				FStateTreeReference* StateTreeReference = static_cast<FStateTreeReference*>(RawData[0]);
-				check(StateTreeReference);
-
-				const UStateTree* StateTreeAsset = StateTreeReference->GetStateTree();
-				
-				FInstancedPropertyBag NewValue;
-				bool bSetNewValue = true;
-
-				if (RawData.Num() == 1)
-				{
-					bSetNewValue = StateTreeReference->RequiresParametersSync();
-
-					// When only one element is selected we can migrate the Parameters in case we are able to preserve some matching parameters
-					if (bSetNewValue)
-					{
-						StateTreeReference->SyncParametersToMatchStateTree(NewValue);
-					}
-				}
-				else if (StateTreeAsset)
-				{
-					// In case where multiple elements are selected we use the default parameters from the StateTree since
-					// the property handle will set value on all elements at once. In this scenario we can't "migrate" parameters since
-					// they might be associated to different StateTree so the resulting parameters would have been different for each element. 
-					NewValue = StateTreeAsset->GetDefaultParameters();
-				}
-
-				// Set new value through property handle to propagate changes to loaded instances
-				if (bSetNewValue)
-				{
-					const FStructProperty* StructProperty = CastFieldChecked<FStructProperty>(ParametersProperty->GetProperty());
-
-					FString TextValue;
-					StructProperty->Struct->ExportText(TextValue, &NewValue, /*Defaults*/nullptr, /*OwnerObject*/nullptr, EPropertyPortFlags::PPF_None, /*ExportRootScope*/nullptr);
-					ensure(ParametersProperty->SetValueFromFormattedString(TextValue, EPropertyValueSetFlags::DefaultFlags) == FPropertyAccess::Result::Success);
-				}
-			}
+			SyncParameters();
 		}));
 
 	// Registers a delegate to be notified when the associated StateTree asset get successfully recompiled
@@ -123,12 +86,22 @@ void FStateTreeReferenceDetails::SyncParameters(const UStateTree* StateTreeToSyn
 {
 	check(StructPropertyHandle.IsValid());
 
+	// This function will get called in 3 situations:
+	// - auto update when the property customization is created
+	// - auto update when a state tree is compiled.
+	// - when the associate state tree asset is changed
+	// The auto updates probably should not create undoable transactions, but we need to
+	// handle the change via the property notification system or else the data might not update
+	// if we're modifying cached data (e.g. instanced struct).
+	
 	TArray<UObject*> OuterObjects;
 	StructPropertyHandle->GetOuterObjects(OuterObjects);
 
 	TArray<void*> RawData;
 	StructPropertyHandle->AccessRawData(RawData);
 
+	bool bNeedsSync = false;
+	
 	for (int32 i = 0; i < RawData.Num(); i++)
 	{
 		if (FStateTreeReference* StateTreeReference = static_cast<FStateTreeReference*>(RawData[i]))
@@ -139,18 +112,34 @@ void FStateTreeReferenceDetails::SyncParameters(const UStateTree* StateTreeToSyn
 				&& (StateTreeToSync == nullptr || StateTreeAsset == StateTreeToSync) 
 				&& StateTreeReference->RequiresParametersSync())
 			{
-				StateTreeReference->SyncParameters();
+				bNeedsSync = true;
+			}
+		}
+	}
 
-				// StateTree compilation is not a Transaction so we simply dirty outer object
-				if (OuterObjects.Num() == RawData.Num())
+	if (bNeedsSync)
+	{
+		FScopedTransaction Transaction(LOCTEXT("SyncParameters", "Update StateTree reference"));
+		StructPropertyHandle->NotifyPreChange();
+
+		for (int32 i = 0; i < RawData.Num(); i++)
+		{
+			if (FStateTreeReference* StateTreeReference = static_cast<FStateTreeReference*>(RawData[i]))
+			{
+				const UStateTree* StateTreeAsset = StateTreeReference->GetStateTree();
+
+				if (StateTreeAsset
+					&& (StateTreeToSync == nullptr || StateTreeAsset == StateTreeToSync) 
+					&& StateTreeReference->RequiresParametersSync())
 				{
-					if (const UObject* Outer = OuterObjects[i])
-					{
-						Outer->MarkPackageDirty();
-					}
+					OuterObjects[i]->Modify();
+					StateTreeReference->SyncParameters();
 				}
 			}
 		}
+
+		StructPropertyHandle->NotifyPostChange(EPropertyChangeType::ValueSet);
+		StructPropertyHandle->NotifyFinishedChangingProperties();
 	}
 }
 
diff --git a/Engine/Plugins/Runtime/StateTree/Source/StateTreeEditorModule/Private/Customizations/StateTreeReferenceDetails.h b/Engine/Plugins/Runtime/StateTree/Source/StateTreeEditorModule/Private/Customizations/StateTreeReferenceDetails.h
--- a/Engine/Plugins/Runtime/StateTree/Source/StateTreeEditorModule/Private/Customizations/StateTreeReferenceDetails.h
+++ b/Engine/Plugins/Runtime/StateTree/Source/StateTreeEditorModule/Private/Customizations/StateTreeReferenceDetails.h
@@ -6,6 +6,7 @@
 
 class IPropertyHandle;
 class UStateTree;
+class IPropertyUtilities;
 
 /**
  * Type customization for FStateTreeReference.
@@ -34,4 +35,5 @@ private:
 	void SyncParameters(const UStateTree* StateTreeToSync = nullptr) const;
 
 	TSharedPtr<IPropertyHandle> StructPropertyHandle;
+	TSharedPtr<IPropertyUtilities> PropUtils;
 };
