StateTree: Fix StateTreeReference details non-unity build.

diff --git a/Engine/Plugins/Runtime/StateTree/Source/StateTreeEditorModule/Private/Customizations/StateTreeReferenceDetails.cpp b/Engine/Plugins/Runtime/StateTree/Source/StateTreeEditorModule/Private/Customizations/StateTreeReferenceDetails.cpp
--- a/Engine/Plugins/Runtime/StateTree/Source/StateTreeEditorModule/Private/Customizations/StateTreeReferenceDetails.cpp
+++ b/Engine/Plugins/Runtime/StateTree/Source/StateTreeEditorModule/Private/Customizations/StateTreeReferenceDetails.cpp
@@ -8,6 +8,7 @@
 #include "IPropertyUtilities.h"
 #include "PropertyCustomizationHelpers.h"
 #include "StateTreeReference.h"
+#include "ScopedTransaction.h"
 
 #define LOCTEXT_NAMESPACE "StateTreeEditor"
 
