StateTree: Revert clearing data views in UpdateInstanceData() as it was parameters and external data too.

diff --git a/Engine/Plugins/Runtime/StateTree/Source/StateTreeModule/Private/StateTreeExecutionContext.cpp b/Engine/Plugins/Runtime/StateTree/Source/StateTreeModule/Private/StateTreeExecutionContext.cpp
--- a/Engine/Plugins/Runtime/StateTree/Source/StateTreeModule/Private/StateTreeExecutionContext.cpp
+++ b/Engine/Plugins/Runtime/StateTree/Source/StateTreeModule/Private/StateTreeExecutionContext.cpp
@@ -616,9 +616,6 @@ void FStateTreeExecutionContext::UpdateInstanceData(const FStateTreeActiveStates
 	FStateTreeExecutionState& Exec = GetExecState();
 	Exec.FirstTaskStructIndex = FStateTreeIndex16(FirstTaskStructIndex);
 	Exec.FirstTaskObjectIndex = FStateTreeIndex16(FirstTaskObjectIndex);
-
-	// Clear dataviews so that we dont have anything point to potentially stale data (e.g. the pointers for instanced structs may have changed).
-	DataViews.Init({}, DataViews.Num());
 }
 
 EStateTreeRunStatus FStateTreeExecutionContext::EnterState(const FStateTreeTransitionResult& Transition)
