StateTree: Allow to bind to deep property paths (including arrays and instanced struct/object)
- refactored the property binding representation, editor binding shave now more structure, and removed intermediate representation
- added functionality to resolve property paths agains a known value
- added instanced struct and object indirection types
- added editor functionality to allow to bind to further than first level of properties
- refactored editor tree traversal, allow to access values too
- simplified statetree node ui
- requires to recompile trees, bumped version

diff --git a/Engine/Plugins/Runtime/StateTree/Config/BaseStateTree.ini b/Engine/Plugins/Runtime/StateTree/Config/BaseStateTree.ini
--- a/Engine/Plugins/Runtime/StateTree/Config/BaseStateTree.ini
+++ b/Engine/Plugins/Runtime/StateTree/Config/BaseStateTree.ini
@@ -60,3 +60,6 @@
 +PropertyRedirects=(OldName="/Script/StateTreeModule.StateTree.NamedExternalDataDescs",NewName="/Script/StateTreeModule.StateTree.ContextDataDescs")
 
 +PropertyRedirects=(OldName="/Script/StateTreeModule.StateTreeTransition.GateDelay",NewName="/Script/StateTreeModule.StateTreeTransition.DelayDuration")
+
++StructRedirects=(OldName="StateTreeEditorPropertyBinding",NewName="/script/StateTreeModule.StateTreePropertyPathBinding")
++StructRedirects=(OldName="StateTreePropCopyBatch",NewName="/script/StateTreeModule.StateTreePropertyCopyBatch")
diff --git a/Engine/Plugins/Runtime/StateTree/Source/StateTreeEditorModule/Private/Customizations/StateTreeBindingExtension.cpp b/Engine/Plugins/Runtime/StateTree/Source/StateTreeEditorModule/Private/Customizations/StateTreeBindingExtension.cpp
--- a/Engine/Plugins/Runtime/StateTree/Source/StateTreeEditorModule/Private/Customizations/StateTreeBindingExtension.cpp
+++ b/Engine/Plugins/Runtime/StateTree/Source/StateTreeEditorModule/Private/Customizations/StateTreeBindingExtension.cpp
@@ -12,6 +12,7 @@
 #include "StateTreeNodeBase.h"
 #include "Styling/AppStyle.h"
 #include "UObject/EnumProperty.h"
+#include "Widgets/Layout/SBox.h"
 
 #define LOCTEXT_NAMESPACE "StateTreeEditor"
 
@@ -37,32 +38,114 @@ UObject* FindEditorBindingsOwner(UObject* InObject)
 	return Result;
 }
 
-void GetStructPropertyPath(TSharedPtr<IPropertyHandle> InPropertyHandle, FStateTreeEditorPropertyPath& OutPath)
+void MakeStructPropertyPathFromBindingChain(const FGuid StructID, const TArray<FBindingChainElement>& InBindingChain, FStateTreePropertyPath& OutPath)
 {
+	OutPath.Reset();
+	OutPath.SetStructID(StructID);
+	
+	for (const FBindingChainElement& Element : InBindingChain)
+	{
+		if (const FProperty* Property = Element.Field.Get<FProperty>())
+		{
+			OutPath.AddPathSegment(Property->GetFName(), Element.ArrayIndex);
+		}
+		else if (const UFunction* Function = Element.Field.Get<UFunction>())
+		{
+			OutPath.AddPathSegment(Function->GetFName());
+		}
+	}
+}
+
+EStateTreePropertyUsage MakeStructPropertyPathFromPropertyHandle(TSharedPtr<const IPropertyHandle> InPropertyHandle, FStateTreePropertyPath& OutPath)
+{
+	OutPath.Reset();
+
 	FGuid StructID;
-	TArray<FBindingChainElement> BindingChain;
+	TArray<FStateTreePropertyPathSegment> PathSegments;
+	EStateTreePropertyUsage ResultUsage = EStateTreePropertyUsage::Invalid; 
 
-	if (FProperty* CurrentProperty = InPropertyHandle->GetProperty())
+	TSharedPtr<const IPropertyHandle> CurrentPropertyHandle = InPropertyHandle;
+	while (CurrentPropertyHandle.IsValid())
 	{
-		// Keep track of the path.
-		BindingChain.Insert(FBindingChainElement(CurrentProperty, InPropertyHandle->GetIndexInArray()), 0);
-		
-		if (const FString* IDString = InPropertyHandle->GetInstanceMetaData(StateTreeNodeIDName))
+		const FProperty* Property = CurrentPropertyHandle->GetProperty();
+		if (Property)
 		{
-			LexFromString(StructID, **IDString);
+			FStateTreePropertyPathSegment& Segment = PathSegments.InsertDefaulted_GetRef(0); // Traversing from leaf to root, insert in reverse.
+
+			// Store path up to the property which has ID.
+			Segment.SetName(Property->GetFName());
+			Segment.SetArrayIndex(CurrentPropertyHandle->GetIndexInArray());
+
+			// Store type of the object (e.g. for instanced objects or instanced structs).
+			if (const FObjectProperty* ObjectProperty = CastField<FObjectProperty>(Property))
+			{
+				if (ObjectProperty->HasAnyPropertyFlags(CPF_PersistentInstance | CPF_InstancedReference))
+				{
+					const UObject* Object = nullptr;
+					if (CurrentPropertyHandle->GetValue(Object) == FPropertyAccess::Success)
+					{
+						if (Object)
+						{
+							Segment.SetInstanceStruct(Object->GetClass());
+						}
+					}
+				}
+			}
+			else if (const FStructProperty* StructProperty = CastField<FStructProperty>(Property))
+			{
+				if (StructProperty->Struct == TBaseStructure<FInstancedStruct>::Get())
+				{
+					void* Address = nullptr;
+					if (CurrentPropertyHandle->GetValueData(Address) == FPropertyAccess::Success)
+					{
+						if (Address)
+						{
+							FInstancedStruct& Struct = *static_cast<FInstancedStruct*>(Address);
+							Segment.SetInstanceStruct(Struct.GetScriptStruct());
+						}
+					}
+				}
+			}
+
+			// Array access is represented as: "Array, PropertyInArray[Index]", we're traversing from leaf to root, skip the node without index.
+			// Advancing the node before ID test, since the array is on the instance data, the ID will be on the Array node.
+			if (Segment.GetArrayIndex() != INDEX_NONE)
+			{
+				TSharedPtr<const IPropertyHandle> ParentPropertyHandle = CurrentPropertyHandle->GetParentHandle();
+				if (ParentPropertyHandle.IsValid())
+				{
+					const FProperty* ParentProperty = ParentPropertyHandle->GetProperty();
+					if (ParentProperty
+						&& ParentProperty->IsA<FArrayProperty>()
+						&& Property->GetFName() == ParentProperty->GetFName())
+					{
+						CurrentPropertyHandle = ParentPropertyHandle;
+					}
+				}
+			}
+
+			// Bindable property must have node ID
+			if (const FString* IDString = CurrentPropertyHandle->GetInstanceMetaData(UE::StateTree::PropertyBinding::StateTreeNodeIDName))
+			{
+				LexFromString(StructID, **IDString);
+				ResultUsage = UE::StateTree::Compiler::GetUsageFromMetaData(Property);
+				break;
+			}
 		}
+		
+		CurrentPropertyHandle = CurrentPropertyHandle->GetParentHandle();
 	}
-	
-	if (StructID.IsValid())
+
+	if (!StructID.IsValid())
 	{
-		OutPath.StructID = StructID;
-		IPropertyAccessEditor& PropertyAccessEditor = IModularFeatures::Get().GetModularFeature<IPropertyAccessEditor>("PropertyAccessEditor");
-		PropertyAccessEditor.MakeStringPath(BindingChain, OutPath.Path);
+		ResultUsage = EStateTreePropertyUsage::Invalid;
 	}
 	else
 	{
-		OutPath = FStateTreeEditorPropertyPath();
+		OutPath = FStateTreePropertyPath(StructID, PathSegments);
 	}
+
+	return ResultUsage;
 }
 
 FText GetSectionNameFromDataSource(const EStateTreeBindableStructSource Source)
@@ -86,22 +169,54 @@ FText GetSectionNameFromDataSource(const EStateTreeBindableStructSource Source)
 	}
 }
 
+// @todo: there's a similar function in StateTreeNodeDetails.cpp, merge.
+FText GetPropertyTypeText(const FProperty* Property)
+{
+	FEdGraphPinType PinType;
+	const UEdGraphSchema_K2* Schema = GetDefault<UEdGraphSchema_K2>();
+	Schema->ConvertPropertyToPinType(Property, PinType);
+				
+	const FName PinSubCategory = PinType.PinSubCategory;
+	const UObject* PinSubCategoryObject = PinType.PinSubCategoryObject.Get();
+	if (PinSubCategory != UEdGraphSchema_K2::PSC_Bitmask && PinSubCategoryObject)
+	{
+		if (const UField* Field = Cast<const UField>(PinSubCategoryObject))
+		{
+			return Field->GetDisplayNameText();
+		}
+		return FText::FromString(PinSubCategoryObject->GetName());
+	}
 
-FOnStateTreeBindingChanged STATETREEEDITORMODULE_API OnStateTreeBindingChanged;
+	return UEdGraphSchema_K2::GetCategoryText(PinType.PinCategory, NAME_None, true);
+}
+
+FOnStateTreePropertyBindingChanged STATETREEEDITORMODULE_API OnStateTreePropertyBindingChanged;
 
 } // UE::StateTree::PropertyBinding
 
 bool FStateTreeBindingExtension::IsPropertyExtendable(const UClass* InObjectClass, const IPropertyHandle& PropertyHandle) const
 {
-	// Bindable property must have node ID
-	if (PropertyHandle.GetInstanceMetaData(UE::StateTree::PropertyBinding::StateTreeNodeIDName) == nullptr)
+	const FProperty* Property = PropertyHandle.GetProperty();
+	if (Property->HasAnyPropertyFlags(CPF_PersistentInstance | CPF_EditorOnly | CPF_Config | CPF_InstancedReference))
 	{
 		return false;
 	}
+	
+	FStateTreePropertyPath TargetPath;
+	// Figure out the structs we're editing, and property path relative to current property.
+	const EStateTreePropertyUsage Usage = UE::StateTree::PropertyBinding::MakeStructPropertyPathFromPropertyHandle(PropertyHandle.AsShared(), TargetPath);
 
-	// Only inputs and parameters are bindable.
-	const EStateTreePropertyUsage Usage = UE::StateTree::Compiler::GetUsageFromMetaData(PropertyHandle.GetProperty());
-	return Usage == EStateTreePropertyUsage::Input || Usage == EStateTreePropertyUsage::Context || Usage == EStateTreePropertyUsage::Parameter;
+	if (Usage == EStateTreePropertyUsage::Input || Usage == EStateTreePropertyUsage::Context)
+	{
+		// Allow to bind only to the main level on input and context properties.
+		return TargetPath.GetSegments().Num() == 1;
+	}
+	if (Usage == EStateTreePropertyUsage::Parameter)
+	{
+		return true;
+	}
+	
+	return false;  
 }
 
 
@@ -122,7 +237,7 @@ void FStateTreeBindingExtension::ExtendWidgetRow(FDetailWidgetRow& InWidgetRow,
 	TArray<FStateTreeBindableStructDesc> AccessibleStructs;
 
 	// The struct and property where we're binding.
-	FStateTreeEditorPropertyPath TargetPath;
+	FStateTreePropertyPath TargetPath;
 
 	TArray<UObject*> OuterObjects;
 	InPropertyHandle->GetOuterObjects(OuterObjects);
@@ -132,12 +247,12 @@ void FStateTreeBindingExtension::ExtendWidgetRow(FDetailWidgetRow& InWidgetRow,
 		OwnerObject = UE::StateTree::PropertyBinding::FindEditorBindingsOwner(OuterObjects[0]);
 
 		// Figure out the structs we're editing, and property path relative to current property.
-		UE::StateTree::PropertyBinding::GetStructPropertyPath(InPropertyHandle, TargetPath);
+		UE::StateTree::PropertyBinding::MakeStructPropertyPathFromPropertyHandle(InPropertyHandle, TargetPath);
 
 		if (IStateTreeEditorPropertyBindingsOwner* BindingOwner = Cast<IStateTreeEditorPropertyBindingsOwner>(OwnerObject))
 		{
 			EditorBindings = BindingOwner->GetPropertyEditorBindings();
-			BindingOwner->GetAccessibleStructs(TargetPath.StructID, AccessibleStructs);
+			BindingOwner->GetAccessibleStructs(TargetPath.GetStructID(), AccessibleStructs);
 			for (FStateTreeBindableStructDesc& StructDesc : AccessibleStructs)
 			{
 				const UStruct* Struct = StructDesc.Struct;
@@ -148,6 +263,26 @@ void FStateTreeBindingExtension::ExtendWidgetRow(FDetailWidgetRow& InWidgetRow,
 				ContextStruct.Section = UE::StateTree::PropertyBinding::GetSectionNameFromDataSource(StructDesc.DataSource);
 			}
 		}
+
+		// Wrap value widget 
+		if (EditorBindings)
+		{
+			auto IsValueVisible = TAttribute<EVisibility>::Create([TargetPath, EditorBindings]() -> EVisibility
+			{
+				return EditorBindings->HasPropertyBinding(TargetPath) ? EVisibility::Collapsed : EVisibility::Visible;
+			});
+
+			TSharedPtr<SWidget> ValueWidget = InWidgetRow.ValueContent().Widget;
+			InWidgetRow.ValueContent()
+			[
+				SNew(SBox)
+				.Visibility(IsValueVisible)
+				[
+					ValueWidget.ToSharedRef()
+				]
+			];
+		}
+		
 	}
 
 	FProperty* Property = InPropertyHandle->GetProperty();
@@ -198,14 +333,14 @@ void FStateTreeBindingExtension::ExtendWidgetRow(FDetailWidgetRow& InWidgetRow,
 				UE::StateTree::PropertyHelpers::GetStructValue(InPropertyHandle, AnyEnum);
 
 				// If the enum class is not specified, allow to bind to any enum, if the class is specified allow only that enum.
-				if (FByteProperty* ByteProperty = CastField<FByteProperty>(InProperty))
+				if (const FByteProperty* ByteProperty = CastField<FByteProperty>(InProperty))
 				{
 					if (UEnum* Enum = ByteProperty->GetIntPropertyEnum())
 					{
 						bCanBind = bAllowAnyBinding || AnyEnum.Enum == Enum;
 					}
 				}
-				else if (FEnumProperty* EnumProperty = CastField<FEnumProperty>(InProperty))
+				else if (const FEnumProperty* EnumProperty = CastField<FEnumProperty>(InProperty))
 				{
 					bCanBind = bAllowAnyBinding || AnyEnum.Enum == EnumProperty->GetEnum();
 				}
@@ -229,7 +364,7 @@ void FStateTreeBindingExtension::ExtendWidgetRow(FDetailWidgetRow& InWidgetRow,
 			else
 			{
 				// Note: We support type promotion here
-				bCanBind = FStateTreePropertyBindingCompiler::GetPropertyCompatibility(InProperty, InPropertyHandle->GetProperty()) != EPropertyAccessCompatibility::Incompatible;
+				bCanBind = FStateTreePropertyBindings::GetPropertyCompatibility(InProperty, InPropertyHandle->GetProperty()) != EStateTreePropertyAccessCompatibility::Incompatible;
 			}
 
 			return bCanBind;
@@ -274,25 +409,24 @@ void FStateTreeBindingExtension::ExtendWidgetRow(FDetailWidgetRow& InWidgetRow,
 			IPropertyAccessEditor& PropertyAccessEditor = IModularFeatures::Get().GetModularFeature<IPropertyAccessEditor>("PropertyAccessEditor");
 			if (EditorBindings && OwnerObject)
 			{
-				if (TargetPath.IsValid() && InBindingChain.Num() > 0)
+				if (TargetPath.GetStructID().IsValid() && InBindingChain.Num() > 0)
 				{
 					// First item in the binding chain is the index in AccessibleStructs.
-					FStateTreeEditorPropertyPath SourcePath;
 					const int32 SourceStructIndex = InBindingChain[0].ArrayIndex;
-					
+						
 					TArray<FBindingChainElement> SourceBindingChain = InBindingChain;
 					SourceBindingChain.RemoveAt(0); // remove struct index.
 
 					check(SourceStructIndex >= 0 && SourceStructIndex < AccessibleStructs.Num());
 
 					// If SourceBindingChain is empty at this stage, it means that the binding points to the source struct itself.
-					SourcePath.StructID = AccessibleStructs[SourceStructIndex].ID;
-					PropertyAccessEditor.MakeStringPath(SourceBindingChain, SourcePath.Path);
-
+					FStateTreePropertyPath SourcePath;
+					UE::StateTree::PropertyBinding::MakeStructPropertyPathFromBindingChain(AccessibleStructs[SourceStructIndex].ID, SourceBindingChain, SourcePath);
+						
 					OwnerObject->Modify();
 					EditorBindings->AddPropertyBinding(SourcePath, TargetPath);
 
-					UE::StateTree::PropertyBinding::OnStateTreeBindingChanged.Broadcast(SourcePath, TargetPath);
+					UE::StateTree::PropertyBinding::OnStateTreePropertyBindingChanged.Broadcast(SourcePath, TargetPath);
 				}
 			}
 		});
@@ -304,8 +438,8 @@ void FStateTreeBindingExtension::ExtendWidgetRow(FDetailWidgetRow& InWidgetRow,
 				OwnerObject->Modify();
 				EditorBindings->RemovePropertyBindings(TargetPath);
 
-				FStateTreeEditorPropertyPath SourcePath; // Null path
-				UE::StateTree::PropertyBinding::OnStateTreeBindingChanged.Broadcast(SourcePath, TargetPath);
+				const FStateTreePropertyPath SourcePath; // Null path
+				UE::StateTree::PropertyBinding::OnStateTreePropertyBindingChanged.Broadcast(SourcePath, TargetPath);
 			}
 		});
 
@@ -316,34 +450,59 @@ void FStateTreeBindingExtension::ExtendWidgetRow(FDetailWidgetRow& InWidgetRow,
 
 	Args.CurrentBindingText = MakeAttributeLambda([EditorBindings, TargetPath, AccessibleStructs]()
 		{
-			const FText MultipleValues = LOCTEXT("MultipleValues", "Multiple Values");
-			const FText Bind = LOCTEXT("Bind", "Bind");
-			FText CurrentValue = Bind;
+			FText CurrentValue = FText::GetEmpty();
 
 			if (EditorBindings)
 			{
-				if (const FStateTreeEditorPropertyPath* SourcePath = EditorBindings->GetPropertyBindingSource(TargetPath))
+				if (const FStateTreePropertyPath* SourcePath = EditorBindings->GetPropertyBindingSource(TargetPath))
 				{
 					FString PropertyName;
 					for (int32 i = 0; i < AccessibleStructs.Num(); i++)
 					{
-						if (AccessibleStructs[i].ID == SourcePath->StructID)
+						if (AccessibleStructs[i].ID == SourcePath->GetStructID())
 						{
 							PropertyName = AccessibleStructs[i].Name.ToString();
 							break;
 						}
 					}
-					for (const FString& Segment : SourcePath->Path)
+					if (!SourcePath->IsPathEmpty())
 					{
-						PropertyName += TEXT(".") + Segment;
+						PropertyName += TEXT(".") + SourcePath->ToString();
 					}
 					CurrentValue = FText::FromString(PropertyName);
 				}
 			}
-			else
+
+			return CurrentValue;
+		});
+
+	Args.CurrentBindingToolTipText = MakeAttributeLambda([EditorBindings, TargetPath, AccessibleStructs, Property]()
+		{
+			FText CurrentValue = FText::GetEmpty();
+
+			if (EditorBindings)
+			{
+				if (const FStateTreePropertyPath* SourcePath = EditorBindings->GetPropertyBindingSource(TargetPath))
+				{
+					FString PropertyName;
+					for (int32 i = 0; i < AccessibleStructs.Num(); i++)
+					{
+						if (AccessibleStructs[i].ID == SourcePath->GetStructID())
+						{
+							PropertyName = AccessibleStructs[i].Name.ToString();
+							break;
+						}
+					}
+					if (!SourcePath->IsPathEmpty())
+					{
+						PropertyName += TEXT(".") + SourcePath->ToString();
+					}
+					CurrentValue = FText::Format(LOCTEXT("ExistingBindingTooltip", "Property is bound to {0}."), FText::FromString(PropertyName));
+				}
+			}
+			else if (Property != nullptr)
 			{
-				// StateTreeEditorData is not valid if there's multiple objects selected.
-				CurrentValue = MultipleValues;
+				CurrentValue = FText::Format(LOCTEXT("BindTooltip", "Bind value to {0} property."), UE::StateTree::PropertyBinding::GetPropertyTypeText(Property));
 			}
 
 			return CurrentValue;
@@ -368,6 +527,7 @@ void FStateTreeBindingExtension::ExtendWidgetRow(FDetailWidgetRow& InWidgetRow,
 			return BindingColor;
 		});
 
+	Args.BindButtonStyle = &FAppStyle::Get().GetWidgetStyle<FButtonStyle>("HoverHintOnly");
 	Args.bAllowNewBindings = false;
 	Args.bAllowArrayElementBindings = false;
 	Args.bAllowUObjectFunctions = false;
diff --git a/Engine/Plugins/Runtime/StateTree/Source/StateTreeEditorModule/Private/Customizations/StateTreeBindingExtension.h b/Engine/Plugins/Runtime/StateTree/Source/StateTreeEditorModule/Private/Customizations/StateTreeBindingExtension.h
--- a/Engine/Plugins/Runtime/StateTree/Source/StateTreeEditorModule/Private/Customizations/StateTreeBindingExtension.h
+++ b/Engine/Plugins/Runtime/StateTree/Source/StateTreeEditorModule/Private/Customizations/StateTreeBindingExtension.h
@@ -8,9 +8,9 @@ enum class EStateTreePropertyUsage : uint8;
 
 class IPropertyHandle;
 class IDetailLayoutBuilder;
-struct FStateTreeEditorPropertyPath;
 class IPropertyAccessEditor;
 struct FStateTreeEditorPropertyPath;
+struct FStateTreePropertyPath;
 
 namespace UE::StateTree::PropertyBinding
 {
@@ -23,22 +23,17 @@ namespace UE::StateTree::PropertyBinding
 	UObject* FindEditorBindingsOwner(UObject* InObject);
 
 	/**
-	 * Returns property path for a specific property. The property's metadata "StateTreeItemID" is expected to specify the containing struct ID.
+	 * Returns property path for a specific property.
+	 * Walks towards root up until a property with metadata  "StateTreeItemID" is found. 
+	 * The property's metadata "StateTreeItemID" is expected to specify the containing struct ID.
 	 * @param InPropertyHandle Handle to the property to find path for.
 	 * @param OutPath Resulting property path.
+	 * @return usage type extracted from the root-most property.
 	 */
-	void GetStructPropertyPath(TSharedPtr<IPropertyHandle> InPropertyHandle, FStateTreeEditorPropertyPath& OutPath);
-
-	/**
-	 * Parses property usage from a property handle.
-	 * @param InPropertyHandle Handle to the property to find path for.
-	 * @return Parsed property usage.
-	 */
-	EStateTreePropertyUsage ParsePropertyUsage(TSharedPtr<const IPropertyHandle> InPropertyHandle);
-	
-	DECLARE_MULTICAST_DELEGATE_TwoParams(FOnStateTreeBindingChanged, const FStateTreeEditorPropertyPath& /*SourcePath*/, const FStateTreeEditorPropertyPath& /*TargetPath*/);
-	extern STATETREEEDITORMODULE_API FOnStateTreeBindingChanged OnStateTreeBindingChanged;
+	EStateTreePropertyUsage MakeStructPropertyPathFromPropertyHandle(TSharedPtr<const IPropertyHandle> InPropertyHandle, FStateTreePropertyPath& OutPath);
 
+	DECLARE_MULTICAST_DELEGATE_TwoParams(FOnStateTreePropertyBindingChanged, const FStateTreePropertyPath& /*SourcePath*/, const FStateTreePropertyPath& /*TargetPath*/);
+	extern STATETREEEDITORMODULE_API FOnStateTreePropertyBindingChanged OnStateTreePropertyBindingChanged;
 } // UE::StateTree::PropertyBinding
 
 class FStateTreeBindingExtension : public IDetailPropertyExtensionHandler
diff --git a/Engine/Plugins/Runtime/StateTree/Source/StateTreeEditorModule/Private/Customizations/StateTreeEditorDataDetails.cpp b/Engine/Plugins/Runtime/StateTree/Source/StateTreeEditorModule/Private/Customizations/StateTreeEditorDataDetails.cpp
--- a/Engine/Plugins/Runtime/StateTree/Source/StateTreeEditorModule/Private/Customizations/StateTreeEditorDataDetails.cpp
+++ b/Engine/Plugins/Runtime/StateTree/Source/StateTreeEditorModule/Private/Customizations/StateTreeEditorDataDetails.cpp
@@ -103,7 +103,7 @@ void FStateTreeEditorDataDetails::CustomizeDetails(IDetailLayoutBuilder& DetailB
 						.BorderImage(new FSlateRoundedBoxBrush(FStyleColors::Hover, 6.f))
 						[
 							SNew(STextBlock)
-							.TextStyle(FStateTreeEditorStyle::Get(), "StateTree.Node.Operand")
+							.TextStyle(FStateTreeEditorStyle::Get(), "StateTree.Param.Label")
 							.ColorAndOpacity(FStyleColors::Foreground)
 							.Text(LOCTEXT("LabelContext", "CONTEXT"))
 							.ToolTipText(LOCTEXT("ContextSourceTooltip", "This is Context Object, it passed in from where the StateTree is being used."))
diff --git a/Engine/Plugins/Runtime/StateTree/Source/StateTreeEditorModule/Private/Customizations/StateTreeEditorNodeDetails.cpp b/Engine/Plugins/Runtime/StateTree/Source/StateTreeEditorModule/Private/Customizations/StateTreeEditorNodeDetails.cpp
--- a/Engine/Plugins/Runtime/StateTree/Source/StateTreeEditorModule/Private/Customizations/StateTreeEditorNodeDetails.cpp
+++ b/Engine/Plugins/Runtime/StateTree/Source/StateTreeEditorModule/Private/Customizations/StateTreeEditorNodeDetails.cpp
@@ -28,6 +28,7 @@
 #include "HAL/PlatformApplicationMisc.h"
 #include "ScopedTransaction.h"
 #include "Widgets/Input/SSearchBox.h"
+#include "EditorFontGlyphs.h"
 
 #define LOCTEXT_NAMESPACE "StateTreeEditor"
 
@@ -110,7 +111,7 @@ namespace UE::StateTreeEditor::Internal
 			// Pass the node ID to binding extension. Since the properties are added using AddChildStructure(), we break the hierarchy and cannot access parent.
 			ChildPropHandle->SetInstanceMetaData(UE::StateTree::PropertyBinding::StateTreeNodeIDName, LexToString(ID));
 
-			FStateTreeEditorPropertyPath Path(ID, *Property->GetFName().ToString());
+			FStateTreePropertyPath Path(ID, *Property->GetFName().ToString());
 			TSharedPtr<SWidget> NameWidget;
 			TSharedPtr<SWidget> ValueWidget;
 			FDetailWidgetRow Row;
@@ -121,8 +122,13 @@ namespace UE::StateTreeEditor::Internal
 					return EditorPropBindings->HasPropertyBinding(Path) ? EVisibility::Collapsed : EVisibility::Visible;
 				});
 
+			bool bShowChildren = true;
+			
 			if (Usage == EStateTreePropertyUsage::Input || Usage == EStateTreePropertyUsage::Output || Usage == EStateTreePropertyUsage::Context)
 			{
+				// Do not show children for input, output and context.
+				bShowChildren = false;
+				
 				FEdGraphPinType PinType;
 				const UEdGraphSchema_K2* Schema = GetDefault<UEdGraphSchema_K2>();
 				Schema->ConvertPropertyToPinType(Property, PinType);
@@ -181,7 +187,7 @@ namespace UE::StateTreeEditor::Internal
 				}
 				
 				ChildRow
-					.CustomWidget(/*bShowChildren*/false)
+					.CustomWidget(bShowChildren)
 					.NameContent()
 					[
 						SNew(SHorizontalBox)
@@ -204,7 +210,7 @@ namespace UE::StateTreeEditor::Internal
 							.Visibility(Label.IsEmpty() ? EVisibility::Collapsed : EVisibility::Visible)
 							[
 								SNew(STextBlock)
-								.TextStyle(FStateTreeEditorStyle::Get(), "StateTree.Node.Operand")
+								.TextStyle(FStateTreeEditorStyle::Get(), "StateTree.Param.Label")
 								.ColorAndOpacity(FStyleColors::Foreground)
 								.Text(Label)
 								.ToolTipText(LabelToolTip)
@@ -239,23 +245,6 @@ namespace UE::StateTreeEditor::Internal
 						]
 					];
 			}
-			else
-			{
-				ChildRow
-					.CustomWidget(/*bShowChildren*/true)
-					.NameContent()
-					[
-						NameWidget.ToSharedRef()
-					]
-					.ValueContent()
-					[
-						SNew(SBox)
-						.Visibility(IsValueVisible)
-						[
-							ValueWidget.ToSharedRef()
-						]
-					];
-			}
 		}
 	}
 } // UE::StateTreeEditor::Internal
@@ -292,7 +281,7 @@ TSharedRef<IPropertyTypeCustomization> FStateTreeEditorNodeDetails::MakeInstance
 
 FStateTreeEditorNodeDetails::~FStateTreeEditorNodeDetails()
 {
-	UE::StateTree::PropertyBinding::OnStateTreeBindingChanged.Remove(OnBindingChangedHandle);
+	UE::StateTree::PropertyBinding::OnStateTreePropertyBindingChanged.Remove(OnBindingChangedHandle);
 }
 
 void FStateTreeEditorNodeDetails::CustomizeHeader(TSharedRef<class IPropertyHandle> StructPropertyHandle, class FDetailWidgetRow& HeaderRow, IPropertyTypeCustomizationUtils& StructCustomizationUtils)
@@ -329,7 +318,7 @@ void FStateTreeEditorNodeDetails::CustomizeHeader(TSharedRef<class IPropertyHand
 	const FResetToDefaultOverride ResetOverride = FResetToDefaultOverride::Create(IsResetVisible, ResetHandler);
 
 	UE::StateTree::Delegates::OnIdentifierChanged.AddSP(this, &FStateTreeEditorNodeDetails::OnIdentifierChanged);
-	OnBindingChangedHandle = UE::StateTree::PropertyBinding::OnStateTreeBindingChanged.AddRaw(this, &FStateTreeEditorNodeDetails::OnBindingChanged);
+	OnBindingChangedHandle = UE::StateTree::PropertyBinding::OnStateTreePropertyBindingChanged.AddRaw(this, &FStateTreeEditorNodeDetails::OnBindingChanged);
 
 	FindOuterObjects();
 
@@ -403,19 +392,52 @@ void FStateTreeEditorNodeDetails::CustomizeHeader(TSharedRef<class IPropertyHand
 					.Text(this, &FStateTreeEditorNodeDetails::GetOpenParens)
 					.Visibility(this, &FStateTreeEditorNodeDetails::IsConditionVisible)
 				]
-				// Name (Eval/Task)
+				// Name & type selection
 				+ SHorizontalBox::Slot()
 				.AutoWidth()
+				.Padding(FMargin(FMargin(4.0f, 0.0f, 0.0f, 0.0f)))
 				.VAlign(VAlign_Center)
 				[
-					SNew(SEditableTextBox)
-					.IsEnabled(TAttribute<bool>(this, &FStateTreeEditorNodeDetails::IsNameEnabled))
-					.Text(this, &FStateTreeEditorNodeDetails::GetName)
-					.OnTextCommitted(this, &FStateTreeEditorNodeDetails::OnNameCommitted)
-					.SelectAllTextWhenFocused(true)
-					.RevertTextOnEscape(true)
-					.Style(FStateTreeEditorStyle::Get(), "StateTree.Node.Name")
-					.Visibility(this, &FStateTreeEditorNodeDetails::IsNameVisible)
+					SAssignNew(ComboButton, SComboButton)
+					.OnGetMenuContent(this, &FStateTreeEditorNodeDetails::GeneratePicker)
+					.ToolTipText(this, &FStateTreeEditorNodeDetails::GetDisplayValueString)
+					.ContentPadding(FMargin(2, 2, 2, 2))
+					.ButtonContent()
+					[
+						SNew(SHorizontalBox)
+						+ SHorizontalBox::Slot()
+						.AutoWidth()
+						.VAlign(VAlign_Center)
+						.Padding(0, 0, 4, 0)
+						[
+							SNew(STextBlock)
+							.Text(FEditorFontGlyphs::Paper_Plane)
+							.TextStyle(FStateTreeEditorStyle::Get(), "StateTree.DetailsIcon")
+							.Visibility_Lambda([this]()
+							{
+								// Only show on tasks
+								const UScriptStruct* ScriptStruct = nullptr;
+								if (const FStateTreeEditorNode* Node = GetCommonNode())
+								{
+									ScriptStruct = Node->Node.GetScriptStruct();
+								}
+								return ScriptStruct != nullptr && ScriptStruct->IsChildOf(FStateTreeTaskBase::StaticStruct()) ? EVisibility::Visible : EVisibility::Collapsed;
+							})
+						]
+						+ SHorizontalBox::Slot()
+						.VAlign(VAlign_Center)
+						.Padding(0, 0, 8, 0)
+						[
+							SNew(SEditableText)
+							.IsEnabled(TAttribute<bool>(this, &FStateTreeEditorNodeDetails::IsNameEnabled))
+							.Text(this, &FStateTreeEditorNodeDetails::GetName)
+							.OnTextCommitted(this, &FStateTreeEditorNodeDetails::OnNameCommitted)
+							.SelectAllTextWhenFocused(true)
+							.RevertTextOnEscape(true)
+							.Style(FStateTreeEditorStyle::Get(), "StateTree.Node.Name")
+							.Visibility(this, &FStateTreeEditorNodeDetails::IsNameVisible)
+						]
+					]
 				]
 				// Close parens
 				+ SHorizontalBox::Slot()
@@ -427,42 +449,13 @@ void FStateTreeEditorNodeDetails::CustomizeHeader(TSharedRef<class IPropertyHand
 					.Text(this, &FStateTreeEditorNodeDetails::GetCloseParens)
 					.Visibility(this, &FStateTreeEditorNodeDetails::IsConditionVisible)
 				]
-			]
-			// Class picker
-			+ SHorizontalBox::Slot()
-			.FillWidth(1.0f)
-			.Padding(FMargin(FMargin(4.0f, 0.0f, 0.0f, 0.0f)))
-			.VAlign(VAlign_Center)
-			[
-				SAssignNew(ComboButton, SComboButton)
-				.OnGetMenuContent(this, &FStateTreeEditorNodeDetails::GeneratePicker)
-				.ToolTipText(this, &FStateTreeEditorNodeDetails::GetDisplayValueString)
-				.ContentPadding(0.f)
-				.ButtonContent()
+
+				+ SHorizontalBox::Slot()
+				.AutoWidth()
 				[
-					SNew(SHorizontalBox)
-					+ SHorizontalBox::Slot()
-					.AutoWidth()
-					.VAlign(VAlign_Center)
-					.Padding(0.0f, 0.0f, 4.0f, 0.0f)
-					[
-						SNew(SImage)
-						.Image(this, &FStateTreeEditorNodeDetails::GetDisplayValueIcon)
-					]
-					+ SHorizontalBox::Slot()
-					.VAlign(VAlign_Center)
-					[
-						SNew(STextBlock)
-						.Text(this, &FStateTreeEditorNodeDetails::GetDisplayValueString)
-						.Font(IDetailLayoutBuilder::GetDetailFont())
-					]
+					StructPropertyHandle->CreateDefaultPropertyButtonWidgets()
 				]
 			]
-			+ SHorizontalBox::Slot()
-			.AutoWidth()
-			[
-				StructPropertyHandle->CreateDefaultPropertyButtonWidgets()
-			]
 		]
 		.OverrideResetToDefault(ResetOverride)
 		.CopyAction(FUIAction(FExecuteAction::CreateSP(this, &FStateTreeEditorNodeDetails::OnCopyNode)))
@@ -611,6 +604,7 @@ void FStateTreeEditorNodeDetails::CustomizeChildren(TSharedRef<class IPropertyHa
 	FGuid ID;
 	UE::StateTree::PropertyHelpers::GetStructValue<FGuid>(IDProperty, ID);
 
+	// Node
 	TSharedRef<FBindableNodeInstanceDetails> NodeDetails = MakeShareable(new FBindableNodeInstanceDetails(NodeProperty, FGuid(), EditorData));
 	StructBuilder.AddCustomBuilder(NodeDetails);
 
@@ -690,7 +684,7 @@ void FStateTreeEditorNodeDetails::OnIdentifierChanged(const UStateTree& InStateT
 	}
 }
 
-void FStateTreeEditorNodeDetails::OnBindingChanged(const FStateTreeEditorPropertyPath& SourcePath, const FStateTreeEditorPropertyPath& TargetPath)
+void FStateTreeEditorNodeDetails::OnBindingChanged(const FStateTreePropertyPath& SourcePath, const FStateTreePropertyPath& TargetPath)
 {
 	check(StructProperty);
 
@@ -711,7 +705,7 @@ void FStateTreeEditorNodeDetails::OnBindingChanged(const FStateTreeEditorPropert
 		UObject* OuterObject = OuterObjects[i]; // Immediate outer, i.e StateTreeState
 		if (Node != nullptr && EditorData != nullptr && Node->Node.IsValid() && Node->Instance.IsValid())
 		{
-			if (Node->ID == TargetPath.StructID)
+			if (Node->ID == TargetPath.GetStructID())
 			{
 				if (FStateTreeConditionBase* Condition = Node->Node.GetMutablePtr<FStateTreeConditionBase>())
 				{
diff --git a/Engine/Plugins/Runtime/StateTree/Source/StateTreeEditorModule/Private/Customizations/StateTreeEditorNodeDetails.h b/Engine/Plugins/Runtime/StateTree/Source/StateTreeEditorModule/Private/Customizations/StateTreeEditorNodeDetails.h
--- a/Engine/Plugins/Runtime/StateTree/Source/StateTreeEditorModule/Private/Customizations/StateTreeEditorNodeDetails.h
+++ b/Engine/Plugins/Runtime/StateTree/Source/StateTreeEditorModule/Private/Customizations/StateTreeEditorNodeDetails.h
@@ -15,7 +15,7 @@ class UStateTree;
 class UStateTreeState;
 class UStateTreeEditorData;
 struct EVisibility;
-struct FStateTreeEditorPropertyPath;
+struct FStateTreePropertyPath;
 enum class EStateTreeConditionOperand : uint8;
 
 /**
@@ -79,7 +79,7 @@ private:
 	void OnClassPicked(const UClass* InClass) const;
 
 	void OnIdentifierChanged(const UStateTree& StateTree);
-	void OnBindingChanged(const FStateTreeEditorPropertyPath& SourcePath, const FStateTreeEditorPropertyPath& TargetPath);
+	void OnBindingChanged(const FStateTreePropertyPath& SourcePath, const FStateTreePropertyPath& TargetPath);
 	void FindOuterObjects();
 
 	// Stores a category path segment, or a node type.
diff --git a/Engine/Plugins/Runtime/StateTree/Source/StateTreeEditorModule/Private/Customizations/StateTreeStateParametersDetails.cpp b/Engine/Plugins/Runtime/StateTree/Source/StateTreeEditorModule/Private/Customizations/StateTreeStateParametersDetails.cpp
--- a/Engine/Plugins/Runtime/StateTree/Source/StateTreeEditorModule/Private/Customizations/StateTreeStateParametersDetails.cpp
+++ b/Engine/Plugins/Runtime/StateTree/Source/StateTreeEditorModule/Private/Customizations/StateTreeStateParametersDetails.cpp
@@ -49,7 +49,7 @@ public:
 		// Pass the node ID to binding extension. Since the properties are added using AddChildStructure(), we break the hierarchy and cannot access parent.
 		ChildPropHandle->SetInstanceMetaData(UE::StateTree::PropertyBinding::StateTreeNodeIDName, LexToString(ID));
 
-		FStateTreeEditorPropertyPath Path(ID, *Property->GetFName().ToString());
+		FStateTreePropertyPath Path(ID, *Property->GetFName().ToString());
 		TSharedPtr<SWidget> NameWidget;
 		TSharedPtr<SWidget> ValueWidget;
 		FDetailWidgetRow Row;
diff --git a/Engine/Plugins/Runtime/StateTree/Source/StateTreeEditorModule/Private/StateTreeCompiler.cpp b/Engine/Plugins/Runtime/StateTree/Source/StateTreeEditorModule/Private/StateTreeCompiler.cpp
--- a/Engine/Plugins/Runtime/StateTree/Source/StateTreeEditorModule/Private/StateTreeCompiler.cpp
+++ b/Engine/Plugins/Runtime/StateTree/Source/StateTreeEditorModule/Private/StateTreeCompiler.cpp
@@ -174,6 +174,7 @@ bool FStateTreeCompiler::Compile(UStateTree& InStateTree)
 	if (!StateTree->Link())
 	{
 		StateTree->ResetCompiled();
+		Log.Reportf(EMessageSeverity::Error, TEXT("Unexpected failure to link the StateTree asset. See log for more info."));
 		return false;
 	}
 	
@@ -451,8 +452,8 @@ bool FStateTreeCompiler::CreateStateTasksAndParameters()
 					};
 
 					// Check that the bindings for this struct are still all valid.
-					TArray<FStateTreeEditorPropertyBinding> Bindings;
-					if (!GetAndValidateBindings(LinkedParamsDesc, Bindings))
+					TArray<FStateTreePropertyPathBinding> Bindings;
+					if (!GetAndValidateBindings(LinkedParamsDesc, FStateTreeDataView(SourceState->Parameters.Parameters.GetMutableValue()), Bindings))
 					{
 						return false;
 					}
@@ -759,6 +760,8 @@ bool FStateTreeCompiler::CreateCondition(UStateTreeState& State, const FStateTre
 
 	Cond.Operand = Operand;
 	Cond.DeltaIndent = DeltaIndent;
+
+	FStateTreeDataView InstanceDataView;
 	
 	if (CondNode.Instance.IsValid())
 	{
@@ -776,11 +779,7 @@ bool FStateTreeCompiler::CreateCondition(UStateTreeState& State, const FStateTre
 		}
 		Cond.InstanceIndex = FStateTreeIndex16(InstanceIndex);
 		Cond.bInstanceIsObject = false;
-
-		if (!CompileAndValidateNode(&State, StructDesc, Node, FStateTreeDataView(SharedInstanceStructs[InstanceIndex])))
-		{
-			return false;
-		}
+		InstanceDataView = FStateTreeDataView(SharedInstanceStructs[InstanceIndex]);
 	}
 	else
 	{
@@ -801,19 +800,20 @@ bool FStateTreeCompiler::CreateCondition(UStateTreeState& State, const FStateTre
 		}
 		Cond.InstanceIndex = FStateTreeIndex16(InstanceIndex);
 		Cond.bInstanceIsObject = true;
+		InstanceDataView = FStateTreeDataView(Instance);
+	}
 
-		if (!CompileAndValidateNode(&State, StructDesc, Node, FStateTreeDataView(Instance)))
-		{
-			return false;
-		}
+	if (!CompileAndValidateNode(&State, StructDesc, Node, InstanceDataView))
+	{
+		return false;
 	}
 
 	// Mark the struct as binding source.
 	const int32 SourceStructIndex = BindingsCompiler.AddSourceStruct(StructDesc);
 
 	// Check that the bindings for this struct are still all valid.
-	TArray<FStateTreeEditorPropertyBinding> Bindings;
-	if (!GetAndValidateBindings(StructDesc, Bindings))
+	TArray<FStateTreePropertyPathBinding> Bindings;
+	if (!GetAndValidateBindings(StructDesc, InstanceDataView, Bindings))
 	{
 		return false;
 	}
@@ -842,7 +842,7 @@ bool FStateTreeCompiler::CreateCondition(UStateTreeState& State, const FStateTre
 	return true;
 }
 
-bool FStateTreeCompiler::CompileAndValidateNode(const UStateTreeState* SourceState, const FStateTreeBindableStructDesc& NodeDesc, FStructView NodeView, const FStateTreeDataView InstanceData) const
+bool FStateTreeCompiler::CompileAndValidateNode(const UStateTreeState* SourceState, const FStateTreeBindableStructDesc& NodeDesc, FStructView NodeView, const FStateTreeDataView InstanceData)
 {
 	if (!NodeView.IsValid())
 	{
@@ -924,6 +924,7 @@ bool FStateTreeCompiler::CreateTask(UStateTreeState* State, const FStateTreeEdit
 	// Copy the task
 	FInstancedStruct& Node = Nodes.Add_GetRef(TaskNode.Node);
 	FStateTreeTaskBase& Task = Node.GetMutable<FStateTreeTaskBase>();
+	FStateTreeDataView InstanceDataView;
 
 	if (TaskNode.Instance.IsValid())
 	{
@@ -941,11 +942,7 @@ bool FStateTreeCompiler::CreateTask(UStateTreeState* State, const FStateTreeEdit
 		}
 		Task.InstanceIndex = FStateTreeIndex16(InstanceIndex);
 		Task.bInstanceIsObject = false;
-
-		if (!CompileAndValidateNode(State, StructDesc, Node,  FStateTreeDataView(InstanceStructs[InstanceIndex])))
-		{
-			return false;
-		}
+		InstanceDataView = FStateTreeDataView(InstanceStructs[InstanceIndex]);
 	}
 	else
 	{
@@ -966,19 +963,20 @@ bool FStateTreeCompiler::CreateTask(UStateTreeState* State, const FStateTreeEdit
 		}
 		Task.InstanceIndex = FStateTreeIndex16(InstanceIndex);
 		Task.bInstanceIsObject = true;
+		InstanceDataView = FStateTreeDataView(Instance);
+	}
 
-		if (!CompileAndValidateNode(State, StructDesc, Node,  FStateTreeDataView(Instance)))
-		{
-			return false;
-		}
+	if (!CompileAndValidateNode(State, StructDesc, Node,  InstanceDataView))
+	{
+		return false;
 	}
 
 	// Mark the instance as binding source.
 	const int32 SourceStructIndex = BindingsCompiler.AddSourceStruct(StructDesc);
-
+	
 	// Check that the bindings for this struct are still all valid.
-	TArray<FStateTreeEditorPropertyBinding> Bindings;
-	if (!GetAndValidateBindings(StructDesc, Bindings))
+	TArray<FStateTreePropertyPathBinding> Bindings;
+	if (!GetAndValidateBindings(StructDesc, InstanceDataView, Bindings))
 	{
 		return false;
 	}
@@ -1037,7 +1035,8 @@ bool FStateTreeCompiler::CreateEvaluator(const FStateTreeEditorNode& EvalNode)
 	// Copy the evaluator
 	FInstancedStruct& Node = Nodes.Add_GetRef(EvalNode.Node);
 	FStateTreeEvaluatorBase& Eval = Node.GetMutable<FStateTreeEvaluatorBase>();
-
+	FStateTreeDataView InstanceDataView;
+	
 	if (EvalNode.Instance.IsValid())
 	{
 		// Struct Instance
@@ -1054,11 +1053,7 @@ bool FStateTreeCompiler::CreateEvaluator(const FStateTreeEditorNode& EvalNode)
 		}
 		Eval.InstanceIndex = FStateTreeIndex16(InstanceIndex);
 		Eval.bInstanceIsObject = false;
-
-		if (!CompileAndValidateNode(nullptr, StructDesc, Node,  FStateTreeDataView(InstanceStructs[InstanceIndex])))
-		{
-			return false;
-		}
+		InstanceDataView = FStateTreeDataView(InstanceStructs[InstanceIndex]);
 	}
 	else
 	{
@@ -1079,19 +1074,20 @@ bool FStateTreeCompiler::CreateEvaluator(const FStateTreeEditorNode& EvalNode)
 		}
 		Eval.InstanceIndex = FStateTreeIndex16(InstanceIndex);
 		Eval.bInstanceIsObject = true;
+		InstanceDataView = FStateTreeDataView(Instance);
+	}
 
-		if (!CompileAndValidateNode(nullptr, StructDesc, Node,  FStateTreeDataView(Instance)))
-		{
-			return false;
-		}
+	if (!CompileAndValidateNode(nullptr, StructDesc, Node,  InstanceDataView))
+	{
+		return false;
 	}
-		
+
 	// Mark the instance as binding source.
 	const int32 SourceStructIndex = BindingsCompiler.AddSourceStruct(StructDesc);
 
 	// Check that the bindings for this struct are still all valid.
-	TArray<FStateTreeEditorPropertyBinding> Bindings;
-	if (!GetAndValidateBindings(StructDesc, Bindings))
+	TArray<FStateTreePropertyPathBinding> Bindings;
+	if (!GetAndValidateBindings(StructDesc, InstanceDataView, Bindings))
 	{
 		return false;
 	}
@@ -1120,16 +1116,17 @@ bool FStateTreeCompiler::CreateEvaluator(const FStateTreeEditorNode& EvalNode)
 	return true;
 }
 
-bool FStateTreeCompiler::IsPropertyAnyEnum(const FStateTreeBindableStructDesc& Struct, FStateTreeEditorPropertyPath Path) const
+bool FStateTreeCompiler::IsPropertyAnyEnum(const FStateTreeBindableStructDesc& Struct, FStateTreePropertyPath Path) const
 {
 	bool bIsAnyEnum = false;
-	TArray<FStateTreePropertySegment> Segments;
-	const FProperty* LeafProperty = nullptr;
-	int32 LeafArrayIndex = INDEX_NONE;
-	const bool bResolved = FStateTreePropertyBindingCompiler::ResolvePropertyPath(Struct, Path, Segments, LeafProperty, LeafArrayIndex);
-	if (bResolved && LeafProperty)
+
+	TArray<FStateTreePropertyPathIndirection> Indirection;
+	const bool bResolved = Path.ResolveIndirections(Struct.Struct, Indirection);
+	
+	if (bResolved && Indirection.Num() > 0)
 	{
-		if (const FProperty* OwnerProperty = LeafProperty->GetOwnerProperty())
+		check(Indirection.Last().GetProperty());
+		if (const FProperty* OwnerProperty = Indirection.Last().GetProperty()->GetOwnerProperty())
 		{
 			if (const FStructProperty* OwnerStructProperty = CastField<FStructProperty>(OwnerProperty))
 			{
@@ -1140,18 +1137,18 @@ bool FStateTreeCompiler::IsPropertyAnyEnum(const FStateTreeBindableStructDesc& S
 	return bIsAnyEnum;
 }
 
-bool FStateTreeCompiler::ValidateStructRef(const FStateTreeBindableStructDesc& SourceStruct, FStateTreeEditorPropertyPath SourcePath,
-											const FStateTreeBindableStructDesc& TargetStruct, FStateTreeEditorPropertyPath TargetPath) const
+bool FStateTreeCompiler::ValidateStructRef(const FStateTreeBindableStructDesc& SourceStruct, FStateTreePropertyPath SourcePath,
+											const FStateTreeBindableStructDesc& TargetStruct, FStateTreePropertyPath TargetPath) const
 {
-	TArray<FStateTreePropertySegment> Segments;
-
-	const FProperty* TargetLeafProperty = nullptr;
-	int32 TargetLeafArrayIndex = INDEX_NONE;
-	if (FStateTreePropertyBindingCompiler::ResolvePropertyPath(TargetStruct, TargetPath, Segments, TargetLeafProperty, TargetLeafArrayIndex) == false)
+	FString ResolveError;
+	TArray<FStateTreePropertyPathIndirection> TargetIndirection;
+	if (!TargetPath.ResolveIndirections(TargetStruct.Struct, TargetIndirection, &ResolveError))
 	{
 		// This will later be reported by the bindings compiler.
-		return true;
+		Log.Reportf(EMessageSeverity::Error, TargetStruct, TEXT("Failed to resolve binding path in %s: %s"), *TargetStruct.Name.ToString(), *ResolveError);
+		return false;
 	}
+	const FProperty* TargetLeafProperty = TargetIndirection.Num() > 0 ? TargetIndirection.Last().GetProperty() : nullptr;
 
 	// Early out if the target is not FStateTreeStructRef.
 	const FStructProperty* TargetStructProperty = CastField<FStructProperty>(TargetLeafProperty);
@@ -1169,14 +1166,15 @@ bool FStateTreeCompiler::ValidateStructRef(const FStateTreeBindableStructDesc& S
 				*TargetBaseStructName, *TargetStruct.Name.ToString(), *TargetPath.ToString());
 		return false;
 	}
-	
-	const FProperty* SourceLeafProperty = nullptr;
-	int32 SourceLeafArrayIndex = INDEX_NONE;
-	if (FStateTreePropertyBindingCompiler::ResolvePropertyPath(SourceStruct, SourcePath, Segments, SourceLeafProperty, SourceLeafArrayIndex) == false)
+
+	TArray<FStateTreePropertyPathIndirection> SourceIndirection;
+	if (!SourcePath.ResolveIndirections(SourceStruct.Struct, SourceIndirection, &ResolveError))
 	{
 		// This will later be reported by the bindings compiler.
-		return true;
+		Log.Reportf(EMessageSeverity::Error, SourceStruct, TEXT("Failed to resolve binding path in %s: %s"), *SourceStruct.Name.ToString(), *ResolveError);
+		return false;
 	}
+	const FProperty* SourceLeafProperty = SourceIndirection.Num() > 0 ? SourceIndirection.Last().GetProperty() : nullptr;
 
 	// Exit if the source is not a struct property.
 	const FStructProperty* SourceStructProperty = CastField<FStructProperty>(SourceLeafProperty);
@@ -1224,7 +1222,30 @@ bool FStateTreeCompiler::ValidateStructRef(const FStateTreeBindableStructDesc& S
 	return true;
 }
 
-bool FStateTreeCompiler::GetAndValidateBindings(const FStateTreeBindableStructDesc& TargetStruct, TArray<FStateTreeEditorPropertyBinding>& OutBindings) const
+FStateTreeDataView FStateTreeCompiler::GetBindingSourceValue(const int32 SourceIndex)
+{
+	for (const FInstancedStruct& Node : Nodes)
+	{
+		if (const FStateTreeNodeBase* NodeBase = Node.GetPtr<FStateTreeNodeBase>())
+		{
+			if (NodeBase->DataViewIndex.Get() == SourceIndex)
+			{
+				if (NodeBase->bInstanceIsObject)
+				{
+					return FStateTreeDataView(InstanceObjects[NodeBase->InstanceIndex.Get()]);
+				}
+				else
+				{
+					return FStateTreeDataView(InstanceStructs[NodeBase->InstanceIndex.Get()]);
+				}
+			}
+		}
+	}
+
+	return {};
+}
+
+bool FStateTreeCompiler::GetAndValidateBindings(const FStateTreeBindableStructDesc& TargetStruct, FStateTreeDataView TargetValue, TArray<FStateTreePropertyPathBinding>& OutBindings) const
 {
 	check(EditorData);
 	
@@ -1238,28 +1259,48 @@ bool FStateTreeCompiler::GetAndValidateBindings(const FStateTreeBindableStructDe
 	
 	OutBindings.Reset();
 	
-	for (const FStateTreeEditorPropertyBinding& Binding : EditorData->EditorBindings.GetBindings())
+	for (FStateTreePropertyPathBinding& Binding : EditorData->EditorBindings.GetMutableBindings())
 	{
-		if (Binding.TargetPath.StructID != TargetStruct.ID)
+		if (Binding.GetTargetPath().GetStructID() != TargetStruct.ID)
 		{
 			continue;
 		}
 
 		// Source must be one of the source structs we have discovered in the tree.
-		const FGuid SourceStructID = Binding.SourcePath.StructID;
+		const FGuid SourceStructID = Binding.GetSourcePath().GetStructID();
 		const int32 SourceStructIdx = BindingsCompiler.GetSourceStructIndexByID(SourceStructID);
 		if (SourceStructIdx == INDEX_NONE)
 		{
 			Log.Reportf(EMessageSeverity::Error, TargetStruct,
 						TEXT("Failed to find binding source property '%s' for target '%s:%s'."),
-						*Binding.SourcePath.ToString(), *TargetStruct.Name.ToString(), *Binding.TargetPath.ToString());
+						*Binding.GetSourcePath().ToString(), *TargetStruct.Name.ToString(), *Binding.GetTargetPath().ToString());
 			return false;
 		}
 		const FStateTreeBindableStructDesc& SourceStruct = BindingsCompiler.GetSourceStructDesc(SourceStructIdx);
 
+		// Source path should not have any instanced indirections.
+		if (Binding.GetSourcePath().HasAnyInstancedIndirection())
+		{
+			Log.Reportf(EMessageSeverity::Error, TargetStruct,
+						TEXT("Malformed source property path for binding source property '%s' for target '%s:%s'."),
+						*Binding.GetSourcePath().ToString(), *TargetStruct.Name.ToString(), *Binding.GetTargetPath().ToString());
+			return false;
+		}
+
+		// Update path instance types from latest data. E.g. binding may have been created for instanced object of type FooB, and changed to FooA.
+ 		// @todo: not liking how this mutates the Binding.TargetPath, but currently we dont track well the instanced object changes.
+
+		if (!Binding.GetMutableTargetPath().UpdateInstanceStructsFromValue(TargetValue))
+		{
+			Log.Reportf(EMessageSeverity::Error, TargetStruct,
+						TEXT("Malformed target property path for binding source property '%s' for target '%s:%s'."),
+						*Binding.GetSourcePath().ToString(), *TargetStruct.Name.ToString(), *Binding.GetTargetPath().ToString());
+			return false;
+		}
+		
 		// Source must be accessible by the target struct via all execution paths.
 		TArray<FStateTreeBindableStructDesc> AccessibleStructs;
-		EditorData->GetAccessibleStructs(Binding.TargetPath.StructID, AccessibleStructs);
+		EditorData->GetAccessibleStructs(Binding.GetTargetPath().GetStructID(), AccessibleStructs);
 
 		const bool bSourceAccessible = AccessibleStructs.ContainsByPredicate([SourceStructID](const FStateTreeBindableStructDesc& Structs)
 			{
@@ -1270,27 +1311,27 @@ bool FStateTreeCompiler::GetAndValidateBindings(const FStateTreeBindableStructDe
 		{
 			Log.Reportf(EMessageSeverity::Error, TargetStruct,
 						TEXT("Property '%s:%s' cannot be bound to '%s:%s', because the binding source '%s' is not updated before '%s' in the tree."),
-						*SourceStruct.Name.ToString(), *Binding.SourcePath.ToString(),
-						*TargetStruct.Name.ToString(), *Binding.TargetPath.ToString(),
+						*SourceStruct.Name.ToString(), *Binding.GetSourcePath().ToString(),
+						*TargetStruct.Name.ToString(), *Binding.GetTargetPath().ToString(),
 						*SourceStruct.Name.ToString(), *TargetStruct.Name.ToString());
 			return false;
 		}
-		
+
 		// Special case fo AnyEnum. StateTreeBindingExtension allows AnyEnums to bind to other enum types.
 		// The actual copy will be done via potential type promotion copy, into the value property inside the AnyEnum.
 		// We amend the paths here to point to the 'Value' property.
-		const bool bSourceIsAnyEnum = IsPropertyAnyEnum(SourceStruct, Binding.SourcePath);
-		const bool bTargetIsAnyEnum = IsPropertyAnyEnum(TargetStruct, Binding.TargetPath);
+		const bool bSourceIsAnyEnum = IsPropertyAnyEnum(SourceStruct, Binding.GetSourcePath());
+		const bool bTargetIsAnyEnum = IsPropertyAnyEnum(TargetStruct, Binding.GetTargetPath());
 		if (bSourceIsAnyEnum || bTargetIsAnyEnum)
 		{
-			FStateTreeEditorPropertyBinding ModifiedBinding(Binding);
+			FStateTreePropertyPathBinding ModifiedBinding(Binding);
 			if (bSourceIsAnyEnum)
 			{
-				ModifiedBinding.SourcePath.Path.Add(GET_MEMBER_NAME_STRING_CHECKED(FStateTreeAnyEnum, Value));
+				ModifiedBinding.GetMutableSourcePath().AddPathSegment(GET_MEMBER_NAME_STRING_CHECKED(FStateTreeAnyEnum, Value));
 			}
 			if (bTargetIsAnyEnum)
 			{
-				ModifiedBinding.TargetPath.Path.Add(GET_MEMBER_NAME_STRING_CHECKED(FStateTreeAnyEnum, Value));
+				ModifiedBinding.GetMutableTargetPath().AddPathSegment(GET_MEMBER_NAME_STRING_CHECKED(FStateTreeAnyEnum, Value));
 			}
 			OutBindings.Add(ModifiedBinding);
 		}
@@ -1300,20 +1341,20 @@ bool FStateTreeCompiler::GetAndValidateBindings(const FStateTreeBindableStructDe
 		}
 
 		// Check if the bindings is for struct ref and validate the types.
-		if (!ValidateStructRef(SourceStruct, Binding.SourcePath, TargetStruct, Binding.TargetPath))
+		if (!ValidateStructRef(SourceStruct, Binding.GetSourcePath(), TargetStruct, Binding.GetTargetPath()))
 		{
 			return false;
 		}
 	}
 
 
-	auto IsPropertyBound = [&OutBindings](const FString& PropertyName)
+	auto IsPropertyBound = [&OutBindings](const FName& PropertyName)
 	{
-		return OutBindings.ContainsByPredicate([&PropertyName](const FStateTreeEditorPropertyBinding& Binding)
+		return OutBindings.ContainsByPredicate([&PropertyName](const FStateTreePropertyPathBinding& Binding)
 			{
 				// We're looping over just the first level of properties on the struct, so we assume that the path is just one item
 				// (or two in case of AnyEnum, because we expand the path to Property.Value, see code above).
-				return Binding.TargetPath.Path.Num() >= 1 && Binding.TargetPath.Path[0] == PropertyName;
+				return Binding.GetTargetPath().GetSegments().Num() >= 1 && Binding.GetTargetPath().GetSegments()[0].GetName() == PropertyName;
 			});
 	};
 
@@ -1324,7 +1365,7 @@ bool FStateTreeCompiler::GetAndValidateBindings(const FStateTreeBindableStructDe
 	{
 		const FProperty* Property = *It;
 		check(Property);
-		const FString PropertyName = Property->GetName();
+		const FName PropertyName = Property->GetFName();
 		const EStateTreePropertyUsage Usage = UE::StateTree::Compiler::GetUsageFromMetaData(Property);
 		if (Usage == EStateTreePropertyUsage::Input)
 		{
@@ -1335,7 +1376,7 @@ bool FStateTreeCompiler::GetAndValidateBindings(const FStateTreeBindableStructDe
 			{
 				Log.Reportf(EMessageSeverity::Error, TargetStruct,
 					TEXT("Input property '%s:%s' is expected to have a binding."),
-					*TargetStruct.Name.ToString(), *PropertyName);
+					*TargetStruct.Name.ToString(), *PropertyName.ToString());
 				bResult = false;
 			}
 		}
@@ -1356,7 +1397,7 @@ bool FStateTreeCompiler::GetAndValidateBindings(const FStateTreeBindableStructDe
 			{
 				Log.Reportf(EMessageSeverity::Error, TargetStruct,
 					TEXT("Context property '%s:%s' type is expected to be Object Reference or Struct."),
-					*TargetStruct.Name.ToString(), *PropertyName);
+					*TargetStruct.Name.ToString(), *PropertyName.ToString());
 				bResult = false;
 				continue;
 			}
@@ -1365,24 +1406,18 @@ bool FStateTreeCompiler::GetAndValidateBindings(const FStateTreeBindableStructDe
 
 			if (!bIsBound)
 			{
-				const FStateTreeBindableStructDesc Desc = EditorData->FindContextData(ContextObjectType, PropertyName);
+				const FStateTreeBindableStructDesc Desc = EditorData->FindContextData(ContextObjectType, PropertyName.ToString());
 
 				if (Desc.IsValid())
 				{
 					// Add automatic binding to Context data.
-					FStateTreeEditorPropertyBinding Binding;
-					Binding.SourcePath.StructID = Desc.ID;
-
-					Binding.TargetPath.StructID = TargetStruct.ID;
-					Binding.TargetPath.Path.Add(PropertyName);
-
-					OutBindings.Add(Binding);
+					OutBindings.Emplace(FStateTreePropertyPath(Desc.ID), FStateTreePropertyPath(TargetStruct.ID, PropertyName));
 				}
 				else
 				{
 					Log.Reportf(EMessageSeverity::Error, TargetStruct,
 						TEXT("Cound not find matching Context object for Context property '%s:%s'. Property must have munual binding."),
-						*TargetStruct.Name.ToString(), *PropertyName);
+						*TargetStruct.Name.ToString(), *PropertyName.ToString());
 					bResult = false;
 				}
 			}
diff --git a/Engine/Plugins/Runtime/StateTree/Source/StateTreeEditorModule/Private/StateTreeEditor.cpp b/Engine/Plugins/Runtime/StateTree/Source/StateTreeEditorModule/Private/StateTreeEditor.cpp
--- a/Engine/Plugins/Runtime/StateTree/Source/StateTreeEditorModule/Private/StateTreeEditor.cpp
+++ b/Engine/Plugins/Runtime/StateTree/Source/StateTreeEditorModule/Private/StateTreeEditor.cpp
@@ -654,9 +654,9 @@ namespace UE::StateTree::Editor::Internal
 			return;
 		}
 
-		TMap<FGuid, const UStruct*> AllStructIDs;
-		TreeData->GetAllStructIDs(AllStructIDs);
-		TreeData->GetPropertyEditorBindings()->RemoveUnusedBindings(AllStructIDs);
+		TMap<FGuid, const FStateTreeDataView> AllStructValues;
+		TreeData->GetAllStructValues(AllStructValues);
+		TreeData->GetPropertyEditorBindings()->RemoveUnusedBindings(AllStructValues);
 	}
 
 	void UpdateLinkedStateParameters(UStateTree& StateTree)
diff --git a/Engine/Plugins/Runtime/StateTree/Source/StateTreeEditorModule/Private/StateTreeEditorData.cpp b/Engine/Plugins/Runtime/StateTree/Source/StateTreeEditorModule/Private/StateTreeEditorData.cpp
--- a/Engine/Plugins/Runtime/StateTree/Source/StateTreeEditorModule/Private/StateTreeEditorData.cpp
+++ b/Engine/Plugins/Runtime/StateTree/Source/StateTreeEditorModule/Private/StateTreeEditorData.cpp
@@ -83,9 +83,9 @@ void UStateTreeEditorData::PostEditChangeChainProperty(FPropertyChangedChainEven
 			if (MemberProperty->GetFName() == GET_MEMBER_NAME_CHECKED(UStateTreeEditorData, Evaluators)
 				|| MemberProperty->GetFName() == GET_MEMBER_NAME_CHECKED(UStateTreeEditorData, GlobalTasks))
 			{
-				TMap<FGuid, const UStruct*> AllStructIDs;
-				GetAllStructIDs(AllStructIDs);
-				EditorBindings.RemoveUnusedBindings(AllStructIDs);
+				TMap<FGuid, const FStateTreeDataView> AllStructValues;
+				GetAllStructValues(AllStructValues);
+				EditorBindings.RemoveUnusedBindings(AllStructValues);
 			}
 		}
 	}
@@ -96,8 +96,8 @@ void UStateTreeEditorData::GetAccessibleStructs(const FGuid TargetStructID, TArr
 {
 	// Find the states that are updated before the current state.
 	TArray<const UStateTreeState*> Path;
-	const UStateTreeState* CurrentState = GetStateByStructID(TargetStructID);
-	for (const UStateTreeState* State = CurrentState; State != nullptr; State = State->Parent)
+	const UStateTreeState* State = GetStateByStructID(TargetStructID);
+	while (State != nullptr)
 	{
 		Path.Insert(State, 0);
 
@@ -106,6 +106,8 @@ void UStateTreeEditorData::GetAccessibleStructs(const FGuid TargetStructID, TArr
 		{
 			break;
 		}
+
+		State = State->Parent;
 	}
 	
 	GetAccessibleStructs(Path, TargetStructID, OutStructDescs);
@@ -115,83 +117,23 @@ void UStateTreeEditorData::GetAccessibleStructs(const TConstArrayView<const USta
 {
 	const UStateTree* StateTree = GetTypedOuter<UStateTree>();
 	checkf(StateTree, TEXT("UStateTreeEditorData should only be allocated within a UStateTree"));
-	
-	// All parameters are accessible
-	if (const UScriptStruct* PropertyBagStruct = RootParameters.Parameters.GetPropertyBagStruct())
-	{
-		OutStructDescs.Emplace(TEXT("Parameters"), PropertyBagStruct, EStateTreeBindableStructSource::Parameter, RootParameters.ID);
-	}
 
-	// All named context objects declared by the schema are accessible
-	if (Schema != nullptr)
-	{
-		for (const FStateTreeExternalDataDesc& Desc : Schema->GetContextDataDescs())
-		{
-			OutStructDescs.Emplace(Desc.Name, Desc.Struct, EStateTreeBindableStructSource::Context, Desc.ID);
-		}	
-	}
 
-	// State parameters
-	const UStateTreeState* RootState = Path.Num() > 0 ? Path[0] : nullptr;
-	if (RootState != nullptr
-		&& RootState->Type == EStateTreeStateType::Subtree
-		&& RootState->Parameters.Parameters.GetPropertyBagStruct() != nullptr)
+	EStateTreeVisitor BaseProgress = VisitGlobalNodes([&OutStructDescs, TargetStructID]
+		(const UStateTreeState* State, const FStateTreeBindableStructDesc& Desc, const FStateTreeDataView Value)
 	{
-		if (const UScriptStruct* PropertyBagStruct = RootState->Parameters.Parameters.GetPropertyBagStruct())
-		{
-			OutStructDescs.Emplace(RootState->Name, PropertyBagStruct, EStateTreeBindableStructSource::State, RootState->Parameters.ID);
-		}
-	}
-
-	bool bFoundTarget = false;
-
-	// Evaluators
-	// Evaluators can access other evaluators that come before them.
-	for (const FStateTreeEditorNode& Node : Evaluators)
-	{
-		if (const FStateTreeEvaluatorBase* Evaluator = Node.Node.GetPtr<FStateTreeEvaluatorBase>())
+		if (Desc.ID == TargetStructID)
 		{
-			// Stop iterating as soon as we find the target node.
-			if (Node.ID == TargetStructID)
-			{
-				bFoundTarget = true;
-				break;
-			}
-
-			// Collect evaluators accessible so far.
-			FStateTreeBindableStructDesc& Desc = OutStructDescs.AddDefaulted_GetRef();
-			Desc.Struct = Evaluator->GetInstanceDataType();
-			Desc.Name = Evaluator->Name;
-			Desc.ID = Node.ID;
-			Desc.DataSource = EStateTreeBindableStructSource::Evaluator;
+			return EStateTreeVisitor::Break;
 		}
-	}
-
-	// Global Tasks
-	// Global tasks can access evaluators and other global tasks that come before them.
-	for (const FStateTreeEditorNode& Node : GlobalTasks)
-	{
-		if (const FStateTreeTaskBase* Task = Node.Node.GetPtr<FStateTreeTaskBase>())
-		{
-			// Stop iterating as soon as we find the target node.
-			if (Node.ID == TargetStructID)
-			{
-				bFoundTarget = true;
-				break;
-			}
+		
+		OutStructDescs.Add(Desc);
+		
+		return EStateTreeVisitor::Continue;
+	});
 
-			// Collect global tasks accessible so far.
-			FStateTreeBindableStructDesc& Desc = OutStructDescs.AddDefaulted_GetRef();
-			Desc.Struct = Task->GetInstanceDataType();
-			Desc.Name = Task->Name;
-			Desc.ID = Node.ID;
-			Desc.DataSource = EStateTreeBindableStructSource::GlobalTask;
-		}
-	}
 
-	// Conditions and Tasks
-	// Visit the tree in execution order. Conditions and tasks can access tasks that are executed before them.
-	if (!bFoundTarget)
+	if (BaseProgress == EStateTreeVisitor::Continue)
 	{
 		TArray<FStateTreeBindableStructDesc> TaskDescs;
 
@@ -202,31 +144,33 @@ void UStateTreeEditorData::GetAccessibleStructs(const TConstArrayView<const USta
 				continue;
 			}
 			
-			VisitStateNodes(*State, [&OutStructDescs, &TaskDescs, TargetStructID]
-				(const UStateTreeState* State, const FGuid& ID, const FName& Name, const EStateTreeNodeType NodeType, const UScriptStruct* NodeStruct, const UStruct* InstanceStruct)
+			const EStateTreeVisitor StateProgress = VisitStateNodes(*State, [&OutStructDescs, &TaskDescs, TargetStructID]
+				(const UStateTreeState* State, const FStateTreeBindableStructDesc& Desc, const FStateTreeDataView Value)
 				{
 					// Stop iterating as soon as we find the target node.
-					if (ID == TargetStructID)
+					if (Desc.ID == TargetStructID)
 					{
 						OutStructDescs.Append(TaskDescs);
 						return EStateTreeVisitor::Break;
 					}
 
-					// Not at target yet, collect all tasks accessible so far.
-					if (NodeStruct->IsChildOf(FStateTreeTaskBase::StaticStruct()))
+					// Not at target yet, collect all bindable source accessible so far.
+					if (Desc.DataSource == EStateTreeBindableStructSource::Task
+						|| Desc.DataSource == EStateTreeBindableStructSource::State)
 					{
-						FStateTreeBindableStructDesc& Desc = TaskDescs.AddDefaulted_GetRef();
-						Desc.Struct = InstanceStruct;
-						Desc.Name = Name;
-						Desc.ID = ID;
-						Desc.DataSource = EStateTreeBindableStructSource::Task;
+						TaskDescs.Add(Desc);
 					}
-				
+							
 					return EStateTreeVisitor::Continue;
 				});
+			
+			if (StateProgress == EStateTreeVisitor::Break)
+			{
+				break;
+			}
 		}
 	}
-
+	
 	OutStructDescs.StableSort([](const FStateTreeBindableStructDesc& A, const FStateTreeBindableStructDesc& B)
 	{
 		return (uint8)A.DataSource < (uint8)B.DataSource;
@@ -296,100 +240,32 @@ bool UStateTreeEditorData::GetStructByID(const FGuid StructID, FStateTreeBindabl
 {
 	bool bResult = false;
 
-	// All parameters
-	if (StructID == RootParameters.ID)
-	{
-		OutStructDesc.Struct = RootParameters.Parameters.GetPropertyBagStruct();
-		OutStructDesc.Name = FName(TEXT("Parameters"));
-		OutStructDesc.ID = RootParameters.ID;
-		bResult = true;
-	}
-
-	// All named external data items declared by the schema
-	if (!bResult && Schema != nullptr)
-	{
-		for (const FStateTreeExternalDataDesc& Desc : Schema->GetContextDataDescs())
-		{
-			if (StructID == Desc.ID)
-			{
-				OutStructDesc.Struct = Desc.Struct;
-				OutStructDesc.Name = Desc.Name;
-				OutStructDesc.ID = Desc.ID;
-				bResult = true;
-				break;
-			}
-		}
-	}
-
-	// Evaluators
-	if (!bResult)
-	{
-		for (const FStateTreeEditorNode& Node : Evaluators)
-		{
-			if (const FStateTreeEvaluatorBase* Evaluator = Node.Node.GetPtr<FStateTreeEvaluatorBase>())
-			{
-				if (StructID == Node.ID)
-				{
-					OutStructDesc.Struct = Evaluator->GetInstanceDataType();
-					OutStructDesc.Name = Evaluator->Name;
-					OutStructDesc.ID = Node.ID;
-					bResult = true;
-					break;
-				}
-			}
-		}
-	}
-
-	// Global tasks
-	if (!bResult)
+	VisitAllNodes([&OutStructDesc, StructID](const UStateTreeState* State, const FStateTreeBindableStructDesc& Desc, const FStateTreeDataView Value)
 	{
-		for (const FStateTreeEditorNode& Node : GlobalTasks)
+		if (Desc.ID == StructID)
 		{
-			if (const FStateTreeTaskBase* Task = Node.Node.GetPtr<FStateTreeTaskBase>())
-			{
-				if (StructID == Node.ID)
-				{
-					OutStructDesc.Struct = Task->GetInstanceDataType();
-					OutStructDesc.Name = Task->Name;
-					OutStructDesc.ID = Node.ID;
-					bResult = true;
-					break;
-				}
-			}
+			OutStructDesc = Desc;
+			return EStateTreeVisitor::Break;
 		}
-	}
-
-	if (!bResult)
-	{
-		VisitHierarchyNodes([&bResult, &OutStructDesc, StructID](const UStateTreeState* State, const FGuid& ID, const FName& Name, const EStateTreeNodeType NodeType, const UScriptStruct* NodeStruct, const UStruct* InstanceStruct)
-			{
-				if (ID == StructID)
-				{
-					OutStructDesc.Struct = InstanceStruct;
-					OutStructDesc.Name = Name;
-					OutStructDesc.ID = ID;
-					bResult = true;
-					return EStateTreeVisitor::Break;
-				}
-				return EStateTreeVisitor::Continue;
-			});
-	}
-
-	return bResult;
+		return EStateTreeVisitor::Continue;
+	});
+	
+	return OutStructDesc.IsValid();
 }
 
 const UStateTreeState* UStateTreeEditorData::GetStateByStructID(const FGuid TargetStructID) const
 {
 	const UStateTreeState* Result = nullptr;
 
-	VisitHierarchyNodes([&Result, TargetStructID](const UStateTreeState* State, const FGuid& ID, const FName& Name, const EStateTreeNodeType NodeType, const UScriptStruct* NodeStruct, const UStruct* InstanceStruct)
+	VisitHierarchyNodes([&Result, TargetStructID](const UStateTreeState* State, const FStateTreeBindableStructDesc& Desc, const FStateTreeDataView Value)
 		{
-			if (ID == TargetStructID)
+			if (Desc.ID == TargetStructID)
 			{
 				Result = State;
 				return EStateTreeVisitor::Break;
 			}
 			return EStateTreeVisitor::Continue;
+			
 		});
 
 	return Result;
@@ -420,39 +296,23 @@ void UStateTreeEditorData::GetAllStructIDs(TMap<FGuid, const UStruct*>& AllStruc
 	const UStateTree* StateTree = GetTypedOuter<UStateTree>();
 	checkf(StateTree, TEXT("UStateTreeEditorData should only be allocated within a UStateTree"));
 
-	// All parameters
-	AllStructs.Emplace(RootParameters.ID, RootParameters.Parameters.GetPropertyBagStruct());
-
-	// All named external data items declared by the schema
-	if (Schema != nullptr)
-	{
-		for (const FStateTreeExternalDataDesc& Desc : Schema->GetContextDataDescs())
+	VisitAllNodes([&AllStructs](const UStateTreeState* State, const FStateTreeBindableStructDesc& Desc, const FStateTreeDataView Value)
 		{
 			AllStructs.Emplace(Desc.ID, Desc.Struct);
-		}	
-	}
+			return EStateTreeVisitor::Continue;
+		});
+}
 
-	// Evaluators
-	for (const FStateTreeEditorNode& Node : Evaluators)
-	{
-		if (const FStateTreeEvaluatorBase* Evaluator = Node.Node.GetPtr<FStateTreeEvaluatorBase>())
-		{
-			AllStructs.Emplace(Node.ID, Evaluator->GetInstanceDataType());
-		}
-	}
+void UStateTreeEditorData::GetAllStructValues(TMap<FGuid, const FStateTreeDataView>& AllValues) const
+{
+	AllValues.Reset();
 
-	// Global tasks
-	for (const FStateTreeEditorNode& Node : GlobalTasks)
-	{
-		if (const FStateTreeTaskBase* Task = Node.Node.GetPtr<FStateTreeTaskBase>())
-		{
-			AllStructs.Emplace(Node.ID, Task->GetInstanceDataType());
-		}
-	}
+	const UStateTree* StateTree = GetTypedOuter<UStateTree>();
+	checkf(StateTree, TEXT("UStateTreeEditorData should only be allocated within a UStateTree"));
 
-	VisitHierarchyNodes([&AllStructs](const UStateTreeState* State, const FGuid& ID, const FName& Name, const EStateTreeNodeType NodeType, const UScriptStruct* NodeStruct, const UStruct* InstanceStruct)
+	VisitAllNodes([&AllValues](const UStateTreeState* State, const FStateTreeBindableStructDesc& Desc, const FStateTreeDataView Value)
 		{
-			AllStructs.Add(ID, InstanceStruct);
+			AllValues.Emplace(Desc.ID, Value);
 			return EStateTreeVisitor::Continue;
 		});
 }
@@ -535,7 +395,120 @@ EStateTreeVisitor UStateTreeEditorData::VisitStateNodes(const UStateTreeState& S
 	return bContinue ? EStateTreeVisitor::Continue : EStateTreeVisitor::Break;
 }
 
-void UStateTreeEditorData::VisitHierarchy(TFunctionRef<EStateTreeVisitor(UStateTreeState& State, UStateTreeState* ParentState)> InFunc) const
+
+EStateTreeVisitor UStateTreeEditorData::VisitStateNodes(const UStateTreeState& State, TFunctionRef<EStateTreeVisitor(const UStateTreeState* State, const FStateTreeBindableStructDesc& Desc, const FStateTreeDataView Value)> InFunc) const
+{
+	bool bContinue = true;
+
+	if (bContinue)
+	{
+		// Bindable state parameters for subtree or linked tree.
+		if ((State.Type == EStateTreeStateType::Subtree
+				|| State.Type == EStateTreeStateType::Linked)
+			&& State.Parameters.Parameters.IsValid())
+		{
+			FStateTreeBindableStructDesc Desc;
+			Desc.Struct = State.Parameters.Parameters.GetPropertyBagStruct();
+			Desc.Name = State.Name;
+			Desc.ID = State.Parameters.ID;
+			Desc.DataSource = EStateTreeBindableStructSource::State;
+
+			if (InFunc(&State, Desc, FStateTreeDataView(const_cast<FInstancedPropertyBag&>(State.Parameters.Parameters).GetMutableValue())) == EStateTreeVisitor::Break)
+			{
+				bContinue = false;
+			}
+		}
+	}
+	
+	if (bContinue)
+	{
+		// Enter conditions
+		for (const FStateTreeEditorNode& Node : State.EnterConditions)
+		{
+			if (const FStateTreeConditionBase* Cond = Node.Node.GetPtr<FStateTreeConditionBase>())
+			{
+				FStateTreeBindableStructDesc Desc;
+				Desc.Struct = Cond->GetInstanceDataType();
+				Desc.Name = Cond->Name;
+				Desc.ID = Node.ID;
+				Desc.DataSource = EStateTreeBindableStructSource::Condition;
+
+				if (InFunc(&State, Desc, Node.GetInstance()) == EStateTreeVisitor::Break)
+				{
+					bContinue = false;
+					break;
+				}
+			}
+		}
+	}
+	if (bContinue)
+	{
+		// Tasks
+		for (const FStateTreeEditorNode& Node : State.Tasks)
+		{
+			if (const FStateTreeTaskBase* Task = Node.Node.GetPtr<FStateTreeTaskBase>())
+			{
+				FStateTreeBindableStructDesc Desc;
+				Desc.Struct = Task->GetInstanceDataType();
+				Desc.Name = Task->Name;
+				Desc.ID = Node.ID;
+				Desc.DataSource = EStateTreeBindableStructSource::Task;
+
+				if (InFunc(&State, Desc, Node.GetInstance()) == EStateTreeVisitor::Break)
+				{
+					bContinue = false;
+					break;
+				}
+			}
+		}
+	}
+	if (bContinue)
+	{
+		if (const FStateTreeTaskBase* Task = State.SingleTask.Node.GetPtr<FStateTreeTaskBase>())
+		{
+			FStateTreeBindableStructDesc Desc;
+			Desc.Struct = Task->GetInstanceDataType();
+			Desc.Name = Task->Name;
+			Desc.ID = State.SingleTask.ID;
+			Desc.DataSource = EStateTreeBindableStructSource::Task;
+
+			if (InFunc(&State, Desc, State.SingleTask.GetInstance()) == EStateTreeVisitor::Break)
+			{
+				bContinue = false;
+			}
+		}
+
+	}
+	if (bContinue)
+	{
+		// Transitions
+		for (const FStateTreeTransition& Transition : State.Transitions)
+		{
+			for (const FStateTreeEditorNode& Node : Transition.Conditions)
+			{
+				if (const FStateTreeConditionBase* Cond = Node.Node.GetPtr<FStateTreeConditionBase>())
+				{
+					FStateTreeBindableStructDesc Desc;
+					Desc.Struct = Cond->GetInstanceDataType();
+					Desc.Name = Cond->Name;
+					Desc.ID = Node.ID;
+					Desc.DataSource = EStateTreeBindableStructSource::Condition;
+
+					if (InFunc(&State, Desc, Node.GetInstance()) == EStateTreeVisitor::Break)
+					{
+						bContinue = false;
+						break;
+					}
+				}
+			}
+		}
+	}
+
+	return bContinue ? EStateTreeVisitor::Continue : EStateTreeVisitor::Break;
+}
+
+
+EStateTreeVisitor UStateTreeEditorData::VisitHierarchy(TFunctionRef<EStateTreeVisitor(UStateTreeState& State, UStateTreeState* ParentState)> InFunc) const
 {
 	using FStatePair = TTuple<UStateTreeState*, UStateTreeState*>; 
 	TArray<FStatePair> Stack;
@@ -576,12 +549,114 @@ void UStateTreeEditorData::VisitHierarchy(TFunctionRef<EStateTreeVisitor(UStateT
 			break;
 		}
 	}
+
+	return EStateTreeVisitor::Continue;
 }
 
 void UStateTreeEditorData::VisitHierarchyNodes(TFunctionRef<EStateTreeVisitor(const UStateTreeState* State, const FGuid& ID, const FName& Name, const EStateTreeNodeType NodeType, const UScriptStruct* NodeStruct, const UStruct* InstanceStruct)> InFunc) const
 {
+	PRAGMA_DISABLE_DEPRECATION_WARNINGS
 	VisitHierarchy([this, &InFunc](const UStateTreeState& State, UStateTreeState* /*ParentState*/)
 	{
 		return VisitStateNodes(State, InFunc);
 	});
+	PRAGMA_ENABLE_DEPRECATION_WARNINGS
+}
+
+EStateTreeVisitor UStateTreeEditorData::VisitGlobalNodes(TFunctionRef<EStateTreeVisitor(const UStateTreeState* State, const FStateTreeBindableStructDesc& Desc, const FStateTreeDataView Value)> InFunc) const
+{
+	// Root parameters
+	{
+		FStateTreeBindableStructDesc Desc;
+		Desc.Struct = RootParameters.Parameters.GetPropertyBagStruct();
+		Desc.Name = FName(TEXT("Parameters"));
+		Desc.ID = RootParameters.ID;
+		Desc.DataSource = EStateTreeBindableStructSource::Parameter;
+		
+		if (InFunc(nullptr, Desc, FStateTreeDataView(const_cast<FInstancedPropertyBag&>(RootParameters.Parameters).GetMutableValue())) == EStateTreeVisitor::Break)
+		{
+			return EStateTreeVisitor::Break;
+		}
+	}
+
+	// All named external data items declared by the schema
+	if (Schema != nullptr)
+	{
+		for (const FStateTreeExternalDataDesc& ContextDesc : Schema->GetContextDataDescs())
+		{
+			FStateTreeBindableStructDesc Desc;
+			Desc.Struct = ContextDesc.Struct;
+			Desc.Name = ContextDesc.Name;
+			Desc.ID = ContextDesc.ID;
+			Desc.DataSource = EStateTreeBindableStructSource::Context;
+			
+			if (InFunc(nullptr, Desc, FStateTreeDataView()) == EStateTreeVisitor::Break)
+			{
+				return EStateTreeVisitor::Break;
+			}
+		}
+	}
+
+	// Evaluators
+	for (const FStateTreeEditorNode& Node : Evaluators)
+	{
+		if (const FStateTreeEvaluatorBase* Evaluator = Node.Node.GetPtr<FStateTreeEvaluatorBase>())
+		{
+			FStateTreeBindableStructDesc Desc;
+			Desc.Struct = Evaluator->GetInstanceDataType();
+			Desc.Name = Evaluator->Name;
+			Desc.ID = Node.ID;
+			Desc.DataSource = EStateTreeBindableStructSource::Evaluator;
+
+			if (InFunc(nullptr, Desc, Node.GetInstance()) == EStateTreeVisitor::Break)
+			{
+				return EStateTreeVisitor::Break;
+			}
+		}
+	}
+
+	// Global tasks
+	for (const FStateTreeEditorNode& Node : GlobalTasks)
+	{
+		if (const FStateTreeTaskBase* Task = Node.Node.GetPtr<FStateTreeTaskBase>())
+		{
+			FStateTreeBindableStructDesc Desc;
+			Desc.Struct = Task->GetInstanceDataType();
+			Desc.Name = Task->Name;
+			Desc.ID = Node.ID;
+			Desc.DataSource = EStateTreeBindableStructSource::GlobalTask;
+
+			if (InFunc(nullptr, Desc, Node.GetInstance()) == EStateTreeVisitor::Break)
+			{
+				return EStateTreeVisitor::Break;
+			}
+		}
+	}
+
+	return  EStateTreeVisitor::Continue;
+}
+
+EStateTreeVisitor UStateTreeEditorData::VisitHierarchyNodes(TFunctionRef<EStateTreeVisitor(const UStateTreeState* State, const FStateTreeBindableStructDesc& Desc, const FStateTreeDataView Value)> InFunc) const
+{
+	return VisitHierarchy([this, &InFunc](const UStateTreeState& State, UStateTreeState* /*ParentState*/)
+	{
+		return VisitStateNodes(State, InFunc);
+	});
+}
+
+EStateTreeVisitor UStateTreeEditorData::VisitAllNodes(TFunctionRef<EStateTreeVisitor(const UStateTreeState* State, const FStateTreeBindableStructDesc& Desc, const FStateTreeDataView Value)> InFunc) const
+{
+	if (VisitGlobalNodes(InFunc) == EStateTreeVisitor::Break)
+	{
+		return EStateTreeVisitor::Break;
+	}
+
+	return VisitHierarchyNodes(InFunc);
+}
+
+PRAGMA_DISABLE_DEPRECATION_WARNINGS
+void UStateTreeEditorData::AddPropertyBinding(const FStateTreeEditorPropertyPath& SourcePath, const FStateTreeEditorPropertyPath& TargetPath)
+{
+	EditorBindings.AddPropertyBinding(UE::StateTree::Private::ConvertEditorPath(SourcePath), UE::StateTree::Private::ConvertEditorPath(TargetPath));
 }
+PRAGMA_ENABLE_DEPRECATION_WARNINGS
\ No newline at end of file
diff --git a/Engine/Plugins/Runtime/StateTree/Source/StateTreeEditorModule/Private/StateTreeEditorPropertyBindings.cpp b/Engine/Plugins/Runtime/StateTree/Source/StateTreeEditorModule/Private/StateTreeEditorPropertyBindings.cpp
--- a/Engine/Plugins/Runtime/StateTree/Source/StateTreeEditorModule/Private/StateTreeEditorPropertyBindings.cpp
+++ b/Engine/Plugins/Runtime/StateTree/Source/StateTreeEditorModule/Private/StateTreeEditorPropertyBindings.cpp
@@ -2,10 +2,11 @@
 
 #include "StateTreeEditorPropertyBindings.h"
 #include "StateTreePropertyBindingCompiler.h"
+#include "Misc/EnumerateRange.h"
+#include "PropertyPathHelpers.h"
 
 #include UE_INLINE_GENERATED_CPP_BY_NAME(StateTreeEditorPropertyBindings)
 
-
 UStateTreeEditorPropertyBindingsOwner::UStateTreeEditorPropertyBindingsOwner(const FObjectInitializer& ObjectInitializer)
 	: Super(ObjectInitializer)
 {
@@ -13,87 +14,148 @@ UStateTreeEditorPropertyBindingsOwner::UStateTreeEditorPropertyBindingsOwner(con
 
 //////////////////////////////////////////////////////////////////////////
 
-void FStateTreeEditorPropertyBindings::AddPropertyBinding(const FStateTreeEditorPropertyPath& SourcePath, const FStateTreeEditorPropertyPath& TargetPath)
+void FStateTreeEditorPropertyBindings::AddPropertyBinding(const FStateTreePropertyPath& SourcePath, const FStateTreePropertyPath& TargetPath)
 {
 	RemovePropertyBindings(TargetPath);
-	PropertyBindings.Add(FStateTreeEditorPropertyBinding(SourcePath, TargetPath));
+	PropertyBindings.Add(FStateTreePropertyPathBinding(SourcePath, TargetPath));
 }
 
-void FStateTreeEditorPropertyBindings::RemovePropertyBindings(const FStateTreeEditorPropertyPath& TargetPath)
+void FStateTreeEditorPropertyBindings::RemovePropertyBindings(const FStateTreePropertyPath& TargetPath)
 {
-	PropertyBindings.RemoveAll([TargetPath](const FStateTreeEditorPropertyBinding& Binding)
+	PropertyBindings.RemoveAll([&TargetPath](const FStateTreePropertyPathBinding& Binding)
 		{
-			return Binding.TargetPath == TargetPath;
+			return Binding.GetTargetPath() == TargetPath;
 		});
 }
 
 void FStateTreeEditorPropertyBindings::CopyBindings(const FGuid FromStructID, const FGuid ToStructID)
 {
 	// Copy all bindings that target "FromStructID" and retarget them to "ToStructID".
-	TArray<FStateTreeEditorPropertyBinding> NewBindings;
-	for (const FStateTreeEditorPropertyBinding& Binding : PropertyBindings)
+	TArray<FStateTreePropertyPathBinding> NewBindings;
+	for (const FStateTreePropertyPathBinding& Binding : PropertyBindings)
 	{
-		if (Binding.TargetPath.StructID == FromStructID)
+		if (Binding.GetTargetPath().GetStructID() == FromStructID)
 		{
-			FStateTreeEditorPropertyBinding& NewBinding = NewBindings.Add_GetRef(Binding);
-			NewBinding.TargetPath.StructID = ToStructID;
+			NewBindings.Emplace(Binding.GetSourcePath(), FStateTreePropertyPath(ToStructID, Binding.GetTargetPath().GetSegments()));
 		}
 	}
 
-	for (const FStateTreeEditorPropertyBinding& NewBinding : NewBindings)
-	{
-		AddPropertyBinding(NewBinding.SourcePath, NewBinding.TargetPath);
-	}
+	PropertyBindings.Append(NewBindings);
 }
 
-bool FStateTreeEditorPropertyBindings::HasPropertyBinding(const FStateTreeEditorPropertyPath& TargetPath) const
+bool FStateTreeEditorPropertyBindings::HasPropertyBinding(const FStateTreePropertyPath& TargetPath) const
 {
-	return PropertyBindings.ContainsByPredicate([TargetPath](const FStateTreeEditorPropertyBinding& Binding)
+	return PropertyBindings.ContainsByPredicate([&TargetPath](const FStateTreePropertyPathBinding& Binding)
 		{
-			return Binding.TargetPath == TargetPath;
+			return Binding.GetTargetPath() == TargetPath;
 		});
 }
 
-const FStateTreeEditorPropertyPath* FStateTreeEditorPropertyBindings::GetPropertyBindingSource(const FStateTreeEditorPropertyPath& TargetPath) const
+const FStateTreePropertyPath* FStateTreeEditorPropertyBindings::GetPropertyBindingSource(const FStateTreePropertyPath& TargetPath) const
 {
-	const FStateTreeEditorPropertyBinding* Binding = PropertyBindings.FindByPredicate([TargetPath](const FStateTreeEditorPropertyBinding& Binding)
+	const FStateTreePropertyPathBinding* Binding = PropertyBindings.FindByPredicate([&TargetPath](const FStateTreePropertyPathBinding& Binding)
 		{
-			return Binding.TargetPath == TargetPath;
+			return Binding.GetTargetPath() == TargetPath;
 		});
-	return Binding ? &Binding->SourcePath : nullptr;
+	return Binding ? &Binding->GetSourcePath() : nullptr;
 }
 
-void FStateTreeEditorPropertyBindings::GetPropertyBindingsFor(const FGuid StructID, TArray<FStateTreeEditorPropertyBinding>& OutBindings) const
+void FStateTreeEditorPropertyBindings::GetPropertyBindingsFor(const FGuid StructID, TArray<FStateTreePropertyPathBinding>& OutBindings) const
 {
-	OutBindings = PropertyBindings.FilterByPredicate([StructID](const FStateTreeEditorPropertyBinding& Binding)
+	OutBindings = PropertyBindings.FilterByPredicate([StructID](const FStateTreePropertyPathBinding& Binding)
 		{
-			return Binding.TargetPath.StructID == StructID && Binding.IsValid();
+			return Binding.GetSourcePath().GetStructID().IsValid()
+				&& Binding.GetTargetPath().GetStructID() == StructID;
 		});
 }
 
-void FStateTreeEditorPropertyBindings::RemoveUnusedBindings(const TMap<FGuid, const UStruct*>& ValidStructs)
+void FStateTreeEditorPropertyBindings::RemoveUnusedBindings(const TMap<FGuid, const FStateTreeDataView>& ValidStructs)
 {
-	PropertyBindings.RemoveAll([ValidStructs](const FStateTreeEditorPropertyBinding& Binding)
+	PropertyBindings.RemoveAll([ValidStructs](FStateTreePropertyPathBinding& Binding)
 		{
 			// Remove binding if it's target struct has been removed
-			if (!ValidStructs.Contains(Binding.TargetPath.StructID))
+			if (!ValidStructs.Contains(Binding.GetTargetPath().GetStructID()))
 			{
+				// Remove
 				return true;
 			}
 
-			// Remove binding if the target property cannot be resolved, likely renamed property.
+			// Target path should always have at least one segment (copy bind directly on a target struct/object). 
+			if (Binding.GetTargetPath().IsPathEmpty())
+			{
+				return true;
+			}
+
+			// Remove binding if path containing instanced indirections (e.g. instance struct or object) cannot be resolved.
 			// TODO: Try to use core redirect to find new name.
-			const UStruct* Struct = ValidStructs.FindChecked(Binding.TargetPath.StructID);
+			const FStateTreeDataView Value = ValidStructs.FindChecked(Binding.GetTargetPath().GetStructID());
+			FString Error;
+			TArray<FStateTreePropertyPathIndirection> Indirections;
+			if (!Binding.GetTargetPath().ResolveIndirectionsWithValue(Value, Indirections, &Error))
+			{
+				UE_LOG(LogStateTree, Verbose, TEXT("Removing binding to %s because binding target path cannot be resolved: %s"),
+					*Binding.GetSourcePath().ToString(), *Error); // Error contains the target path.
+				
+				// Remove
+				return true;
+			}
+
+			return false;
+		});
+}
+
+PRAGMA_DISABLE_DEPRECATION_WARNINGS
+void FStateTreeEditorPropertyBindings::AddPropertyBinding(const FStateTreeEditorPropertyPath& SourcePath, const FStateTreeEditorPropertyPath& TargetPath)
+{
+	AddPropertyBinding(UE::StateTree::Private::ConvertEditorPath(SourcePath), UE::StateTree::Private::ConvertEditorPath(TargetPath));
+}
+
+void FStateTreeEditorPropertyBindings::RemovePropertyBindings(const FStateTreeEditorPropertyPath& TargetPath)
+{
+	RemovePropertyBindings(UE::StateTree::Private::ConvertEditorPath(TargetPath));
+}
 
-			FStateTreeBindableStructDesc StructDesc;
-			StructDesc.Struct = Struct;
-			StructDesc.ID = Binding.TargetPath.StructID;
+bool FStateTreeEditorPropertyBindings::HasPropertyBinding(const FStateTreeEditorPropertyPath& TargetPath) const
+{
+	return HasPropertyBinding(UE::StateTree::Private::ConvertEditorPath(TargetPath));
+}
+
+const FStateTreeEditorPropertyPath* FStateTreeEditorPropertyBindings::GetPropertyBindingSource(const FStateTreeEditorPropertyPath& TargetPath) const
+{
+	static FStateTreeEditorPropertyPath Dummy;
+	const FStateTreePropertyPath* SourcePath = GetPropertyBindingSource(UE::StateTree::Private::ConvertEditorPath(TargetPath));
+	if (SourcePath != nullptr)
+	{
+		Dummy = UE::StateTree::Private::ConvertEditorPath(*SourcePath);
+		return &Dummy;
+	}
+	return nullptr;
+}
+PRAGMA_ENABLE_DEPRECATION_WARNINGS
+
+void FStateTreeEditorPropertyBindings::RemoveUnusedBindings(const TMap<FGuid, const UStruct*>& ValidStructs)
+{
+	PropertyBindings.RemoveAll([ValidStructs](const FStateTreePropertyPathBinding& Binding)
+		{
+			// Remove binding if it's target struct has been removed
+			if (!ValidStructs.Contains(Binding.GetTargetPath().GetStructID()))
+			{
+				return true;
+			}
+
+			// Target path should always have at least one segment (copy bind directly on a target struct/object). 
+			if (Binding.GetTargetPath().IsPathEmpty())
+			{
+				return true;
+			}
 		
-			TArray<FStateTreePropertySegment> Segments;
-			const FProperty* LeafProperty = nullptr;
-			int32 LeafArrayIndex = INDEX_NONE;
-			if (!FStateTreePropertyBindingCompiler::ResolvePropertyPath(StructDesc, Binding.TargetPath, Segments, LeafProperty, LeafArrayIndex))
+			// Remove binding if path containing instanced indirections (e.g. instance struct or object) cannot be resolved.
+			// TODO: Try to use core redirect to find new name.
+			const UStruct* Struct = ValidStructs.FindChecked(Binding.GetTargetPath().GetStructID());
+			TArray<FStateTreePropertyPathIndirection> Indirections;
+			if (!Binding.GetTargetPath().ResolveIndirections(Struct, Indirections))
 			{
+				// Remove
 				return true;
 			}
 
@@ -101,6 +163,7 @@ void FStateTreeEditorPropertyBindings::RemoveUnusedBindings(const TMap<FGuid, co
 		});
 }
 
+
 //////////////////////////////////////////////////////////////////////////
 
 FStateTreeBindingLookup::FStateTreeBindingLookup(IStateTreeEditorPropertyBindingsOwner* InBindingOwner)
@@ -108,58 +171,50 @@ FStateTreeBindingLookup::FStateTreeBindingLookup(IStateTreeEditorPropertyBinding
 {
 }
 
-const FStateTreeEditorPropertyPath* FStateTreeBindingLookup::GetPropertyBindingSource(const FStateTreeEditorPropertyPath& InTargetPath) const
+const FStateTreePropertyPath* FStateTreeBindingLookup::GetPropertyBindingSource(const FStateTreePropertyPath& InTargetPath) const
 {
 	check(BindingOwner);
-	FStateTreeEditorPropertyBindings* EditorBindings = BindingOwner->GetPropertyEditorBindings();
+	const FStateTreeEditorPropertyBindings* EditorBindings = BindingOwner->GetPropertyEditorBindings();
 	check(EditorBindings);
 
 	return EditorBindings->GetPropertyBindingSource(InTargetPath);
 }
 
-FText FStateTreeBindingLookup::GetPropertyPathDisplayName(const FStateTreeEditorPropertyPath& InPath) const
+FText FStateTreeBindingLookup::GetPropertyPathDisplayName(const FStateTreePropertyPath& InPath) const
 {
 	check(BindingOwner);
-	FStateTreeEditorPropertyBindings* EditorBindings = BindingOwner->GetPropertyEditorBindings();
+	const FStateTreeEditorPropertyBindings* EditorBindings = BindingOwner->GetPropertyEditorBindings();
 	check(EditorBindings);
 
 	FString Result;
 
 	FStateTreeBindableStructDesc Struct;
-	if (BindingOwner->GetStructByID(InPath.StructID, Struct))
+	if (BindingOwner->GetStructByID(InPath.GetStructID(), Struct))
 	{
 		Result = Struct.Name.ToString();
 	}
 
-	for (const FString& Segment : InPath.Path)
-	{
-		Result += TEXT(".") + Segment;
-	}
+	Result += TEXT(".") + InPath.ToString();
 
 	return FText::FromString(Result);
 }
 
-const FProperty* FStateTreeBindingLookup::GetPropertyPathLeafProperty(const FStateTreeEditorPropertyPath& InPath) const
+const FProperty* FStateTreeBindingLookup::GetPropertyPathLeafProperty(const FStateTreePropertyPath& InPath) const
 {
 	check(BindingOwner);
-	FStateTreeEditorPropertyBindings* EditorBindings = BindingOwner->GetPropertyEditorBindings();
+	const FStateTreeEditorPropertyBindings* EditorBindings = BindingOwner->GetPropertyEditorBindings();
 	check(EditorBindings);
 
 	const FProperty* Result = nullptr;
 	FStateTreeBindableStructDesc Struct;
-	if (BindingOwner->GetStructByID(InPath.StructID, Struct))
+	if (BindingOwner->GetStructByID(InPath.GetStructID(), Struct))
 	{
-		// TODO: Could use inline allocator here, since there are usually only few segments. Needs API change in ResolvePropertyPath().
-		TArray<FStateTreePropertySegment> Segments;
-		const FProperty* LeafProperty = nullptr;
-		int32 LeafArrayIndex = INDEX_NONE;
-		if (FStateTreePropertyBindingCompiler::ResolvePropertyPath(Struct, InPath, Segments, LeafProperty, LeafArrayIndex))
+		TArray<FStateTreePropertyPathIndirection> Indirection;
+		if (InPath.ResolveIndirections(Struct.Struct, Indirection) && Indirection.Num() > 0)
 		{
-			Result = LeafProperty;
+			return Indirection.Last().GetProperty();
 		}
 	}
 
 	return Result;
 }
-
-
diff --git a/Engine/Plugins/Runtime/StateTree/Source/StateTreeEditorModule/Private/StateTreeEditorStyle.cpp b/Engine/Plugins/Runtime/StateTree/Source/StateTreeEditorModule/Private/StateTreeEditorStyle.cpp
--- a/Engine/Plugins/Runtime/StateTree/Source/StateTreeEditorModule/Private/StateTreeEditorStyle.cpp
+++ b/Engine/Plugins/Runtime/StateTree/Source/StateTreeEditorModule/Private/StateTreeEditorStyle.cpp
@@ -128,12 +128,16 @@ void FStateTreeEditorStyle::Initialize()
 
 	StyleSet->Set("StateTree.Node.Operand", FTextBlockStyle(NormalText)
 		.SetFont(FAppStyle::GetFontStyle(TEXT("PropertyWindow.BoldFont")))
-		.SetFontSize(7));
+		.SetFontSize(8));
 
 	StyleSet->Set("StateTree.Node.Parens", FTextBlockStyle(NormalText)
 		.SetFont(FAppStyle::GetFontStyle(TEXT("PropertyWindow.NormalFont")))
 		.SetFontSize(12));
 
+	StyleSet->Set("StateTree.Param.Label", FTextBlockStyle(NormalText)
+		.SetFont(FAppStyle::GetFontStyle(TEXT("PropertyWindow.BoldFont")))
+		.SetFontSize(7));
+
 	// Condition Indent combo button
 	const FButtonStyle IndentButton = FButtonStyle()
 		.SetNormal(FSlateRoundedBoxBrush(FLinearColor::Transparent, 2.0f))
@@ -147,9 +151,9 @@ void FStateTreeEditorStyle::Initialize()
 	
 	StyleSet->Set("StateTree.Node.Indent.ComboBox", FComboButtonStyle(ComboButtonStyle).SetButtonStyle(IndentButton));
 
-	const FEditableTextBoxStyle& NormalEditableTextBox = FCoreStyle::Get().GetWidgetStyle<FEditableTextBoxStyle>("NormalEditableTextBox");
-	FEditableTextBoxStyle NameEditStyle(NormalEditableTextBox);
-	NameEditStyle.TextStyle.Font.Size = 10;
+	const FEditableTextStyle& NormalEditableText = FCoreStyle::Get().GetWidgetStyle<FEditableTextStyle>("NormalEditableText");
+	FEditableTextStyle NameEditStyle(NormalEditableText);
+	NameEditStyle.Font.Size = 10;
 	StyleSet->Set("StateTree.Node.Name", NameEditStyle);
 	
 	FSlateStyleRegistry::RegisterSlateStyle(*StyleSet.Get());
diff --git a/Engine/Plugins/Runtime/StateTree/Source/StateTreeEditorModule/Private/StateTreePropertyBindingCompiler.cpp b/Engine/Plugins/Runtime/StateTree/Source/StateTreeEditorModule/Private/StateTreePropertyBindingCompiler.cpp
--- a/Engine/Plugins/Runtime/StateTree/Source/StateTreeEditorModule/Private/StateTreePropertyBindingCompiler.cpp
+++ b/Engine/Plugins/Runtime/StateTree/Source/StateTreeEditorModule/Private/StateTreePropertyBindingCompiler.cpp
@@ -5,6 +5,7 @@
 #include "StateTreeCompiler.h"
 #include "StateTreeCompilerLog.h"
 #include "StateTreeEditorPropertyBindings.h"
+#include "Misc/EnumerateRange.h"
 
 #include UE_INLINE_GENERATED_CPP_BY_NAME(StateTreePropertyBindingCompiler)
 
@@ -17,7 +18,7 @@ bool FStateTreePropertyBindingCompiler::Init(FStateTreePropertyBindings& InPrope
 	return true;
 }
 
-bool FStateTreePropertyBindingCompiler::CompileBatch(const FStateTreeBindableStructDesc& TargetStruct, TConstArrayView<FStateTreeEditorPropertyBinding> EditorPropertyBindings, int32& OutBatchIndex)
+bool FStateTreePropertyBindingCompiler::CompileBatch(const FStateTreeBindableStructDesc& TargetStruct, TConstArrayView<FStateTreePropertyPathBinding> BatchPropertyBindings, int32& OutBatchIndex)
 {
 	check(Log);
 	check(PropertyBindings);
@@ -25,19 +26,16 @@ bool FStateTreePropertyBindingCompiler::CompileBatch(const FStateTreeBindableStr
 
 	StoreSourceStructs();
 
-	TArray<FStateTreePropertySegment> SourceSegments;
-	TArray<FStateTreePropertySegment> TargetSegments;
-	
-	const int32 BindingsBegin = PropertyBindings->PropertyBindings.Num();
+	const int32 BindingsBegin = PropertyBindings->PropertyPathBindings.Num();
 
-	for (const FStateTreeEditorPropertyBinding& EditorBinding : EditorPropertyBindings)
+	for (const FStateTreePropertyPathBinding& Binding : BatchPropertyBindings)
 	{
-		if (EditorBinding.TargetPath.StructID != TargetStruct.ID)
+		if (Binding.GetTargetPath().GetStructID() != TargetStruct.ID)
 		{
 			continue;
 		}
 		// Source must be in the source array
-		const FGuid SourceStructID = EditorBinding.SourcePath.StructID;
+		const FGuid SourceStructID = Binding.GetSourcePath().GetStructID();
 		const int32 SourceStructIdx = SourceStructs.IndexOfByPredicate([SourceStructID](const FStateTreeBindableStructDesc& Struct)
 			{
 				return (Struct.ID == SourceStructID);
@@ -50,92 +48,50 @@ bool FStateTreePropertyBindingCompiler::CompileBatch(const FStateTreeBindableStr
 		}
 		const FStateTreeBindableStructDesc& SourceStruct = SourceStructs[SourceStructIdx];
 
-		FStateTreePropertyBinding& NewBinding = PropertyBindings->PropertyBindings.AddDefaulted_GetRef();
-
-		SourceSegments.Reset();
-		TargetSegments.Reset();
-
-		// Resolve paths
-		const FProperty* SourceLeafProperty = nullptr;
-		int32 SourceLeafArrayIndex = INDEX_NONE;
-		if (!ResolvePropertyPath(SourceStruct, EditorBinding.SourcePath, SourceSegments, SourceLeafProperty, SourceLeafArrayIndex, Log, &TargetStruct))
+		FString Error;
+		TArray<FStateTreePropertyPathIndirection> SourceIndirections;
+		TArray<FStateTreePropertyPathIndirection> TargetIndirections;
+		
+		if (!Binding.GetSourcePath().ResolveIndirections(SourceStruct.Struct, SourceIndirections, &Error))
 		{
-			Log->Reportf(EMessageSeverity::Error, TargetStruct,
-				TEXT("Could not resolve path to '%s:%s'."),
-				*SourceStruct.Name.ToString(), *EditorBinding.SourcePath.ToString());
+			Log->Reportf(EMessageSeverity::Error, TargetStruct, TEXT("Resolving path in %s: %s"), *SourceStruct.Name.ToString(), *Error);
 			return false;
 		}
 
-		// Destination container is set to 0, it is assumed to be passed in when doing the batch copy.
-		const FProperty* TargetLeafProperty = nullptr;
-		int32 TargetLeafArrayIndex = INDEX_NONE;
-		if (!ResolvePropertyPath(TargetStruct, EditorBinding.TargetPath, TargetSegments, TargetLeafProperty, TargetLeafArrayIndex, Log, &TargetStruct))
+		if (!Binding.GetTargetPath().ResolveIndirections(TargetStruct.Struct, TargetIndirections, &Error))
 		{
-			Log->Reportf(EMessageSeverity::Error, TargetStruct,
-				TEXT("Could not resolve path to '%s:%s'."),
-				*TargetStruct.Name.ToString(), *EditorBinding.TargetPath.ToString());
+			Log->Reportf(EMessageSeverity::Error, TargetStruct, TEXT("Resolving path in %s: %s"), *TargetStruct.Name.ToString(), *Error);
 			return false;
 		}
 
-		auto StorePropertyPath = [this, &TargetStruct](FStateTreePropertySegment& FirstPathSegment, TArray<FStateTreePropertySegment>& Segments) -> bool
-		{
-			// The path is empty when directly bound to the target struct.
-			if (Segments.IsEmpty())
-			{
-				return true;
-			}
-			
-			FirstPathSegment = Segments[0];
-			FStateTreePropertySegment* PrevSegment = &FirstPathSegment; 
-			for (int32 Index = 1; Index < Segments.Num(); Index++)
-			{
-				const int32 SegmentIndex = PropertyBindings->PropertySegments.Num();
-				FStateTreePropertySegment& NewSegment = PropertyBindings->PropertySegments.Add_GetRef(Segments[Index]);
-
-				if (const auto Validation = UE::StateTree::Compiler::IsValidIndex16(SegmentIndex); Validation.DidFail())
-				{
-					Validation.Log(*Log, TEXT("SegmentIndex"), TargetStruct);
-					return false;
-				}
-				PrevSegment->NextIndex = FStateTreeIndex16(SegmentIndex);
-				PrevSegment = &NewSegment;
-			}
-			return true;
-		};
-
-		if (!StorePropertyPath(NewBinding.SourcePath, SourceSegments))
-		{
-			return false;
-		}
-		if (!StorePropertyPath(NewBinding.TargetPath, TargetSegments))
+		FStateTreePropertyCopy DummyCopy;
+		FStateTreePropertyPathIndirection LastSourceIndirection = !SourceIndirections.IsEmpty() ? SourceIndirections.Last() : FStateTreePropertyPathIndirection(SourceStruct.Struct);
+		FStateTreePropertyPathIndirection LastTargetIndirection = !TargetIndirections.IsEmpty() ? TargetIndirections.Last() : FStateTreePropertyPathIndirection(TargetStruct.Struct);
+		if (!PropertyBindings->ResolveCopyType(DummyCopy, LastSourceIndirection, LastTargetIndirection))
 		{
+			Log->Reportf(EMessageSeverity::Error, TargetStruct,
+			TEXT("Failed to resolve property copy type between %s:%s and %s:%s."),
+				*SourceStruct.Name.ToString(), *Binding.GetSourcePath().ToString(),
+				*TargetStruct.Name.ToString(), *Binding.GetTargetPath().ToString());
 			return false;
 		}
 
 		if (const auto Validation = UE::StateTree::Compiler::IsValidIndex16(SourceStructIdx); Validation.DidFail())
 		{
-			Validation.Log(*Log, TEXT("SourceStructIdx"), TargetStruct);
+			if (Log != nullptr)
+			{
+				Validation.Log(*Log, TEXT("Source Struct Index"), SourceStruct);
+			}
 			return false;
 		}
-		NewBinding.SourceStructIndex = FStateTreeIndex16(SourceStructIdx);
-		
-		NewBinding.CopyType = GetCopyType(SourceStruct.Struct, SourceLeafProperty, SourceLeafArrayIndex,
-			TargetStruct.Struct, TargetLeafProperty, TargetLeafArrayIndex);
 
-		if (NewBinding.CopyType == EStateTreePropertyCopyType::None)
-		{
-			Log->Reportf(EMessageSeverity::Error, TargetStruct,
-				TEXT("Could not resolve copy type for properties '%s:%s' and '%s:%s'."),
-				*SourceStruct.Name.ToString(), *EditorBinding.SourcePath.ToString(),
-				*TargetStruct.Name.ToString(), *EditorBinding.TargetPath.ToString());
-			return false;
-		}
+		PropertyBindings->PropertyPathBindings.Emplace(FStateTreeIndex16(SourceStructIdx), Binding.GetSourcePath(), Binding.GetTargetPath());
 	}
 
-	const int32 BindingsEnd = PropertyBindings->PropertyBindings.Num();
+	const int32 BindingsEnd = PropertyBindings->PropertyPathBindings.Num();
 	if (BindingsBegin != BindingsEnd)
 	{
-		FStateTreePropCopyBatch& Batch = PropertyBindings->CopyBatches.AddDefaulted_GetRef();
+		FStateTreePropertyCopyBatch& Batch = PropertyBindings->CopyBatches.AddDefaulted_GetRef();
 		Batch.TargetStruct = TargetStruct;
 		Batch.BindingsBegin = IntCastChecked<uint16>(BindingsBegin);
 		Batch.BindingsEnd = IntCastChecked<uint16>(BindingsEnd);
@@ -169,228 +125,7 @@ int32 FStateTreePropertyBindingCompiler::GetSourceStructIndexByID(const FGuid& I
 	return SourceStructs.IndexOfByPredicate([ID](const FStateTreeBindableStructDesc& Structs) { return (Structs.ID == ID); });
 }
 
-EStateTreePropertyCopyType FStateTreePropertyBindingCompiler::GetCopyType(const UStruct* SourceStruct, const FProperty* SourceProperty, const int32 SourceArrayIndex,
-																		  const UStruct* TargetStruct, const FProperty* TargetProperty, const int32 TargetArrayIndex) const
-{
-	if (SourceProperty == nullptr)
-	{
-		// Copy directly from the source struct, target must be.
-		if (const FStructProperty* StructProperty = CastField<FStructProperty>(TargetProperty))
-		{
-			if (StructProperty->Struct == SourceStruct)
-			{
-				return EStateTreePropertyCopyType::CopyStruct;
-			}
-		}
-		else if (const FObjectPropertyBase* ObjectProperty = CastField<FObjectPropertyBase>(TargetProperty))
-		{
-			if (SourceStruct->IsChildOf(ObjectProperty->PropertyClass))
-			{
-				return EStateTreePropertyCopyType::CopyObject;
-			}
-		}
-		
-		return EStateTreePropertyCopyType::None;
-	}
-	
-	if (const FArrayProperty* SourceArrayProperty = CastField<FArrayProperty>(SourceProperty))
-	{
-		// use the array's inner property if we are not trying to copy the whole array
-		if (SourceArrayIndex != INDEX_NONE)
-		{
-			SourceProperty = SourceArrayProperty->Inner;
-		}
-	}
-
-	if (const FArrayProperty* TargetArrayProperty = CastField<FArrayProperty>(TargetProperty))
-	{
-		// use the array's inner property if we are not trying to copy the whole array
-		if (TargetArrayIndex != INDEX_NONE)
-		{
-			TargetProperty = TargetArrayProperty->Inner;
-		}
-	}
-
-	// Handle FStateTreeStructRef
-	if (const FStructProperty* TargetStructProperty = CastField<const FStructProperty>(TargetProperty))
-	{
-		if (TargetStructProperty->Struct == TBaseStructure<FStateTreeStructRef>::Get())
-		{
-			if (const FStructProperty* SourceStructProperty = CastField<const FStructProperty>(SourceProperty))
-			{
-				// FStateTreeStructRef to FStateTreeStructRef is copied as usual.
-				if (SourceStructProperty->Struct != TBaseStructure<FStateTreeStructRef>::Get())
-				{
-					return EStateTreePropertyCopyType::StructReference;
-				}
-			}
-		}
-	}
-
-	const EPropertyAccessCompatibility Compatibility = GetPropertyCompatibility(SourceProperty, TargetProperty);
-
-	// Extract underlying types for enums
-	if (const FEnumProperty* EnumPropertyA = CastField<const FEnumProperty>(SourceProperty))
-	{
-		SourceProperty = EnumPropertyA->GetUnderlyingProperty();
-	}
-
-	if (const FEnumProperty* EnumPropertyB = CastField<const FEnumProperty>(TargetProperty))
-	{
-		TargetProperty = EnumPropertyB->GetUnderlyingProperty();
-	}
-
-	if (Compatibility == EPropertyAccessCompatibility::Compatible)
-	{
-		if (CastField<FNameProperty>(TargetProperty))
-		{
-			return EStateTreePropertyCopyType::CopyName;
-		}
-		else if (CastField<FBoolProperty>(TargetProperty))
-		{
-			return EStateTreePropertyCopyType::CopyBool;
-		}
-		else if (CastField<FStructProperty>(TargetProperty))
-		{
-			return EStateTreePropertyCopyType::CopyStruct;
-		}
-		else if (CastField<FObjectPropertyBase>(TargetProperty))
-		{
-			return EStateTreePropertyCopyType::CopyObject;
-		}
-		else if (CastField<FArrayProperty>(TargetProperty) && TargetProperty->HasAnyPropertyFlags(CPF_EditFixedSize))
-		{
-			// only apply array copying rules if the destination array is fixed size, otherwise it will be 'complex'
-			return EStateTreePropertyCopyType::CopyFixedArray;
-		}
-		else if (TargetProperty->PropertyFlags & CPF_IsPlainOldData)
-		{
-			return EStateTreePropertyCopyType::CopyPlain;
-		}
-		else
-		{
-			return EStateTreePropertyCopyType::CopyComplex;
-		}
-	}
-	else if (Compatibility == EPropertyAccessCompatibility::Promotable)
-	{
-		if (SourceProperty->IsA<FBoolProperty>())
-		{
-			if (TargetProperty->IsA<FByteProperty>())
-			{
-				return EStateTreePropertyCopyType::PromoteBoolToByte;
-			}
-			else if (TargetProperty->IsA<FIntProperty>())
-			{
-				return EStateTreePropertyCopyType::PromoteBoolToInt32;
-			}
-			else if (TargetProperty->IsA<FUInt32Property>())
-			{
-				return EStateTreePropertyCopyType::PromoteBoolToUInt32;
-			}
-			else if (TargetProperty->IsA<FInt64Property>())
-			{
-				return EStateTreePropertyCopyType::PromoteBoolToInt64;
-			}
-			else if (TargetProperty->IsA<FFloatProperty>())
-			{
-				return EStateTreePropertyCopyType::PromoteBoolToFloat;
-			}
-			else if (TargetProperty->IsA<FDoubleProperty>())
-			{
-				return EStateTreePropertyCopyType::PromoteBoolToDouble;
-			}
-		}
-		else if (SourceProperty->IsA<FByteProperty>())
-		{
-			if (TargetProperty->IsA<FIntProperty>())
-			{
-				return EStateTreePropertyCopyType::PromoteByteToInt32;
-			}
-			else if (TargetProperty->IsA<FUInt32Property>())
-			{
-				return EStateTreePropertyCopyType::PromoteByteToUInt32;
-			}
-			else if (TargetProperty->IsA<FInt64Property>())
-			{
-				return EStateTreePropertyCopyType::PromoteByteToInt64;
-			}
-			else if (TargetProperty->IsA<FFloatProperty>())
-			{
-				return EStateTreePropertyCopyType::PromoteByteToFloat;
-			}
-			else if (TargetProperty->IsA<FDoubleProperty>())
-			{
-				return EStateTreePropertyCopyType::PromoteByteToDouble;
-			}
-		}
-		else if (SourceProperty->IsA<FIntProperty>())
-		{
-			if (TargetProperty->IsA<FInt64Property>())
-			{
-				return EStateTreePropertyCopyType::PromoteInt32ToInt64;
-			}
-			else if (TargetProperty->IsA<FFloatProperty>())
-			{
-				return EStateTreePropertyCopyType::PromoteInt32ToFloat;
-			}
-			else if (TargetProperty->IsA<FDoubleProperty>())
-			{
-				return EStateTreePropertyCopyType::PromoteInt32ToDouble;
-			}
-		}
-		else if (SourceProperty->IsA<FUInt32Property>())
-		{
-			if (TargetProperty->IsA<FInt64Property>())
-			{
-				return EStateTreePropertyCopyType::PromoteUInt32ToInt64;
-			}
-			else if (TargetProperty->IsA<FFloatProperty>())
-			{
-				return EStateTreePropertyCopyType::PromoteUInt32ToFloat;
-			}
-			else if (TargetProperty->IsA<FDoubleProperty>())
-			{
-				return EStateTreePropertyCopyType::PromoteUInt32ToDouble;
-			}
-		}
-		else if (SourceProperty->IsA<FFloatProperty>())
-		{
-			if (TargetProperty->IsA<FIntProperty>())
-			{
-				return EStateTreePropertyCopyType::PromoteFloatToInt32;
-			}
-			else if (TargetProperty->IsA<FInt64Property>())
-			{
-				return EStateTreePropertyCopyType::PromoteFloatToInt64;
-			}
-			else if (TargetProperty->IsA<FDoubleProperty>())
-			{
-				return EStateTreePropertyCopyType::PromoteFloatToDouble;
-			}
-		}
-		else if (SourceProperty->IsA<FDoubleProperty>())
-		{
-			if (TargetProperty->IsA<FIntProperty>())
-			{
-				return EStateTreePropertyCopyType::DemoteDoubleToInt32;
-			}
-			else if (TargetProperty->IsA<FInt64Property>())
-			{
-				return EStateTreePropertyCopyType::DemoteDoubleToInt64;
-			}
-			else if (TargetProperty->IsA<FFloatProperty>())
-			{
-				return EStateTreePropertyCopyType::DemoteDoubleToFloat;
-			}
-		}
-	}
-	
-	ensureMsgf(false, TEXT("Couldnt determine property copy type (%s -> %s)"), *SourceProperty->GetNameCPP(), *TargetProperty->GetNameCPP());
-
-	return EStateTreePropertyCopyType::None;
-}
-
+PRAGMA_DISABLE_DEPRECATION_WARNINGS
 bool FStateTreePropertyBindingCompiler::ResolvePropertyPath(const FStateTreeBindableStructDesc& InStructDesc, const FStateTreeEditorPropertyPath& InPath,
 															TArray<FStateTreePropertySegment>& OutSegments, const FProperty*& OutLeafProperty, int32& OutLeafArrayIndex,
 															FStateTreeCompilerLog* InLog, const FStateTreeBindableStructDesc* InLogContextStruct)
@@ -579,122 +314,18 @@ bool FStateTreePropertyBindingCompiler::ResolvePropertyPath(const FStateTreeBind
 
 EPropertyAccessCompatibility FStateTreePropertyBindingCompiler::GetPropertyCompatibility(const FProperty* FromProperty, const FProperty* ToProperty)
 {
-	if (FromProperty == ToProperty)
+	const EStateTreePropertyAccessCompatibility Result = FStateTreePropertyBindings::GetPropertyCompatibility(FromProperty, ToProperty);
+	if (Result == EStateTreePropertyAccessCompatibility::Compatible)
 	{
 		return EPropertyAccessCompatibility::Compatible;
 	}
-
-	if (FromProperty == nullptr || ToProperty == nullptr)
+	if (Result == EStateTreePropertyAccessCompatibility::Promotable)
 	{
-		return EPropertyAccessCompatibility::Incompatible;
+		return EPropertyAccessCompatibility::Promotable;
 	}
-
-	// Special case for object properties since InPropertyA->SameType(InPropertyB) requires both properties to be of the exact same class.
-	// In our case we want to be able to bind a source property if its class is a child of the target property class.
-	if (FromProperty->IsA<FObjectPropertyBase>() && ToProperty->IsA<FObjectPropertyBase>())
-	{
-		const FObjectPropertyBase* SourceProperty = CastField<FObjectPropertyBase>(FromProperty);
-		const FObjectPropertyBase* TargetProperty = CastField<FObjectPropertyBase>(ToProperty);
-		return (SourceProperty->PropertyClass->IsChildOf(TargetProperty->PropertyClass)) ? EPropertyAccessCompatibility::Compatible : EPropertyAccessCompatibility::Incompatible;
-	}
-
-
-	// When copying to an enum property, expect FromProperty to be the same enum.
-	auto GetPropertyEnum = [](const FProperty* Property) -> const UEnum*
-	{
-		if (const FByteProperty* ByteProperty = CastField<FByteProperty>(Property))
-		{
-			return ByteProperty->GetIntPropertyEnum();
-		}
-		if (const FEnumProperty* EnumProperty = CastField<FEnumProperty>(Property))
-		{
-			return EnumProperty->GetEnum();
-		}
-		return nullptr;
-	};
-	
-	if (const UEnum* ToPropertyEnum = GetPropertyEnum(ToProperty))
-	{
-		const UEnum* FromPropertyEnum = GetPropertyEnum(FromProperty);
-		return (ToPropertyEnum == FromPropertyEnum) ? EPropertyAccessCompatibility::Compatible : EPropertyAccessCompatibility::Incompatible;
-	}
-	
-	// Allow source enums to be promoted to numbers.
-	if (const FEnumProperty* EnumPropertyA = CastField<const FEnumProperty>(FromProperty))
-	{
-		FromProperty = EnumPropertyA->GetUnderlyingProperty();
-	}
-
-	if (FromProperty->SameType(ToProperty))
-	{
-		return EPropertyAccessCompatibility::Compatible;
-	}
-	else
-	{
-		// Not directly compatible, check for promotions
-		if (FromProperty->IsA<FBoolProperty>())
-		{
-			if (ToProperty->IsA<FByteProperty>()
-				|| ToProperty->IsA<FIntProperty>()
-				|| ToProperty->IsA<FUInt32Property>()
-				|| ToProperty->IsA<FInt64Property>()
-				|| ToProperty->IsA<FFloatProperty>()
-				|| ToProperty->IsA<FDoubleProperty>())
-			{
-				return EPropertyAccessCompatibility::Promotable;
-			}
-		}
-		else if (FromProperty->IsA<FByteProperty>())
-		{
-			if (ToProperty->IsA<FIntProperty>()
-				|| ToProperty->IsA<FUInt32Property>()
-				|| ToProperty->IsA<FInt64Property>()
-				|| ToProperty->IsA<FFloatProperty>()
-				|| ToProperty->IsA<FDoubleProperty>())
-			{
-				return EPropertyAccessCompatibility::Promotable;
-			}
-		}
-		else if (FromProperty->IsA<FIntProperty>())
-		{
-			if (ToProperty->IsA<FInt64Property>()
-				|| ToProperty->IsA<FFloatProperty>()
-				|| ToProperty->IsA<FDoubleProperty>())
-			{
-				return EPropertyAccessCompatibility::Promotable;
-			}
-		}
-		else if (FromProperty->IsA<FUInt32Property>())
-		{
-			if (ToProperty->IsA<FInt64Property>()
-				|| ToProperty->IsA<FFloatProperty>()
-				|| ToProperty->IsA<FDoubleProperty>())
-			{
-				return EPropertyAccessCompatibility::Promotable;
-			}
-		}
-		else if (FromProperty->IsA<FFloatProperty>())
-		{
-			if (ToProperty->IsA<FIntProperty>()
-				|| ToProperty->IsA<FInt64Property>()
-				|| ToProperty->IsA<FDoubleProperty>())
-			{
-				return EPropertyAccessCompatibility::Promotable;
-			}
-		}
-		else if (FromProperty->IsA<FDoubleProperty>())
-		{
-			if (ToProperty->IsA<FIntProperty>()
-				|| ToProperty->IsA<FInt64Property>()
-				|| ToProperty->IsA<FFloatProperty>())
-			{
-				return EPropertyAccessCompatibility::Promotable;
-			}
-		}
-	}
-
 	return EPropertyAccessCompatibility::Incompatible;
 }
+PRAGMA_ENABLE_DEPRECATION_WARNINGS
 
 void FStateTreePropertyBindingCompiler::StoreSourceStructs()
 {
@@ -714,4 +345,3 @@ void FStateTreePropertyBindingCompiler::StoreSourceStructs()
 		}
 	}
 }
-
diff --git a/Engine/Plugins/Runtime/StateTree/Source/StateTreeEditorModule/Private/StateTreeState.cpp b/Engine/Plugins/Runtime/StateTree/Source/StateTreeEditorModule/Private/StateTreeState.cpp
--- a/Engine/Plugins/Runtime/StateTree/Source/StateTreeEditorModule/Private/StateTreeState.cpp
+++ b/Engine/Plugins/Runtime/StateTree/Source/StateTreeEditorModule/Private/StateTreeState.cpp
@@ -235,9 +235,9 @@ void UStateTreeState::PostEditChangeChainProperty(FPropertyChangedChainEvent& Pr
 				TreeData->Modify();
 				FStateTreeEditorPropertyBindings* Bindings = TreeData->GetPropertyEditorBindings();
 				check(Bindings);
-				TMap<FGuid, const UStruct*> AllStructIDs;
-				TreeData->GetAllStructIDs(AllStructIDs);
-				Bindings->RemoveUnusedBindings(AllStructIDs);
+				TMap<FGuid, const FStateTreeDataView> AllStructValues;
+				TreeData->GetAllStructValues(AllStructValues);
+				Bindings->RemoveUnusedBindings(AllStructValues);
 			}
 		}
 	}
diff --git a/Engine/Plugins/Runtime/StateTree/Source/StateTreeEditorModule/Public/StateTreeCompiler.h b/Engine/Plugins/Runtime/StateTree/Source/StateTreeEditorModule/Public/StateTreeCompiler.h
--- a/Engine/Plugins/Runtime/StateTree/Source/StateTreeEditorModule/Public/StateTreeCompiler.h
+++ b/Engine/Plugins/Runtime/StateTree/Source/StateTreeEditorModule/Public/StateTreeCompiler.h
@@ -52,12 +52,14 @@ private:
 	bool CreateCondition(UStateTreeState& State, const FStateTreeEditorNode& CondNode, const EStateTreeConditionOperand Operand, const int8 DeltaIndent);
 	bool CreateTask(UStateTreeState* State, const FStateTreeEditorNode& TaskNode, bool& bOutHasTransitionTasks);
 	bool CreateEvaluator(const FStateTreeEditorNode& EvalNode);
-	bool GetAndValidateBindings(const FStateTreeBindableStructDesc& TargetStruct, TArray<FStateTreeEditorPropertyBinding>& OutBindings) const;
-	bool IsPropertyAnyEnum(const FStateTreeBindableStructDesc& Struct, FStateTreeEditorPropertyPath Path) const;
-	bool ValidateStructRef(const FStateTreeBindableStructDesc& SourceStruct, FStateTreeEditorPropertyPath SourcePath,
-							const FStateTreeBindableStructDesc& TargetStruct, FStateTreeEditorPropertyPath TargetPath) const;
-	bool CompileAndValidateNode(const UStateTreeState* SourceState, const FStateTreeBindableStructDesc& NodeDesc, FStructView NodeView, const FStateTreeDataView InstanceData) const;
+	bool GetAndValidateBindings(const FStateTreeBindableStructDesc& TargetStruct, FStateTreeDataView TargetValue, TArray<FStateTreePropertyPathBinding>& OutBindings) const;
+	bool IsPropertyAnyEnum(const FStateTreeBindableStructDesc& Struct, FStateTreePropertyPath Path) const;
+	bool ValidateStructRef(const FStateTreeBindableStructDesc& SourceStruct, FStateTreePropertyPath SourcePath,
+							const FStateTreeBindableStructDesc& TargetStruct, FStateTreePropertyPath TargetPath) const;
+	bool CompileAndValidateNode(const UStateTreeState* SourceState, const FStateTreeBindableStructDesc& NodeDesc, FStructView NodeView, const FStateTreeDataView InstanceData);
 
+	FStateTreeDataView GetBindingSourceValue(const int32 SourceIndex);
+	
 	FStateTreeCompilerLog& Log;
 	UStateTree* StateTree = nullptr;
 	UStateTreeEditorData* EditorData = nullptr;
diff --git a/Engine/Plugins/Runtime/StateTree/Source/StateTreeEditorModule/Public/StateTreeEditorData.h b/Engine/Plugins/Runtime/StateTree/Source/StateTreeEditorModule/Public/StateTreeEditorData.h
--- a/Engine/Plugins/Runtime/StateTree/Source/StateTreeEditorModule/Public/StateTreeEditorData.h
+++ b/Engine/Plugins/Runtime/StateTree/Source/StateTreeEditorModule/Public/StateTreeEditorData.h
@@ -39,32 +39,44 @@ public:
 	virtual void PostEditChangeChainProperty(FPropertyChangedChainEvent& PropertyChangedEvent) override;
 #endif
 
-	/** Returns parent state of a struct, or nullptr if not found. */
+	/** @returns parent state of a struct, or nullptr if not found. */
 	const UStateTreeState* GetStateByStructID(const FGuid TargetStructID) const;
 
-	/** Returns state based on its ID, or nullptr if not found. */
+	/** @returns state based on its ID, or nullptr if not found. */
 	const UStateTreeState* GetStateByID(const FGuid StateID) const;
 
-	/** Gets the IDs of all bindable structs in the StateTree. */
-	void GetAllStructIDs(TMap<FGuid, const UStruct*>& AllStructs) const;
+	/** @returns the IDs and instance values of all bindable structs in the StateTree. */
+	void GetAllStructValues(TMap<FGuid, const FStateTreeDataView>& AllValues) const;
 
 	/**
 	* Iterates over all structs that are related to binding
 	* @param InFunc function called at each node, should return true if visiting is continued or false to stop.
 	*/
-	void VisitHierarchy(TFunctionRef<EStateTreeVisitor(UStateTreeState& State, UStateTreeState* ParentState)> InFunc) const;
+	EStateTreeVisitor VisitHierarchy(TFunctionRef<EStateTreeVisitor(UStateTreeState& State, UStateTreeState* ParentState)> InFunc) const;
 
 	/**
-	 * Iterates over all structs that are related to binding
+	 * Iterates over all structs at the global level (context, tree parameters, evaluators, global tasks) that are related to binding.
 	 * @param InFunc function called at each node, should return true if visiting is continued or false to stop.
 	 */
-	void VisitHierarchyNodes(TFunctionRef<EStateTreeVisitor(const UStateTreeState* State, const FGuid& ID, const FName& Name, const EStateTreeNodeType NodeType, const UScriptStruct* NodeStruct, const UStruct* InstanceStruct)> InFunc) const;
+	EStateTreeVisitor VisitGlobalNodes(TFunctionRef<EStateTreeVisitor(const UStateTreeState* State, const FStateTreeBindableStructDesc& Desc, const FStateTreeDataView Value)> InFunc) const;
+
+	/**
+	 * Iterates over all structs in the state hierarchy that are related to binding.
+	 * @param InFunc function called at each node, should return true if visiting is continued or false to stop.
+	 */
+	EStateTreeVisitor VisitHierarchyNodes(TFunctionRef<EStateTreeVisitor(const UStateTreeState* State, const FStateTreeBindableStructDesc& Desc, const FStateTreeDataView Value)> InFunc) const;
+
+	/**
+	 * Iterates over all structs that are related to binding.
+	 * @param InFunc function called at each node, should return true if visiting is continued or false to stop.
+	 */
+	EStateTreeVisitor VisitAllNodes(TFunctionRef<EStateTreeVisitor(const UStateTreeState* State, const FStateTreeBindableStructDesc& Desc, const FStateTreeDataView Value)> InFunc) const;
 
 	/**
 	 * Iterates over all nodes in a given state.
 	 * @param InFunc function called at each node, should return true if visiting is continued or false to stop.
 	 */
-	EStateTreeVisitor VisitStateNodes(const UStateTreeState& State, TFunctionRef<EStateTreeVisitor(const UStateTreeState* State, const FGuid& ID, const FName& Name, const EStateTreeNodeType NodeType, const UScriptStruct* NodeStruct, const UStruct* InstanceStruct)> InFunc) const;
+	EStateTreeVisitor VisitStateNodes(const UStateTreeState& State, TFunctionRef<EStateTreeVisitor(const UStateTreeState* State, const FStateTreeBindableStructDesc& Desc, const FStateTreeDataView Value)> InFunc) const;
 
 	/**
 	 * Returns array of nodes along the execution path, up to the TargetStruct.
@@ -80,7 +92,16 @@ public:
 	 * @param ObjectNameHint Name to use if multiple context objects of same type are found. 
 	 */
 	FStateTreeBindableStructDesc FindContextData(const UStruct* ObjectType, const FString ObjectNameHint) const;
-	
+
+	UE_DEPRECATED(5.3, "Use VisitHierarchyNodes with State, Desc, Value instead.")
+	void VisitHierarchyNodes(TFunctionRef<EStateTreeVisitor(const UStateTreeState* State, const FGuid& ID, const FName& Name, const EStateTreeNodeType NodeType, const UScriptStruct* NodeStruct, const UStruct* InstanceStruct)> InFunc) const;
+
+	UE_DEPRECATED(5.3, "Use VisitStateNodes with State, Desc, Value instead.")
+	EStateTreeVisitor VisitStateNodes(const UStateTreeState& State, TFunctionRef<EStateTreeVisitor(const UStateTreeState* State, const FGuid& ID, const FName& Name, const EStateTreeNodeType NodeType, const UScriptStruct* NodeStruct, const UStruct* InstanceStruct)> InFunc) const;
+
+	UE_DEPRECATED(5.3, "Use GetAllStructValues with values instead.")
+	void GetAllStructIDs(TMap<FGuid, const UStruct*>& AllStructs) const;
+
 	// StateTree Builder API
 
 	/**
@@ -144,10 +165,33 @@ public:
 	/**
 	 * Adds property binding between two structs.
 	 */
-	void AddPropertyBinding(const FStateTreeEditorPropertyPath& SourcePath, const FStateTreeEditorPropertyPath& TargetPath)
+	void AddPropertyBinding(const FStateTreePropertyPath& SourcePath, const FStateTreePropertyPath& TargetPath)
 	{
 		EditorBindings.AddPropertyBinding(SourcePath, TargetPath);
 	}
+
+PRAGMA_DISABLE_DEPRECATION_WARNINGS
+	UE_DEPRECATED(5.3, "Use version with FStateTreePropertyPath instead.")
+	void AddPropertyBinding(const FStateTreeEditorPropertyPath& SourcePath, const FStateTreeEditorPropertyPath& TargetPath);
+PRAGMA_ENABLE_DEPRECATION_WARNINGS
+
+	/**
+	 * Adds property binding between two structs.
+	 */
+	bool AddPropertyBinding(const FStateTreeEditorNode& SourceNode, const FString SourcePathStr, const FStateTreeEditorNode& TargetNode, const FString TargetPathStr)
+	{
+		FStateTreePropertyPath SourcePath;
+		FStateTreePropertyPath TargetPath;
+		SourcePath.SetStructID(SourceNode.ID);
+		TargetPath.SetStructID(TargetNode.ID);
+		if (SourcePath.FromString(SourcePathStr) && TargetPath.FromString(TargetPathStr))
+		{
+			EditorBindings.AddPropertyBinding(SourcePath, TargetPath);
+			return true;
+		}
+		return false;
+	}
+
 	// ~StateTree Builder API
 
 	/** Schema describing which inputs, evaluators, and tasks a StateTree can contain */	
diff --git a/Engine/Plugins/Runtime/StateTree/Source/StateTreeEditorModule/Public/StateTreeEditorNode.h b/Engine/Plugins/Runtime/StateTree/Source/StateTreeEditorModule/Public/StateTreeEditorNode.h
--- a/Engine/Plugins/Runtime/StateTree/Source/StateTreeEditorModule/Public/StateTreeEditorNode.h
+++ b/Engine/Plugins/Runtime/StateTree/Source/StateTreeEditorModule/Public/StateTreeEditorNode.h
@@ -40,6 +40,16 @@ struct STATETREEEDITORMODULE_API FStateTreeEditorNode
 		return FName();
 	}
 
+	const FStateTreeDataView GetInstance() const
+	{
+		return InstanceObject ? FStateTreeDataView(InstanceObject) : FStateTreeDataView(const_cast<FInstancedStruct&>(Instance));
+	}
+
+	FStateTreeDataView GetInstance()
+	{
+		return InstanceObject ? FStateTreeDataView(InstanceObject) : FStateTreeDataView(Instance);
+	}
+
 	UPROPERTY(EditDefaultsOnly, Category = Node)
 	FInstancedStruct Node;
 
diff --git a/Engine/Plugins/Runtime/StateTree/Source/StateTreeEditorModule/Public/StateTreeEditorPropertyBindings.h b/Engine/Plugins/Runtime/StateTree/Source/StateTreeEditorModule/Public/StateTreeEditorPropertyBindings.h
--- a/Engine/Plugins/Runtime/StateTree/Source/StateTreeEditorModule/Public/StateTreeEditorPropertyBindings.h
+++ b/Engine/Plugins/Runtime/StateTree/Source/StateTreeEditorModule/Public/StateTreeEditorPropertyBindings.h
@@ -6,29 +6,6 @@
 #include "StateTreePropertyBindings.h"
 #include "StateTreeEditorPropertyBindings.generated.h"
 
-/**
- * Editor representation of a property binding in StateTree
- */
-USTRUCT()
-struct STATETREEEDITORMODULE_API FStateTreeEditorPropertyBinding
-{
-	GENERATED_BODY()
-
-	FStateTreeEditorPropertyBinding() = default;
-	FStateTreeEditorPropertyBinding(const FStateTreeEditorPropertyPath& InSourcePath, const FStateTreeEditorPropertyPath& InTargetPath) : SourcePath(InSourcePath), TargetPath(InTargetPath) {}
-
-	bool IsValid() const { return SourcePath.IsValid() && TargetPath.IsValid(); }
-
-	/** Source property path of the binding */
-	UPROPERTY()
-	FStateTreeEditorPropertyPath SourcePath;
-
-	/** Target property path of the binding */
-	UPROPERTY()
-	FStateTreeEditorPropertyPath TargetPath;
-};
-
-
 /**
  * Editor representation of a all property bindings in a StateTree
  */
@@ -42,19 +19,19 @@ struct STATETREEEDITORMODULE_API FStateTreeEditorPropertyBindings
 	 * @param SourcePath Binding source property path.
 	 * @param TargetPath Binding target property path.
 	 */
-	void AddPropertyBinding(const FStateTreeEditorPropertyPath& SourcePath, const FStateTreeEditorPropertyPath& TargetPath);
+	void AddPropertyBinding(const FStateTreePropertyPath& SourcePath, const FStateTreePropertyPath& TargetPath);
 	
 	/**
 	 * Removes all bindings to target path.
 	 * @param TargetPath Target property path.
 	 */ 
-	void RemovePropertyBindings(const FStateTreeEditorPropertyPath& TargetPath);
+	void RemovePropertyBindings(const FStateTreePropertyPath& TargetPath);
 	
 	/**
 	 * @param TargetPath Target property path.
 	 * @return True of the target path has any bindings.
 	 */
-	bool HasPropertyBinding(const FStateTreeEditorPropertyPath& TargetPath) const;
+	bool HasPropertyBinding(const FStateTreePropertyPath& TargetPath) const;
 
 	/**
 	 * Copies property bindings from an existing struct to another.
@@ -67,28 +44,47 @@ struct STATETREEEDITORMODULE_API FStateTreeEditorPropertyBindings
 	/**
 	 * @return Source path for given target path, or null if binding does not exists.
 	 */
-	const FStateTreeEditorPropertyPath* GetPropertyBindingSource(const FStateTreeEditorPropertyPath& TargetPath) const;
+	const FStateTreePropertyPath* GetPropertyBindingSource(const FStateTreePropertyPath& TargetPath) const;
 	
 	/**
 	 * Returns all bindings for a specified structs based in struct ID.
 	 * @param StructID ID of the struct to find bindings for.
 	 * @param OutBindings Bindings for specified struct.
 	 */
-	void GetPropertyBindingsFor(const FGuid StructID, TArray<FStateTreeEditorPropertyBinding>& OutBindings) const;
+	void GetPropertyBindingsFor(const FGuid StructID, TArray<FStateTreePropertyPathBinding>& OutBindings) const;
 	
 	/**
 	 * Removes bindings which do not point to valid structs IDs.
 	 * @param ValidStructs Set of struct IDs that are currently valid.
 	 */
-	void RemoveUnusedBindings(const TMap<FGuid, const UStruct*>& ValidStructs);
+	void RemoveUnusedBindings(const TMap<FGuid, const FStateTreeDataView>& ValidStructs);
 
 	/** @return array view to all bindings. */
-	TConstArrayView<FStateTreeEditorPropertyBinding> GetBindings() const { return PropertyBindings; }
+	TConstArrayView<FStateTreePropertyPathBinding> GetBindings() const { return PropertyBindings; }
 
-protected:
+	TArrayView<FStateTreePropertyPathBinding> GetMutableBindings() { return PropertyBindings; }
+
+PRAGMA_DISABLE_DEPRECATION_WARNINGS
+	UE_DEPRECATED(5.3, "Use version with FStateTreePropertyPath instead.")
+	void AddPropertyBinding(const FStateTreeEditorPropertyPath& SourcePath, const FStateTreeEditorPropertyPath& TargetPath);
+
+	UE_DEPRECATED(5.3, "Use version with FStateTreePropertyPath instead.")
+	void RemovePropertyBindings(const FStateTreeEditorPropertyPath& TargetPath);
+
+	UE_DEPRECATED(5.3, "Use version with FStateTreePropertyPath instead.")
+	bool HasPropertyBinding(const FStateTreeEditorPropertyPath& TargetPath) const;
+
+	UE_DEPRECATED(5.3, "Use version with FStateTreePropertyPath instead.")
+	const FStateTreeEditorPropertyPath* GetPropertyBindingSource(const FStateTreeEditorPropertyPath& TargetPath) const;
+
+	UE_DEPRECATED(5.3, "Use RemoveUnusedBindings with values instead.")
+	void RemoveUnusedBindings(const TMap<FGuid, const UStruct*>& ValidStructs);
+PRAGMA_ENABLE_DEPRECATION_WARNINGS
+	
+private:
 
 	UPROPERTY()
-	TArray<FStateTreeEditorPropertyBinding> PropertyBindings;
+	TArray<FStateTreePropertyPathBinding> PropertyBindings;
 };
 
 
@@ -129,11 +125,13 @@ struct STATETREEEDITORMODULE_API FStateTreeBindingLookup : public IStateTreeBind
 {
 	FStateTreeBindingLookup(IStateTreeEditorPropertyBindingsOwner* InBindingOwner);
 
-	virtual const FStateTreeEditorPropertyPath* GetPropertyBindingSource(const FStateTreeEditorPropertyPath& InTargetPath) const override;
-	virtual FText GetPropertyPathDisplayName(const FStateTreeEditorPropertyPath& InTargetPath) const override;
-	virtual const FProperty* GetPropertyPathLeafProperty(const FStateTreeEditorPropertyPath& InPath) const override;
-
 	IStateTreeEditorPropertyBindingsOwner* BindingOwner = nullptr;
+
+protected:
+	virtual const FStateTreePropertyPath* GetPropertyBindingSource(const FStateTreePropertyPath& InTargetPath) const override;
+	virtual FText GetPropertyPathDisplayName(const FStateTreePropertyPath& InTargetPath) const override;
+	virtual const FProperty* GetPropertyPathLeafProperty(const FStateTreePropertyPath& InPath) const override;
+
 };
 
 #if UE_ENABLE_INCLUDE_ORDER_DEPRECATED_IN_5_2
diff --git a/Engine/Plugins/Runtime/StateTree/Source/StateTreeEditorModule/Public/StateTreePropertyBindingCompiler.h b/Engine/Plugins/Runtime/StateTree/Source/StateTreeEditorModule/Public/StateTreePropertyBindingCompiler.h
--- a/Engine/Plugins/Runtime/StateTree/Source/StateTreeEditorModule/Public/StateTreePropertyBindingCompiler.h
+++ b/Engine/Plugins/Runtime/StateTree/Source/StateTreeEditorModule/Public/StateTreePropertyBindingCompiler.h
@@ -7,7 +7,7 @@
 
 enum class EPropertyAccessCompatibility;
 struct FStateTreeCompilerLog;
-struct FStateTreeEditorPropertyBinding;
+struct FStateTreePropertyPathBinding;
 
 /**
  * Helper class to compile editor representation of property bindings into runtime representation.
@@ -32,7 +32,7 @@ struct STATETREEEDITORMODULE_API FStateTreePropertyBindingCompiler
 	  * @param OutBatchIndex - Resulting batch index, if index is INDEX_NONE, no bindings were found and no batch was generated.
 	  * @return True on success, false on failure.
 	 */
-	[[nodiscard]] bool CompileBatch(const FStateTreeBindableStructDesc& TargetStruct, TConstArrayView<FStateTreeEditorPropertyBinding> EditorPropertyBindings, int32& OutBatchIndex);
+	[[nodiscard]] bool CompileBatch(const FStateTreeBindableStructDesc& TargetStruct, TConstArrayView<FStateTreePropertyPathBinding> PropertyBindings, int32& OutBatchIndex);
 
 	/** Finalizes compilation, should be called once all batches are compiled. */
 	void Finalize();
@@ -54,6 +54,7 @@ struct STATETREEEDITORMODULE_API FStateTreePropertyBindingCompiler
 		return SourceStructs[Index];
 	}
 
+PRAGMA_DISABLE_DEPRECATION_WARNINGS
 	/**
 	 * Resolves a string based property path in specified struct into segments of property names and access types.
 	 * If logging is required both, Log and LogContextStruct needs to be non-null.
@@ -66,6 +67,7 @@ struct STATETREEEDITORMODULE_API FStateTreePropertyBindingCompiler
 	 * @param LogContextStruct Pointer to bindable struct desc where the property path belongs to.
 	 * @return True of the property was solved successfully.
 	 */
+	UE_DEPRECATED(5.3, "Use FStateTreePropertyPath resolve indirection methods instead.")
 	[[nodiscard]] static bool ResolvePropertyPath(const FStateTreeBindableStructDesc& InStructDesc, const FStateTreeEditorPropertyPath& InPath,
 												  TArray<FStateTreePropertySegment>& OutSegments, const FProperty*& OutLeafProperty, int32& OutLeafArrayIndex,
 												  FStateTreeCompilerLog* Log = nullptr, const FStateTreeBindableStructDesc* LogContextStruct = nullptr);
@@ -76,14 +78,14 @@ struct STATETREEEDITORMODULE_API FStateTreePropertyBindingCompiler
 	 * @param ToProperty Property to copy to.
 	 * @return Incompatible if the properties cannot be copied, Compatible if they are trivially copyable, or Promotable if numeric values can be promoted to another numeric type.
 	 */
+	UE_DEPRECATED(5.3, "Use FStateTreePropertyBindings::GetPropertyCompatibility instead.")
 	static EPropertyAccessCompatibility GetPropertyCompatibility(const FProperty* FromProperty, const FProperty* ToProperty);
-	
+PRAGMA_ENABLE_DEPRECATION_WARNINGS
+
 protected:
 
 	void StoreSourceStructs();
 
-	EStateTreePropertyCopyType GetCopyType(const UStruct* SourceStruct, const FProperty* SourceProperty, const int32 SourceArrayIndex, const UStruct* TargetStruct, const FProperty* TargetProperty, const int32 TargetArrayIndex) const;
-
 	UPROPERTY()
 	TArray<FStateTreeBindableStructDesc> SourceStructs;
 
diff --git a/Engine/Plugins/Runtime/StateTree/Source/StateTreeModule/Private/Conditions/StateTreeCommonConditions.cpp b/Engine/Plugins/Runtime/StateTree/Source/StateTreeModule/Private/Conditions/StateTreeCommonConditions.cpp
--- a/Engine/Plugins/Runtime/StateTree/Source/StateTreeModule/Private/Conditions/StateTreeCommonConditions.cpp
+++ b/Engine/Plugins/Runtime/StateTree/Source/StateTreeModule/Private/Conditions/StateTreeCommonConditions.cpp
@@ -127,9 +127,9 @@ bool FStateTreeCompareEnumCondition::TestCondition(FStateTreeExecutionContext& C
 }
 
 #if WITH_EDITOR
-void FStateTreeCompareEnumCondition::OnBindingChanged(const FGuid& ID, FStateTreeDataView InstanceData, const FStateTreeEditorPropertyPath& SourcePath, const FStateTreeEditorPropertyPath& TargetPath, const IStateTreeBindingLookup& BindingLookup)
+void FStateTreeCompareEnumCondition::OnBindingChanged(const FGuid& ID, FStateTreeDataView InstanceData, const FStateTreePropertyPath& SourcePath, const FStateTreePropertyPath& TargetPath, const IStateTreeBindingLookup& BindingLookup)
 {
-	if (!TargetPath.IsValid())
+	if (!TargetPath.GetStructID().IsValid())
 	{
 		return;
 	}
@@ -137,7 +137,8 @@ void FStateTreeCompareEnumCondition::OnBindingChanged(const FGuid& ID, FStateTre
 	FInstanceDataType& Instance = InstanceData.GetMutable<FInstanceDataType>();
 
 	// Left has changed, update enums from the leaf property.
-	if (TargetPath.Path.Last() == GET_MEMBER_NAME_STRING_CHECKED(FInstanceDataType, Left))
+	if (!TargetPath.IsPathEmpty()
+		&& TargetPath.GetSegments().Last().GetName() == GET_MEMBER_NAME_CHECKED(FInstanceDataType, Left))
 	{
 		if (const FProperty* LeafProperty = BindingLookup.GetPropertyPathLeafProperty(SourcePath))
 		{
diff --git a/Engine/Plugins/Runtime/StateTree/Source/StateTreeModule/Private/StateTree.cpp b/Engine/Plugins/Runtime/StateTree/Source/StateTreeModule/Private/StateTree.cpp
--- a/Engine/Plugins/Runtime/StateTree/Source/StateTreeModule/Private/StateTree.cpp
+++ b/Engine/Plugins/Runtime/StateTree/Source/StateTreeModule/Private/StateTree.cpp
@@ -222,10 +222,9 @@ bool UStateTree::Link()
 		}
 
 		// Update property bag structs before resolving binding.
-		const TArrayView<FStateTreeBindableStructDesc> SourceStructs = PropertyBindings.GetSourceStructs();
-		const TArrayView<FStateTreePropCopyBatch> CopyBatches = PropertyBindings.GetCopyBatches();
+		const TArrayView<FStateTreeBindableStructDesc> SourceStructs = PropertyBindings.SourceStructs;
+		const TArrayView<FStateTreePropertyCopyBatch> CopyBatches = PropertyBindings.CopyBatches;
 
-		
 		// Reconcile out of date classes.
 		for (FStateTreeBindableStructDesc& SourceStruct : SourceStructs)
 		{
@@ -237,7 +236,7 @@ bool UStateTree::Link()
 				}
 			}
 		}
-		for (FStateTreePropCopyBatch& CopyBatch : CopyBatches)
+		for (FStateTreePropertyCopyBatch& CopyBatch : CopyBatches)
 		{
 			if (const UClass* TargetClass = Cast<UClass>(CopyBatch.TargetStruct.Struct))
 			{
@@ -296,7 +295,7 @@ bool UStateTree::Link()
 
 				if (Params.BindingsBatch.IsValid())
 				{
-					FStateTreePropCopyBatch& Batch = CopyBatches[Params.BindingsBatch.Get()];
+					FStateTreePropertyCopyBatch& Batch = CopyBatches[Params.BindingsBatch.Get()];
 					Batch.TargetStruct.Struct = Params.Parameters.GetPropertyBagStruct();
 				}
 			}
diff --git a/Engine/Plugins/Runtime/StateTree/Source/StateTreeModule/Private/StateTreePropertyBindings.cpp b/Engine/Plugins/Runtime/StateTree/Source/StateTreeModule/Private/StateTreePropertyBindings.cpp
--- a/Engine/Plugins/Runtime/StateTree/Source/StateTreeModule/Private/StateTreePropertyBindings.cpp
+++ b/Engine/Plugins/Runtime/StateTree/Source/StateTreeModule/Private/StateTreePropertyBindings.cpp
@@ -1,9 +1,82 @@
 // Copyright Epic Games, Inc. All Rights Reserved.
 #include "StateTreePropertyBindings.h"
 #include "UObject/EnumProperty.h"
+#include "Misc/EnumerateRange.h"
+#include "PropertyPathHelpers.h"
 
 #include UE_INLINE_GENERATED_CPP_BY_NAME(StateTreePropertyBindings)
 
+namespace UE::StateTree::Private
+{
+
+#if WITH_EDITORONLY_DATA
+PRAGMA_DISABLE_DEPRECATION_WARNINGS
+	FStateTreePropertyPath ConvertEditorPath(const FStateTreeEditorPropertyPath& InEditorPath)
+	{
+		FStateTreePropertyPath Path;
+		Path.SetStructID(InEditorPath.StructID);
+
+		for (const FString& Segment : InEditorPath.Path)
+		{
+			const TCHAR* PropertyNamePtr = nullptr;
+			int32 PropertyNameLength = 0;
+			int32 ArrayIndex = INDEX_NONE;
+			PropertyPathHelpers::FindFieldNameAndArrayIndex(Segment.Len(), *Segment, PropertyNameLength, &PropertyNamePtr, ArrayIndex);
+			FString PropertyNameString(PropertyNameLength, PropertyNamePtr);
+			const FName PropertyName(*PropertyNameString, FNAME_Find);
+			Path.AddPathSegment(PropertyName, ArrayIndex);
+		}
+		return Path;
+	}
+
+	FStateTreeEditorPropertyPath ConvertEditorPath(const FStateTreePropertyPath& InPath)
+	{
+		FStateTreeEditorPropertyPath Path;
+		Path.StructID = InPath.GetStructID();
+		for (const FStateTreePropertyPathSegment& Segment : InPath.GetSegments())
+		{
+			if (Segment.GetArrayIndex() != INDEX_NONE)
+			{
+				Path.Path.Add(FString::Printf(TEXT("%s[%d]"), *Segment.GetName().ToString(), Segment.GetArrayIndex()));
+			}
+			else
+			{
+				Path.Path.Add(Segment.GetName().ToString());
+			}
+		}
+		return Path;
+	}
+PRAGMA_ENABLE_DEPRECATION_WARNINGS
+#endif // WITH_EDITORONLY_DATA
+}; // UE::StateTree::Private 
+
+
+//----------------------------------------------------------------//
+//  FStateTreePropertyPathBinding
+//----------------------------------------------------------------//
+void FStateTreePropertyPathBinding::PostSerialize(const FArchive& Ar)
+{
+#if WITH_EDITORONLY_DATA
+PRAGMA_DISABLE_DEPRECATION_WARNINGS
+	if (SourcePath_DEPRECATED.IsValid())
+	{
+		SourcePropertyPath = UE::StateTree::Private::ConvertEditorPath(SourcePath_DEPRECATED);
+		SourcePath_DEPRECATED.StructID = FGuid();
+		SourcePath_DEPRECATED.Path.Reset();
+	}
+
+	if (TargetPath_DEPRECATED.IsValid())
+	{
+		TargetPropertyPath = UE::StateTree::Private::ConvertEditorPath(TargetPath_DEPRECATED);
+		TargetPath_DEPRECATED.StructID = FGuid();
+		TargetPath_DEPRECATED.Path.Reset();
+	}
+PRAGMA_ENABLE_DEPRECATION_WARNINGS
+#endif // WITH_EDITORONLY_DATA
+
+}
+
+
 //----------------------------------------------------------------//
 //  FStateTreePropertyBindings
 //----------------------------------------------------------------//
@@ -12,33 +85,39 @@ void FStateTreePropertyBindings::Reset()
 {
 	SourceStructs.Reset();
 	CopyBatches.Reset();
-	PropertyBindings.Reset();
-	PropertySegments.Reset();
+	PropertyPathBindings.Reset();
 	PropertyCopies.Reset();
 	PropertyIndirections.Reset();
+	
 	bBindingsResolved = false;
 }
 
 bool FStateTreePropertyBindings::ResolvePaths()
 {
 	PropertyIndirections.Reset();
-	PropertyCopies.SetNum(PropertyBindings.Num());
+	PropertyCopies.SetNum(PropertyPathBindings.Num());
 
 	bBindingsResolved = true;
 
 	bool bResult = true;
 	
-	for (const FStateTreePropCopyBatch& Batch : CopyBatches)
+	for (const FStateTreePropertyCopyBatch& Batch : CopyBatches)
 	{
 		for (int32 i = Batch.BindingsBegin; i != Batch.BindingsEnd; i++)
 		{
-			const FStateTreePropertyBinding& Binding = PropertyBindings[i];
-			FStateTreePropCopy& Copy = PropertyCopies[i];
+			const FStateTreePropertyPathBinding& Binding = PropertyPathBindings[i];
+			FStateTreePropertyCopy& Copy = PropertyCopies[i];
+
+			if (!Binding.GetCompiledSourceStructIndex().IsValid())
+			{
+				UE_LOG(LogStateTree, Error, TEXT("%hs: '%s' Invalid source struct for property binding %s."), __FUNCTION__, *Binding.GetSourcePath().ToString());
+				bResult = false;
+				continue;
+			}
 
-			Copy.SourceStructIndex = Binding.SourceStructIndex;
-			Copy.Type = Binding.CopyType;
+			Copy.SourceStructIndex = Binding.GetCompiledSourceStructIndex();
 
-			const UStruct* SourceStruct = SourceStructs[Binding.SourceStructIndex.Get()].Struct;
+			const UStruct* SourceStruct = SourceStructs[Copy.SourceStructIndex.Get()].Struct;
 			const UStruct* TargetStruct = Batch.TargetStruct.Struct;
 			if (!SourceStruct || !TargetStruct)
 			{
@@ -50,10 +129,12 @@ bool FStateTreePropertyBindings::ResolvePaths()
 
 			// Resolve paths and validate the copy. Stops on first failure.
 			bool bSuccess = true;
-			bSuccess = bSuccess && ResolvePath(SourceStruct, Binding.SourcePath, Copy.SourceIndirection, Copy.SourceLeafProperty);
-			bSuccess = bSuccess && ResolvePath(TargetStruct, Binding.TargetPath, Copy.TargetIndirection, Copy.TargetLeafProperty);
-			bSuccess = bSuccess && ValidateCopy(Copy);
-			if (!bSuccess)
+			FStateTreePropertyPathIndirection SourceLeafIndirection;
+			FStateTreePropertyPathIndirection TargetLeafIndirection;
+			bSuccess = bSuccess && ResolvePath(SourceStruct, Binding.GetSourcePath(), Copy.SourceIndirection, SourceLeafIndirection);
+			bSuccess = bSuccess && ResolvePath(TargetStruct, Binding.GetTargetPath(), Copy.TargetIndirection, TargetLeafIndirection);
+			bSuccess = bSuccess && ResolveCopyType(Copy, SourceLeafIndirection, TargetLeafIndirection);
+			if (!bSuccess) 
 			{
 				// Resolving or validating failed, make the copy a nop.
 				Copy.Type = EStateTreePropertyCopyType::None;
@@ -65,113 +146,122 @@ bool FStateTreePropertyBindings::ResolvePaths()
 	return bResult;
 }
 
-bool FStateTreePropertyBindings::ResolvePath(const UStruct* Struct, const FStateTreePropertySegment& FirstPathSegment, FStateTreePropertyIndirection& OutFirstIndirection, const FProperty*& OutLeafProperty)
+bool FStateTreePropertyBindings::ResolvePath(const UStruct* Struct, const FStateTreePropertyPath& Path, FStateTreePropertyIndirection& OutFirstIndirection, FStateTreePropertyPathIndirection& OutLeafIndirection)
 {
 	if (!Struct)
+ 	{
+		UE_LOG(LogStateTree, Error, TEXT("%hs: '%s' Invalid source struct."), __FUNCTION__, *Path.ToString());
+		return false;
+	}
+
+	FString Error;
+	TArray<FStateTreePropertyPathIndirection> PathIndirections;
+	if (!Path.ResolveIndirections(Struct, PathIndirections, &Error))
 	{
-		UE_LOG(LogStateTree, Error, TEXT("%s: '%s' Invalid source struct."), ANSI_TO_TCHAR(__FUNCTION__), *GetPathAsString(FirstPathSegment));
+		UE_LOG(LogStateTree, Error, TEXT("%hs: %s"), __FUNCTION__, *Error);
 		return false;
 	}
 
-	const UStruct* CurrentStruct = Struct;
-	const FProperty* LeafProperty = nullptr;
-	TArray<FStateTreePropertyIndirection, TInlineAllocator<16>> TempIndirections;
-	const FStateTreePropertySegment* Segment = &FirstPathSegment;
+	// Casts array index to FStateTreeIndex16 (clamping INDEX_NONE to 0) or returns false if out if index bounds.
+	auto CastArrayIndexToIndex16 = [](const int32 Index)
+	{
+		const int32 ClampedIndex = FMath::Max(0, Index);
+		if (!FStateTreeIndex16::IsValidIndex(ClampedIndex))
+		{
+			return FStateTreeIndex16();
+		}
+		return FStateTreeIndex16(ClampedIndex);
+	};
 
-	while (Segment != nullptr && !Segment->IsEmpty())
+	TArray<FStateTreePropertyIndirection, TInlineAllocator<16>> TempIndirections;
+	for (FStateTreePropertyPathIndirection& PathIndirection : PathIndirections)
 	{
 		FStateTreePropertyIndirection& Indirection = TempIndirections.AddDefaulted_GetRef();
 
-		const bool bFinalSegment = Segment->NextIndex.IsValid() == false;;
-
-		if (!ensure(CurrentStruct))
-		{
-			UE_LOG(LogStateTree, Error, TEXT("%s: '%s' Invalid struct."), ANSI_TO_TCHAR(__FUNCTION__), *GetPathAsString(FirstPathSegment, Segment, TEXT("<"), TEXT(">")));
-			return false;
-		}
-		FProperty* Property = CurrentStruct->FindPropertyByName(Segment->Name);
-		if (!Property)
-		{
-			// TODO: use core redirects to fix up the name.
-			UE_LOG(LogStateTree, Error, TEXT("%s: Malformed path '%s', could not to find property '%s%s.%s'."),
-				ANSI_TO_TCHAR(__FUNCTION__), *GetPathAsString(FirstPathSegment, Segment, TEXT("<"), TEXT(">")),
-				CurrentStruct->GetPrefixCPP(), *CurrentStruct->GetName(), *Segment->Name.ToString());
-			return false;
-		}
-		Indirection.ArrayIndex = FStateTreeIndex16(Segment->ArrayIndex.IsValid() ? Segment->ArrayIndex.Get() : 0);
-		Indirection.Type = Segment->Type;
-		Indirection.Offset = IntCastChecked<uint16>(Property->GetOffset_ForInternal() + Property->ElementSize * Indirection.ArrayIndex.Get());
+		Indirection.Offset = PathIndirection.GetPropertyOffset();
+		Indirection.Type = PathIndirection.GetAccessType();
 
-		// Check to see if it is an array access first.
-		if (const FArrayProperty* ArrayProperty = CastField<FArrayProperty>(Property))
+		if (Indirection.Type == EStateTreePropertyAccessType::IndexArray)
 		{
-			if (const FStructProperty* ArrayOfStructsProperty = CastField<FStructProperty>(ArrayProperty->Inner))
+			if (const FArrayProperty* ArrayProperty = CastField<FArrayProperty>(PathIndirection.GetProperty()))
 			{
-				CurrentStruct = ArrayOfStructsProperty->Struct;
+				Indirection.ArrayProperty = ArrayProperty;
+				Indirection.ArrayIndex = CastArrayIndexToIndex16(PathIndirection.GetArrayIndex());
+				if (!Indirection.ArrayIndex.IsValid())
+				{
+					UE_LOG(LogStateTree, Error, TEXT("%hs: Array index %d at %s, is too large."),
+						__FUNCTION__, PathIndirection.GetArrayIndex(), *Path.ToString(PathIndirection.GetPathSegmentIndex(), TEXT("<"), TEXT(">")));
+					return false;
+				}
 			}
-			else if (const FObjectPropertyBase* ArrayOfObjectsProperty = CastField<FObjectPropertyBase>(ArrayProperty->Inner))
+			else
 			{
-				CurrentStruct = ArrayOfObjectsProperty->PropertyClass;
+				UE_LOG(LogStateTree, Error, TEXT("%hs: Expect property %s to be array property."),
+					__FUNCTION__, *Path.ToString(PathIndirection.GetPathSegmentIndex(), TEXT("<"), TEXT(">")));
+				return false;
 			}
-			Indirection.ArrayProperty = ArrayProperty;
-		}
-		// Leaf segments all get treated the same, plain, struct or object
-		else if (bFinalSegment)
-		{
-			CurrentStruct = nullptr;
-		}
-		// Check to see if this is a simple structure (eg. not an array of structures)
-		else if (const FStructProperty* StructProperty = CastField<FStructProperty>(Property))
-		{
-			CurrentStruct = StructProperty->Struct;
-		}
-		// Check to see if this is a simple object (eg. not an array of objects)
-		else if (const FObjectProperty* ObjectProperty = CastField<FObjectProperty>(Property))
-		{
-			CurrentStruct = ObjectProperty->PropertyClass;
 		}
-		// Check to see if this is a simple weak object property (eg. not an array of weak objects).
-		else if (const FWeakObjectProperty* WeakObjectProperty = CastField<FWeakObjectProperty>(Property))
+		else if (Indirection.Type == EStateTreePropertyAccessType::StructInstance
+				|| Indirection.Type == EStateTreePropertyAccessType::ObjectInstance)
 		{
-			CurrentStruct = WeakObjectProperty->PropertyClass;
-		}
-		// Check to see if this is a simple soft object property (eg. not an array of soft objects).
-		else if (const FSoftObjectProperty* SoftObjectProperty = CastField<FSoftObjectProperty>(Property))
-		{
-			CurrentStruct = SoftObjectProperty->PropertyClass;
-		}
-		else
-		{
-			UE_LOG(LogStateTree, Error, TEXT("%s: Unsupported segment %s in path '%s'."),
-			    ANSI_TO_TCHAR(__FUNCTION__), *Property->GetCPPType(), *GetPathAsString(FirstPathSegment, Segment, TEXT("<"), TEXT(">")));
-		}
-
-		if (bFinalSegment)
-		{
-			LeafProperty = Property;
+			if (PathIndirection.GetInstanceStruct())
+			{
+				Indirection.InstanceStruct = PathIndirection.GetInstanceStruct();
+			}
+			else
+			{
+				UE_LOG(LogStateTree, Error, TEXT("%hs: Expect instanced property access %s to have instance type specified."),
+					__FUNCTION__, *Path.ToString(PathIndirection.GetPathSegmentIndex(), TEXT("<"), TEXT(">")));
+				return false;
+			}
 		}
-
-		Segment = Segment->NextIndex.IsValid() ? &PropertySegments[Segment->NextIndex.Get()] : nullptr;
 	}
 
 	if (TempIndirections.Num() > 0)
 	{
-		// Collapse adjacent offset indirections
 		for (int32 Index = 0; Index < TempIndirections.Num(); Index++)
 		{
 			FStateTreePropertyIndirection& Indirection = TempIndirections[Index];
-			if (Indirection.Type == EStateTreePropertyAccessType::Offset && (Index + 1) < TempIndirections.Num())
+			if ((Index + 1) < TempIndirections.Num())
 			{
 				const FStateTreePropertyIndirection& NextIndirection = TempIndirections[Index + 1];
-				if (NextIndirection.Type == EStateTreePropertyAccessType::Offset)
+				if (Indirection.Type == EStateTreePropertyAccessType::Offset
+					&& NextIndirection.Type == EStateTreePropertyAccessType::Offset)
 				{
+					// Collapse adjacent offset indirections
 					Indirection.Offset += NextIndirection.Offset;
 					TempIndirections.RemoveAt(Index + 1);
 					Index--;
 				}
+				else if (Indirection.Type == EStateTreePropertyAccessType::IndexArray
+					&& NextIndirection.Type == EStateTreePropertyAccessType::Offset
+					&& NextIndirection.Offset == 0)
+				{
+					// Remove empty offset after array indexing.
+					TempIndirections.RemoveAt(Index + 1);
+					Index--;
+				}
+				else if (Indirection.Type == EStateTreePropertyAccessType::StructInstance
+					&& NextIndirection.Type == EStateTreePropertyAccessType::Offset
+					&& NextIndirection.Offset == 0)
+				{
+					// Remove empty offset after struct indirection.
+					TempIndirections.RemoveAt(Index + 1);
+					Index--;
+				}
+				else if ((Indirection.Type == EStateTreePropertyAccessType::Object
+						|| Indirection.Type == EStateTreePropertyAccessType::ObjectInstance)
+					&& NextIndirection.Type == EStateTreePropertyAccessType::Offset
+					&& NextIndirection.Offset == 0)
+				{
+					// Remove empty offset after object indirection.
+					TempIndirections.RemoveAt(Index + 1);
+					Index--;
+				}
 			}
 		}
 
+		OutLeafIndirection = PathIndirections.Last(); 
 
 		// Store indirections
 		OutFirstIndirection = TempIndirections[0];
@@ -190,29 +280,75 @@ bool FStateTreePropertyBindings::ResolvePath(const UStruct* Struct, const FState
 		// Zero offset will return the struct itself.
 		OutFirstIndirection.Offset = 0;
 		OutFirstIndirection.Type = EStateTreePropertyAccessType::Offset;
+
+		OutLeafIndirection = FStateTreePropertyPathIndirection(Struct);
 	}
-	
-	OutLeafProperty = LeafProperty;
 
 	return true;
 }
 
-bool FStateTreePropertyBindings::ValidateCopy(FStateTreePropCopy& Copy) const
+bool FStateTreePropertyBindings::ResolveCopyType(FStateTreePropertyCopy& Copy, const FStateTreePropertyPathIndirection& SourceIndirection,const FStateTreePropertyPathIndirection& TargetIndirection)
 {
-	const FProperty* SourceProperty = Copy.SourceLeafProperty;
-	const FProperty* TargetProperty = Copy.TargetLeafProperty;
+	// @todo: see if GetPropertyCompatibility() can be implemented as call to ResolveCopyType() instead so that we write this logic just once.
+	
+	const FProperty* SourceProperty = SourceIndirection.GetProperty();
+	const UStruct* SourceStruct = SourceIndirection.GetContainerStruct();
+	
+	const FProperty* TargetProperty = TargetIndirection.GetProperty();
+	const UStruct* TargetStruct = TargetIndirection.GetContainerStruct();
 
-	if (!TargetProperty)
+	if (!SourceStruct || !TargetStruct)
 	{
 		return false;
 	}
 
-	// If source property is nullptr, we're binding directly to the source.
+	Copy.SourceLeafProperty = SourceProperty;
+	Copy.TargetLeafProperty = TargetProperty;
+	Copy.CopySize = 0;
+	Copy.Type = EStateTreePropertyCopyType::None;
+	
 	if (SourceProperty == nullptr)
 	{
-		return TargetProperty->IsA<FStructProperty>() || TargetProperty->IsA<FObjectPropertyBase>();
+		// Copy directly from the source struct, target must be.
+		if (const FStructProperty* TargetStructProperty = CastField<FStructProperty>(TargetProperty))
+		{
+			if (TargetStructProperty->Struct == SourceStruct)
+			{
+				Copy.Type = EStateTreePropertyCopyType::CopyStruct;
+				return true;
+			}
+		}
+		else if (const FObjectPropertyBase* TargetObjectProperty = CastField<FObjectPropertyBase>(TargetProperty))
+		{
+			if (SourceStruct->IsChildOf(TargetObjectProperty->PropertyClass))
+			{
+				Copy.Type = EStateTreePropertyCopyType::CopyObject;
+				return true;
+			}
+		}
+
+		return false;
+	}
+
+	// Handle FStateTreeStructRef
+	if (const FStructProperty* TargetStructProperty = CastField<const FStructProperty>(TargetProperty))
+	{
+		if (TargetStructProperty->Struct == TBaseStructure<FStateTreeStructRef>::Get())
+		{
+			if (const FStructProperty* SourceStructProperty = CastField<const FStructProperty>(SourceProperty))
+			{
+				// FStateTreeStructRef to FStateTreeStructRef is copied as usual.
+				if (SourceStructProperty->Struct != TBaseStructure<FStateTreeStructRef>::Get())
+				{
+					Copy.Type = EStateTreePropertyCopyType::StructReference;
+					return true;
+				}
+			}
+		}
 	}
 
+	const EStateTreePropertyAccessCompatibility Compatibility = FStateTreePropertyBindings::GetPropertyCompatibility(SourceProperty, TargetProperty);
+
 	// Extract underlying types for enums
 	if (const FEnumProperty* EnumPropertyA = CastField<const FEnumProperty>(SourceProperty))
 	{
@@ -224,134 +360,304 @@ bool FStateTreePropertyBindings::ValidateCopy(FStateTreePropCopy& Copy) const
 		TargetProperty = EnumPropertyB->GetUnderlyingProperty();
 	}
 
-	bool bResult = true;
-	switch (Copy.Type)
+	if (Compatibility == EStateTreePropertyAccessCompatibility::Compatible)
 	{
-	case EStateTreePropertyCopyType::CopyPlain:
-		Copy.CopySize = SourceProperty->ElementSize * SourceProperty->ArrayDim;
-		bResult = (SourceProperty->PropertyFlags & CPF_IsPlainOldData) != 0 && (TargetProperty->PropertyFlags & CPF_IsPlainOldData) != 0;
-		break;
-	case EStateTreePropertyCopyType::CopyComplex:
-		bResult = true;
-		break;
-	case EStateTreePropertyCopyType::CopyBool:
-		bResult = SourceProperty->IsA<FBoolProperty>() && TargetProperty->IsA<FBoolProperty>();
-		break;
-	case EStateTreePropertyCopyType::CopyStruct:
-		bResult = SourceProperty->IsA<FStructProperty>() && TargetProperty->IsA<FStructProperty>();
-		break;
-	case EStateTreePropertyCopyType::CopyObject:
-		bResult = SourceProperty->IsA<FObjectPropertyBase>() && TargetProperty->IsA<FObjectPropertyBase>();
-		break;
-	case EStateTreePropertyCopyType::CopyName:
-		bResult = SourceProperty->IsA<FNameProperty>() && TargetProperty->IsA<FNameProperty>();
-		break;
-	case EStateTreePropertyCopyType::CopyFixedArray:
-		bResult = SourceProperty->IsA<FArrayProperty>() && TargetProperty->IsA<FArrayProperty>();
-		break;
-	case EStateTreePropertyCopyType::StructReference:
+		if (CastField<FNameProperty>(TargetProperty))
+		{
+			Copy.Type = EStateTreePropertyCopyType::CopyName;
+			return true;
+		}
+		else if (CastField<FBoolProperty>(TargetProperty))
+		{
+			Copy.Type = EStateTreePropertyCopyType::CopyBool;
+			return true;
+		}
+		else if (CastField<FStructProperty>(TargetProperty))
+		{
+			Copy.Type = EStateTreePropertyCopyType::CopyStruct;
+			return true;
+		}
+		else if (CastField<FObjectPropertyBase>(TargetProperty))
+		{
+			Copy.Type = EStateTreePropertyCopyType::CopyObject;
+			return true;
+		}
+		else if (CastField<FArrayProperty>(TargetProperty) && TargetProperty->HasAnyPropertyFlags(CPF_EditFixedSize))
+		{
+			// only apply array copying rules if the destination array is fixed size, otherwise it will be 'complex'
+			Copy.Type = EStateTreePropertyCopyType::CopyFixedArray;
+			return true;
+		}
+		else if (TargetProperty->PropertyFlags & CPF_IsPlainOldData)
+		{
+			Copy.Type = EStateTreePropertyCopyType::CopyPlain;
+			Copy.CopySize = SourceProperty->ElementSize * SourceProperty->ArrayDim;
+			return true;
+		}
+		else
+		{
+			Copy.Type = EStateTreePropertyCopyType::CopyComplex;
+			return true;
+		}
+	}
+	else if (Compatibility == EStateTreePropertyAccessCompatibility::Promotable)
 	{
-		const FStructProperty* SourceStructProperty = CastField<const FStructProperty>(SourceProperty);
-		const FStructProperty* TargetStructProperty = CastField<const FStructProperty>(TargetProperty);
-
-		bResult = SourceStructProperty != nullptr
-			&& TargetStructProperty != nullptr
-			&& SourceStructProperty->Struct != TBaseStructure<FStateTreeStructRef>::Get()
-			&& TargetStructProperty->Struct == TBaseStructure<FStateTreeStructRef>::Get();
-		break;
+		if (SourceProperty->IsA<FBoolProperty>())
+		{
+			if (TargetProperty->IsA<FByteProperty>())
+			{
+				Copy.Type = EStateTreePropertyCopyType::PromoteBoolToByte;
+				return true;
+			}
+			else if (TargetProperty->IsA<FIntProperty>())
+			{
+				Copy.Type = EStateTreePropertyCopyType::PromoteBoolToInt32;
+				return true;
+			}
+			else if (TargetProperty->IsA<FUInt32Property>())
+			{
+				Copy.Type = EStateTreePropertyCopyType::PromoteBoolToUInt32;
+				return true;
+			}
+			else if (TargetProperty->IsA<FInt64Property>())
+			{
+				Copy.Type = EStateTreePropertyCopyType::PromoteBoolToInt64;
+				return true;
+			}
+			else if (TargetProperty->IsA<FFloatProperty>())
+			{
+				Copy.Type = EStateTreePropertyCopyType::PromoteBoolToFloat;
+				return true;
+			}
+			else if (TargetProperty->IsA<FDoubleProperty>())
+			{
+				Copy.Type = EStateTreePropertyCopyType::PromoteBoolToDouble;
+				return true;
+			}
+		}
+		else if (SourceProperty->IsA<FByteProperty>())
+		{
+			if (TargetProperty->IsA<FIntProperty>())
+			{
+				Copy.Type = EStateTreePropertyCopyType::PromoteByteToInt32;
+				return true;
+			}
+			else if (TargetProperty->IsA<FUInt32Property>())
+			{
+				Copy.Type = EStateTreePropertyCopyType::PromoteByteToUInt32;
+				return true;
+			}
+			else if (TargetProperty->IsA<FInt64Property>())
+			{
+				Copy.Type = EStateTreePropertyCopyType::PromoteByteToInt64;
+				return true;
+			}
+			else if (TargetProperty->IsA<FFloatProperty>())
+			{
+				Copy.Type = EStateTreePropertyCopyType::PromoteByteToFloat;
+				return true;
+			}
+			else if (TargetProperty->IsA<FDoubleProperty>())
+			{
+				Copy.Type = EStateTreePropertyCopyType::PromoteByteToDouble;
+				return true;
+			}
+		}
+		else if (SourceProperty->IsA<FIntProperty>())
+		{
+			if (TargetProperty->IsA<FInt64Property>())
+			{
+				Copy.Type = EStateTreePropertyCopyType::PromoteInt32ToInt64;
+				return true;
+			}
+			else if (TargetProperty->IsA<FFloatProperty>())
+			{
+				Copy.Type = EStateTreePropertyCopyType::PromoteInt32ToFloat;
+				return true;
+			}
+			else if (TargetProperty->IsA<FDoubleProperty>())
+			{
+				Copy.Type = EStateTreePropertyCopyType::PromoteInt32ToDouble;
+				return true;
+			}
+		}
+		else if (SourceProperty->IsA<FUInt32Property>())
+		{
+			if (TargetProperty->IsA<FInt64Property>())
+			{
+				Copy.Type = EStateTreePropertyCopyType::PromoteUInt32ToInt64;
+				return true;
+			}
+			else if (TargetProperty->IsA<FFloatProperty>())
+			{
+				Copy.Type = EStateTreePropertyCopyType::PromoteUInt32ToFloat;
+				return true;
+			}
+			else if (TargetProperty->IsA<FDoubleProperty>())
+			{
+				Copy.Type = EStateTreePropertyCopyType::PromoteUInt32ToDouble;
+				return true;
+			}
+		}
+		else if (SourceProperty->IsA<FFloatProperty>())
+		{
+			if (TargetProperty->IsA<FIntProperty>())
+			{
+				Copy.Type = EStateTreePropertyCopyType::PromoteFloatToInt32;
+				return true;
+			}
+			else if (TargetProperty->IsA<FInt64Property>())
+			{
+				Copy.Type = EStateTreePropertyCopyType::PromoteFloatToInt64;
+				return true;
+			}
+			else if (TargetProperty->IsA<FDoubleProperty>())
+			{
+				Copy.Type = EStateTreePropertyCopyType::PromoteFloatToDouble;
+				return true;
+			}
+		}
+		else if (SourceProperty->IsA<FDoubleProperty>())
+		{
+			if (TargetProperty->IsA<FIntProperty>())
+			{
+				Copy.Type = EStateTreePropertyCopyType::DemoteDoubleToInt32;
+				return true;
+			}
+			else if (TargetProperty->IsA<FInt64Property>())
+			{
+				Copy.Type = EStateTreePropertyCopyType::DemoteDoubleToInt64;
+				return true;
+			}
+			else if (TargetProperty->IsA<FFloatProperty>())
+			{
+				Copy.Type = EStateTreePropertyCopyType::DemoteDoubleToFloat;
+				return true;
+			}
+		}
 	}
+	
+	ensureMsgf(false, TEXT("Couldnt determine property copy type (%s -> %s)"), *SourceProperty->GetNameCPP(), *TargetProperty->GetNameCPP());
 
-	// Bool promotions		
-	case EStateTreePropertyCopyType::PromoteBoolToByte:
-		bResult = SourceProperty->IsA<FBoolProperty>() && TargetProperty->IsA<FByteProperty>();
-		break;
-	case EStateTreePropertyCopyType::PromoteBoolToInt32:
-		bResult = SourceProperty->IsA<FBoolProperty>() && TargetProperty->IsA<FIntProperty>();
-		break;
-	case EStateTreePropertyCopyType::PromoteBoolToUInt32:
-		bResult = SourceProperty->IsA<FBoolProperty>() && TargetProperty->IsA<FUInt32Property>();
-		break;
-	case EStateTreePropertyCopyType::PromoteBoolToInt64:
-		bResult = SourceProperty->IsA<FBoolProperty>() && TargetProperty->IsA<FInt64Property>();
-		break;
-	case EStateTreePropertyCopyType::PromoteBoolToFloat:
-		bResult = SourceProperty->IsA<FBoolProperty>() && TargetProperty->IsA<FFloatProperty>();
-		break;
-	case EStateTreePropertyCopyType::PromoteBoolToDouble:
-		bResult = SourceProperty->IsA<FBoolProperty>() && TargetProperty->IsA<FDoubleProperty>();
-		break;
+	return false;
+}
 
-	// Byte promotions
-	case EStateTreePropertyCopyType::PromoteByteToInt32:
-		bResult = SourceProperty->IsA<FByteProperty>() && TargetProperty->IsA<FIntProperty>();
-		break;
-	case EStateTreePropertyCopyType::PromoteByteToUInt32:
-		bResult = SourceProperty->IsA<FByteProperty>() && TargetProperty->IsA<FUInt32Property>();
-		break;
-	case EStateTreePropertyCopyType::PromoteByteToInt64:
-		bResult = SourceProperty->IsA<FByteProperty>() && TargetProperty->IsA<FInt64Property>();
-		break;
-	case EStateTreePropertyCopyType::PromoteByteToFloat:
-		bResult = SourceProperty->IsA<FByteProperty>() && TargetProperty->IsA<FFloatProperty>();
-		break;
-	case EStateTreePropertyCopyType::PromoteByteToDouble:
-		bResult = SourceProperty->IsA<FByteProperty>() && TargetProperty->IsA<FDoubleProperty>();
-		break;
+EStateTreePropertyAccessCompatibility FStateTreePropertyBindings::GetPropertyCompatibility(const FProperty* FromProperty, const FProperty* ToProperty)
+{
+	if (FromProperty == ToProperty)
+	{
+		return EStateTreePropertyAccessCompatibility::Compatible;
+	}
 
-	// Int32 promotions
-	case EStateTreePropertyCopyType::PromoteInt32ToInt64:
-		bResult = SourceProperty->IsA<FIntProperty>() && TargetProperty->IsA<FInt64Property>();
-		break;
-	case EStateTreePropertyCopyType::PromoteInt32ToFloat:
-		bResult = SourceProperty->IsA<FIntProperty>() && TargetProperty->IsA<FFloatProperty>();
-		break;
-	case EStateTreePropertyCopyType::PromoteInt32ToDouble:
-		bResult = SourceProperty->IsA<FIntProperty>() && TargetProperty->IsA<FDoubleProperty>();
-		break;
+	if (FromProperty == nullptr || ToProperty == nullptr)
+	{
+		return EStateTreePropertyAccessCompatibility::Incompatible;
+	}
 
-	// Uint32 promotions
-	case EStateTreePropertyCopyType::PromoteUInt32ToInt64:
-		bResult = SourceProperty->IsA<FUInt32Property>() && TargetProperty->IsA<FInt64Property>();
-		break;
-	case EStateTreePropertyCopyType::PromoteUInt32ToFloat:
-		bResult = SourceProperty->IsA<FUInt32Property>() && TargetProperty->IsA<FFloatProperty>();
-		break;
-	case EStateTreePropertyCopyType::PromoteUInt32ToDouble:
-		bResult = SourceProperty->IsA<FUInt32Property>() && TargetProperty->IsA<FDoubleProperty>();
-		break;
+	// Special case for object properties since InPropertyA->SameType(InPropertyB) requires both properties to be of the exact same class.
+	// In our case we want to be able to bind a source property if its class is a child of the target property class.
+	if (FromProperty->IsA<FObjectPropertyBase>() && ToProperty->IsA<FObjectPropertyBase>())
+	{
+		const FObjectPropertyBase* SourceProperty = CastField<FObjectPropertyBase>(FromProperty);
+		const FObjectPropertyBase* TargetProperty = CastField<FObjectPropertyBase>(ToProperty);
+		return (SourceProperty->PropertyClass->IsChildOf(TargetProperty->PropertyClass)) ? EStateTreePropertyAccessCompatibility::Compatible : EStateTreePropertyAccessCompatibility::Incompatible;
+	}
 
-	// Float promotions
-	case EStateTreePropertyCopyType::PromoteFloatToInt32:
-		bResult = SourceProperty->IsA<FFloatProperty>() && TargetProperty->IsA<FIntProperty>();
-		break;
-	case EStateTreePropertyCopyType::PromoteFloatToInt64:
-		bResult = SourceProperty->IsA<FFloatProperty>() && TargetProperty->IsA<FInt64Property>();
-		break;
-	case EStateTreePropertyCopyType::PromoteFloatToDouble:
-		bResult = SourceProperty->IsA<FFloatProperty>() && TargetProperty->IsA<FDoubleProperty>();
-		break;
+	// When copying to an enum property, expect FromProperty to be the same enum.
+	auto GetPropertyEnum = [](const FProperty* Property) -> const UEnum*
+	{
+		if (const FByteProperty* ByteProperty = CastField<FByteProperty>(Property))
+		{
+			return ByteProperty->GetIntPropertyEnum();
+		}
+		if (const FEnumProperty* EnumProperty = CastField<FEnumProperty>(Property))
+		{
+			return EnumProperty->GetEnum();
+		}
+		return nullptr;
+	};
+	
+	if (const UEnum* ToPropertyEnum = GetPropertyEnum(ToProperty))
+	{
+		const UEnum* FromPropertyEnum = GetPropertyEnum(FromProperty);
+		return (ToPropertyEnum == FromPropertyEnum) ? EStateTreePropertyAccessCompatibility::Compatible : EStateTreePropertyAccessCompatibility::Incompatible;
+	}
+	
+	// Allow source enums to be promoted to numbers.
+	if (const FEnumProperty* EnumPropertyA = CastField<const FEnumProperty>(FromProperty))
+	{
+		FromProperty = EnumPropertyA->GetUnderlyingProperty();
+	}
 
-	// Double promotions
-	case EStateTreePropertyCopyType::DemoteDoubleToInt32:
-		bResult = SourceProperty->IsA<FDoubleProperty>() && TargetProperty->IsA<FIntProperty>();
-		break;
-	case EStateTreePropertyCopyType::DemoteDoubleToInt64:
-		bResult = SourceProperty->IsA<FDoubleProperty>() && TargetProperty->IsA<FInt64Property>();
-		break;
-	case EStateTreePropertyCopyType::DemoteDoubleToFloat:
-		bResult = SourceProperty->IsA<FDoubleProperty>() && TargetProperty->IsA<FFloatProperty>();
-		break;
-	default:
-		UE_LOG(LogStateTree, Error, TEXT("FStateTreePropertyBindings::ValidateCopy: Unhandled copy type %s between '%s' and '%s'"),
-			*StaticEnum<EStateTreePropertyCopyType>()->GetValueAsString(Copy.Type), *SourceProperty->GetNameCPP(), *TargetProperty->GetNameCPP());
-		bResult = false;
-		break;
+	if (FromProperty->SameType(ToProperty))
+	{
+		return EStateTreePropertyAccessCompatibility::Compatible;
+	}
+	else
+	{
+		// Not directly compatible, check for promotions
+		if (FromProperty->IsA<FBoolProperty>())
+		{
+			if (ToProperty->IsA<FByteProperty>()
+				|| ToProperty->IsA<FIntProperty>()
+				|| ToProperty->IsA<FUInt32Property>()
+				|| ToProperty->IsA<FInt64Property>()
+				|| ToProperty->IsA<FFloatProperty>()
+				|| ToProperty->IsA<FDoubleProperty>())
+			{
+				return EStateTreePropertyAccessCompatibility::Promotable;
+			}
+		}
+		else if (FromProperty->IsA<FByteProperty>())
+		{
+			if (ToProperty->IsA<FIntProperty>()
+				|| ToProperty->IsA<FUInt32Property>()
+				|| ToProperty->IsA<FInt64Property>()
+				|| ToProperty->IsA<FFloatProperty>()
+				|| ToProperty->IsA<FDoubleProperty>())
+			{
+				return EStateTreePropertyAccessCompatibility::Promotable;
+			}
+		}
+		else if (FromProperty->IsA<FIntProperty>())
+		{
+			if (ToProperty->IsA<FInt64Property>()
+				|| ToProperty->IsA<FFloatProperty>()
+				|| ToProperty->IsA<FDoubleProperty>())
+			{
+				return EStateTreePropertyAccessCompatibility::Promotable;
+			}
+		}
+		else if (FromProperty->IsA<FUInt32Property>())
+		{
+			if (ToProperty->IsA<FInt64Property>()
+				|| ToProperty->IsA<FFloatProperty>()
+				|| ToProperty->IsA<FDoubleProperty>())
+			{
+				return EStateTreePropertyAccessCompatibility::Promotable;
+			}
+		}
+		else if (FromProperty->IsA<FFloatProperty>())
+		{
+			if (ToProperty->IsA<FIntProperty>()
+				|| ToProperty->IsA<FInt64Property>()
+				|| ToProperty->IsA<FDoubleProperty>())
+			{
+				return EStateTreePropertyAccessCompatibility::Promotable;
+			}
+		}
+		else if (FromProperty->IsA<FDoubleProperty>())
+		{
+			if (ToProperty->IsA<FIntProperty>()
+				|| ToProperty->IsA<FInt64Property>()
+				|| ToProperty->IsA<FFloatProperty>())
+			{
+				return EStateTreePropertyAccessCompatibility::Promotable;
+			}
+		}
 	}
 
-	UE_CLOG(!bResult, LogStateTree, Error, TEXT("FStateTreePropertyBindings::ValidateCopy: Failed to validate copy type %s between '%s' and '%s'"),
-		*StaticEnum<EStateTreePropertyCopyType>()->GetValueAsString(Copy.Type), *SourceProperty->GetNameCPP(), *TargetProperty->GetNameCPP());
-	
-	return bResult;
+	return EStateTreePropertyAccessCompatibility::Incompatible;
 }
 
 uint8* FStateTreePropertyBindings::GetAddress(FStateTreeDataView InStructView, const FStateTreePropertyIndirection& FirstIndirection, const FProperty* LeafProperty) const
@@ -359,6 +665,7 @@ uint8* FStateTreePropertyBindings::GetAddress(FStateTreeDataView InStructView, c
 	uint8* Address = InStructView.GetMutableMemory();
 	if (Address == nullptr)
 	{
+		// Failed indirection, will be reported by caller.
 		return nullptr;
 	}
 
@@ -393,17 +700,51 @@ uint8* FStateTreePropertyBindings::GetAddress(FStateTreeDataView InStructView, c
 			Address = reinterpret_cast<uint8*>(Object);
 			break;
 		}
+		case EStateTreePropertyAccessType::ObjectInstance:
+		{
+			check(Indirection->InstanceStruct);
+			UObject* Object = *reinterpret_cast<UObject**>(Address + Indirection->Offset);
+			if (Object
+				&& Object->GetClass()->IsChildOf(Indirection->InstanceStruct))
+			{
+				Address = reinterpret_cast<uint8*>(Object);
+			}
+			else
+			{
+				// Failed indirection, will be reported by caller.
+				return nullptr;
+			}
+			break;
+		}
+		case EStateTreePropertyAccessType::StructInstance:
+		{
+			check(Indirection->InstanceStruct);
+			FInstancedStruct& InstancedStruct = *reinterpret_cast<FInstancedStruct*>(Address + Indirection->Offset);
+			const UScriptStruct* InstanceType = InstancedStruct.GetScriptStruct(); 
+			if (InstanceType != nullptr
+				&& InstanceType->IsChildOf(Indirection->InstanceStruct))
+			{
+				Address = InstancedStruct.GetMutableMemory();
+			}
+			else
+			{
+				// Failed indirection, will be reported by caller.
+				return nullptr;
+			}
+			break;
+		}
 		case EStateTreePropertyAccessType::IndexArray:
 		{
 			check(Indirection->ArrayProperty);
 			FScriptArrayHelper Helper(Indirection->ArrayProperty, Address + Indirection->Offset);
 			if (Helper.IsValidIndex(Indirection->ArrayIndex.Get()))
 			{
-				Address = reinterpret_cast<uint8*>(Helper.GetRawPtr(Indirection->ArrayIndex.Get()));
+				Address = Helper.GetRawPtr(Indirection->ArrayIndex.Get());
 			}
 			else
 			{
-				Address = nullptr;
+				// Failed indirection, will be reported by caller.
+				return nullptr;
 			}
 			break;
 		}
@@ -418,7 +759,7 @@ uint8* FStateTreePropertyBindings::GetAddress(FStateTreeDataView InStructView, c
 	return Address;
 }
 
-void FStateTreePropertyBindings::PerformCopy(const FStateTreePropCopy& Copy, uint8* SourceAddress, uint8* TargetAddress) const
+void FStateTreePropertyBindings::PerformCopy(const FStateTreePropertyCopy& Copy, uint8* SourceAddress, uint8* TargetAddress) const
 {
 	// Source property can be null
 	check(SourceAddress);
@@ -578,7 +919,7 @@ bool FStateTreePropertyBindings::CopyTo(TConstArrayView<FStateTreeDataView> Sour
 	}
 
 	check(CopyBatches.IsValidIndex(TargetBatchIndex.Get()));
-	const FStateTreePropCopyBatch& Batch = CopyBatches[TargetBatchIndex.Get()];
+	const FStateTreePropertyCopyBatch& Batch = CopyBatches[TargetBatchIndex.Get()];
 
 	check(TargetStructView.IsValid());
 	check(TargetStructView.GetStruct() == Batch.TargetStruct.Struct);
@@ -587,7 +928,7 @@ bool FStateTreePropertyBindings::CopyTo(TConstArrayView<FStateTreeDataView> Sour
 	
 	for (int32 i = Batch.BindingsBegin; i != Batch.BindingsEnd; i++)
 	{
-		const FStateTreePropCopy& Copy = PropertyCopies[i];
+		const FStateTreePropertyCopy& Copy = PropertyCopies[i];
 		// Copies that fail to be resolved (i.e. property path does not resolve, types changed) will be marked as None, skip them.
 		if (Copy.Type == EStateTreePropertyCopyType::None)
 		{
@@ -603,8 +944,14 @@ bool FStateTreePropertyBindings::CopyTo(TConstArrayView<FStateTreeDataView> Sour
 			uint8* SourceAddress = GetAddress(SourceStructView, Copy.SourceIndirection, Copy.SourceLeafProperty);
 			uint8* TargetAddress = GetAddress(TargetStructView, Copy.TargetIndirection, Copy.TargetLeafProperty);
 			
-			check (SourceAddress != nullptr && TargetAddress != nullptr);
-			PerformCopy(Copy, SourceAddress, TargetAddress);
+			if (SourceAddress != nullptr && TargetAddress != nullptr)
+			{
+				PerformCopy(Copy, SourceAddress, TargetAddress);
+			}
+			else
+			{
+				bResult = false;
+			}
 		}
 		else
 		{
@@ -615,7 +962,7 @@ bool FStateTreePropertyBindings::CopyTo(TConstArrayView<FStateTreeDataView> Sour
 	return bResult;
 }
 
-void FStateTreePropertyBindings::PerformResetObjects(const FStateTreePropCopy& Copy, uint8* TargetAddress) const
+void FStateTreePropertyBindings::PerformResetObjects(const FStateTreePropertyCopy& Copy, uint8* TargetAddress) const
 {
 	// Source property can be null
 	check(Copy.TargetLeafProperty);
@@ -664,7 +1011,7 @@ bool FStateTreePropertyBindings::ResetObjects(const FStateTreeIndex16 TargetBatc
 	}
 
 	check(CopyBatches.IsValidIndex(TargetBatchIndex.Get()));
-	const FStateTreePropCopyBatch& Batch = CopyBatches[TargetBatchIndex.Get()];
+	const FStateTreePropertyCopyBatch& Batch = CopyBatches[TargetBatchIndex.Get()];
 
 	check(TargetStructView.IsValid());
 	check(TargetStructView.GetStruct() == Batch.TargetStruct.Struct);
@@ -673,7 +1020,7 @@ bool FStateTreePropertyBindings::ResetObjects(const FStateTreeIndex16 TargetBatc
 	
 	for (int32 i = Batch.BindingsBegin; i != Batch.BindingsEnd; i++)
 	{
-		const FStateTreePropCopy& Copy = PropertyCopies[i];
+		const FStateTreePropertyCopy& Copy = PropertyCopies[i];
 		// Copies that fail to be resolved (i.e. property path does not resolve, types changed) will be marked as None, skip them.
 		if (Copy.Type == EStateTreePropertyCopyType::None)
 		{
@@ -702,7 +1049,7 @@ void FStateTreePropertyBindings::DebugPrintInternalLayout(FString& OutString) co
 	/** Array of copy batches. */
 	OutString += FString::Printf(TEXT("\nCopyBatches (%d)\n  [ %-40s | %-40s | %-8s [%-3s:%-3s[ ]\n"), CopyBatches.Num(),
 		TEXT("Target Type"), TEXT("Target Name"), TEXT("Bindings"), TEXT("Beg"), TEXT("End"));
-	for (const FStateTreePropCopyBatch& CopyBatch : CopyBatches)
+	for (const FStateTreePropertyCopyBatch& CopyBatch : CopyBatches)
 	{
 		OutString += FString::Printf(TEXT("  | %-40s | %-40s | %8s [%3d:%-3d[ |\n"),
 									 CopyBatch.TargetStruct.Struct ? *CopyBatch.TargetStruct.Struct->GetName() : TEXT("null"),
@@ -711,36 +1058,11 @@ void FStateTreePropertyBindings::DebugPrintInternalLayout(FString& OutString) co
 	}
 
 	/** Array of property bindings, resolved into arrays of copies before use. */
-	OutString += FString::Printf(TEXT("\nPropertyBindings (%d)\n  [ %-20s | %-4s | %-4s | %-10s | %-20s | %-4s | %-7s | %-20s | %-7s | %-20s ]\n"),
-		PropertyBindings.Num(),
-		TEXT("Source Name"), TEXT("Arr#"), TEXT("Next"), TEXT("Access"),
-		TEXT("Target Name"), TEXT("Arr#"), TEXT("Next"), TEXT("Access"),
-		TEXT("Struct#"), TEXT("Copy Type"));
-	for (const FStateTreePropertyBinding& PropertyBinding : PropertyBindings)
-	{
-		OutString += FString::Printf(TEXT("  | %-20s | %4d | %4d | %-10s | %-20s | %4d | %4d | %-10s | %7d | %-20s | \n"),
-									 *PropertyBinding.SourcePath.Name.ToString(),
-									 PropertyBinding.SourcePath.ArrayIndex.Get(),
-									 PropertyBinding.SourcePath.NextIndex.Get(),
-									 *UEnum::GetDisplayValueAsText(PropertyBinding.TargetPath.Type).ToString(),
-									 *PropertyBinding.TargetPath.Name.ToString(),
-									 PropertyBinding.TargetPath.ArrayIndex.Get(),
-									 PropertyBinding.TargetPath.NextIndex.Get(),
-									 *UEnum::GetDisplayValueAsText(PropertyBinding.TargetPath.Type).ToString(),
-									 PropertyBinding.SourceStructIndex.Get(),
-									 *UEnum::GetValueAsString(PropertyBinding.CopyType));
-	}
-
-	/** Array of property segments, indexed by property paths. */
-	OutString += FString::Printf(TEXT("\nPropertySegments (%d)\n  [ %-20s | %4s | %4s | %-10s ]\n"), PropertySegments.Num(),
-		TEXT("Name"), TEXT("Arr#"), TEXT("Next"), TEXT("Access"));
-	for (const FStateTreePropertySegment& PropertySegment : PropertySegments)
-	{
-		OutString += FString::Printf(TEXT("  | %-20s | %4d | %4d | %-10s |\n"),
-									 *PropertySegment.Name.ToString(),
-									 PropertySegment.ArrayIndex.Get(),
-									 PropertySegment.NextIndex.Get(),
-									 *UEnum::GetDisplayValueAsText(PropertySegment.Type).ToString());
+	OutString += FString::Printf(TEXT("\nPropertyPathBindings (%d)\n"), PropertyPathBindings.Num());
+	for (const FStateTreePropertyPathBinding& PropertyBinding : PropertyPathBindings)
+	{
+		OutString += FString::Printf(TEXT("\n  Source: %s | Target: %s"),
+					*PropertyBinding.GetSourcePath().ToString(), *PropertyBinding.GetSourcePath().ToString());
 	}
 
 	/** Array of property copies */
@@ -748,7 +1070,7 @@ void FStateTreePropertyBindings::DebugPrintInternalLayout(FString& OutString) co
 		TEXT("Src Idx"), TEXT("Off."), TEXT("Next"), TEXT("Type"),
 		TEXT("Tgt Idx"), TEXT("Off."), TEXT("Next"), TEXT("Type"),
 		TEXT("Struct"), TEXT("Copy Type"), TEXT("Size"));
-	for (const FStateTreePropCopy& PropertyCopy : PropertyCopies)
+	for (const FStateTreePropertyCopy& PropertyCopy : PropertyCopies)
 	{
 		OutString += FString::Printf(TEXT("  | %7d | %4d | %4d | %-10s | %7d | %4d | %4d | %-10s | %7d | %-20s | %4d |\n"),
 					PropertyCopy.SourceIndirection.ArrayIndex.Get(),
@@ -777,35 +1099,396 @@ void FStateTreePropertyBindings::DebugPrintInternalLayout(FString& OutString) co
 	}
 }
 
-FString FStateTreePropertyBindings::GetPathAsString(const FStateTreePropertySegment& FirstPathSegment, const FStateTreePropertySegment* HighlightedSegment, const TCHAR* HighlightPrefix, const TCHAR* HighlightPostfix)
+
+//----------------------------------------------------------------//
+//  FStateTreePropertyPath
+//----------------------------------------------------------------//
+
+bool FStateTreePropertyPath::FromString(const FString& InPath)
+{
+	Segments.Reset();
+	
+	if (InPath.IsEmpty())
+	{
+		return true;
+	}
+	
+	bool bResult = true;
+	TArray<FString> PathSegments;
+	InPath.ParseIntoArray(PathSegments, TEXT("."), /*InCullEmpty=*/false);
+	
+	for (const FString& Segment : PathSegments)
+	{
+		if (Segment.IsEmpty())
+		{
+			bResult = false;
+			break;
+		}
+
+		int32 FirstBracket = INDEX_NONE;
+		int32 LastBracket = INDEX_NONE;
+		if (Segment.FindChar(TEXT('['), FirstBracket)
+			&& Segment.FindLastChar(TEXT(']'), LastBracket))
+		{
+			const int32 NameStringLength = FirstBracket;
+			const int32 IndexStringLength = LastBracket - FirstBracket - 1;
+			if (NameStringLength < 1
+				|| IndexStringLength <= 0)
+			{
+				bResult = false;
+				break;
+			}
+
+			const FString NameString = Segment.Left(FirstBracket);
+			const FString IndexString = Segment.Mid(FirstBracket + 1, IndexStringLength);
+			int32 ArrayIndex = INDEX_NONE;
+			LexFromString(ArrayIndex, *IndexString);
+			if (ArrayIndex < 0)
+			{
+				bResult = false;
+				break;
+			}
+			
+			AddPathSegment(FName(NameString), ArrayIndex);
+		}
+		else
+		{
+			AddPathSegment(FName(Segment));
+		}
+	}
+
+	if (!bResult)
+	{
+		Segments.Reset();
+	}
+	
+	return bResult;
+}
+
+bool FStateTreePropertyPath::UpdateInstanceStructsFromValue(const FStateTreeDataView BaseValueView, FString* OutError)
+{
+	TArray<FStateTreePropertyPathIndirection> Indirections;
+	if (!ResolveIndirectionsWithValue(BaseValueView, Indirections, OutError))
+	{
+		return false;
+	}
+
+	for (FStateTreePropertyPathSegment& Segment : Segments)
+	{
+		Segment.SetInstanceStruct(nullptr);
+	}
+	
+	for (const FStateTreePropertyPathIndirection& Indirection : Indirections)
+	{
+		if (Indirection.InstanceStruct != nullptr)
+		{
+			Segments[Indirection.PathSegmentIndex].SetInstanceStruct(Indirection.InstanceStruct);
+		}
+	}
+
+	return true;
+}
+
+FString FStateTreePropertyPath::ToString(const int32 HighlightedSegment, const TCHAR* HighlightPrefix, const TCHAR* HighlightPostfix, const bool bOutputInstances) const
 {
 	FString Result;
-	const FStateTreePropertySegment* Segment = &FirstPathSegment;
-	while (Segment != nullptr)
+	for (TEnumerateRef<const FStateTreePropertyPathSegment> Segment : EnumerateRange(Segments))
 	{
-		if (!Result.IsEmpty())
+		if (Segment.GetIndex() > 0)
 		{
 			Result += TEXT(".");
 		}
-		
-		if (Segment == HighlightedSegment && HighlightPrefix)
+		if (Segment.GetIndex() == HighlightedSegment && HighlightPrefix)
 		{
 			Result += HighlightPrefix;
 		}
 
-		Result += Segment->Name.ToString();
+		if (bOutputInstances && Segment->GetInstanceStruct())
+		{
+			Result += FString::Printf(TEXT("(%s)"), *GetNameSafe(Segment->GetInstanceStruct()));
+		}
 
-		if (Segment == HighlightedSegment && HighlightPostfix)
+		if (Segment->GetArrayIndex() >= 0)
 		{
-			Result += HighlightPostfix;
+			Result += FString::Printf(TEXT("%s[%d]"), *Segment->GetName().ToString(), Segment->GetArrayIndex());
+		}
+		else
+		{
+			Result += Segment->GetName().ToString();
 		}
 
-		Segment = Segment->NextIndex.IsValid() ? &PropertySegments[Segment->NextIndex.Get()] : nullptr;
+		if (Segment.GetIndex() == HighlightedSegment && HighlightPostfix)
+		{
+			Result += HighlightPostfix;
+		}
 	}
 	return Result;
 }
 
 
+bool FStateTreePropertyPath::ResolveIndirections(const UStruct* BaseStruct, TArray<FStateTreePropertyPathIndirection>& OutIndirections, FString* OutError) const
+{
+	return ResolveIndirectionsWithValue(FStateTreeDataView(BaseStruct, nullptr), OutIndirections, OutError);
+}
+
+bool FStateTreePropertyPath::ResolveIndirectionsWithValue(const FStateTreeDataView BaseValueView, TArray<FStateTreePropertyPathIndirection>& OutIndirections, FString* OutError) const
+{
+	OutIndirections.Reset();
+	if (OutError)
+	{
+		OutError->Reset();
+	}
+	
+	// Nothing to do for an empty path.
+	if (IsPathEmpty())
+	{
+		return true;
+	}
+
+	const uint8* CurrentAddress = BaseValueView.GetMemory();
+	const UStruct* CurrentStruct = BaseValueView.GetStruct();
+
+	const bool bWithValue = CurrentAddress != nullptr;
+	
+	for (const TEnumerateRef<const FStateTreePropertyPathSegment> Segment : EnumerateRange(Segments))
+	{
+		if (CurrentStruct == nullptr || (bWithValue && CurrentAddress == nullptr))
+		{
+			if (OutError)
+			{
+				*OutError = FString::Printf(TEXT("Malformed path '%s'."),
+					*ToString(Segment.GetIndex(), TEXT("<"), TEXT(">")));
+			}
+			OutIndirections.Reset();
+			return false;
+		}
+
+		const FProperty* Property = CurrentStruct->FindPropertyByName(Segment->GetName());
+		if (Property == nullptr)
+		{
+			if (OutError)
+			{
+				*OutError = FString::Printf(TEXT("Malformed path '%s', could not find property '%s%s::%s'."),
+					*ToString(Segment.GetIndex(), TEXT("<"), TEXT(">")),
+					CurrentStruct->GetPrefixCPP(), *CurrentStruct->GetName(), *Segment->GetName().ToString());
+			}
+			OutIndirections.Reset();
+			return false;
+		}
+
+		const FArrayProperty* ArrayProperty = CastField<FArrayProperty>(Property);
+		int ArrayIndex = 0;
+		int32 Offset = 0;
+		if (ArrayProperty && Segment->GetArrayIndex() != INDEX_NONE)
+		{
+			FStateTreePropertyPathIndirection& Indirection = OutIndirections.AddDefaulted_GetRef();
+			Indirection.Property = Property;
+			Indirection.ContainerAddress = CurrentAddress;
+			Indirection.ContainerStruct = CurrentStruct;
+			Indirection.InstanceStruct = nullptr;
+			Indirection.ArrayIndex = Segment->GetArrayIndex();
+			Indirection.PropertyOffset = ArrayProperty->GetOffset_ForInternal();
+			Indirection.PathSegmentIndex = Segment.GetIndex();
+			Indirection.AccessType = EStateTreePropertyAccessType::IndexArray; 
+			
+			ArrayIndex = 0;
+			Offset = 0;
+			Property = ArrayProperty->Inner;
+
+			if (bWithValue)
+			{
+				FScriptArrayHelper Helper(ArrayProperty, CurrentAddress + ArrayProperty->GetOffset_ForInternal());
+				if (!Helper.IsValidIndex(Segment->GetArrayIndex()))
+				{
+					if (OutError)
+					{
+						*OutError = FString::Printf(TEXT("Index %d out of range (num elements %d) trying to access dynamic array '%s'."),
+							Segment->GetArrayIndex(), Helper.Num(), *ToString(Segment.GetIndex(), TEXT("<"), TEXT(">")));
+					}
+					OutIndirections.Reset();
+					return false;
+				}
+				CurrentAddress = Helper.GetRawPtr(Segment->GetArrayIndex());
+			}
+		}
+		else
+		{
+			if (Segment->GetArrayIndex() > Property->ArrayDim)
+			{
+				if (OutError)
+				{
+					*OutError = FString::Printf(TEXT("Index %d out of range %d trying to access static array '%s'."),
+						Segment->GetArrayIndex(), Property->ArrayDim, *ToString(Segment.GetIndex(), TEXT("<"), TEXT(">")));
+				}
+				OutIndirections.Reset();
+				return false;
+			}
+			ArrayIndex = FMath::Max(0, Segment->GetArrayIndex());
+			Offset = Property->GetOffset_ForInternal() + Property->ElementSize * ArrayIndex;
+		}
+
+		FStateTreePropertyPathIndirection& Indirection = OutIndirections.AddDefaulted_GetRef();
+		Indirection.Property = Property;
+		Indirection.ContainerAddress = CurrentAddress;
+		Indirection.ContainerStruct = CurrentStruct;
+		Indirection.ArrayIndex = ArrayIndex;
+		Indirection.PropertyOffset = Offset;
+		Indirection.PathSegmentIndex = Segment.GetIndex();
+		Indirection.AccessType = EStateTreePropertyAccessType::Offset; 
+
+		const bool bLastSegment = Segment.GetIndex() == (Segments.Num() - 1);
+
+		if (!bLastSegment)
+		{
+			if (const FStructProperty* StructProperty = CastField<FStructProperty>(Property))
+			{
+				if (bWithValue)
+				{
+					if (StructProperty->Struct == TBaseStructure<FInstancedStruct>::Get())
+					{
+						const FInstancedStruct& InstancedStruct = *reinterpret_cast<const FInstancedStruct*>(CurrentAddress + Offset);
+						if (!InstancedStruct.IsValid())
+						{
+							if (OutError)
+							{
+								*OutError = FString::Printf(TEXT("Expecting valid instanced struct value at path '%s'."),
+									*ToString(Segment.GetIndex(), TEXT("<"), TEXT(">")));
+							}
+							OutIndirections.Reset();
+							return false;
+						}
+						const UScriptStruct* ValueInstanceStructType = InstancedStruct.GetScriptStruct();
+
+						CurrentAddress = InstancedStruct.GetMemory();
+						CurrentStruct = ValueInstanceStructType;
+						Indirection.InstanceStruct = CurrentStruct;
+						Indirection.AccessType = EStateTreePropertyAccessType::StructInstance; 
+					}
+					else
+					{
+						CurrentAddress = CurrentAddress + Offset;
+						CurrentStruct = StructProperty->Struct;
+						Indirection.AccessType = EStateTreePropertyAccessType::Offset;
+					}
+				}
+				else
+				{
+					if (Segment->GetInstanceStruct())
+					{
+						CurrentStruct = Segment->GetInstanceStruct();
+						Indirection.InstanceStruct = CurrentStruct;
+						Indirection.AccessType = EStateTreePropertyAccessType::StructInstance;
+					}
+					else
+					{
+						CurrentStruct = StructProperty->Struct;
+						Indirection.AccessType = EStateTreePropertyAccessType::Offset;
+					}
+				}
+			}
+			else if (const FObjectProperty* ObjectProperty = CastField<FObjectProperty>(Property))
+			{
+				if (bWithValue)
+				{
+					const UObject* Object = *reinterpret_cast<UObject* const*>(CurrentAddress + Offset);
+					CurrentAddress = reinterpret_cast<const uint8*>(Object);
+					if (Property->HasAnyPropertyFlags(CPF_PersistentInstance | CPF_InstancedReference))
+					{
+						if (!Object)
+						{
+							if (OutError)
+							{
+								*OutError = FString::Printf(TEXT("Expecting valid instanced object value at path '%s'."),
+									*ToString(Segment.GetIndex(), TEXT("<"), TEXT(">")));
+							}
+							OutIndirections.Reset();
+							return false;
+						}
+						CurrentStruct = Object->GetClass();
+						Indirection.InstanceStruct = CurrentStruct;
+						Indirection.AccessType = EStateTreePropertyAccessType::ObjectInstance;
+					}
+					else
+					{
+						CurrentStruct = ObjectProperty->PropertyClass;
+						Indirection.AccessType = EStateTreePropertyAccessType::Object;
+					}
+				}
+				else
+				{
+					if (Segment->GetInstanceStruct())
+					{
+						CurrentStruct = Segment->GetInstanceStruct();
+						Indirection.InstanceStruct = CurrentStruct;
+						Indirection.AccessType = EStateTreePropertyAccessType::ObjectInstance;
+					}
+					else
+					{
+						CurrentStruct = ObjectProperty->PropertyClass;
+						Indirection.AccessType = EStateTreePropertyAccessType::Object;
+					}
+				}
+			}
+			// Check to see if this is a simple weak object property (eg. not an array of weak objects).
+			else if (const FWeakObjectProperty* WeakObjectProperty = CastField<FWeakObjectProperty>(Property))
+			{
+				const TWeakObjectPtr<UObject>& WeakObjectPtr = *reinterpret_cast<const TWeakObjectPtr<UObject>*>(CurrentAddress + Offset);
+				const UObject* Object = WeakObjectPtr.Get();
+				CurrentAddress = reinterpret_cast<const uint8*>(Object);
+				CurrentStruct = WeakObjectProperty->PropertyClass;
+				Indirection.AccessType = EStateTreePropertyAccessType::WeakObject;
+			}
+			// Check to see if this is a simple soft object property (eg. not an array of soft objects).
+			else if (const FSoftObjectProperty* SoftObjectProperty = CastField<FSoftObjectProperty>(Property))
+			{
+				const FSoftObjectPtr& SoftObjectPtr = *reinterpret_cast<const FSoftObjectPtr*>(CurrentAddress + Offset);
+				const UObject* Object = SoftObjectPtr.Get();
+				CurrentAddress = reinterpret_cast<const uint8*>(Object);
+				CurrentStruct = SoftObjectProperty->PropertyClass;
+				Indirection.AccessType = EStateTreePropertyAccessType::SoftObject;
+			}
+			else
+			{
+				// We get here if we encounter a property type that is not supported for indirection (e.g. Map or Set).
+				if (OutError)
+				{
+					*OutError = FString::Printf(TEXT("Unsupported property indirection type %s in path '%s'."),
+						*Property->GetCPPType(), *ToString(Segment.GetIndex(), TEXT("<"), TEXT(">")));
+				}
+				OutIndirections.Reset();
+				return false;
+			}
+		}
+	}
+
+	return true;
+}
+
+bool FStateTreePropertyPath::operator==(const FStateTreePropertyPath& RHS) const
+{
+#if WITH_EDITORONLY_DATA
+	if (StructID != RHS.StructID)
+	{
+		return false;
+	}
+#endif // WITH_EDITORONLY_DATA
+	if (Segments.Num() != RHS.Segments.Num())
+	{
+		return false;
+	}
+
+	for (TEnumerateRef<const FStateTreePropertyPathSegment> Segment : EnumerateRange(Segments))
+	{
+		if (*Segment != RHS.Segments[Segment.GetIndex()])
+		{
+			return false;
+		}
+	}
+
+	return true;
+}
+
+
 //----------------------------------------------------------------//
 //  FStateTreeEditorPropertyPath
 //----------------------------------------------------------------//
@@ -834,6 +1517,7 @@ FString FStateTreeEditorPropertyPath::ToString(const int32 HighlightedSegment, c
 	return Result;
 }
 
+PRAGMA_DISABLE_DEPRECATION_WARNINGS
 bool FStateTreeEditorPropertyPath::operator==(const FStateTreeEditorPropertyPath& RHS) const
 {
 	if (StructID != RHS.StructID)
@@ -856,4 +1540,4 @@ bool FStateTreeEditorPropertyPath::operator==(const FStateTreeEditorPropertyPath
 
 	return true;
 }
-
+PRAGMA_ENABLE_DEPRECATION_WARNINGS
diff --git a/Engine/Plugins/Runtime/StateTree/Source/StateTreeModule/Public/Conditions/StateTreeCommonConditions.h b/Engine/Plugins/Runtime/StateTree/Source/StateTreeModule/Public/Conditions/StateTreeCommonConditions.h
--- a/Engine/Plugins/Runtime/StateTree/Source/StateTreeModule/Public/Conditions/StateTreeCommonConditions.h
+++ b/Engine/Plugins/Runtime/StateTree/Source/StateTreeModule/Public/Conditions/StateTreeCommonConditions.h
@@ -165,7 +165,7 @@ struct STATETREEMODULE_API FStateTreeCompareEnumCondition : public FStateTreeCon
 	virtual bool TestCondition(FStateTreeExecutionContext& Context) const override;
 
 #if WITH_EDITOR
-	virtual void OnBindingChanged(const FGuid& ID, FStateTreeDataView InstanceData, const FStateTreeEditorPropertyPath& SourcePath, const FStateTreeEditorPropertyPath& TargetPath, const IStateTreeBindingLookup& BindingLookup) override;
+	virtual void OnBindingChanged(const FGuid& ID, FStateTreeDataView InstanceData, const FStateTreePropertyPath& SourcePath, const FStateTreePropertyPath& TargetPath, const IStateTreeBindingLookup& BindingLookup) override;
 #endif
 
 	UPROPERTY(EditAnywhere, Category = "Parameter")
diff --git a/Engine/Plugins/Runtime/StateTree/Source/StateTreeModule/Public/StateTree.h b/Engine/Plugins/Runtime/StateTree/Source/StateTreeModule/Public/StateTree.h
--- a/Engine/Plugins/Runtime/StateTree/Source/StateTreeModule/Public/StateTree.h
+++ b/Engine/Plugins/Runtime/StateTree/Source/StateTreeModule/Public/StateTree.h
@@ -32,6 +32,8 @@ struct STATETREEMODULE_API FStateTreeCustomVersion
 		TransitionDelay,
 		// Added external transitions
 		AddedExternalTransitions,
+		// Changed how bindings are represented
+		ChangedBindingsRepresentation,
 
 		// -----<new versions can be added above this line>-------------------------------------------------
 		VersionPlusOne,
diff --git a/Engine/Plugins/Runtime/StateTree/Source/StateTreeModule/Public/StateTreeConditionBase.h b/Engine/Plugins/Runtime/StateTree/Source/StateTreeModule/Public/StateTreeConditionBase.h
--- a/Engine/Plugins/Runtime/StateTree/Source/StateTreeModule/Public/StateTreeConditionBase.h
+++ b/Engine/Plugins/Runtime/StateTree/Source/StateTreeModule/Public/StateTreeConditionBase.h
@@ -6,6 +6,7 @@
 #include "StateTreeConditionBase.generated.h"
 
 struct FStateTreeEditorPropertyPath;
+struct FStateTreePropertyPath;
 struct FStateTreeExecutionContext;
 struct IStateTreeBindingLookup;
 
@@ -24,6 +25,10 @@ struct STATETREEMODULE_API FStateTreeConditionBase : public FStateTreeNodeBase
 	GENERATED_BODY()
 
 #if WITH_EDITOR
+PRAGMA_DISABLE_DEPRECATION_WARNINGS	
+	UE_DEPRECATED(5.3, "Use version with FStateTreePropertyPath instead.")
+	virtual void OnBindingChanged(const FGuid& ID, FStateTreeDataView InstanceData, const FStateTreeEditorPropertyPath& SourcePath, const FStateTreeEditorPropertyPath& TargetPath, const IStateTreeBindingLookup& BindingLookup) final {}
+PRAGMA_ENABLE_DEPRECATION_WARNINGS
 	/**
 	 * Called when binding of any of the properties in the condition changes.
 	 * @param ID ID of the item, can be used make property paths to this item.
@@ -32,7 +37,7 @@ struct STATETREEMODULE_API FStateTreeConditionBase : public FStateTreeNodeBase
 	 * @param TargetPath Target path of the new binding (the property in the condition).
 	 * @param BindingLookup Reference to binding lookup which can be used to reason about property paths.
 	 */
-	virtual void OnBindingChanged(const FGuid& ID, FStateTreeDataView InstanceData, const FStateTreeEditorPropertyPath& SourcePath, const FStateTreeEditorPropertyPath& TargetPath, const IStateTreeBindingLookup& BindingLookup) {}
+	virtual void OnBindingChanged(const FGuid& ID, FStateTreeDataView InstanceData, const FStateTreePropertyPath& SourcePath, const FStateTreePropertyPath& TargetPath, const IStateTreeBindingLookup& BindingLookup) {}
 #endif
 	
 	/** @return True if the condition passes. */
diff --git a/Engine/Plugins/Runtime/StateTree/Source/StateTreeModule/Public/StateTreePropertyBindings.h b/Engine/Plugins/Runtime/StateTree/Source/StateTreeModule/Public/StateTreePropertyBindings.h
--- a/Engine/Plugins/Runtime/StateTree/Source/StateTreeModule/Public/StateTreePropertyBindings.h
+++ b/Engine/Plugins/Runtime/StateTree/Source/StateTreeModule/Public/StateTreePropertyBindings.h
@@ -6,6 +6,9 @@
 #include "StateTreePropertyBindings.generated.h"
 
 class FProperty;
+struct FStateTreePropertyPath;
+struct FStateTreePropertyBindingCompiler;
+class UStateTree;
 
 UENUM()
 enum class EStateTreeBindableStructSource : uint8
@@ -26,49 +29,6 @@
 	Condition,
 };
 
-/**
- * Descriptor for a struct or class that can be a binding source or target.
- * Each struct has unique identifier, which is used to distinguish them, and name that is mostly for debugging and UI.
- */
-USTRUCT()
-struct STATETREEMODULE_API FStateTreeBindableStructDesc
-{
-	GENERATED_BODY()
-
-	FStateTreeBindableStructDesc() = default;
-
-#if WITH_EDITORONLY_DATA
-	FStateTreeBindableStructDesc(const FName InName, const UStruct* InStruct, const EStateTreeBindableStructSource InDataSource, const FGuid InGuid)
-		: Struct(InStruct), Name(InName), DataSource(InDataSource), ID(InGuid)
-	{
-	}
-
-	bool operator==(const FStateTreeBindableStructDesc& RHS) const
-	{
-		return ID == RHS.ID && Struct == RHS.Struct; // Not checking name, it's cosmetic.
-	}
-#endif
-
-	bool IsValid() const { return Struct != nullptr; }
-	
-	/** The type of the struct or class. */
-	UPROPERTY()
-	TObjectPtr<const UStruct> Struct = nullptr;
-
-	/** Name of the container (cosmetic). */
-	UPROPERTY()
-	FName Name;
-
-	UPROPERTY()
-	EStateTreeBindableStructSource DataSource = EStateTreeBindableStructSource::Context;
-	
-#if WITH_EDITORONLY_DATA
-	/** Unique identifier of the struct. */
-	UPROPERTY()
-	FGuid ID;
-#endif
-};
-
 
 /**
  * The type of access a property path describes.
@@ -80,6 +40,8 @@
 	Object,			// Access node needs to dereference an object at its current address
 	WeakObject,		// Access is a weak object
 	SoftObject,		// Access is a soft object
+	ObjectInstance,	// Access node needs to dereference an object of specific type at its current address
+	StructInstance,	// Access node needs to dereference an instanced struct of specific type at its current address
 	IndexArray,		// Access node indexes a dynamic array
 };
 
@@ -139,12 +101,425 @@
 	DemoteDoubleToFloat,
 };
 
+/** Enum describing property compatibility */
+enum class EStateTreePropertyAccessCompatibility
+{
+	/** Properties are incompatible */
+	Incompatible,
+	/** Properties are directly compatible */
+	Compatible,	
+	/** Properties can be copied with a simple type promotion */
+	Promotable,
+};
 
 /**
- * Describes a segment of a property path (see FStateTreePropertyPath). Used for storage only.
+ * Descriptor for a struct or class that can be a binding source or target.
+ * Each struct has unique identifier, which is used to distinguish them, and name that is mostly for debugging and UI.
  */
 USTRUCT()
-struct STATETREEMODULE_API FStateTreePropertySegment
+struct STATETREEMODULE_API FStateTreeBindableStructDesc
 {
 	GENERATED_BODY()
 
+	FStateTreeBindableStructDesc() = default;
+
+#if WITH_EDITORONLY_DATA
+	FStateTreeBindableStructDesc(const FName InName, const UStruct* InStruct, const EStateTreeBindableStructSource InDataSource, const FGuid InGuid)
+		: Struct(InStruct), Name(InName), DataSource(InDataSource), ID(InGuid)
+	{
+	}
+
+	bool operator==(const FStateTreeBindableStructDesc& RHS) const
+	{
+		return ID == RHS.ID && Struct == RHS.Struct; // Not checking name, it's cosmetic.
+	}
+#endif
+
+	bool IsValid() const { return Struct != nullptr; }
+	
+	/** The type of the struct or class. */
+	UPROPERTY()
+	TObjectPtr<const UStruct> Struct = nullptr;
+
+	/** Name of the container (cosmetic). */
+	UPROPERTY()
+	FName Name;
+
+	UPROPERTY()
+	EStateTreeBindableStructSource DataSource = EStateTreeBindableStructSource::Context;
+	
+#if WITH_EDITORONLY_DATA
+	/** Unique identifier of the struct. */
+	UPROPERTY()
+	FGuid ID;
+#endif
+};
+
+/** Struct describing a path segment in FStateTreePropertyPath. */
+USTRUCT()
+struct STATETREEMODULE_API FStateTreePropertyPathSegment
+{
+	GENERATED_BODY()
+
+	FStateTreePropertyPathSegment() = default;
+
+	explicit FStateTreePropertyPathSegment(const FName InName, const int32 InArrayIndex = INDEX_NONE, const UStruct* InInstanceStruct = nullptr)
+		: Name(InName)
+		, ArrayIndex(InArrayIndex)
+		, InstanceStruct(InInstanceStruct)
+	{
+	}
+
+	bool operator==(const FStateTreePropertyPathSegment& RHS) const
+	{
+		return Name == RHS.Name && InstanceStruct == RHS.InstanceStruct && ArrayIndex == RHS.ArrayIndex;
+	}
+
+	bool operator!=(const FStateTreePropertyPathSegment& RHS) const
+	{
+		return !(*this == RHS);
+	}
+
+	void SetName(const FName InName)
+	{
+		Name = InName;
+	}
+
+	FName GetName() const
+	{
+		return Name;
+	}
+
+	void SetArrayIndex(const int32 InArrayIndex)
+	{
+		ArrayIndex = InArrayIndex;
+	}
+
+	int32 GetArrayIndex() const
+	{
+		return ArrayIndex;
+	}
+
+	void SetInstanceStruct(const UStruct* InInstanceStruct)
+	{
+		InstanceStruct = InInstanceStruct;
+	}
+
+	const UStruct* GetInstanceStruct() const
+	{
+		return InstanceStruct;
+	}
+	
+private:
+	/** Name of the property */
+	UPROPERTY()
+	FName Name;
+
+	/** Array index if the property is dynamic or static array. */
+	UPROPERTY()
+	int32 ArrayIndex = INDEX_NONE;
+
+	/** Type of the instanced struct or object reference by the property at the segment. This allows the path to be resolved when it points to a specific instance. */
+	UPROPERTY()
+	TObjectPtr<const UStruct> InstanceStruct = nullptr;
+};
+
+/**
+ * Struct describing an indirection at specific segment at path.
+ * Returned by FStateTreePropertyPath::ResolveIndirections() and FStateTreePropertyPath::ResolveIndirectionsWithValue().
+ * Generally there's one indirection per FProperty. Containers have one path segment but two indirection (container + inner type).
+ */
+struct FStateTreePropertyPathIndirection
+{
+	FStateTreePropertyPathIndirection() = default;
+	explicit FStateTreePropertyPathIndirection(const UStruct* InContainerStruct)
+		: ContainerStruct(InContainerStruct)
+	{
+	}
+
+	const FProperty* GetProperty() const { return Property; }
+	const uint8* GetContainerAddress() const { return ContainerAddress; }
+	const UStruct* GetInstanceStruct() const { return InstanceStruct; }
+	const UStruct* GetContainerStruct() const { return ContainerStruct; }
+	int32 GetArrayIndex() const { return ArrayIndex; }
+	int32 GetPropertyOffset() const { return PropertyOffset; }
+	int32 GetPathSegmentIndex() const { return PathSegmentIndex; }
+	EStateTreePropertyAccessType GetAccessType() const { return AccessType; }
+	const uint8* GetPropertyAddress() const { return ContainerAddress + PropertyOffset; }
+	
+private:
+	/** Property at the indirection. */
+	const FProperty* Property = nullptr;
+	
+	/** Address of the container class/struct where the property belongs to. Only valid if created with ResolveIndirectionsWithValue() */
+	const uint8* ContainerAddress = nullptr;
+	
+	/** Type of the container class/struct. */
+	const UStruct* ContainerStruct = nullptr;
+	
+	/** Type of the instance class/struct of when AccessType is ObjectInstance or StructInstance. */
+	const UStruct* InstanceStruct = nullptr;
+	
+	/** Array index for static and dynamic arrays. Note: static array indexing is baked in the PropertyOffset. */
+	int32 ArrayIndex = 0;
+	
+	/** Offset of the property relative to ContainerAddress. Includes static array indexing. */
+	int32 PropertyOffset = 0;
+	
+	/** Index of the path segment where indirection originated from. */
+	int32 PathSegmentIndex = 0;
+	
+	/** How to access the data through the indirection. */
+	EStateTreePropertyAccessType AccessType = EStateTreePropertyAccessType::Offset;
+
+	friend FStateTreePropertyPath;
+};
+	
+
+/**
+ * Representation of a property path used for property binding in StateTree.
+ */
+USTRUCT()
+struct STATETREEMODULE_API FStateTreePropertyPath
+{
+	GENERATED_BODY()
+
+	FStateTreePropertyPath() = default;
+
+#if WITH_EDITORONLY_DATA
+	explicit FStateTreePropertyPath(const FGuid InStructID)
+		: StructID(InStructID)
+	{
+	}
+
+	explicit FStateTreePropertyPath(const FGuid InStructID, const FName PropertyName)
+		: StructID(InStructID)
+	{
+		Segments.Emplace(PropertyName);
+	}
+	
+	explicit FStateTreePropertyPath(const FGuid InStructID, TConstArrayView<FStateTreePropertyPathSegment> InSegments)
+		: StructID(InStructID)
+		, Segments(InSegments)
+	{
+	}
+#endif // WITH_EDITORONLY_DATA
+	
+	/**
+	 * Parses path from string. The path should be in format: Foo.Bar[1].Baz
+	 * @param InPath Path string to parse
+	 * @return true if path was parsed successfully.
+	 */
+	bool FromString(const FString& InPath);
+
+	/**
+	 * Returns the property path as a one string. Highlight allows to decorate a specific segment.
+	 * @param HighlightedSegment Index of the highlighted path segment
+	 * @param HighlightPrefix String to append before highlighted segment
+	 * @param HighlightPostfix String to append after highlighted segment
+	 * @param bOutputInstances if true, the instance struct types will be output. 
+	 */
+	FString ToString(const int32 HighlightedSegment = INDEX_NONE, const TCHAR* HighlightPrefix = nullptr, const TCHAR* HighlightPostfix = nullptr, const bool bOutputInstances = false) const;
+
+	/**
+	 * Resolves the property path against base struct type. The path is assumed to be relative to the BaseStruct.
+	 * @param BaseStruct Base struct/class type the path is relative to.
+	 * @param OutIndirections Indirections describing how the properties were accessed. 
+	 * @param OutError Optional, pointer to string where error will be logged if update fails.
+	 * @return true of path could be resolved against the base value, and instance types were updated.
+	 */
+	bool ResolveIndirections(const UStruct* BaseStruct, TArray<FStateTreePropertyPathIndirection>& OutIndirections, FString* OutError = nullptr) const;
+	
+	/**
+	 * Resolves the property path against base value. The path is assumed to be relative to the BaseValueView.
+	 * @param BaseValueView Base value the path is relative to.
+	 * @param OutIndirections Indirections describing how the properties were accessed. 
+	 * @param OutError Optional, pointer to string where error will be logged if update fails.
+	 * @return true of path could be resolved against the base value, and instance types were updated.
+	 */
+	bool ResolveIndirectionsWithValue(const FStateTreeDataView BaseValueView, TArray<FStateTreePropertyPathIndirection>& OutIndirections, FString* OutError = nullptr) const;
+
+	/**
+	 * Updates property segment InstanceStruct types from value. The path is expected to be relative to the base value.
+	 * By storing the instance types on the path, we can resolve the path without the base value later.
+	 * @param BaseValueView Base value the path is relative to.
+	 * @param OutError Optional, pointer to string where error will be logged if update fails.
+	 * @return true of path could be resolved against the base value, and instance types were updated.
+	 */
+	bool UpdateInstanceStructsFromValue(const FStateTreeDataView BaseValueView, FString* OutError = nullptr);
+
+	/** @return true if the path is empty. In that case the path points to the struct. */
+	bool IsPathEmpty() const { return Segments.IsEmpty(); }
+
+	/** @return true if any of the path segments is and indirection via instanced struct or object. */
+	bool HasAnyInstancedIndirection() const
+	{
+		return Segments.ContainsByPredicate([](const FStateTreePropertyPathSegment& Segment)
+		{
+			return Segment.GetInstanceStruct() != nullptr;
+		});
+	}
+
+	/** Reset the path to empty. */
+	void Reset()
+	{
+#if WITH_EDITORONLY_DATA
+		StructID = FGuid();
+#endif
+		Segments.Reset();
+	}
+
+#if WITH_EDITORONLY_DATA
+	const FGuid& GetStructID() const { return StructID; }
+	void SetStructID(const FGuid NewStructID) { StructID = NewStructID; }
+#endif // WITH_EDITORONLY_DATA
+
+	TConstArrayView<FStateTreePropertyPathSegment> GetSegments() const { return Segments; }
+	TArrayView<FStateTreePropertyPathSegment> GetMutableSegments() { return Segments; }
+	int32 NumSegments() const { return Segments.Num(); }
+	const FStateTreePropertyPathSegment& GetSegment(const int32 Index) const { return Segments[Index]; }
+
+	/** Adds a path segment to the path. */
+	void AddPathSegment(const FName InName, const int32 InArrayIndex = INDEX_NONE, const UStruct* InInstanceType = nullptr)
+	{
+		Segments.Emplace(InName, InArrayIndex, InInstanceType);
+	}
+
+	/** Adds a path segment to the path. */
+	void AddPathSegment(const FStateTreePropertyPathSegment& PathSegment)
+	{
+		Segments.Add(PathSegment);
+	}
+
+	/** Test if paths are equal. */
+	bool operator==(const FStateTreePropertyPath& RHS) const;
+
+private:
+#if WITH_EDITORONLY_DATA
+	/** ID of the struct this property path is relative to. */
+	UPROPERTY()
+	FGuid StructID;
+#endif // WITH_EDITORONLY_DATA
+
+	/** Path segments pointing to a specific property on the path. */
+	UPROPERTY()
+	TArray<FStateTreePropertyPathSegment> Segments;
+};
+
+
+USTRUCT()
+struct UE_DEPRECATED(5.3, "Use FStateTreePropertyPath instead.") STATETREEMODULE_API FStateTreeEditorPropertyPath
+{
+	GENERATED_BODY()
+
+	FStateTreeEditorPropertyPath() = default;
+	FStateTreeEditorPropertyPath(const FGuid& InStructID, const TCHAR* PropertyName) : StructID(InStructID) { Path.Add(PropertyName); }
+	FStateTreeEditorPropertyPath(const FStateTreeEditorPropertyPath& InPath) : StructID(InPath.StructID), Path(InPath.Path) {}
+	
+	/**
+	 * Returns the property path as a one string. Highlight allows to decorate a specific segment.
+	 * @param HighlightedSegment Index of the highlighted path segment
+	 * @param HighlightPrefix String to append before highlighted segment
+	 * @param HighlightPostfix String to append after highlighted segment
+	 */
+	FString ToString(const int32 HighlightedSegment = INDEX_NONE, const TCHAR* HighlightPrefix = nullptr, const TCHAR* HighlightPostfix = nullptr) const;
+
+	/** Handle of the struct this property path is relative to. */
+	UPROPERTY()
+	FGuid StructID;
+
+	/** Property path segments */
+	UPROPERTY()
+	TArray<FString> Path;
+
+	bool IsValid() const { return StructID.IsValid(); }
+
+PRAGMA_DISABLE_DEPRECATION_WARNINGS
+	bool operator==(const FStateTreeEditorPropertyPath& RHS) const;
+PRAGMA_ENABLE_DEPRECATION_WARNINGS
+};
+
+
+/**
+ * Representation of a property binding in StateTree
+ */
+USTRUCT()
+struct STATETREEMODULE_API FStateTreePropertyPathBinding
+{
+	GENERATED_BODY()
+
+PRAGMA_DISABLE_DEPRECATION_WARNINGS	
+	FStateTreePropertyPathBinding() = default;
+	
+	FStateTreePropertyPathBinding(const FStateTreePropertyPath& InSourcePath, const FStateTreePropertyPath& InTargetPath)
+		: SourcePropertyPath(InSourcePath)
+		, TargetPropertyPath(InTargetPath)
+	{
+	}
+
+	FStateTreePropertyPathBinding(const FStateTreeIndex16 InCompiledSourceStructIndex, const FStateTreePropertyPath& InSourcePath, const FStateTreePropertyPath& InTargetPath)
+		: SourcePropertyPath(InSourcePath)
+		, TargetPropertyPath(InTargetPath)
+		, CompiledSourceStructIndex(InCompiledSourceStructIndex)
+	{
+	}
+
+	UE_DEPRECATED(5.3, "Use constructor with FStateTreePropertyPath instead.")
+	FStateTreePropertyPathBinding(const FStateTreeEditorPropertyPath& InSourcePath, const FStateTreeEditorPropertyPath& InTargetPath);
+PRAGMA_ENABLE_DEPRECATION_WARNINGS
+
+	void PostSerialize(const FArchive& Ar);
+
+	const FStateTreePropertyPath& GetSourcePath() const { return SourcePropertyPath; }
+	const FStateTreePropertyPath& GetTargetPath() const { return TargetPropertyPath; }
+
+	FStateTreePropertyPath& GetMutableSourcePath() { return SourcePropertyPath; }
+	FStateTreePropertyPath& GetMutableTargetPath() { return TargetPropertyPath; }
+
+	FStateTreeIndex16 GetCompiledSourceStructIndex() const { return CompiledSourceStructIndex; }
+
+private:
+	/** Source property path of the binding */
+	UPROPERTY()
+	FStateTreePropertyPath SourcePropertyPath;
+
+	/** Target property path of the binding */
+	UPROPERTY()
+	FStateTreePropertyPath TargetPropertyPath;
+
+	/** ID of the struct this property path is relative to. */
+	UPROPERTY()
+	FStateTreeIndex16 CompiledSourceStructIndex;
+
+public:
+#if WITH_EDITORONLY_DATA
+PRAGMA_DISABLE_DEPRECATION_WARNINGS
+	UPROPERTY()
+	FStateTreeEditorPropertyPath SourcePath_DEPRECATED;
+
+	UPROPERTY()
+	FStateTreeEditorPropertyPath TargetPath_DEPRECATED;
+PRAGMA_ENABLE_DEPRECATION_WARNINGS
+#endif // WITH_EDITORONLY_DATA
+};
+
+template<>
+struct TStructOpsTypeTraits<FStateTreePropertyPathBinding> : public TStructOpsTypeTraitsBase2<FStateTreePropertyPathBinding>
+{
+	enum 
+	{
+		WithPostSerialize = true,
+	};
+};
+
+using FStateTreeEditorPropertyBinding UE_DEPRECATED(5.3, "Deprecated struct. Please use FStateTreePropertyPathBinding instead.") = FStateTreePropertyPathBinding;
+
+
+/**
+ * Deprecated. Describes a segment of a property path. Used for storage only.
+ */
+USTRUCT()
+struct UE_DEPRECATED(5.3, "Use FStateTreePropertyPath instead.") STATETREEMODULE_API FStateTreePropertySegment
+{
+	GENERATED_BODY()
+
@@ -170,11 +545,13 @@
 
 
 /**
- * Describes property binding, the property from source path is copied into the property at the target path.
+ * Deprecated. Describes property binding, the property from source path is copied into the property at the target path.
  */
 USTRUCT()
-struct STATETREEMODULE_API FStateTreePropertyBinding
+struct UE_DEPRECATED(5.3, "Use FStateTreePropertyPath instead.") STATETREEMODULE_API FStateTreePropertyBinding
 {
+PRAGMA_DISABLE_DEPRECATION_WARNINGS
+
 	GENERATED_BODY()
 
 	/** Source property path. */
@@ -192,9 +569,13 @@
 	/** The type of copy to use between the properties. */
 	UPROPERTY()
 	EStateTreePropertyCopyType CopyType = EStateTreePropertyCopyType::None;
+
+PRAGMA_ENABLE_DEPRECATION_WARNINGS
 };
 
+
 /**
+ * Used internally.
  * Property indirection is a resolved property path segment, used for accessing properties in structs.
  */
 USTRUCT()
@@ -218,16 +599,22 @@
 	UPROPERTY()
 	EStateTreePropertyAccessType Type = EStateTreePropertyAccessType::Offset;
 
+	/** Type of the struct or object instance in case the segment is pointing into an instanced data. */
+	UPROPERTY()
+	TObjectPtr<const UStruct> InstanceStruct = nullptr;
+
 	/** Cached array property. */
 	const FArrayProperty* ArrayProperty = nullptr;
 };
 
+
 /**
+ * Used internally.
  * Describes property copy, the property from source is copied into the property at the target.
  * Copy target struct is described in the property copy batch.
  */
 USTRUCT()
-struct STATETREEMODULE_API FStateTreePropCopy
+struct STATETREEMODULE_API FStateTreePropertyCopy
 {
 	GENERATED_BODY()
 
@@ -258,12 +645,15 @@
 	EStateTreePropertyCopyType Type = EStateTreePropertyCopyType::None;
 };
 
+using FStateTreePropCopy UE_DEPRECATED(5.3, "Deprecated struct. Please use FStateTreePropertyCopy instead.") = FStateTreePropertyCopy;
+
+
 /**
  * Describes a batch of property copies from many sources to one target struct.
  * Note: The batch is used to reference both bindings and copies (a binding turns into copy when resolved).
  */
 USTRUCT()
-struct STATETREEMODULE_API FStateTreePropCopyBatch
+struct STATETREEMODULE_API FStateTreePropertyCopyBatch
 {
 	GENERATED_BODY()
 
@@ -280,6 +670,9 @@
 	uint16 BindingsEnd = 0;
 };
 
+using FStateTreePropCopyBatch UE_DEPRECATED(5.3, "Deprecated struct. Please use FStateTreePropertyCopy instead.") = FStateTreePropertyCopyBatch;
+
+
 /**
  * Runtime storage and execution of property bindings.
  */
@@ -309,11 +702,12 @@
 	 */
 	int32 GetSourceStructNum() const { return SourceStructs.Num(); }
 
+	UE_DEPRECATED(5.3, "Will be removed in a future version.")
 	TArrayView<FStateTreeBindableStructDesc> GetSourceStructs() { return  SourceStructs; };
 
-	TArrayView<FStateTreePropCopyBatch> GetCopyBatches() { return CopyBatches; };
+	UE_DEPRECATED(5.3, "Will be removed in a future version.")
+	TArrayView<FStateTreePropertyCopyBatch> GetCopyBatches() { return CopyBatches; };
 
-	
 	/**
 	 * Copies a batch of properties from source structs to target struct.
 	 * @param SourceStructViews Views to structs where properties are copied from.
@@ -323,17 +717,26 @@
 	 */
 	bool CopyTo(TConstArrayView<FStateTreeDataView> SourceStructViews, const FStateTreeIndex16 TargetBatchIndex, FStateTreeDataView TargetStructView) const;
 
+	/**
+	 * Resets copied properties in TargetStructView. Can be used e.g. to erase UObject references.
+	 * @param TargetBatchIndex Batch index to copy (see FStateTreePropertyBindingCompiler).
+	 * @param TargetStructView View to struct where properties are copied to.
+	 * @return true if all resets succeeded (a reset can fail e.g. if source or destination struct view is invalid).
+	 */
 	bool ResetObjects(const FStateTreeIndex16 TargetBatchIndex, FStateTreeDataView TargetStructView) const;
 
 	void DebugPrintInternalLayout(FString& OutString) const;
 
-protected:
-	[[nodiscard]] bool ResolvePath(const UStruct* Struct, const FStateTreePropertySegment& FirstPathSegment, FStateTreePropertyIndirection& OutFirstIndirection, const FProperty*& OutLeafProperty);
-	[[nodiscard]] bool ValidateCopy(FStateTreePropCopy& Copy) const;
-	void PerformCopy(const FStateTreePropCopy& Copy, uint8* SourceAddress, uint8* TargetAddress) const;
-	void PerformResetObjects(const FStateTreePropCopy& Copy, uint8* TargetAddress) const;
+	/** @return how properties are compatible for copying. */
+	static EStateTreePropertyAccessCompatibility GetPropertyCompatibility(const FProperty* FromProperty, const FProperty* ToProperty);
+
+private:
+	[[nodiscard]] bool ResolvePath(const UStruct* Struct, const FStateTreePropertyPath& Path, FStateTreePropertyIndirection& OutFirstIndirection, FStateTreePropertyPathIndirection& OutLeafIndirection);
+	[[nodiscard]] bool ResolveCopyType(FStateTreePropertyCopy& Copy, const FStateTreePropertyPathIndirection& SourceIndirection,const FStateTreePropertyPathIndirection& TargetIndirection);
+
+	void PerformCopy(const FStateTreePropertyCopy& Copy, uint8* SourceAddress, uint8* TargetAddress) const;
+	void PerformResetObjects(const FStateTreePropertyCopy& Copy, uint8* TargetAddress) const;
 	uint8* GetAddress(FStateTreeDataView InStructView, const FStateTreePropertyIndirection& FirstIndirection, const FProperty* LeafProperty) const;
-	FString GetPathAsString(const FStateTreePropertySegment& FirstPathSegment, const FStateTreePropertySegment* HighlightedSegment = nullptr, const TCHAR* HighlightPrefix = nullptr, const TCHAR* HighlightPostfix = nullptr);
 
 	/** Array of expected source structs. */
 	UPROPERTY()
@@ -341,19 +744,29 @@
 
 	/** Array of copy batches. */
 	UPROPERTY()
-	TArray<FStateTreePropCopyBatch> CopyBatches;
+	TArray<FStateTreePropertyCopyBatch> CopyBatches;
 
 	/** Array of property bindings, resolved into arrays of copies before use. */
 	UPROPERTY()
-	TArray<FStateTreePropertyBinding> PropertyBindings;
+	TArray<FStateTreePropertyPathBinding> PropertyPathBindings;
 
+#if WITH_EDITORONLY_DATA
+PRAGMA_DISABLE_DEPRECATION_WARNINGS
+	/** Array of property bindings, resolved into arrays of copies before use. */
+	UE_DEPRECATED_FORGAME(5.3, "Use PropertyPathBindings instead.")
+	UPROPERTY()
+	TArray<FStateTreePropertyBinding> PropertyBindings_DEPRECATED;
+
 	/** Array of property segments, indexed by property paths. */
+	UE_DEPRECATED_FORGAME(5.3, "Use PropertyPathBindings instead.")
 	UPROPERTY()
-	TArray<FStateTreePropertySegment> PropertySegments;
-
+	TArray<FStateTreePropertySegment> PropertySegments_DEPRECATED;
+PRAGMA_ENABLE_DEPRECATION_WARNINGS
+#endif // WITH_EDITORONLY_DATA
+	
 	/** Array of property copies */
 	UPROPERTY(Transient)
-	TArray<FStateTreePropCopy> PropertyCopies;
+	TArray<FStateTreePropertyCopy> PropertyCopies;
 
 	/** Array of property indirections, indexed by accesses*/
 	UPROPERTY(Transient)
@@ -362,60 +775,51 @@
 	/** Flag indicating if the properties has been resolved successfully . */
 	bool bBindingsResolved = false;
 
-	friend struct FStateTreePropertyBindingCompiler;
+	friend FStateTreePropertyBindingCompiler;
+	friend UStateTree;
 };
 
 
-
 /**
- * Editor representation of a property path in StateTree.
- * Note: This is defined here for IStateTreeBindingLookup.
- */
-USTRUCT()
-struct STATETREEMODULE_API FStateTreeEditorPropertyPath
-{
-	GENERATED_BODY()
-
-	FStateTreeEditorPropertyPath() = default;
-	FStateTreeEditorPropertyPath(const FGuid& InStructID, const TCHAR* PropertyName) : StructID(InStructID) { Path.Add(PropertyName); }
-	FStateTreeEditorPropertyPath(const FStateTreeEditorPropertyPath& InPath) : StructID(InPath.StructID), Path(InPath.Path) {}
-
-	/**
-	 * Returns the property path as a one string. Highlight allows to decorate a specific segment.
-	 * @param HighlightedSegment Index of the highlighted path segment
-	 * @param HighlightPrefix String to append before highlighted segment
-	 * @param HighlightPostfix String to append after highlighted segment
-	 */
-	FString ToString(const int32 HighlightedSegment = INDEX_NONE, const TCHAR* HighlightPrefix = nullptr, const TCHAR* HighlightPostfix = nullptr) const;
-
-	/** Handle of the struct this property path is relative to. */
-	UPROPERTY()
-	FGuid StructID;
-
-	/** Property path segments */
-	UPROPERTY()
-	TArray<FString> Path;
-
-	bool IsValid() const { return StructID.IsValid(); }
-
-	bool operator==(const FStateTreeEditorPropertyPath& RHS) const;
-};
-
-/**
  * Helper interface to reason about bound properties. The implementation is in the editor plugin.
  */
 struct STATETREEMODULE_API IStateTreeBindingLookup
 {
 	/** @return Source path for given target path, or null if binding does not exists. */
-	virtual const FStateTreeEditorPropertyPath* GetPropertyBindingSource(const FStateTreeEditorPropertyPath& InTargetPath) const PURE_VIRTUAL(IStateTreeBindingLookup::GetPropertyBindingSource, return nullptr; );
+	virtual const FStateTreePropertyPath* GetPropertyBindingSource(const FStateTreePropertyPath& InTargetPath) const PURE_VIRTUAL(IStateTreeBindingLookup::GetPropertyBindingSource, return nullptr; );
 
 	/** @return Display name given property path. */
-	virtual FText GetPropertyPathDisplayName(const FStateTreeEditorPropertyPath& InPath) const PURE_VIRTUAL(IStateTreeBindingLookup::GetPropertyPathDisplayName, return FText::GetEmpty(); );
+	virtual FText GetPropertyPathDisplayName(const FStateTreePropertyPath& InPath) const PURE_VIRTUAL(IStateTreeBindingLookup::GetPropertyPathDisplayName, return FText::GetEmpty(); );
 
 	/** @return Leaf property based on property path. */
-	virtual const FProperty* GetPropertyPathLeafProperty(const FStateTreeEditorPropertyPath& InPath) const PURE_VIRTUAL(IStateTreeBindingLookup::GetPropertyPathLeafProperty, return nullptr; );
+	virtual const FProperty* GetPropertyPathLeafProperty(const FStateTreePropertyPath& InPath) const PURE_VIRTUAL(IStateTreeBindingLookup::GetPropertyPathLeafProperty, return nullptr; );
+
+PRAGMA_DISABLE_DEPRECATION_WARNINGS
+	UE_DEPRECATED(5.3, "Use version with FStateTreePropertyPath instead.")
+	virtual const FStateTreeEditorPropertyPath* GetPropertyBindingSource(const FStateTreeEditorPropertyPath& InTargetPath) const final { return nullptr; }
+	
+	UE_DEPRECATED(5.3, "Use version with FStateTreePropertyPath instead.")
+	virtual FText GetPropertyPathDisplayName(const FStateTreeEditorPropertyPath& InPath) const final { return FText::GetEmpty(); }
+	
+	UE_DEPRECATED(5.3, "Use version with FStateTreePropertyPath instead.")
+	virtual const FProperty* GetPropertyPathLeafProperty(const FStateTreeEditorPropertyPath& InPath) const final { return nullptr; }
+PRAGMA_ENABLE_DEPRECATION_WARNINGS
 };
 
+
+namespace UE::StateTree::Private
+{
+#if WITH_EDITORONLY_DATA
+// Helper functions to convert between property path types.
+PRAGMA_DISABLE_DEPRECATION_WARNINGS
+	UE_DEPRECATED(5.3, "Will be removed when FStateTreeEditorPropertyPath is removed.")
+	extern STATETREEMODULE_API FStateTreePropertyPath ConvertEditorPath(const FStateTreeEditorPropertyPath& InEditorPath);
+	UE_DEPRECATED(5.3, "Will be removed when FStateTreeEditorPropertyPath is removed.")
+	extern STATETREEMODULE_API FStateTreeEditorPropertyPath ConvertEditorPath(const FStateTreePropertyPath& InPath);
+PRAGMA_ENABLE_DEPRECATION_WARNINGS
+#endif // WITH_EDITORONLY_DATA
+}; // UE::StateTree::Private 
+
 #if UE_ENABLE_INCLUDE_ORDER_DEPRECATED_IN_5_2
 #include "CoreMinimal.h"
 #endif
diff --git a/Engine/Plugins/Runtime/StateTree/Source/StateTreeModule/Public/StateTreeTypes.h b/Engine/Plugins/Runtime/StateTree/Source/StateTreeModule/Public/StateTreeTypes.h
--- a/Engine/Plugins/Runtime/StateTree/Source/StateTreeModule/Public/StateTreeTypes.h
+++ b/Engine/Plugins/Runtime/StateTree/Source/StateTreeModule/Public/StateTreeTypes.h
@@ -955,7 +955,7 @@ struct STATETREEMODULE_API FStateTreeDataView
 	FStateTreeDataView() = default;
 
 	// USTRUCT() constructor.
-	FStateTreeDataView(const UScriptStruct* InScriptStruct, uint8* InMemory) : Struct(InScriptStruct), Memory(InMemory)
+	FStateTreeDataView(const UStruct* InStruct, uint8* InMemory) : Struct(InStruct), Memory(InMemory)
 	{
 		// Must have type with valid pointer.
 		check(!Memory || (Memory && Struct));
diff --git a/Engine/Plugins/Runtime/StateTree/Source/StateTreeModule/StateTreeModule.Build.cs b/Engine/Plugins/Runtime/StateTree/Source/StateTreeModule/StateTreeModule.Build.cs
--- a/Engine/Plugins/Runtime/StateTree/Source/StateTreeModule/StateTreeModule.Build.cs
+++ b/Engine/Plugins/Runtime/StateTree/Source/StateTreeModule/StateTreeModule.Build.cs
@@ -25,6 +25,7 @@ namespace UnrealBuildTool.Rules
 			PrivateDependencyModuleNames.AddRange(
 				new string[] {
 					"RenderCore",
+					"PropertyPath",
 				}
 			);
 			
diff --git a/Engine/Plugins/Runtime/StateTree/Source/StateTreeTestSuite/Private/StateTreeTest.cpp b/Engine/Plugins/Runtime/StateTree/Source/StateTreeTestSuite/Private/StateTreeTest.cpp
--- a/Engine/Plugins/Runtime/StateTree/Source/StateTreeTestSuite/Private/StateTreeTest.cpp
+++ b/Engine/Plugins/Runtime/StateTree/Source/StateTreeTestSuite/Private/StateTreeTest.cpp
@@ -49,23 +49,23 @@ struct FStateTreeTest_MakeAndBakeStateTree : FAITestBase
 		
 		// State A
 		auto& TaskB1 = StateA.AddTask<FTestTask_B>();
-		EditorData.AddPropertyBinding(FStateTreeEditorPropertyPath(EvalA.ID, TEXT("IntA")), FStateTreeEditorPropertyPath(TaskB1.ID, TEXT("IntB")));
+		EditorData.AddPropertyBinding(EvalA, TEXT("IntA"), TaskB1, TEXT("IntB"));
 
 		auto& IntCond = StateA.AddEnterCondition<FStateTreeCompareIntCondition>(EGenericAICheck::Less);
 		IntCond.GetInstanceData().Right = 2;
 
-		EditorData.AddPropertyBinding(FStateTreeEditorPropertyPath(EvalA.ID, TEXT("IntA")), FStateTreeEditorPropertyPath(IntCond.ID, TEXT("Left")));
+		EditorData.AddPropertyBinding(EvalA, TEXT("IntA"), IntCond, TEXT("Left"));
 
 		StateA.AddTransition(EStateTreeTransitionTrigger::OnStateCompleted, EStateTreeTransitionType::GotoState, &StateB);
 
 		// State B
 		auto& TaskB2 = StateB.AddTask<FTestTask_B>();
-		EditorData.AddPropertyBinding(FStateTreeEditorPropertyPath(EvalA.ID, TEXT("bBoolA")), FStateTreeEditorPropertyPath(TaskB2.ID, TEXT("bBoolB")));
+		EditorData.AddPropertyBinding(EvalA, TEXT("bBoolA"), TaskB2, TEXT("bBoolB"));
 
 		FStateTreeTransition& Trans = StateB.AddTransition({}, EStateTreeTransitionType::GotoState, &Root);
 		auto& TransFloatCond = Trans.AddCondition<FStateTreeCompareFloatCondition>(EGenericAICheck::Less);
 		TransFloatCond.GetInstanceData().Right = 13.0f;
-		EditorData.AddPropertyBinding(FStateTreeEditorPropertyPath(EvalA.ID, TEXT("FloatA")), FStateTreeEditorPropertyPath(TransFloatCond.ID, TEXT("Left")));
+		EditorData.AddPropertyBinding(EvalA, TEXT("FloatA"), TransFloatCond, TEXT("Left"));
 
 		StateB.AddTransition(EStateTreeTransitionTrigger::OnStateCompleted, EStateTreeTransitionType::Succeeded);
 
@@ -656,9 +656,9 @@ struct FStateTreeTest_TransitionGlobalDataView : FAITestBase
 
 		// State B
 		auto& Task1 = StateB.AddTask<FTestTask_PrintValue>(FName(TEXT("Task1")));
-		EditorData.AddPropertyBinding(FStateTreeEditorPropertyPath(EvalA.ID, TEXT("IntA")), FStateTreeEditorPropertyPath(Task1.ID, TEXT("Value")));
+		EditorData.AddPropertyBinding(EvalA, TEXT("IntA"), Task1, TEXT("Value"));
 		auto& Task2 = StateB.AddTask<FTestTask_PrintValue>(FName(TEXT("Task2")));
-		EditorData.AddPropertyBinding(FStateTreeEditorPropertyPath(GlobalTask.ID, TEXT("Value")), FStateTreeEditorPropertyPath(Task2.ID, TEXT("Value")));
+		EditorData.AddPropertyBinding(GlobalTask, TEXT("Value"), Task2, TEXT("Value"));
 
 		FStateTreeCompilerLog Log;
 		FStateTreeCompiler Compiler(Log);
@@ -692,6 +692,323 @@ struct FStateTreeTest_TransitionGlobalDataView : FAITestBase
 };
 IMPLEMENT_AI_INSTANT_TEST(FStateTreeTest_TransitionGlobalDataView, "System.StateTree.Transition.GlobalDataView");
 
+struct FStateTreeTest_PropertyPathOffset : FAITestBase
+{
+	virtual bool InstantTest() override
+	{
+		FStateTreePropertyPath Path;
+		const bool bParseResult = Path.FromString(TEXT("StructB.B"));
+
+		AITEST_TRUE("Parsing path should succeeed", bParseResult);
+		AITEST_EQUAL("Should have 2 path segments", Path.NumSegments(), 2);
+
+		FString ResolveErrors;
+		TArray<FStateTreePropertyPathIndirection> Indirections;
+		const bool bResolveResult = Path.ResolveIndirections(FStateTreeTest_PropertyStruct::StaticStruct(), Indirections, &ResolveErrors);
+
+		AITEST_TRUE("Resolve path should succeeed", bResolveResult);
+		AITEST_EQUAL("Should have no resolve errors", ResolveErrors.Len(), 0);
+		
+		AITEST_EQUAL("Should have 2 indirections", Indirections.Num(), 2);
+		AITEST_EQUAL("Indirection 0 should be Offset type", Indirections[0].GetAccessType(), EStateTreePropertyAccessType::Offset);
+		AITEST_EQUAL("Indirection 1 should be Offset type", Indirections[1].GetAccessType(), EStateTreePropertyAccessType::Offset);
+
+		return true;
+	}
+};
+IMPLEMENT_AI_INSTANT_TEST(FStateTreeTest_PropertyPathOffset, "System.StateTree.PropertyPath.Offset");
+
+struct FStateTreeTest_PropertyPathParseFail : FAITestBase
+{
+	virtual bool InstantTest() override
+	{
+		{
+			FStateTreePropertyPath Path;
+			const bool bParseResult = Path.FromString(TEXT("")); // empty is valid.
+			AITEST_TRUE("Parsing path should succeed", bParseResult);
+		}
+
+		{
+			FStateTreePropertyPath Path;
+			const bool bParseResult = Path.FromString(TEXT("StructB.[0]B"));
+			AITEST_FALSE("Parsing path should fail", bParseResult);
+		}
+
+		{
+			FStateTreePropertyPath Path;
+			const bool bParseResult = Path.FromString(TEXT("StructB..NoThere"));
+			AITEST_FALSE("Parsing path should fail", bParseResult);
+		}
+
+		{
+			FStateTreePropertyPath Path;
+			const bool bParseResult = Path.FromString(TEXT("."));
+			AITEST_FALSE("Parsing path should fail", bParseResult);
+		}
+
+		{
+			FStateTreePropertyPath Path;
+			const bool bParseResult = Path.FromString(TEXT("StructB..B"));
+			AITEST_FALSE("Parsing path should fail", bParseResult);
+		}
+
+		return true;
+	}
+};
+IMPLEMENT_AI_INSTANT_TEST(FStateTreeTest_PropertyPathParseFail, "System.StateTree.PropertyPath.ParseFail");
+
+struct FStateTreeTest_PropertyPathOffsetFail : FAITestBase
+{
+	virtual bool InstantTest() override
+	{
+		FStateTreePropertyPath Path;
+		const bool bParseResult = Path.FromString(TEXT("StructB.Q"));
+
+		AITEST_TRUE("Parsing path should succeeed", bParseResult);
+		AITEST_EQUAL("Should have 2 path segments", Path.NumSegments(), 2);
+
+		FString ResolveErrors;
+		TArray<FStateTreePropertyPathIndirection> Indirections;
+		const bool bResolveResult = Path.ResolveIndirections(FStateTreeTest_PropertyStruct::StaticStruct(), Indirections, &ResolveErrors);
+
+		AITEST_FALSE("Resolve path should not succeeed", bResolveResult);
+		AITEST_NOT_EQUAL("Should have errors", ResolveErrors.Len(), 0);
+		
+		AITEST_EQUAL("Should have 0 indirections", Indirections.Num(), 0);
+
+		return true;
+	}
+};
+IMPLEMENT_AI_INSTANT_TEST(FStateTreeTest_PropertyPathOffsetFail, "System.StateTree.PropertyPath.OffsetFail");
+
+struct FStateTreeTest_PropertyPathObject : FAITestBase
+{
+	virtual bool InstantTest() override
+	{
+		FStateTreePropertyPath Path;
+		const bool bParseResult = Path.FromString(TEXT("InstancedObject.A"));
+
+		AITEST_TRUE("Parsing path should succeeed", bParseResult);
+		AITEST_EQUAL("Should have 2 path segments", Path.NumSegments(), 2);
+
+		UStateTreeTest_PropertyObject* Object = NewObject<UStateTreeTest_PropertyObject>();
+		Object->InstancedObject = NewObject<UStateTreeTest_PropertyObjectInstanced>();
+		
+		const bool bUpdateResult = Path.UpdateInstanceStructsFromValue(FStateTreeDataView(Object));
+
+		AITEST_TRUE("Update instance types should succeeed", bUpdateResult);
+		AITEST_TRUE("Path segment 0 instance type should be UStateTreeTest_PropertyObjectInstanced", Path.GetSegment(0).GetInstanceStruct() == UStateTreeTest_PropertyObjectInstanced::StaticClass());
+		AITEST_TRUE("Path segment 1 instance type should be nullptr", Path.GetSegment(1).GetInstanceStruct() == nullptr);
+
+		return true;
+	}
+};
+IMPLEMENT_AI_INSTANT_TEST(FStateTreeTest_PropertyPathObject, "System.StateTree.PropertyPath.Object");
+
+struct FStateTreeTest_PropertyPathWrongObject : FAITestBase
+{
+	virtual bool InstantTest() override
+	{
+		FStateTreePropertyPath Path;
+		const bool bParseResult = Path.FromString(TEXT("InstancedObject.B"));
+
+		AITEST_TRUE("Parsing path should succeeed", bParseResult);
+		AITEST_EQUAL("Should have 2 path segments", Path.NumSegments(), 2);
+
+		UStateTreeTest_PropertyObject* Object = NewObject<UStateTreeTest_PropertyObject>();
+
+		Object->InstancedObject = NewObject<UStateTreeTest_PropertyObjectInstancedWithB>();
+		{
+			FString ResolveErrors;
+			TArray<FStateTreePropertyPathIndirection> Indirections;
+			const bool bResolveResult = Path.ResolveIndirectionsWithValue(FStateTreeDataView(Object), Indirections, &ResolveErrors);
+
+			AITEST_TRUE("Resolve path should succeeed", bResolveResult);
+			AITEST_EQUAL("Should have 2 indirections", Indirections.Num(), 2);
+			AITEST_TRUE("Object ", Indirections[0].GetAccessType() == EStateTreePropertyAccessType::ObjectInstance);
+			AITEST_TRUE("Object ", Indirections[0].GetContainerStruct() == Object->GetClass());
+			AITEST_TRUE("Object ", Indirections[0].GetInstanceStruct() == UStateTreeTest_PropertyObjectInstancedWithB::StaticClass());
+			AITEST_EQUAL("Should not have error", ResolveErrors.Len(), 0);
+		}
+
+		Object->InstancedObject = NewObject<UStateTreeTest_PropertyObjectInstanced>();
+		{
+			FString ResolveErrors;
+			TArray<FStateTreePropertyPathIndirection> Indirections;
+			const bool bResolveResult = Path.ResolveIndirectionsWithValue(FStateTreeDataView(Object), Indirections, &ResolveErrors);
+
+			AITEST_FALSE("Resolve path should fail", bResolveResult);
+			AITEST_EQUAL("Should have 0 indirections", Indirections.Num(), 0);
+			AITEST_NOT_EQUAL("Should have error", ResolveErrors.Len(), 0);
+		}
+		
+		return true;
+	}
+};
+IMPLEMENT_AI_INSTANT_TEST(FStateTreeTest_PropertyPathWrongObject, "System.StateTree.PropertyPath.WrongObject");
+
+struct FStateTreeTest_PropertyPathArray : FAITestBase
+{
+	virtual bool InstantTest() override
+	{
+		FStateTreePropertyPath Path;
+		const bool bParseResult = Path.FromString(TEXT("ArrayOfInts[1]"));
+
+		AITEST_TRUE("Parsing path should succeeed", bParseResult);
+		AITEST_EQUAL("Should have 1 path segments", Path.NumSegments(), 1);
+
+		UStateTreeTest_PropertyObject* Object = NewObject<UStateTreeTest_PropertyObject>();
+		Object->ArrayOfInts.Add(42);
+		Object->ArrayOfInts.Add(123);
+
+		FString ResolveErrors;
+		TArray<FStateTreePropertyPathIndirection> Indirections;
+		const bool bResolveResult = Path.ResolveIndirectionsWithValue(FStateTreeDataView(Object), Indirections, &ResolveErrors);
+
+		AITEST_TRUE("Resolve path should succeeed", bResolveResult);
+		AITEST_EQUAL("Should have no resolve errors", ResolveErrors.Len(), 0);
+		AITEST_EQUAL("Should have 2 indirections", Indirections.Num(), 2);
+		AITEST_EQUAL("Indirection 0 should be IndexArray type", Indirections[0].GetAccessType(), EStateTreePropertyAccessType::IndexArray);
+		AITEST_EQUAL("Indirection 1 should be Offset type", Indirections[1].GetAccessType(), EStateTreePropertyAccessType::Offset);
+
+		const int32 Value = *reinterpret_cast<const int32*>(Indirections[1].GetPropertyAddress());
+		AITEST_EQUAL("Value should be 123", Value, 123);
+		
+		return true;
+	}
+};
+IMPLEMENT_AI_INSTANT_TEST(FStateTreeTest_PropertyPathArray, "System.StateTree.PropertyPath.Array");
+
+struct FStateTreeTest_PropertyPathArrayInvalidIndex : FAITestBase
+{
+	virtual bool InstantTest() override
+	{
+		FStateTreePropertyPath Path;
+		const bool bParseResult = Path.FromString(TEXT("ArrayOfInts[123]"));
+
+		AITEST_TRUE("Parsing path should succeeed", bParseResult);
+		AITEST_EQUAL("Should have 1 path segments", Path.NumSegments(), 1);
+
+		UStateTreeTest_PropertyObject* Object = NewObject<UStateTreeTest_PropertyObject>();
+		Object->ArrayOfInts.Add(42);
+		Object->ArrayOfInts.Add(123);
+
+		FString ResolveErrors;
+		TArray<FStateTreePropertyPathIndirection> Indirections;
+		const bool bResolveResult = Path.ResolveIndirectionsWithValue(FStateTreeDataView(Object), Indirections, &ResolveErrors);
+
+		AITEST_FALSE("Resolve path should fail", bResolveResult);
+		
+		return true;
+	}
+};
+IMPLEMENT_AI_INSTANT_TEST(FStateTreeTest_PropertyPathArrayInvalidIndex, "System.StateTree.PropertyPath.ArrayInvalidIndex");
+
+struct FStateTreeTest_PropertyPathArrayOfStructs : FAITestBase
+{
+	virtual bool InstantTest() override
+	{
+		FStateTreePropertyPath Path1;
+		Path1.FromString(TEXT("ArrayOfStruct[0].B"));
+
+		FStateTreePropertyPath Path2;
+		Path2.FromString(TEXT("ArrayOfStruct[2].StructB.B"));
+
+		UStateTreeTest_PropertyObject* Object = NewObject<UStateTreeTest_PropertyObject>();
+		Object->ArrayOfStruct.AddDefaulted_GetRef().B = 3;
+		Object->ArrayOfStruct.AddDefaulted();
+		Object->ArrayOfStruct.AddDefaulted_GetRef().StructB.B = 42;
+
+		{
+			FString ResolveErrors;
+			TArray<FStateTreePropertyPathIndirection> Indirections;
+			const bool bResolveResult = Path1.ResolveIndirectionsWithValue(FStateTreeDataView(Object), Indirections, &ResolveErrors);
+
+			AITEST_TRUE("Resolve path1 should succeeed", bResolveResult);
+			AITEST_EQUAL("Should have no resolve errors", ResolveErrors.Len(), 0);
+			AITEST_EQUAL("Should have 3 indirections", Indirections.Num(), 3);
+			AITEST_EQUAL("Indirection 0 should be ArrayIndex type", Indirections[0].GetAccessType(), EStateTreePropertyAccessType::IndexArray);
+			AITEST_EQUAL("Indirection 1 should be Offset type", Indirections[1].GetAccessType(), EStateTreePropertyAccessType::Offset);
+			AITEST_EQUAL("Indirection 2 should be Offset type", Indirections[2].GetAccessType(), EStateTreePropertyAccessType::Offset);
+
+			const int32 Value = *reinterpret_cast<const int32*>(Indirections[2].GetPropertyAddress());
+			AITEST_EQUAL("Value should be 3", Value, 3);
+		}
+
+		{
+			FString ResolveErrors;
+			TArray<FStateTreePropertyPathIndirection> Indirections;
+			const bool bResolveResult = Path2.ResolveIndirectionsWithValue(FStateTreeDataView(Object), Indirections, &ResolveErrors);
+
+			AITEST_TRUE("Resolve path2 should succeeed", bResolveResult);
+			AITEST_EQUAL("Should have no resolve errors", ResolveErrors.Len(), 0);
+			AITEST_EQUAL("Should have 4 indirections", Indirections.Num(), 4);
+			AITEST_EQUAL("Indirection 0 should be ArrayIndex type", Indirections[0].GetAccessType(), EStateTreePropertyAccessType::IndexArray);
+			AITEST_EQUAL("Indirection 1 should be Offset type", Indirections[1].GetAccessType(), EStateTreePropertyAccessType::Offset);
+			AITEST_EQUAL("Indirection 2 should be Offset type", Indirections[2].GetAccessType(), EStateTreePropertyAccessType::Offset);
+			AITEST_EQUAL("Indirection 3 should be Offset type", Indirections[3].GetAccessType(), EStateTreePropertyAccessType::Offset);
+
+			const int32 Value = *reinterpret_cast<const int32*>(Indirections[3].GetPropertyAddress());
+			AITEST_EQUAL("Value should be 42", Value, 42);
+		}
+		return true;
+	}
+};
+IMPLEMENT_AI_INSTANT_TEST(FStateTreeTest_PropertyPathArrayOfStructs, "System.StateTree.PropertyPath.ArrayOfStructs");
+
+struct FStateTreeTest_PropertyPathArrayOfInstancedObjects : FAITestBase
+{
+	virtual bool InstantTest() override
+	{
+		FStateTreePropertyPath Path;
+		Path.FromString(TEXT("ArrayOfInstancedStructs[0].B"));
+
+		FStateTreeTest_PropertyStruct Struct;
+		Struct.B = 123;
+		
+		UStateTreeTest_PropertyObject* Object = NewObject<UStateTreeTest_PropertyObject>();
+		Object->ArrayOfInstancedStructs.Emplace(FConstStructView::Make(Struct));
+
+		const bool bUpdateResult = Path.UpdateInstanceStructsFromValue(FStateTreeDataView(Object));
+		AITEST_TRUE("Update instance types should succeeed", bUpdateResult);
+		AITEST_EQUAL("Should have 2 path segments", Path.NumSegments(), 2);
+		AITEST_TRUE("Path segment 0 instance type should be FStateTreeTest_PropertyStruct", Path.GetSegment(0).GetInstanceStruct() == FStateTreeTest_PropertyStruct::StaticStruct());
+		AITEST_TRUE("Path segment 1 instance type should be nullptr", Path.GetSegment(1).GetInstanceStruct() == nullptr);
+
+		{
+			FString ResolveErrors;
+			TArray<FStateTreePropertyPathIndirection> Indirections;
+			const bool bResolveResult = Path.ResolveIndirections(UStateTreeTest_PropertyObject::StaticClass(), Indirections, &ResolveErrors);
+
+			AITEST_TRUE("Resolve path should succeeed", bResolveResult);
+			AITEST_EQUAL("Should have no resolve errors", ResolveErrors.Len(), 0);
+			AITEST_EQUAL("Should have 3 indirections", Indirections.Num(), 3);
+			AITEST_EQUAL("Indirection 0 should be ArrayIndex type", Indirections[0].GetAccessType(), EStateTreePropertyAccessType::IndexArray);
+			AITEST_EQUAL("Indirection 1 should be StructInstance type", Indirections[1].GetAccessType(), EStateTreePropertyAccessType::StructInstance);
+			AITEST_EQUAL("Indirection 2 should be Offset type", Indirections[2].GetAccessType(), EStateTreePropertyAccessType::Offset);
+		}
+
+		{
+			FString ResolveErrors;
+			TArray<FStateTreePropertyPathIndirection> Indirections;
+			const bool bResolveResult = Path.ResolveIndirectionsWithValue(FStateTreeDataView(Object), Indirections, &ResolveErrors);
+
+			AITEST_TRUE("Resolve path should succeeed", bResolveResult);
+			AITEST_EQUAL("Should have no resolve errors", ResolveErrors.Len(), 0);
+			AITEST_EQUAL("Should have 3 indirections", Indirections.Num(), 3);
+			AITEST_EQUAL("Indirection 0 should be ArrayIndex type", Indirections[0].GetAccessType(), EStateTreePropertyAccessType::IndexArray);
+			AITEST_EQUAL("Indirection 1 should be StructInstance type", Indirections[1].GetAccessType(), EStateTreePropertyAccessType::StructInstance);
+			AITEST_EQUAL("Indirection 2 should be Offset type", Indirections[2].GetAccessType(), EStateTreePropertyAccessType::Offset);
+
+			const int32 Value = *reinterpret_cast<const int32*>(Indirections[2].GetPropertyAddress());
+			AITEST_EQUAL("Value should be 123", Value, 123);
+		}
+		
+		return true;
+	}
+};
+IMPLEMENT_AI_INSTANT_TEST(FStateTreeTest_PropertyPathArrayOfInstancedObjects, "System.StateTree.PropertyPath.ArrayOfInstancedObjects");
+
 UE_ENABLE_OPTIMIZATION_SHIP
 
 #undef LOCTEXT_NAMESPACE
diff --git a/Engine/Plugins/Runtime/StateTree/Source/StateTreeTestSuite/Private/StateTreeTestTypes.h b/Engine/Plugins/Runtime/StateTree/Source/StateTreeTestSuite/Private/StateTreeTestTypes.h
--- a/Engine/Plugins/Runtime/StateTree/Source/StateTreeTestSuite/Private/StateTreeTestTypes.h
+++ b/Engine/Plugins/Runtime/StateTree/Source/StateTreeTestSuite/Private/StateTreeTestTypes.h
@@ -303,3 +303,81 @@ struct FStateTreeTestRunContext
 	int32 Count = 0;
 };
 
+
+USTRUCT()
+struct FStateTreeTest_PropertyStructB
+{
+	GENERATED_BODY()
+
+	UPROPERTY(EditAnywhere, Category = "")
+	int32 B = 0;
+};
+
+USTRUCT()
+struct FStateTreeTest_PropertyStruct
+{
+	GENERATED_BODY()
+
+	UPROPERTY(EditAnywhere, Category = "")
+	int32 A = 0;
+
+	UPROPERTY(EditAnywhere, Category = "")
+	int32 B = 0;
+
+	UPROPERTY(EditAnywhere, Category = "")
+	FStateTreeTest_PropertyStructB StructB;
+};
+
+UCLASS(HideDropdown)
+class UStateTreeTest_PropertyObjectInstanced : public UObject
+{
+	GENERATED_BODY()
+public:
+
+	UPROPERTY(EditAnywhere, Category = "")
+	int32 A = 0;
+
+	UPROPERTY(EditAnywhere, Category = "")
+	FInstancedStruct InstancedStruct;
+
+	UPROPERTY(EditAnywhere, Category = "")
+	TArray<FGameplayTag> ArrayOfTags;
+};
+
+UCLASS(HideDropdown)
+class UStateTreeTest_PropertyObjectInstancedWithB : public UStateTreeTest_PropertyObjectInstanced
+{
+	GENERATED_BODY()
+public:
+
+	UPROPERTY(EditAnywhere, Category = "")
+	int32 B = 0;
+};
+
+UCLASS(HideDropdown)
+class UStateTreeTest_PropertyObject : public UObject
+{
+	GENERATED_BODY()
+public:
+	
+	UPROPERTY(EditAnywhere, Instanced, Category = "")
+	TObjectPtr<UStateTreeTest_PropertyObjectInstanced> InstancedObject;
+
+	UPROPERTY(EditAnywhere, Instanced, Category = "")
+	TArray<TObjectPtr<UStateTreeTest_PropertyObjectInstanced>> ArrayOfInstancedObjects;
+
+	UPROPERTY(EditAnywhere, Category = "")
+	TArray<int32> ArrayOfInts;
+
+	UPROPERTY(EditAnywhere, Category = "")
+	FInstancedStruct InstancedStruct;
+
+	UPROPERTY(EditAnywhere, Category = "")
+	TArray<FInstancedStruct> ArrayOfInstancedStructs;
+
+	UPROPERTY(EditAnywhere, Category = "")
+	FStateTreeTest_PropertyStruct Struct;
+
+	UPROPERTY(EditAnywhere, Category = "")
+	TArray<FStateTreeTest_PropertyStruct> ArrayOfStruct;
+};
