StateTree: Fix extra %s when logging in FStateTreePropertyBindings::ResolvePaths()

diff --git a/Engine/Plugins/Runtime/StateTree/Source/StateTreeModule/Private/StateTreePropertyBindings.cpp b/Engine/Plugins/Runtime/StateTree/Source/StateTreeModule/Private/StateTreePropertyBindings.cpp
--- a/Engine/Plugins/Runtime/StateTree/Source/StateTreeModule/Private/StateTreePropertyBindings.cpp
+++ b/Engine/Plugins/Runtime/StateTree/Source/StateTreeModule/Private/StateTreePropertyBindings.cpp
@@ -110,7 +110,7 @@ bool FStateTreePropertyBindings::ResolvePaths()
 
 			if (!Binding.GetCompiledSourceStructIndex().IsValid())
 			{
-				UE_LOG(LogStateTree, Error, TEXT("%hs: '%s' Invalid source struct for property binding %s."), __FUNCTION__, *Binding.GetSourcePath().ToString());
+				UE_LOG(LogStateTree, Error, TEXT("%hs: Invalid source struct for property binding %s."), __FUNCTION__, *Binding.GetSourcePath().ToString());
 				bResult = false;
 				continue;
 			}
