StateTree: Allow class cache to discover derived native types.
- Fixes bug where BP classes that derive from natively derived UStateTreeXXXBlueprintBase classes would not show up in the node picker

diff --git a/Engine/Plugins/Runtime/StateTree/Source/StateTreeEditorModule/Private/StateTreeNodeClassCache.cpp b/Engine/Plugins/Runtime/StateTree/Source/StateTreeEditorModule/Private/StateTreeNodeClassCache.cpp
--- a/Engine/Plugins/Runtime/StateTree/Source/StateTreeEditorModule/Private/StateTreeNodeClassCache.cpp
+++ b/Engine/Plugins/Runtime/StateTree/Source/StateTreeEditorModule/Private/StateTreeNodeClassCache.cpp
@@ -12,6 +12,7 @@
 #include "AssetRegistry/ARFilter.h"
 #include "Logging/MessageLog.h"
 #include "AssetRegistry/AssetRegistryModule.h"
+#include "Misc/EnumerateRange.h"
 
 #include UE_INLINE_GENERATED_CPP_BY_NAME(StateTreeNodeClassCache)
 
@@ -131,10 +132,8 @@ void FStateTreeNodeClassCache::AddRootStruct(UStruct* RootStruct)
 	if (RootClasses.ContainsByPredicate([RootStruct](const FRootClassContainer& RootClass){ return RootClass.BaseStruct == RootStruct; }))
 	{
 		return;
-		return;
 	}
 	
-	RootClassNameToIndex.Add(RootStruct->GetName(), RootClasses.Num());
 	RootClasses.Emplace(RootStruct);
 
 	InvalidateCache();
@@ -173,7 +172,7 @@ void FStateTreeNodeClassCache::OnAssetRemoved(const FAssetData& AssetData)
 		ConstructorHelpers::StripObjectClass(AssetClassName);
 		AssetClassName = FPackageName::ObjectPathToObjectName(AssetClassName);
 
-		if (const int32* RootClassIndex = RootClassNameToIndex.Find(AssetNativeParentClassName))
+		if (const int32* RootClassIndex = ClassNameToRootIndex.Find(AssetNativeParentClassName))
 		{
 			FRootClassContainer& RootClass = RootClasses[*RootClassIndex];
 			const FName ClassName(AssetClassName);
@@ -189,6 +188,7 @@ void FStateTreeNodeClassCache::InvalidateCache()
 		RootClass.ClassData.Reset();
 		RootClass.bUpdated = false;
 	}
+	ClassNameToRootIndex.Reset();
 }
 
 void FStateTreeNodeClassCache::OnReloadComplete(EReloadCompleteReason Reason)
@@ -203,13 +203,13 @@ void FStateTreeNodeClassCache::UpdateBlueprintClass(const FAssetData& AssetData)
 
 	if (AssetData.GetTagValue(FBlueprintTags::GeneratedClassPath, AssetClassName) && AssetData.GetTagValue(FBlueprintTags::NativeParentClassPath, AssetNativeParentClassName))
 	{
-		UObject* Outer1 = nullptr;
-		ResolveName(Outer1, AssetClassName, false, false);
+		UObject* AssetClassPackage = nullptr;
+		ResolveName(AssetClassPackage, AssetClassName, false, false);
 
-		UObject* Outer2 = nullptr;
-		ResolveName(Outer2, AssetNativeParentClassName, false, false);
+		UObject* AssetNativeParentClassPackage = nullptr;
+		ResolveName(AssetNativeParentClassPackage, AssetNativeParentClassName, false, false);
 
-		if (const int32* RootClassIndex = RootClassNameToIndex.Find(AssetNativeParentClassName))
+		if (const int32* RootClassIndex = ClassNameToRootIndex.Find(AssetNativeParentClassName))
 		{
 			FRootClassContainer& RootClass = RootClasses[*RootClassIndex];
 
@@ -231,31 +231,33 @@ void FStateTreeNodeClassCache::UpdateBlueprintClass(const FAssetData& AssetData)
 
 void FStateTreeNodeClassCache::CacheClasses()
 {
-	for (FRootClassContainer& RootClass : RootClasses)
+	for (TEnumerateRef<FRootClassContainer> RootClass : EnumerateRange(RootClasses))
 	{
-		RootClass.ClassData.Reset();
-		RootClass.bUpdated = true;
-		
+		RootClass->ClassData.Reset();
+		RootClass->bUpdated = true;
+
 		// gather all native classes
-		if (const UClass* Class = Cast<UClass>(RootClass.BaseStruct))
+		if (const UClass* Class = Cast<UClass>(RootClass->BaseStruct))
 		{
 			for (TObjectIterator<UClass> It; It; ++It)
 			{
 				UClass* TestClass = *It;
 				if (TestClass->HasAnyClassFlags(CLASS_Native) && TestClass->IsChildOf(Class))
 				{
-					RootClass.ClassData.Add(MakeShareable(new FStateTreeNodeClassData(TestClass)));
+					RootClass->ClassData.Add(MakeShareable(new FStateTreeNodeClassData(TestClass)));
+					ClassNameToRootIndex.Add(TestClass->GetName(), RootClass.GetIndex());
 				}
 			}
 		}
-		else if (const UScriptStruct* ScriptStruct = Cast<UScriptStruct>(RootClass.BaseStruct))
+		else if (const UScriptStruct* ScriptStruct = Cast<UScriptStruct>(RootClass->BaseStruct))
 		{
 			for (TObjectIterator<UScriptStruct> It; It; ++It)
 			{
 				UScriptStruct* TestStruct = *It;
 				if (TestStruct->IsChildOf(ScriptStruct))
 				{
-					RootClass.ClassData.Add(MakeShareable(new FStateTreeNodeClassData(TestStruct)));
+					RootClass->ClassData.Add(MakeShareable(new FStateTreeNodeClassData(TestStruct)));
+					ClassNameToRootIndex.Add(TestStruct->GetName(), RootClass.GetIndex());
 				}
 			}
 		}
diff --git a/Engine/Plugins/Runtime/StateTree/Source/StateTreeEditorModule/Public/StateTreeNodeClassCache.h b/Engine/Plugins/Runtime/StateTree/Source/StateTreeEditorModule/Public/StateTreeNodeClassCache.h
--- a/Engine/Plugins/Runtime/StateTree/Source/StateTreeEditorModule/Public/StateTreeNodeClassCache.h
+++ b/Engine/Plugins/Runtime/StateTree/Source/StateTreeEditorModule/Public/StateTreeNodeClassCache.h
@@ -113,7 +113,7 @@ private:
 	};
 
 	TArray<FRootClassContainer> RootClasses;
-	TMap<FString, int32> RootClassNameToIndex;
+	TMap<FString, int32> ClassNameToRootIndex;
 };
 
 #if UE_ENABLE_INCLUDE_ORDER_DEPRECATED_IN_5_2
