StateTree: remove unnecessary null test to keep static analysis happy.

diff --git a/Engine/Plugins/Runtime/StateTree/Source/StateTreeEditorModule/Private/StateTreePropertyBindingCompiler.cpp b/Engine/Plugins/Runtime/StateTree/Source/StateTreeEditorModule/Private/StateTreePropertyBindingCompiler.cpp
--- a/Engine/Plugins/Runtime/StateTree/Source/StateTreeEditorModule/Private/StateTreePropertyBindingCompiler.cpp
+++ b/Engine/Plugins/Runtime/StateTree/Source/StateTreeEditorModule/Private/StateTreePropertyBindingCompiler.cpp
@@ -78,10 +78,7 @@ bool FStateTreePropertyBindingCompiler::CompileBatch(const FStateTreeBindableStr
 
 		if (const auto Validation = UE::StateTree::Compiler::IsValidIndex16(SourceStructIdx); Validation.DidFail())
 		{
-			if (Log != nullptr)
-			{
-				Validation.Log(*Log, TEXT("Source Struct Index"), SourceStruct);
-			}
+			Validation.Log(*Log, TEXT("Source Struct Index"), SourceStruct);
 			return false;
 		}
 
