StateTree: improve error reporting when pasting state tree nodes

diff --git a/Engine/Plugins/Runtime/StateTree/Source/StateTreeEditorModule/Private/Customizations/StateTreeEditorNodeDetails.cpp b/Engine/Plugins/Runtime/StateTree/Source/StateTreeEditorModule/Private/Customizations/StateTreeEditorNodeDetails.cpp
--- a/Engine/Plugins/Runtime/StateTree/Source/StateTreeEditorModule/Private/Customizations/StateTreeEditorNodeDetails.cpp
+++ b/Engine/Plugins/Runtime/StateTree/Source/StateTreeEditorModule/Private/Customizations/StateTreeEditorNodeDetails.cpp
@@ -3,7 +3,6 @@
 #include "StateTreeEditorNodeDetails.h"
 #include "Brushes/SlateRoundedBoxBrush.h"
 #include "DetailLayoutBuilder.h"
-#include "Framework/Views/TableViewMetadata.h"
 #include "IPropertyUtilities.h"
 #include "IDetailChildrenBuilder.h"
 #include "StateTree.h"
@@ -29,6 +28,8 @@
 #include "ScopedTransaction.h"
 #include "Widgets/Input/SSearchBox.h"
 #include "EditorFontGlyphs.h"
+#include "Framework/Notifications/NotificationManager.h"
+#include "Widgets/Notifications/SNotificationList.h"
 
 #define LOCTEXT_NAMESPACE "StateTreeEditor"
 
@@ -477,72 +478,105 @@ void FStateTreeEditorNodeDetails::OnPasteNode()
 	FString PastedText;
 	FPlatformApplicationMisc::ClipboardPaste(PastedText);
 
-	if (!PastedText.IsEmpty())
+	if (PastedText.IsEmpty())
 	{
-		// create node from the 
-		FStateTreeEditorNode TempNode;
-		UScriptStruct* NodeScriptStruct = TBaseStructure<FStateTreeEditorNode>::Get();
+		return;
+	}
+	
+	// Create node from the clipboard data to figure out the node type. 
+	FStateTreeEditorNode TempNode;
+	UScriptStruct* NodeScriptStruct = TBaseStructure<FStateTreeEditorNode>::Get();
 
-		FStateTreeDefaultValueImportErrorContext ErrorPipe;
-		NodeScriptStruct->ImportText(*PastedText, &TempNode, nullptr, EPropertyPortFlags::PPF_None, &ErrorPipe, NodeScriptStruct->GetName());
+	TArray<UObject*> OuterObjects; 
+	StructProperty->GetOuterObjects(OuterObjects);
+	if (OuterObjects.IsEmpty())
+	{
+		return;
+	}
 
-		if (ErrorPipe.NumErrors == 0)
+	FStateTreeDefaultValueImportErrorContext ErrorPipe;
+	NodeScriptStruct->ImportText(*PastedText, &TempNode, nullptr, PPF_None, &ErrorPipe, NodeScriptStruct->GetName());
+	
+	const UStruct* NodeTypeStruct = TempNode.Node.GetScriptStruct(); 
+	// Only allow valid node types for this property (e.g. do not mix task with conditions).
+	if (ErrorPipe.NumErrors > 0 || !NodeTypeStruct || !NodeTypeStruct->IsChildOf(BaseScriptStruct))
+	{
+		FNotificationInfo NotificationInfo(FText::GetEmpty());
+		NotificationInfo.Text = FText::Format(LOCTEXT("NotSupportedByType", "This property only accepts nodes of type {0}."), BaseScriptStruct->GetDisplayNameText());
+		NotificationInfo.ExpireDuration = 5.0f;
+		FSlateNotificationManager::Get().AddNotification(NotificationInfo);
+		return;
+	}
+	
+	// Reject nodes that are not allowed by the schema. 
+	const UStateTreeSchema* Schema = EditorData ? EditorData->Schema : nullptr;
+	if (!Schema->IsStructAllowed(TempNode.Node.GetScriptStruct()))
+	{
+		// BP nodes are identified by the instance type.
+		if (NodeTypeStruct->IsChildOf(FStateTreeBlueprintEvaluatorWrapper::StaticStruct())
+			|| NodeTypeStruct->IsChildOf(FStateTreeBlueprintTaskWrapper::StaticStruct())
+			|| NodeTypeStruct->IsChildOf(FStateTreeBlueprintConditionWrapper::StaticStruct()))
 		{
-			// Do not allow to mix and match types.
-			// @todo: Check Schema too and warn user about mismatching schema.
-			if (TempNode.Node.GetScriptStruct() && TempNode.Node.GetScriptStruct()->IsChildOf(BaseScriptStruct))
+			if (const FStateTreeNodeBase* Node = TempNode.Node.GetPtr<FStateTreeNodeBase>())
 			{
-				FScopedTransaction Transaction(LOCTEXT("PasteNode", "Paste Node"));
-
-				StructProperty->NotifyPreChange();
-
-				// Make sure we instantiate new objects when setting the value.
-				StructProperty->SetValueFromFormattedString(PastedText, EPropertyValueSetFlags::InstanceObjects);
+				NodeTypeStruct = Node->GetInstanceDataType();
+			}
+		}
 
-				// Reset GUIDs on paste
-				TArray<UObject*> OuterObjects; 
-				TArray<void*> RawNodeData;
-				StructProperty->GetOuterObjects(OuterObjects);
-				StructProperty->AccessRawData(RawNodeData);
-				if (OuterObjects.Num() == RawNodeData.Num())
-				{
-					for (int32 Index = 0; Index < RawNodeData.Num(); Index++)
-					{
-						UObject* OuterObject = OuterObjects[Index];
-						FStateTreeEditorNode* EditorNode = static_cast<FStateTreeEditorNode*>(RawNodeData[Index]);
-						if (EditorNode && OuterObject)
-						{
-							if (FStateTreeNodeBase* Node = EditorNode->Node.GetMutablePtr<FStateTreeNodeBase>())
-							{
-								Node->Name = FName(Node->Name.ToString() + TEXT(" Copy"));
-							}
+		FNotificationInfo NotificationInfo(FText::GetEmpty());
+		NotificationInfo.Text = FText::Format(LOCTEXT("NotSupportedBySchema", "Node {0} is not supported by {1} schema."),
+									NodeTypeStruct->GetDisplayNameText(), Schema->GetClass()->GetDisplayNameText());
+		NotificationInfo.ExpireDuration = 5.0f;
+		FSlateNotificationManager::Get().AddNotification(NotificationInfo);
+		return;
+	}
+	
+	FScopedTransaction Transaction(LOCTEXT("PasteNode", "Paste Node"));
 
-							const FGuid OldStructID = EditorNode->ID; 
-							EditorNode->ID = FGuid::NewGuid();
+	StructProperty->NotifyPreChange();
 
-							// Copy bindings from the copied node.
-							if (OldStructID.IsValid() && EditorData)
-							{
-								if (FStateTreeEditorPropertyBindings* Bindings = EditorData->GetPropertyEditorBindings())
-								{
-									Bindings->CopyBindings(OldStructID, EditorNode->ID);
-								}
-							}
+	// Make sure we instantiate new objects when setting the value.
+	StructProperty->SetValueFromFormattedString(PastedText, EPropertyValueSetFlags::InstanceObjects);
 
-						}
-					}
+	// Reset GUIDs on paste
+	TArray<void*> RawNodeData;
+	StructProperty->AccessRawData(RawNodeData);
+	if (OuterObjects.Num() == RawNodeData.Num())
+	{
+		for (int32 Index = 0; Index < RawNodeData.Num(); Index++)
+		{
+			UObject* OuterObject = OuterObjects[Index];
+			FStateTreeEditorNode* EditorNode = static_cast<FStateTreeEditorNode*>(RawNodeData[Index]);
+			if (EditorNode && OuterObject)
+			{
+				if (FStateTreeNodeBase* Node = EditorNode->Node.GetMutablePtr<FStateTreeNodeBase>())
+				{
+					Node->Name = FName(Node->Name.ToString() + TEXT(" Copy"));
 				}
 
-				StructProperty->NotifyPostChange(EPropertyChangeType::ValueSet);
-				StructProperty->NotifyFinishedChangingProperties();
+				const FGuid OldStructID = EditorNode->ID; 
+				EditorNode->ID = FGuid::NewGuid();
 
-				if (PropUtils)
+				// Copy bindings from the copied node.
+				if (OldStructID.IsValid() && EditorData)
 				{
-					PropUtils->ForceRefresh();
+					if (FStateTreeEditorPropertyBindings* Bindings = EditorData->GetPropertyEditorBindings())
+					{
+						Bindings->CopyBindings(OldStructID, EditorNode->ID);
+					}
 				}
+
 			}
 		}
 	}
+
+	StructProperty->NotifyPostChange(EPropertyChangeType::ValueSet);
+	StructProperty->NotifyFinishedChangingProperties();
+
+	if (PropUtils)
+	{
+		PropUtils->ForceRefresh();
+	}
 }
 
 bool FStateTreeEditorNodeDetails::ShouldResetToDefault(TSharedPtr<IPropertyHandle> PropertyHandle) const
