StateTree: Fixed not saving STs when executing cmdlet

diff --git a/Engine/Plugins/Runtime/StateTree/Source/StateTreeEditorModule/Private/Commandlets/StateTreeCompileAllCommandlet.cpp b/Engine/Plugins/Runtime/StateTree/Source/StateTreeEditorModule/Private/Commandlets/StateTreeCompileAllCommandlet.cpp
--- a/Engine/Plugins/Runtime/StateTree/Source/StateTreeEditorModule/Private/Commandlets/StateTreeCompileAllCommandlet.cpp
+++ b/Engine/Plugins/Runtime/StateTree/Source/StateTreeEditorModule/Private/Commandlets/StateTreeCompileAllCommandlet.cpp
@@ -122,12 +122,15 @@ bool UStateTreeCompileAllCommandlet::CompileAndSaveStateTree(UStateTree& StateTr
 			else if (SourceControlState->IsCheckedOut() || SourceControlState->IsAdded())
 			{
 				UE_LOG(LogStateTreeCompile, Log, TEXT("Package %s already checked out"), *PackageFileName);
-				return true;
 			}
 			else if (SourceControlState->IsSourceControlled())
 			{
 				UE_LOG(LogStateTreeCompile, Log, TEXT("Checking out package %s from revision control"), *PackageFileName);
-				return SourceControlProvider->Execute(ISourceControlOperation::Create<FCheckOut>(), PackageFileName) == ECommandResult::Succeeded;
+				if (SourceControlProvider->Execute(ISourceControlOperation::Create<FCheckOut>(), PackageFileName) != ECommandResult::Succeeded)
+				{
+					UE_LOG(LogStateTreeCompile, Log, TEXT("Failed to check out package %s from revision control"), *PackageFileName);
+					return false;
+				}
 			}
 		}
 	}
