StateTree: Sort property bindings based on layout.
- the bindings sent to the compiler are in order they were created
- sort the bindings so that e.g. copying and array and a speciific array item will first copy the earliest property (array) and then the rest (specific item)

diff --git a/Engine/Plugins/Runtime/StateTree/Source/StateTreeEditorModule/Private/StateTreePropertyBindingCompiler.cpp b/Engine/Plugins/Runtime/StateTree/Source/StateTreeEditorModule/Private/StateTreePropertyBindingCompiler.cpp
--- a/Engine/Plugins/Runtime/StateTree/Source/StateTreeEditorModule/Private/StateTreePropertyBindingCompiler.cpp
+++ b/Engine/Plugins/Runtime/StateTree/Source/StateTreeEditorModule/Private/StateTreePropertyBindingCompiler.cpp
@@ -26,7 +26,12 @@ bool FStateTreePropertyBindingCompiler::CompileBatch(const FStateTreeBindableStr
 
 	StoreSourceStructs();
 
-	const int32 BindingsBegin = PropertyBindings->PropertyPathBindings.Num();
+	struct FSortedBinding
+	{
+		FStateTreePropertyPathBinding Binding;
+		TArray<FStateTreePropertyPathIndirection> TargetIndirections;
+	};
+	TArray<FSortedBinding> NewBindings;
 
 	for (const FStateTreePropertyPathBinding& Binding : BatchPropertyBindings)
 	{
@@ -82,12 +87,45 @@ bool FStateTreePropertyBindingCompiler::CompileBatch(const FStateTreeBindableStr
 			return false;
 		}
 
-		PropertyBindings->PropertyPathBindings.Emplace(FStateTreeIndex16(SourceStructIdx), Binding.GetSourcePath(), Binding.GetTargetPath());
+		FSortedBinding& NewBinding = NewBindings.AddDefaulted_GetRef();
+		NewBinding.Binding = FStateTreePropertyPathBinding(FStateTreeIndex16(SourceStructIdx), Binding.GetSourcePath(), Binding.GetTargetPath());
+		NewBinding.TargetIndirections = TargetIndirections;
 	}
 
-	const int32 BindingsEnd = PropertyBindings->PropertyPathBindings.Num();
-	if (BindingsBegin != BindingsEnd)
+	if (!NewBindings.IsEmpty())
 	{
+		// Sort bindings base on copy target memory layout.
+		NewBindings.StableSort([](const FSortedBinding& A, const FSortedBinding& B)
+		{
+			const int32 MaxSegments = FMath::Min(A.TargetIndirections.Num(), B.TargetIndirections.Num());
+			for (int32 Index = 0; Index < MaxSegments; Index++)
+			{
+				// If property A is in struct before B, copy A first. 
+				if (A.TargetIndirections[Index].GetPropertyOffset() < B.TargetIndirections[Index].GetPropertyOffset())
+				{
+					return true;
+				}
+				// If A and B points to the same property, choose the one that points to an earlier array item.
+				// Note: this assumes that INDEX_NONE = -1, which means that binding directly to an array comes before an array access,
+				// and non-array access will compare equal (both INDEX_NONE).
+				if (A.TargetIndirections[Index].GetPropertyOffset() == B.TargetIndirections[Index].GetPropertyOffset()
+					&& A.TargetIndirections[Index].GetArrayIndex() < B.TargetIndirections[Index].GetArrayIndex())
+				{
+					return true;
+				}
+			}
+			// We get here if the common path is the same, shorter path wins.
+			return A.TargetIndirections.Num() <= B.TargetIndirections.Num(); 
+		});
+
+		// Store bindings batch.
+		const int32 BindingsBegin = PropertyBindings->PropertyPathBindings.Num();
+		for (const FSortedBinding& NewBinding : NewBindings)
+		{
+			PropertyBindings->PropertyPathBindings.Add(NewBinding.Binding);
+		}
+		const int32 BindingsEnd = PropertyBindings->PropertyPathBindings.Num();
+
 		FStateTreePropertyCopyBatch& Batch = PropertyBindings->CopyBatches.AddDefaulted_GetRef();
 		Batch.TargetStruct = TargetStruct;
 		Batch.BindingsBegin = IntCastChecked<uint16>(BindingsBegin);
diff --git a/Engine/Plugins/Runtime/StateTree/Source/StateTreeTestSuite/Private/StateTreeTest.cpp b/Engine/Plugins/Runtime/StateTree/Source/StateTreeTestSuite/Private/StateTreeTest.cpp
--- a/Engine/Plugins/Runtime/StateTree/Source/StateTreeTestSuite/Private/StateTreeTest.cpp
+++ b/Engine/Plugins/Runtime/StateTree/Source/StateTreeTestSuite/Private/StateTreeTest.cpp
@@ -1009,6 +1009,106 @@ struct FStateTreeTest_PropertyPathArrayOfInstancedObjects : FAITestBase
 };
 IMPLEMENT_AI_INSTANT_TEST(FStateTreeTest_PropertyPathArrayOfInstancedObjects, "System.StateTree.PropertyPath.ArrayOfInstancedObjects");
 
+struct FStateTreeTest_BindingsCompiler : FAITestBase
+{
+	virtual bool InstantTest() override
+	{
+		FStateTreeCompilerLog Log;
+		FStateTreePropertyBindings Bindings;
+		FStateTreePropertyBindingCompiler BindingCompiler;
+
+		const bool bInitResult = BindingCompiler.Init(Bindings, Log);
+		AITEST_TRUE("Expect init to succeed", bInitResult);
+
+		FStateTreeBindableStructDesc SourceADesc;
+		SourceADesc.Name = FName(TEXT("SourceA"));
+		SourceADesc.Struct = TBaseStructure<FStateTreeTest_PropertyCopy>::Get();
+		SourceADesc.DataSource = EStateTreeBindableStructSource::Parameter;
+		SourceADesc.ID = FGuid::NewGuid();
+
+		FStateTreeBindableStructDesc SourceBDesc;
+		SourceBDesc.Name = FName(TEXT("SourceB"));
+		SourceBDesc.Struct = TBaseStructure<FStateTreeTest_PropertyCopy>::Get();
+		SourceBDesc.DataSource = EStateTreeBindableStructSource::Parameter;
+		SourceBDesc.ID = FGuid::NewGuid();
+
+		FStateTreeBindableStructDesc TargetDesc;
+		TargetDesc.Name = FName(TEXT("Target"));
+		TargetDesc.Struct = TBaseStructure<FStateTreeTest_PropertyCopy>::Get();
+		TargetDesc.DataSource = EStateTreeBindableStructSource::Parameter;
+		TargetDesc.ID = FGuid::NewGuid();
+		
+		const int32 SourceAIndex = BindingCompiler.AddSourceStruct(SourceADesc);
+		const int32 SourceBIndex = BindingCompiler.AddSourceStruct(SourceBDesc);
+
+		auto MakeBinding = [](const FGuid& SourceID, const FString& Source, const FGuid& TargetID, const FString& Target)
+		{
+			FStateTreePropertyPath SourcePath;
+			SourcePath.FromString(Source);
+			SourcePath.SetStructID(SourceID);
+
+			FStateTreePropertyPath TargetPath;
+			TargetPath.FromString(Target);
+			TargetPath.SetStructID(TargetID);
+
+			return FStateTreePropertyPathBinding(SourcePath, TargetPath);
+		};
+
+		TArray<FStateTreePropertyPathBinding> PropertyBindings;
+		PropertyBindings.Add(MakeBinding(SourceBDesc.ID, TEXT("Item"), TargetDesc.ID, TEXT("Array[1]")));
+		PropertyBindings.Add(MakeBinding(SourceADesc.ID, TEXT("Item.B"), TargetDesc.ID, TEXT("Array[1].B")));
+		PropertyBindings.Add(MakeBinding(SourceADesc.ID, TEXT("Array"), TargetDesc.ID, TEXT("Array")));
+
+		int32 CopyBatchIndex = INDEX_NONE;
+		const bool bCompileBatchResult = BindingCompiler.CompileBatch(TargetDesc, PropertyBindings, CopyBatchIndex);
+		AITEST_TRUE("CompileBatch should succeed", bCompileBatchResult);
+		AITEST_NOT_EQUAL("CopyBatchIndex should not be INDEX_NONE", CopyBatchIndex, (int32)INDEX_NONE);
+
+		BindingCompiler.Finalize();
+
+		const bool bResolveResult = Bindings.ResolvePaths();
+		AITEST_TRUE("ResolvePaths should succeed", bResolveResult);
+
+		FStateTreeTest_PropertyCopy SourceA;
+		SourceA.Item.B = 123;
+		SourceA.Array.AddDefaulted_GetRef().A = 1;
+		SourceA.Array.AddDefaulted_GetRef().B = 2;
+
+		FStateTreeTest_PropertyCopy SourceB;
+		SourceB.Item.A = 41;
+		SourceB.Item.B = 42;
+
+		FStateTreeTest_PropertyCopy Target;
+
+		AITEST_TRUE("SourceAIndex should be less than max number of source structs.", SourceAIndex < Bindings.GetSourceStructNum());
+		AITEST_TRUE("SourceBIndex should be less than max number of source structs.", SourceBIndex < Bindings.GetSourceStructNum());
+
+		TArray<FStateTreeDataView> SourceViews;
+		SourceViews.SetNum(Bindings.GetSourceStructNum());
+		SourceViews[SourceAIndex] = FStateTreeDataView(TBaseStructure<FStateTreeTest_PropertyCopy>::Get(), (uint8*)&SourceA);
+		SourceViews[SourceBIndex] = FStateTreeDataView(TBaseStructure<FStateTreeTest_PropertyCopy>::Get(), (uint8*)&SourceB);
+		FStateTreeDataView TargetView(TBaseStructure<FStateTreeTest_PropertyCopy>::Get(), (uint8*)&Target);
+		
+		const bool bCopyResult = Bindings.CopyTo(SourceViews, FStateTreeIndex16(CopyBatchIndex), TargetView);
+		AITEST_TRUE("CopyTo should succeed", bCopyResult);
+
+		// Due to binding sorting, we expect them to executed in this order (sorted based on target access, earliest to latest)
+		// SourceA.Array -> Target.Array
+		// SourceB.Item -> Target.Array[1]
+		// SourceA.Item.B -> Target.Array[1].B
+
+		AITEST_EQUAL("Expect TargetArray to be copied from SourceA", Target.Array.Num(), SourceA.Array.Num());
+		AITEST_EQUAL("Expect Target.Array[0].A copied from SourceA.Array[0].A", Target.Array[0].A, SourceA.Array[0].A);
+		AITEST_EQUAL("Expect Target.Array[0].B copied from SourceA.Array[0].B", Target.Array[0].B, SourceA.Array[0].B);
+		AITEST_EQUAL("Expect Target.Array[1].A copied from SourceB.Item.A", Target.Array[1].A, SourceB.Item.A);
+		AITEST_EQUAL("Expect Target.Array[1].B copied from SourceA.Item.B", Target.Array[1].B, SourceA.Item.B);
+		
+		return true;
+	}
+};
+IMPLEMENT_AI_INSTANT_TEST(FStateTreeTest_BindingsCompiler, "System.StateTree.BindingsCompiler");
+
+
 UE_ENABLE_OPTIMIZATION_SHIP
 
 #undef LOCTEXT_NAMESPACE
diff --git a/Engine/Plugins/Runtime/StateTree/Source/StateTreeTestSuite/Private/StateTreeTestTypes.h b/Engine/Plugins/Runtime/StateTree/Source/StateTreeTestSuite/Private/StateTreeTestTypes.h
--- a/Engine/Plugins/Runtime/StateTree/Source/StateTreeTestSuite/Private/StateTreeTestTypes.h
+++ b/Engine/Plugins/Runtime/StateTree/Source/StateTreeTestSuite/Private/StateTreeTestTypes.h
@@ -381,3 +381,15 @@ public:
 	UPROPERTY(EditAnywhere, Category = "")
 	TArray<FStateTreeTest_PropertyStruct> ArrayOfStruct;
 };
+
+USTRUCT()
+struct FStateTreeTest_PropertyCopy
+{
+	GENERATED_BODY()
+
+	UPROPERTY(EditAnywhere, Category = "")
+	FStateTreeTest_PropertyStruct Item;
+
+	UPROPERTY(EditAnywhere, Category = "")
+	TArray<FStateTreeTest_PropertyStruct> Array;
+};
