State Tree: Copy and paste State Tree states.
- Refactored State Tree style to follow more recent style pattern
- Added code to ensure that State Tree states outers are always the parent state (fixes issues with FCustomizableTextObjectFactory)
- Added support for Cut/Copy/Paste of state tree states (across State Trees)
- Improved reporting of bad bindings and state links (compiler & UI)

diff --git a/Engine/Plugins/Runtime/StateTree/Source/StateTreeEditorModule/Private/Customizations/StateTreeBindingExtension.cpp b/Engine/Plugins/Runtime/StateTree/Source/StateTreeEditorModule/Private/Customizations/StateTreeBindingExtension.cpp
--- a/Engine/Plugins/Runtime/StateTree/Source/StateTreeEditorModule/Private/Customizations/StateTreeBindingExtension.cpp
+++ b/Engine/Plugins/Runtime/StateTree/Source/StateTreeEditorModule/Private/Customizations/StateTreeBindingExtension.cpp
@@ -148,6 +148,12 @@ EStateTreePropertyUsage MakeStructPropertyPathFromPropertyHandle(TSharedPtr<cons
 	return ResultUsage;
 }
 
+const FStateTreeBindableStructDesc* FindStruct(TConstArrayView<FStateTreeBindableStructDesc> AccessibleStructs, const FGuid StructID)
+{
+	return AccessibleStructs.FindByPredicate([StructID](const FStateTreeBindableStructDesc& Desc) { return Desc.ID == StructID; });
+}
+
+
 FText GetSectionNameFromDataSource(const EStateTreeBindableStructSource Source)
 {
 	switch (Source)
@@ -456,20 +462,20 @@ void FStateTreeBindingExtension::ExtendWidgetRow(FDetailWidgetRow& InWidgetRow,
 			{
 				if (const FStateTreePropertyPath* SourcePath = EditorBindings->GetPropertyBindingSource(TargetPath))
 				{
-					FString PropertyName;
-					for (int32 i = 0; i < AccessibleStructs.Num(); i++)
+					const FStateTreeBindableStructDesc* SourceDesc = UE::StateTree::PropertyBinding::FindStruct(AccessibleStructs, SourcePath->GetStructID());
+					if (SourceDesc)
 					{
-						if (AccessibleStructs[i].ID == SourcePath->GetStructID())
+						FString PropertyName = SourceDesc->Name.ToString();
+						if (!SourcePath->IsPathEmpty())
 						{
-							PropertyName = AccessibleStructs[i].Name.ToString();
-							break;
+							PropertyName += TEXT(".") + SourcePath->ToString();
 						}
+						CurrentValue = FText::FromString(PropertyName);
 					}
-					if (!SourcePath->IsPathEmpty())
+					else
 					{
-						PropertyName += TEXT(".") + SourcePath->ToString();
+						CurrentValue = FText::Format(LOCTEXT("MissingSource", "???.{0}"), FText::FromString(SourcePath->ToString()));
 					}
-					CurrentValue = FText::FromString(PropertyName);
 				}
 			}
 
@@ -484,46 +490,77 @@ void FStateTreeBindingExtension::ExtendWidgetRow(FDetailWidgetRow& InWidgetRow,
 			{
 				if (const FStateTreePropertyPath* SourcePath = EditorBindings->GetPropertyBindingSource(TargetPath))
 				{
-					FString PropertyName;
-					for (int32 i = 0; i < AccessibleStructs.Num(); i++)
+					const FStateTreeBindableStructDesc* SourceDesc = UE::StateTree::PropertyBinding::FindStruct(AccessibleStructs, SourcePath->GetStructID());
+					if (SourceDesc)
 					{
-						if (AccessibleStructs[i].ID == SourcePath->GetStructID())
+						FString PropertyName = SourceDesc->Name.ToString();
+						if (!SourcePath->IsPathEmpty())
 						{
-							PropertyName = AccessibleStructs[i].Name.ToString();
-							break;
+							PropertyName += TEXT(".") + SourcePath->ToString();
 						}
+						CurrentValue = FText::Format(LOCTEXT("ExistingBindingTooltip", "Property is bound to {0}."), FText::FromString(PropertyName));
 					}
-					if (!SourcePath->IsPathEmpty())
+					else
 					{
-						PropertyName += TEXT(".") + SourcePath->ToString();
+						CurrentValue = FText::Format(LOCTEXT("MissingBindingTooltip", "Missing binding source for property path '{0}'."), FText::FromString(SourcePath->ToString()));
 					}
-					CurrentValue = FText::Format(LOCTEXT("ExistingBindingTooltip", "Property is bound to {0}."), FText::FromString(PropertyName));
 				}
-			}
-			else if (Property != nullptr)
-			{
-				CurrentValue = FText::Format(LOCTEXT("BindTooltip", "Bind value to {0} property."), UE::StateTree::PropertyBinding::GetPropertyTypeText(Property));
+				else if (Property != nullptr)
+				{
+					CurrentValue = FText::Format(LOCTEXT("BindTooltip", "Bind value to {0} property."), UE::StateTree::PropertyBinding::GetPropertyTypeText(Property));
+				}
 			}
 
 			return CurrentValue;
 		});
 
-	Args.CurrentBindingImage = MakeAttributeLambda([]() -> const FSlateBrush*
+	Args.CurrentBindingImage = MakeAttributeLambda([EditorBindings, TargetPath, AccessibleStructs]() -> const FSlateBrush*
 		{
 			static FName PropertyIcon(TEXT("Kismet.Tabs.Variables"));
-			return FAppStyle::GetBrush(PropertyIcon);
+
+			const FSlateBrush* Result = FAppStyle::GetBrush(PropertyIcon);
+		
+			if (EditorBindings)
+			{
+				if (const FStateTreePropertyPath* SourcePath = EditorBindings->GetPropertyBindingSource(TargetPath))
+				{
+					const FStateTreeBindableStructDesc* SourceDesc = UE::StateTree::PropertyBinding::FindStruct(AccessibleStructs, SourcePath->GetStructID());
+					if (!SourceDesc)
+					{
+						// Invalid binding
+						Result = FCoreStyle::Get().GetBrush("Icons.ErrorWithColor");
+					}
+				}
+			}
+
+			return Result;
 		});
 
-	Args.CurrentBindingColor = MakeAttributeLambda([InPropertyHandle]() -> FLinearColor
+	Args.CurrentBindingColor = MakeAttributeLambda([InPropertyHandle, EditorBindings, TargetPath, AccessibleStructs]() -> FLinearColor
 		{
 			const UEdGraphSchema_K2* Schema = GetDefault<UEdGraphSchema_K2>();
 
 			FEdGraphPinType PinType;
 			Schema->ConvertPropertyToPinType(InPropertyHandle->GetProperty(), PinType);
-			const FLinearColor BindingColor = Schema->GetPinTypeColor(PinType);
+		
+			FLinearColor BindingColor = Schema->GetPinTypeColor(PinType);
 
 			// TODO: Handle coloring of type promotion
 
+			if (EditorBindings)
+			{
+				if (const FStateTreePropertyPath* SourcePath = EditorBindings->GetPropertyBindingSource(TargetPath))
+				{
+					const FStateTreeBindableStructDesc* SourceDesc = UE::StateTree::PropertyBinding::FindStruct(AccessibleStructs, SourcePath->GetStructID());
+					if (!SourceDesc)
+					{
+						// Invalid binding
+						BindingColor = FLinearColor::White;
+					}
+				}
+			}
+
+
 			return BindingColor;
 		});
 
diff --git a/Engine/Plugins/Runtime/StateTree/Source/StateTreeEditorModule/Private/Customizations/StateTreeStateLinkDetails.cpp b/Engine/Plugins/Runtime/StateTree/Source/StateTreeEditorModule/Private/Customizations/StateTreeStateLinkDetails.cpp
--- a/Engine/Plugins/Runtime/StateTree/Source/StateTreeEditorModule/Private/Customizations/StateTreeStateLinkDetails.cpp
+++ b/Engine/Plugins/Runtime/StateTree/Source/StateTreeEditorModule/Private/Customizations/StateTreeStateLinkDetails.cpp
@@ -50,9 +50,33 @@ void FStateTreeStateLinkDetails::CustomizeHeader(TSharedRef<class IPropertyHandl
 			.OnGetMenuContent(this, &FStateTreeStateLinkDetails::OnGetStateContent)
 			.ButtonContent()
 			[
-				SNew(STextBlock)
-				.Text(this, &FStateTreeStateLinkDetails::GetCurrentStateDesc)
-				.Font(IDetailLayoutBuilder::GetDetailFont())
+				SNew(SHorizontalBox)
+
+				+ SHorizontalBox::Slot()
+				.AutoWidth()
+				.VAlign(VAlign_Center)
+				.Padding(0,0,4,0)
+				[
+					SNew(SBox)
+					[
+						SNew(SImage)
+						.ToolTipText(LOCTEXT("MissingState", "The specified state cannot be found."))
+						.Visibility_Lambda([this]()
+						{
+							return IsValidLink() ? EVisibility::Collapsed : EVisibility::Visible;
+						})
+						.Image(FAppStyle::GetBrush("Icons.ErrorWithColor"))
+					]
+				]
+
+				+SHorizontalBox::Slot()
+				.AutoWidth()
+				.VAlign(VAlign_Center)
+				[
+					SNew(STextBlock)
+					.Text(this, &FStateTreeStateLinkDetails::GetCurrentStateDesc)
+					.Font(IDetailLayoutBuilder::GetDetailFont())
+				]
 			]
 		];
 
@@ -235,9 +259,7 @@ FText FStateTreeStateLinkDetails::GetCurrentStateDesc() const
 				}
 			}
 
-			FFormatNamedArguments Args;
-			Args.Add(TEXT("Identifier"), FText::FromName(OldName));
-			Result = FText::Format(LOCTEXT("InvalidReference", "Invalid Reference {Identifier}"), Args);
+			return FText::FromName(OldName);
 		}
 		break;
 	default:
@@ -249,6 +271,41 @@ FText FStateTreeStateLinkDetails::GetCurrentStateDesc() const
 	return Result;
 }
 
+bool FStateTreeStateLinkDetails::IsValidLink() const
+{
+	const EStateTreeTransitionType TransitionType = GetTransitionType().Get(EStateTreeTransitionType::Failed);
+
+	bool bIsValid = false;
+	switch (TransitionType)
+	{
+	// all these are unconditionally valid		
+	case EStateTreeTransitionType::None: // fall through on purpose
+	case EStateTreeTransitionType::NextState: // fall through on purpose
+	case EStateTreeTransitionType::Succeeded: // fall through on purpose
+	case EStateTreeTransitionType::Failed:
+		bIsValid = true;
+		break;
+	case EStateTreeTransitionType::GotoState:
+		{
+			bIsValid = false;
+			if (NameProperty && IDProperty)
+			{
+				FGuid StateID;
+				if (UE::StateTree::PropertyHelpers::GetStructValue<FGuid>(IDProperty, StateID) == FPropertyAccess::Success)
+				{
+					bIsValid = !StateID.IsValid() || CachedIDs.Contains(StateID);
+				}
+			}
+		}
+		break;
+	default:
+		bIsValid = false;
+		break;
+	}
+
+	return bIsValid;
+}
+
 TOptional<EStateTreeTransitionType> FStateTreeStateLinkDetails::GetTransitionType() const
 {
 	if (LinkTypeProperty)
diff --git a/Engine/Plugins/Runtime/StateTree/Source/StateTreeEditorModule/Private/Customizations/StateTreeStateLinkDetails.h b/Engine/Plugins/Runtime/StateTree/Source/StateTreeEditorModule/Private/Customizations/StateTreeStateLinkDetails.h
--- a/Engine/Plugins/Runtime/StateTree/Source/StateTreeEditorModule/Private/Customizations/StateTreeStateLinkDetails.h
+++ b/Engine/Plugins/Runtime/StateTree/Source/StateTreeEditorModule/Private/Customizations/StateTreeStateLinkDetails.h
@@ -41,6 +41,7 @@ private:
 	void OnStateComboChange(int Idx);
 	TSharedRef<SWidget> OnGetStateContent() const;
 	FText GetCurrentStateDesc() const;
+	bool IsValidLink() const;
 	TOptional<EStateTreeTransitionType> GetTransitionType() const;
 
 	TArray<FName> CachedNames;
diff --git a/Engine/Plugins/Runtime/StateTree/Source/StateTreeEditorModule/Private/SStateTreeView.cpp b/Engine/Plugins/Runtime/StateTree/Source/StateTreeEditorModule/Private/SStateTreeView.cpp
--- a/Engine/Plugins/Runtime/StateTree/Source/StateTreeEditorModule/Private/SStateTreeView.cpp
+++ b/Engine/Plugins/Runtime/StateTree/Source/StateTreeEditorModule/Private/SStateTreeView.cpp
@@ -12,7 +12,7 @@
 
 #include "StateTreeViewModel.h"
 #include "StateTreeState.h"
-#include "Widgets/Views/STreeView.h"
+#include "StateTreeEditorCommands.h"
 
 #define LOCTEXT_NAMESPACE "StateTreeEditor"
 
@@ -131,8 +131,77 @@ void SStateTreeView::Construct(const FArguments& InArgs, TSharedRef<FStateTreeVi
 	];
 
 	UpdateTree(true);
+
+	CommandList = MakeShared<FUICommandList>();
+	BindCommands();
 }
 
+void SStateTreeView::BindCommands()
+{
+	const FStateTreeEditorCommands& Commands = FStateTreeEditorCommands::Get();
+
+	CommandList->MapAction(
+		Commands.AddSiblingState,
+		FExecuteAction::CreateSP(this, &SStateTreeView::HandleAddSiblingState),
+		FCanExecuteAction());
+
+	CommandList->MapAction(
+		Commands.AddChildState,
+		FExecuteAction::CreateSP(this, &SStateTreeView::HandleAddChildState),
+		FCanExecuteAction::CreateSP(this, &SStateTreeView::HasSelection));
+
+	CommandList->MapAction(
+		Commands.CutStates,
+		FExecuteAction::CreateSP(this, &SStateTreeView::HandleCutSelectedStates),
+		FCanExecuteAction::CreateSP(this, &SStateTreeView::HasSelection));
+
+	CommandList->MapAction(
+		Commands.CopyStates,
+		FExecuteAction::CreateSP(this, &SStateTreeView::HandleCopySelectedStates),
+		FCanExecuteAction::CreateSP(this, &SStateTreeView::HasSelection));
+
+	CommandList->MapAction(
+		Commands.DeleteStates,
+		FExecuteAction::CreateSP(this, &SStateTreeView::HandleDeleteStates),
+		FCanExecuteAction::CreateSP(this, &SStateTreeView::HasSelection));
+
+	CommandList->MapAction(
+		Commands.PasteStatesAsSiblings,
+		FExecuteAction::CreateSP(this, &SStateTreeView::HandlePasteStatesAsSiblings),
+		FCanExecuteAction::CreateSP(this, &SStateTreeView::HasSelection));
+
+	CommandList->MapAction(
+		Commands.PasteStatesAsChildren,
+		FExecuteAction::CreateSP(this, &SStateTreeView::HandlePasteStatesAsChildren),
+		FCanExecuteAction::CreateSP(this, &SStateTreeView::HasSelection));
+
+	CommandList->MapAction(
+		Commands.DuplicateStates,
+		FExecuteAction::CreateSP(this, &SStateTreeView::HandleDuplicateSelectedStates),
+		FCanExecuteAction::CreateSP(this, &SStateTreeView::HasSelection));
+
+	CommandList->MapAction(
+		Commands.RenameState,
+		FExecuteAction::CreateSP(this, &SStateTreeView::HandleRenameState),
+		FCanExecuteAction::CreateSP(this, &SStateTreeView::HasSelection));
+}
+
+bool SStateTreeView::HasSelection() const
+{
+	return StateTreeViewModel && StateTreeViewModel->HasSelection();
+}
+
+FReply SStateTreeView::OnKeyDown(const FGeometry& MyGeometry, const FKeyEvent& InKeyEvent)
+{
+	if(CommandList->ProcessCommandBindings(InKeyEvent))
+	{
+		return FReply::Handled();
+	}
+	else
+	{
+		return SCompoundWidget::OnKeyDown(MyGeometry, InKeyEvent);
+	}
+}
 
 void SStateTreeView::SavePersistentExpandedStates()
 {
@@ -224,7 +293,8 @@ void SStateTreeView::HandleModelStateAdded(UStateTreeState* ParentState, UStateT
 {
 	bItemsDirty = true;
 
-	HandleRenameState(NewState);
+	// Request to rename the state immediately.
+	RequestedRenameState = NewState;
 
 	if (StateTreeViewModel.IsValid())
 	{
@@ -315,40 +385,40 @@ TSharedPtr<SWidget> SStateTreeView::HandleContextMenuOpening()
 		return nullptr;
 	}
 
-	TArray<UStateTreeState*> SelectedStates;
-	StateTreeViewModel->GetSelectedStates(SelectedStates);
-	UStateTreeState* FirstSelectedState = SelectedStates.Num() > 0 ? SelectedStates[0] : nullptr;
-
-	FMenuBuilder MenuBuilder(true, nullptr);
+	FMenuBuilder MenuBuilder(true, CommandList);
 
-	MenuBuilder.AddMenuEntry(
+	MenuBuilder.AddSubMenu(
 		LOCTEXT("AddState", "Add State"),
 		FText(),
-		FSlateIcon(),
-		FUIAction(FExecuteAction::CreateSP(this, &SStateTreeView::HandleAddState, FirstSelectedState)));
-
-	if (FirstSelectedState)
-	{
-		MenuBuilder.AddMenuEntry(
-			LOCTEXT("AddChildState", "Add Child State"),
-			FText(),
-			FSlateIcon(),
-			FUIAction(FExecuteAction::CreateSP(this, &SStateTreeView::HandleAddChildState, FirstSelectedState)));
+		FNewMenuDelegate::CreateLambda([this](FMenuBuilder& MenuBuilder)
+		{
+			MenuBuilder.AddMenuEntry(FStateTreeEditorCommands::Get().AddSiblingState);
+			MenuBuilder.AddMenuEntry(FStateTreeEditorCommands::Get().AddChildState);
+		}),
+		/*bInOpenSubMenuOnClick =*/false,
+		FSlateIcon(FAppStyle::GetAppStyleSetName(), "Icons.Plus")
+	);
 
-		MenuBuilder.AddSeparator();
+	MenuBuilder.AddSeparator();
 
-		MenuBuilder.AddMenuEntry(
-			LOCTEXT("RenameNodeGroup", "Rename"),
-			FText(),
-			FSlateIcon(),
-			FUIAction(FExecuteAction::CreateSP(this, &SStateTreeView::HandleRenameState, FirstSelectedState)));
+	MenuBuilder.AddMenuEntry(FStateTreeEditorCommands::Get().CutStates);
+	MenuBuilder.AddMenuEntry(FStateTreeEditorCommands::Get().CopyStates);
 
-		MenuBuilder.AddMenuEntry(
-			LOCTEXT("DeleteState", "Delete Selected"),
-			FText(),
-			FSlateIcon(),
-			FUIAction(FExecuteAction::CreateSP(this, &SStateTreeView::HandleDeleteItems)));
-	}
+	MenuBuilder.AddSubMenu(
+		LOCTEXT("Paste", "Paste"),
+		FText(),
+		FNewMenuDelegate::CreateLambda([this](FMenuBuilder& MenuBuilder)
+		{
+			MenuBuilder.AddMenuEntry(FStateTreeEditorCommands::Get().PasteStatesAsSiblings);
+			MenuBuilder.AddMenuEntry(FStateTreeEditorCommands::Get().PasteStatesAsChildren);
+		}),
+		/*bInOpenSubMenuOnClick =*/false,
+		FSlateIcon(FAppStyle::GetAppStyleSetName(), "GenericCommands.Paste")
+	);
+	
+	MenuBuilder.AddMenuEntry(FStateTreeEditorCommands::Get().DuplicateStates);
+	MenuBuilder.AddMenuEntry(FStateTreeEditorCommands::Get().DeleteStates);
+	MenuBuilder.AddMenuEntry(FStateTreeEditorCommands::Get().RenameState);
 
 	return MenuBuilder.MakeWidget();
 }
@@ -387,45 +457,90 @@ FReply SStateTreeView::HandleAddStateButton()
 	return FReply::Handled();
 }
 
-void SStateTreeView::HandleAddState(UStateTreeState* AfterState)
+UStateTreeState* SStateTreeView::GetFirstSelectedState() const
 {
-	if (StateTreeViewModel == nullptr)
+	TArray<UStateTreeState*> SelectedStates;
+	if (StateTreeViewModel)
 	{
-		return;
+		StateTreeViewModel->GetSelectedStates(SelectedStates);
 	}
+	return SelectedStates.IsEmpty() ? nullptr : SelectedStates[0];
+}
 
-	StateTreeViewModel->AddState(AfterState);
+void SStateTreeView::HandleAddSiblingState()
+{
+	if (StateTreeViewModel)
+	{
+		StateTreeViewModel->AddState(GetFirstSelectedState());
+	}
 }
 
-void SStateTreeView::HandleRenameState(UStateTreeState* State)
+void SStateTreeView::HandleAddChildState()
 {
-	RequestedRenameState = State;
+	if (StateTreeViewModel)
+	{
+		UStateTreeState* ParentState = GetFirstSelectedState();
+		if (ParentState)
+		{
+			StateTreeViewModel->AddChildState(ParentState);
+			TreeView->SetItemExpansion(ParentState, true);
+		}
+	}
 }
 
-void SStateTreeView::HandleAddChildState(UStateTreeState* ParentState)
+void SStateTreeView::HandleCutSelectedStates()
 {
-	if (StateTreeViewModel == nullptr)
+	if (StateTreeViewModel)
 	{
-		return;
+		StateTreeViewModel->CopySelectedStates();
+		StateTreeViewModel->RemoveSelectedStates();
 	}
+}
 
-	if (ParentState)
+void SStateTreeView::HandleCopySelectedStates()
+{
+	if (StateTreeViewModel)
 	{
-		StateTreeViewModel->AddChildState(ParentState);
-		TreeView->SetItemExpansion(ParentState, true);
+		StateTreeViewModel->CopySelectedStates();
 	}
 }
 
-void SStateTreeView::HandleDeleteItems()
+void SStateTreeView::HandlePasteStatesAsSiblings()
 {
-	if (StateTreeViewModel == nullptr)
+	if (StateTreeViewModel)
 	{
-		return;
+		StateTreeViewModel->PasteStates(GetFirstSelectedState());
 	}
+}
 
-	StateTreeViewModel->RemoveSelectedStates();
+void SStateTreeView::HandlePasteStatesAsChildren()
+{
+	if (StateTreeViewModel)
+	{
+		StateTreeViewModel->PasteStatesAsChildren(GetFirstSelectedState());
+	}
+}
+
+void SStateTreeView::HandleDuplicateSelectedStates()
+{
+	if (StateTreeViewModel)
+	{
+		StateTreeViewModel->CopySelectedStates();
+		StateTreeViewModel->PasteStates(GetFirstSelectedState());
+	}
 }
 
+void SStateTreeView::HandleDeleteStates()
+{
+	if (StateTreeViewModel)
+	{
+		StateTreeViewModel->RemoveSelectedStates();
+	}
+}
 
+void SStateTreeView::HandleRenameState()
+{
+	RequestedRenameState = GetFirstSelectedState();
+}
 
 #undef LOCTEXT_NAMESPACE
diff --git a/Engine/Plugins/Runtime/StateTree/Source/StateTreeEditorModule/Private/SStateTreeView.h b/Engine/Plugins/Runtime/StateTree/Source/StateTreeEditorModule/Private/SStateTreeView.h
--- a/Engine/Plugins/Runtime/StateTree/Source/StateTreeEditorModule/Private/SStateTreeView.h
+++ b/Engine/Plugins/Runtime/StateTree/Source/StateTreeEditorModule/Private/SStateTreeView.h
@@ -45,13 +45,14 @@ public:
 	SLATE_END_ARGS()
 
 	SStateTreeView();
-	~SStateTreeView();
+	virtual ~SStateTreeView() override;
 
 	void Construct(const FArguments& InArgs, TSharedRef<FStateTreeViewModel> StateTreeViewModel);
 
 	void SavePersistentExpandedStates();
 
 private:
+	virtual FReply OnKeyDown(const FGeometry& MyGeometry, const FKeyEvent& InKeyEvent) override;
 	virtual void Tick(const FGeometry& AllottedGeometry, const double InCurrentTime, const float InDeltaTime) override;
 
 	void UpdateTree(bool bExpandPersistent = false);
@@ -73,11 +74,21 @@ private:
 	TSharedPtr<SWidget> HandleContextMenuOpening();
 
 	// Action handlers
+	UStateTreeState* GetFirstSelectedState() const;
 	FReply HandleAddStateButton();
-	void HandleRenameState(UStateTreeState* State);
-	void HandleAddState(UStateTreeState* AfterItem);
-	void HandleAddChildState(UStateTreeState* ParentItem);
-	void HandleDeleteItems();
+	void HandleAddSiblingState();
+	void HandleAddChildState();
+	void HandleCutSelectedStates();
+	void HandleCopySelectedStates();
+	void HandlePasteStatesAsSiblings();
+	void HandlePasteStatesAsChildren();
+	void HandleDuplicateSelectedStates();
+	void HandleRenameState();
+	void HandleDeleteStates();
+
+	bool HasSelection() const;
+
+	void BindCommands();
 
 	TSharedPtr<FStateTreeViewModel> StateTreeViewModel;
 
@@ -86,6 +97,8 @@ private:
 	TSharedPtr<SScrollBox> ViewBox;
 	TArray<TWeakObjectPtr<UStateTreeState>> Subtrees;
 
+	TSharedPtr<FUICommandList> CommandList;
+
 	UStateTreeState* RequestedRenameState;
 	bool bItemsDirty;
 	bool bUpdatingSelection;
diff --git a/Engine/Plugins/Runtime/StateTree/Source/StateTreeEditorModule/Private/SStateTreeViewRow.cpp b/Engine/Plugins/Runtime/StateTree/Source/StateTreeEditorModule/Private/SStateTreeViewRow.cpp
--- a/Engine/Plugins/Runtime/StateTree/Source/StateTreeEditorModule/Private/SStateTreeViewRow.cpp
+++ b/Engine/Plugins/Runtime/StateTree/Source/StateTreeEditorModule/Private/SStateTreeViewRow.cpp
@@ -23,12 +23,12 @@ void SStateTreeViewRow::Construct(const FArguments& InArgs, const TSharedRef<STa
 	WeakState = InState;
 
 	STableRow<TWeakObjectPtr<UStateTreeState>>::ConstructInternal(STableRow::FArguments()
-		.Padding(5.0f)
-		.OnDragDetected(this, &SStateTreeViewRow::HandleDragDetected)
-		.OnCanAcceptDrop(this, &SStateTreeViewRow::HandleCanAcceptDrop)
-		.OnAcceptDrop(this, &SStateTreeViewRow::HandleAcceptDrop)
-		.Style(&FStateTreeEditorStyle::Get()->GetWidgetStyle<FTableRowStyle>("StateTree.Selection"))
-	, InOwnerTableView);
+        .Padding(5.0f)
+        .OnDragDetected(this, &SStateTreeViewRow::HandleDragDetected)
+        .OnCanAcceptDrop(this, &SStateTreeViewRow::HandleCanAcceptDrop)
+        .OnAcceptDrop(this, &SStateTreeViewRow::HandleAcceptDrop)
+        .Style(&FStateTreeEditorStyle::Get().GetWidgetStyle<FTableRowStyle>("StateTree.Selection"))
+    , InOwnerTableView);
 
 	static const FLinearColor TasksBackground = FLinearColor(FColor(17, 117, 131));
 	static const FLinearColor LinkBackground = FLinearColor(FColor(84, 84, 84));
diff --git a/Engine/Plugins/Runtime/StateTree/Source/StateTreeEditorModule/Private/StateTreeCompiler.cpp b/Engine/Plugins/Runtime/StateTree/Source/StateTreeEditorModule/Private/StateTreeCompiler.cpp
--- a/Engine/Plugins/Runtime/StateTree/Source/StateTreeEditorModule/Private/StateTreeCompiler.cpp
+++ b/Engine/Plugins/Runtime/StateTree/Source/StateTreeEditorModule/Private/StateTreeCompiler.cpp
@@ -192,7 +192,7 @@ FStateTreeStateHandle FStateTreeCompiler::GetStateHandle(const FGuid& StateID) c
 	return FStateTreeStateHandle(uint16(*Idx));
 }
 
-UStateTreeState* FStateTreeCompiler::GetState(const FGuid& StateID)
+UStateTreeState* FStateTreeCompiler::GetState(const FGuid& StateID) const
 {
 	const int32* Idx = IDToState.Find(StateID);
 	if (Idx == nullptr)
@@ -686,6 +686,17 @@ bool FStateTreeCompiler::ResolveTransitionState(const UStateTreeState* SourceSta
 {
 	if (Link.LinkType == EStateTreeTransitionType::GotoState)
 	{
+		// Warn if goto state points to another subtree.
+		if (const UStateTreeState* TargetState = GetState(Link.ID))
+		{
+			if (TargetState->GetRootState() != SourceState->GetRootState())
+			{
+				Log.Reportf(EMessageSeverity::Warning,
+					TEXT("Target state '%s' is in different subtree. Verify that this is intentional."),
+					*Link.Name.ToString());
+			}			
+		}
+		
 		OutTransitionHandle = GetStateHandle(Link.ID);
 		if (!OutTransitionHandle.IsValid())
 		{
diff --git a/Engine/Plugins/Runtime/StateTree/Source/StateTreeEditorModule/Private/StateTreeEditor.cpp b/Engine/Plugins/Runtime/StateTree/Source/StateTreeEditorModule/Private/StateTreeEditor.cpp
--- a/Engine/Plugins/Runtime/StateTree/Source/StateTreeEditorModule/Private/StateTreeEditor.cpp
+++ b/Engine/Plugins/Runtime/StateTree/Source/StateTreeEditorModule/Private/StateTreeEditor.cpp
@@ -580,11 +580,7 @@ namespace UE::StateTree::Editor::Internal
 			return;
 		}
 
-		TreeData->VisitHierarchy([](UStateTreeState& State, UStateTreeState* ParentState)
-		{
-			State.Parent = ParentState;
-			return EStateTreeVisitor::Continue;
-		});
+		TreeData->ReparentStates();
 	}
 
 	void ApplySchema(UStateTree& StateTree)
diff --git a/Engine/Plugins/Runtime/StateTree/Source/StateTreeEditorModule/Private/StateTreeEditorCommands.cpp b/Engine/Plugins/Runtime/StateTree/Source/StateTreeEditorModule/Private/StateTreeEditorCommands.cpp
--- a/Engine/Plugins/Runtime/StateTree/Source/StateTreeEditorModule/Private/StateTreeEditorCommands.cpp
+++ b/Engine/Plugins/Runtime/StateTree/Source/StateTreeEditorModule/Private/StateTreeEditorCommands.cpp
@@ -1,16 +1,17 @@
 // Copyright Epic Games, Inc. All Rights Reserved.
 
 #include "StateTreeEditorCommands.h"
+#include "StateTreeEditorStyle.h"
 #include "Styling/AppStyle.h"
 
 #define LOCTEXT_NAMESPACE "StateTreeEditor"
 
 FStateTreeEditorCommands::FStateTreeEditorCommands() 
 	: TCommands<FStateTreeEditorCommands>(
-		"StateTreeEditor", // Context name for fast lookup
+		TEXT("StateTreeEditor"), // Context name for fast lookup
 		LOCTEXT("StateTreeEditor", "StateTree Editor"), // Localized context name for displaying
 		NAME_None,
-		FAppStyle::GetAppStyleSetName()
+		TEXT("StateTreeEditorStyle")
 	)
 {
 }
@@ -18,10 +19,19 @@ FStateTreeEditorCommands::FStateTreeEditorCommands()
 void FStateTreeEditorCommands::RegisterCommands()
 {
 	UI_COMMAND(Compile, "Compile", "Compile the current StateTree.", EUserInterfaceActionType::Button, FInputChord(EKeys::F7));
-
 	UI_COMMAND(SaveOnCompile_Never, "Never", "Sets the save-on-compile option to 'Never', meaning that your StateTree will not be saved when they are compiled", EUserInterfaceActionType::RadioButton, FInputChord());
 	UI_COMMAND(SaveOnCompile_SuccessOnly, "On Success Only", "Sets the save-on-compile option to 'Success Only', meaning that your StateTree will be saved whenever they are successfully compiled", EUserInterfaceActionType::RadioButton, FInputChord());
 	UI_COMMAND(SaveOnCompile_Always, "Always", "Sets the save-on-compile option to 'Always', meaning that your StateTree will be saved whenever they are compiled (even if there were errors)", EUserInterfaceActionType::RadioButton, FInputChord());
+
+	UI_COMMAND(AddSiblingState, "Add Sibling State", "Add a Sibling State", EUserInterfaceActionType::Button, FInputChord());
+	UI_COMMAND(AddChildState, "Add Child State", "Add a Child State", EUserInterfaceActionType::Button, FInputChord());
+	UI_COMMAND(CutStates, "Cut", "Cut Selected States", EUserInterfaceActionType::Button, FInputChord(EModifierKey::Control, EKeys::X));
+	UI_COMMAND(CopyStates, "Copy", "Copy Selected States", EUserInterfaceActionType::Button, FInputChord(EModifierKey::Control, EKeys::C));
+	UI_COMMAND(PasteStatesAsSiblings, "Paste As Siblings", "Paste States as Siblings", EUserInterfaceActionType::Button, FInputChord(EModifierKey::Control, EKeys::V));
+	UI_COMMAND(PasteStatesAsChildren, "Paste As Children", "Paste States as Children", EUserInterfaceActionType::Button, FInputChord(EModifierKey::Control | EModifierKey::Shift, EKeys::V));
+	UI_COMMAND(DuplicateStates, "Duplicate", "Duplicate Selected States", EUserInterfaceActionType::Button, FInputChord(EModifierKey::Control, EKeys::D));
+	UI_COMMAND(DeleteStates, "Delete", "Delete Selected States", EUserInterfaceActionType::Button, FInputChord(EKeys::Delete));
+	UI_COMMAND(RenameState, "Rename", "Rename Selected State", EUserInterfaceActionType::Button, FInputChord(EKeys::F2));
 }
 
 #undef LOCTEXT_NAMESPACE
diff --git a/Engine/Plugins/Runtime/StateTree/Source/StateTreeEditorModule/Private/StateTreeEditorCommands.h b/Engine/Plugins/Runtime/StateTree/Source/StateTreeEditorModule/Private/StateTreeEditorCommands.h
--- a/Engine/Plugins/Runtime/StateTree/Source/StateTreeEditorModule/Private/StateTreeEditorCommands.h
+++ b/Engine/Plugins/Runtime/StateTree/Source/StateTreeEditorModule/Private/StateTreeEditorCommands.h
@@ -20,4 +20,14 @@ public:
 	TSharedPtr<FUICommandInfo> SaveOnCompile_Never;
 	TSharedPtr<FUICommandInfo> SaveOnCompile_SuccessOnly;
 	TSharedPtr<FUICommandInfo> SaveOnCompile_Always;
+
+	TSharedPtr<FUICommandInfo> AddSiblingState;
+	TSharedPtr<FUICommandInfo> AddChildState;
+	TSharedPtr<FUICommandInfo> RenameState;
+	TSharedPtr<FUICommandInfo> CutStates;
+	TSharedPtr<FUICommandInfo> CopyStates;
+	TSharedPtr<FUICommandInfo> PasteStatesAsSiblings;
+	TSharedPtr<FUICommandInfo> PasteStatesAsChildren;
+	TSharedPtr<FUICommandInfo> DuplicateStates;
+	TSharedPtr<FUICommandInfo> DeleteStates;
 };
\ No newline at end of file
diff --git a/Engine/Plugins/Runtime/StateTree/Source/StateTreeEditorModule/Private/StateTreeEditorData.cpp b/Engine/Plugins/Runtime/StateTree/Source/StateTreeEditorModule/Private/StateTreeEditorData.cpp
--- a/Engine/Plugins/Runtime/StateTree/Source/StateTreeEditorModule/Private/StateTreeEditorData.cpp
+++ b/Engine/Plugins/Runtime/StateTree/Source/StateTreeEditorModule/Private/StateTreeEditorData.cpp
@@ -18,6 +18,12 @@ void UStateTreeEditorData::PostInitProperties()
 }
 
 #if WITH_EDITOR
+void UStateTreeEditorData::PostLoad()
+{
+	Super::PostLoad();
+	ReparentStates();
+}
+
 void UStateTreeEditorData::PostEditChangeChainProperty(FPropertyChangedChainEvent& PropertyChangedEvent)
 {
 	Super::PostEditChangeChainProperty(PropertyChangedEvent);
@@ -317,6 +323,22 @@ void UStateTreeEditorData::GetAllStructValues(TMap<FGuid, const FStateTreeDataVi
 		});
 }
 
+void UStateTreeEditorData::ReparentStates()
+{
+	VisitHierarchy([TreeData = this](UStateTreeState& State, UStateTreeState* ParentState) mutable 
+	{
+		UObject* ExpectedOuter = ParentState ? Cast<UObject>(ParentState) : Cast<UObject>(TreeData);
+		if (State.GetOuter() != ExpectedOuter)
+		{
+			State.Rename(nullptr, ExpectedOuter, REN_DontCreateRedirectors | REN_DoNotDirty);
+		}
+		
+		State.Parent = ParentState;
+		
+		return EStateTreeVisitor::Continue;
+	});
+}
+
 EStateTreeVisitor UStateTreeEditorData::VisitStateNodes(const UStateTreeState& State, TFunctionRef<EStateTreeVisitor(const UStateTreeState* State, const FGuid& ID, const FName& Name, const EStateTreeNodeType NodeType, const UScriptStruct* NodeStruct, const UStruct* InstanceStruct)> InFunc) const
 {
 	bool bContinue = true;
diff --git a/Engine/Plugins/Runtime/StateTree/Source/StateTreeEditorModule/Private/StateTreeEditorModule.cpp b/Engine/Plugins/Runtime/StateTree/Source/StateTreeEditorModule/Private/StateTreeEditorModule.cpp
--- a/Engine/Plugins/Runtime/StateTree/Source/StateTreeEditorModule/Private/StateTreeEditorModule.cpp
+++ b/Engine/Plugins/Runtime/StateTree/Source/StateTreeEditorModule/Private/StateTreeEditorModule.cpp
@@ -72,7 +72,7 @@ void FStateTreeEditorModule::StartupModule()
 	MenuExtensibilityManager = MakeShareable(new FExtensibilityManager);
 	ToolBarExtensibilityManager = MakeShareable(new FExtensibilityManager);
 
-	FStateTreeEditorStyle::Initialize();
+	FStateTreeEditorStyle::Register();
 	FStateTreeEditorCommands::Register();
 
 	// Register the details customizer
@@ -96,7 +96,7 @@ void FStateTreeEditorModule::ShutdownModule()
 	MenuExtensibilityManager.Reset();
 	ToolBarExtensibilityManager.Reset();
 
-	FStateTreeEditorStyle::Shutdown();
+	FStateTreeEditorStyle::Unregister();
 	FStateTreeEditorCommands::Unregister();
 
 	// Unregister the details customization
diff --git a/Engine/Plugins/Runtime/StateTree/Source/StateTreeEditorModule/Private/StateTreeEditorStyle.cpp b/Engine/Plugins/Runtime/StateTree/Source/StateTreeEditorModule/Private/StateTreeEditorStyle.cpp
--- a/Engine/Plugins/Runtime/StateTree/Source/StateTreeEditorModule/Private/StateTreeEditorStyle.cpp
+++ b/Engine/Plugins/Runtime/StateTree/Source/StateTreeEditorModule/Private/StateTreeEditorStyle.cpp
@@ -8,38 +8,47 @@
 #include "Brushes/SlateRoundedBoxBrush.h"
 #include "Styling/SlateTypes.h"
 #include "Interfaces/IPluginManager.h"
-
 #include "Misc/Paths.h"
 #include "Styling/StyleColors.h"
+#include "Styling/SlateStyleMacros.h"
 
-#define IMAGE_BRUSH( RelativePath, ... ) FSlateImageBrush( StyleSet->RootToContentDir( RelativePath, TEXT( ".png" ) ), __VA_ARGS__ )
-#define IMAGE_PLUGIN_BRUSH( RelativePath, ... ) FSlateImageBrush( FMeshEditorStyle::InContent( RelativePath, ".png" ), __VA_ARGS__ )
-#define BOX_BRUSH( RelativePath, ... ) FSlateBoxBrush( StyleSet->RootToContentDir( RelativePath, TEXT( ".png" ) ), __VA_ARGS__ )
-#define TTF_CORE_FONT( RelativePath, ... ) FSlateFontInfo( StyleSet->RootToCoreContentDir( RelativePath, TEXT( ".ttf" ) ), __VA_ARGS__ )
-
-#define DEFAULT_FONT(...) FCoreStyle::GetDefaultFontStyle(__VA_ARGS__)
 
-TSharedPtr<FSlateStyleSet> FStateTreeEditorStyle::StyleSet = nullptr;
+#define CORE_FONT( RelativePath, ... ) FSlateFontInfo( RootToCoreContentDir( RelativePath, TEXT( ".ttf" ) ), __VA_ARGS__ )
 
-FString FStateTreeEditorStyle::InContent(const FString& RelativePath, const ANSICHAR* Extension)
+namespace UE::StateTree::Editor
 {
-	static FString ContentDir = IPluginManager::Get().FindPlugin(TEXT("StateTreeEditorModule"))->GetContentDir() / TEXT("Slate");
-	return (ContentDir / RelativePath) + Extension;
-}
 
-void FStateTreeEditorStyle::Initialize()
+class FContentRootScope
 {
-	if (StyleSet.IsValid())
+public:
+	FContentRootScope(FStateTreeEditorStyle* InStyle, const FString& NewContentRoot)
+		: Style(InStyle)
+		, PreviousContentRoot(InStyle->GetContentRootDir())
+	{
+		Style->SetContentRoot(NewContentRoot);
+	}
+
+	~FContentRootScope()
 	{
-		return;
+		Style->SetContentRoot(PreviousContentRoot);
 	}
+private:
+	FStateTreeEditorStyle* Style;
+	FString PreviousContentRoot;
+};
+
+}; // UE::StateTree::Editor
 
+FStateTreeEditorStyle::FStateTreeEditorStyle()
+	: FSlateStyleSet(TEXT("StateTreeEditorStyle"))
+{
 	const FVector2f Icon8x8(8.0f, 8.0f);
 
-	StyleSet = MakeShared<FSlateStyleSet>(GetStyleSetName());
+	const FString EngineEditorSlateDir = FPaths::EngineContentDir() / TEXT("Slate");
+	SetCoreContentRoot(EngineEditorSlateDir);
 
-	StyleSet->SetContentRoot(FPaths::EngineContentDir() / TEXT("Editor/Slate"));
-	StyleSet->SetCoreContentRoot(FPaths::EngineContentDir() / TEXT("Slate"));
+	const FString StateTreePluginContentDir = FPaths::EnginePluginsDir() / TEXT("Runtime/StateTree/Resources");
+	SetContentRoot(StateTreePluginContentDir);
 
 	const FScrollBarStyle ScrollBar = FAppStyle::GetWidgetStyle<FScrollBarStyle>("ScrollBar");
 	const FTextBlockStyle& NormalText = FAppStyle::Get().GetWidgetStyle<FTextBlockStyle>("NormalText");
@@ -49,21 +58,21 @@ void FStateTreeEditorStyle::Initialize()
 		FTextBlockStyle StateIcon = FTextBlockStyle(NormalText)
 			.SetFont(FAppStyle::Get().GetFontStyle("FontAwesome.12"))
 			.SetColorAndOpacity(FLinearColor(230.0f / 255.0f, 230.0f / 255.0f, 230.0f / 255.0f, 0.5f));
-		StyleSet->Set("StateTree.Icon", StateIcon);
+		Set("StateTree.Icon", StateIcon);
 
 		FTextBlockStyle StateDetailsIcon = FTextBlockStyle(NormalText)
 			.SetFont(FAppStyle::Get().GetFontStyle("FontAwesome.10"))
 		    .SetColorAndOpacity(FLinearColor(230.0f / 255.0f, 230.0f / 255.0f, 230.0f / 255.0f, 0.5f));
-		StyleSet->Set("StateTree.DetailsIcon", StateDetailsIcon);
+		Set("StateTree.DetailsIcon", StateDetailsIcon);
 
 		FTextBlockStyle StateTitle = FTextBlockStyle(NormalText)
-			.SetFont(TTF_CORE_FONT("Fonts/Roboto-Bold", 12))
+			.SetFont(CORE_FONT("Fonts/Roboto-Bold", 12))
 			.SetColorAndOpacity(FLinearColor(230.0f / 255.0f, 230.0f / 255.0f, 230.0f / 255.0f, 0.9f));
-		StyleSet->Set("StateTree.State.Title", StateTitle);
+		Set("StateTree.State.Title", StateTitle);
 
 		FEditableTextBoxStyle StateTitleEditableText = FEditableTextBoxStyle()
 			.SetTextStyle(NormalText)
-			.SetFont(TTF_CORE_FONT("Fonts/Roboto-Bold", 12))
+			.SetFont(CORE_FONT("Fonts/Roboto-Bold", 12))
 			.SetBackgroundImageNormal(BOX_BRUSH("Common/TextBox", FMargin(4.0f / 16.0f)))
 			.SetBackgroundImageHovered(BOX_BRUSH("Common/TextBox_Hovered", FMargin(4.0f / 16.0f)))
 			.SetBackgroundImageFocused(BOX_BRUSH("Common/TextBox_Hovered", FMargin(4.0f / 16.0f)))
@@ -71,9 +80,9 @@ void FStateTreeEditorStyle::Initialize()
 			.SetBackgroundColor(FLinearColor(0,0,0,0.1f))
 			.SetPadding(FMargin(0))
 			.SetScrollBarStyle(ScrollBar);
-		StyleSet->Set("StateTree.State.TitleEditableText", StateTitleEditableText);
+		Set("StateTree.State.TitleEditableText", StateTitleEditableText);
 
-		StyleSet->Set("StateTree.State.TitleInlineEditableText", FInlineEditableTextBlockStyle()
+		Set("StateTree.State.TitleInlineEditableText", FInlineEditableTextBlockStyle()
 			.SetTextStyle(StateTitle)
 			.SetEditableTextBoxStyle(StateTitleEditableText));
 	}
@@ -81,27 +90,27 @@ void FStateTreeEditorStyle::Initialize()
 	// Details
 	{
 		FTextBlockStyle StateTitle = FTextBlockStyle(NormalText)
-			.SetFont(TTF_CORE_FONT("Fonts/Roboto-Regular", 10))
+			.SetFont(CORE_FONT("Fonts/Roboto-Regular", 10))
 			.SetColorAndOpacity(FLinearColor(230.0f / 255.0f, 230.0f / 255.0f, 230.0f / 255.0f, 0.75f));
-		StyleSet->Set("StateTree.Details", StateTitle);
+		Set("StateTree.Details", StateTitle);
 	}
 
 	// Details rich text
 	{
-		StyleSet->Set("Details.Normal", FTextBlockStyle(NormalText)
+		Set("Details.Normal", FTextBlockStyle(NormalText)
 			.SetFont(FAppStyle::GetFontStyle(TEXT("PropertyWindow.NormalFont"))));
 
-		StyleSet->Set("Details.Bold", FTextBlockStyle(NormalText)
+		Set("Details.Bold", FTextBlockStyle(NormalText)
 			.SetFont(FAppStyle::GetFontStyle(TEXT("PropertyWindow.BoldFont"))));
 
-		StyleSet->Set("Details.Subdued", FTextBlockStyle(NormalText)
+		Set("Details.Subdued", FTextBlockStyle(NormalText)
 			.SetColorAndOpacity(FSlateColor::UseSubduedForeground())
 			.SetFont(FAppStyle::GetFontStyle(TEXT("PropertyWindow.NormalFont"))));
 	}
 
 	const FLinearColor SelectionColor = FColor(0, 0, 0, 32);
 	const FTableRowStyle& NormalTableRowStyle = FAppStyle::Get().GetWidgetStyle<FTableRowStyle>("TableView.Row");
-	StyleSet->Set("StateTree.Selection",
+	Set("StateTree.Selection",
 		FTableRowStyle(NormalTableRowStyle)
 		.SetActiveBrush(IMAGE_BRUSH("Common/Selection", Icon8x8, SelectionColor))
 		.SetActiveHoveredBrush(IMAGE_BRUSH("Common/Selection", Icon8x8, SelectionColor))
@@ -124,17 +133,17 @@ void FStateTreeEditorStyle::Initialize()
 		.SetNormalPadding(FMargin(2, 2, 2, 2))
 		.SetPressedPadding(FMargin(2, 3, 2, 1));
 
-	StyleSet->Set("StateTree.Node.Operand.ComboBox", FComboButtonStyle(ComboButtonStyle).SetButtonStyle(OperandButton));
+	Set("StateTree.Node.Operand.ComboBox", FComboButtonStyle(ComboButtonStyle).SetButtonStyle(OperandButton));
 
-	StyleSet->Set("StateTree.Node.Operand", FTextBlockStyle(NormalText)
+	Set("StateTree.Node.Operand", FTextBlockStyle(NormalText)
 		.SetFont(FAppStyle::GetFontStyle(TEXT("PropertyWindow.BoldFont")))
 		.SetFontSize(8));
 
-	StyleSet->Set("StateTree.Node.Parens", FTextBlockStyle(NormalText)
+	Set("StateTree.Node.Parens", FTextBlockStyle(NormalText)
 		.SetFont(FAppStyle::GetFontStyle(TEXT("PropertyWindow.NormalFont")))
 		.SetFontSize(12));
 
-	StyleSet->Set("StateTree.Param.Label", FTextBlockStyle(NormalText)
+	Set("StateTree.Param.Label", FTextBlockStyle(NormalText)
 		.SetFont(FAppStyle::GetFontStyle(TEXT("PropertyWindow.BoldFont")))
 		.SetFontSize(7));
 
@@ -149,30 +158,47 @@ void FStateTreeEditorStyle::Initialize()
 		.SetNormalPadding(FMargin(2, 2, 2, 2))
 		.SetPressedPadding(FMargin(2, 3, 2, 1));
 	
-	StyleSet->Set("StateTree.Node.Indent.ComboBox", FComboButtonStyle(ComboButtonStyle).SetButtonStyle(IndentButton));
+	Set("StateTree.Node.Indent.ComboBox", FComboButtonStyle(ComboButtonStyle).SetButtonStyle(IndentButton));
 
 	const FEditableTextStyle& NormalEditableText = FCoreStyle::Get().GetWidgetStyle<FEditableTextStyle>("NormalEditableText");
 	FEditableTextStyle NameEditStyle(NormalEditableText);
 	NameEditStyle.Font.Size = 10;
-	StyleSet->Set("StateTree.Node.Name", NameEditStyle);
-	
-	FSlateStyleRegistry::RegisterSlateStyle(*StyleSet.Get());
-}
+	Set("StateTree.Node.Name", NameEditStyle);
 
+	// Command icons
+	{
+		// From generic
+		UE::StateTree::Editor::FContentRootScope Scope(this, EngineEditorSlateDir);
+		Set("StateTreeEditor.CutStates", new IMAGE_BRUSH_SVG("Starship/Common/Cut", CoreStyleConstants::Icon16x16));
+		Set("StateTreeEditor.CopyStates", new IMAGE_BRUSH_SVG("Starship/Common/Copy", CoreStyleConstants::Icon16x16));
+		Set("StateTreeEditor.DuplicateStates", new IMAGE_BRUSH_SVG("Starship/Common/Duplicate", CoreStyleConstants::Icon16x16));
+		Set("StateTreeEditor.DeleteStates", new IMAGE_BRUSH_SVG("Starship/Common/Delete", CoreStyleConstants::Icon16x16));
+		Set("StateTreeEditor.RenameState", new IMAGE_BRUSH_SVG("Starship/Common/Rename", CoreStyleConstants::Icon16x16));
+	}
 
-void FStateTreeEditorStyle::Shutdown()
-{
-	if (StyleSet.IsValid())
 	{
-		FSlateStyleRegistry::UnRegisterSlateStyle(*StyleSet.Get());
-		ensure(StyleSet.IsUnique());
-		StyleSet.Reset();
+		// From plugin
+		Set("StateTreeEditor.AddSiblingState", new IMAGE_BRUSH_SVG("Icons/Sibling_State", CoreStyleConstants::Icon16x16));
+		Set("StateTreeEditor.AddChildState", new IMAGE_BRUSH_SVG("Icons/Child_State", CoreStyleConstants::Icon16x16));
+
+		Set("StateTreeEditor.PasteStatesAsSiblings", new IMAGE_BRUSH_SVG("Icons/Sibling_State", CoreStyleConstants::Icon16x16));
+		Set("StateTreeEditor.PasteStatesAsChildren", new IMAGE_BRUSH_SVG("Icons/Child_State", CoreStyleConstants::Icon16x16));
 	}
+
 }
 
+void FStateTreeEditorStyle::Register()
+{
+	FSlateStyleRegistry::RegisterSlateStyle(Get());
+}
+
+void FStateTreeEditorStyle::Unregister()
+{
+	FSlateStyleRegistry::UnRegisterSlateStyle(Get());
+}
 
-FName FStateTreeEditorStyle::GetStyleSetName()
+FStateTreeEditorStyle& FStateTreeEditorStyle::Get()
 {
-	static FName StyleName("StateTreeEditorStyle");
-	return StyleName;
+	static FStateTreeEditorStyle Instance;
+	return Instance;
 }
diff --git a/Engine/Plugins/Runtime/StateTree/Source/StateTreeEditorModule/Private/StateTreeViewModel.cpp b/Engine/Plugins/Runtime/StateTree/Source/StateTreeEditorModule/Private/StateTreeViewModel.cpp
--- a/Engine/Plugins/Runtime/StateTree/Source/StateTreeEditorModule/Private/StateTreeViewModel.cpp
+++ b/Engine/Plugins/Runtime/StateTree/Source/StateTreeEditorModule/Private/StateTreeViewModel.cpp
@@ -6,13 +6,169 @@
 #include "StateTreeDelegates.h"
 #include "Editor.h"
 #include "ScopedTransaction.h"
+#include "HAL/PlatformApplicationMisc.h"
+#include "UnrealExporter.h"
+#include "Exporters/Exporter.h"
+#include "Factories.h"
 
 
 #define LOCTEXT_NAMESPACE "StateTreeEditor"
 
-
 namespace UE::StateTree::Editor
 {
+	class FStateTreeStateTextFactory : public FCustomizableTextObjectFactory
+	{
+	public:
+		FStateTreeStateTextFactory()
+			: FCustomizableTextObjectFactory(GWarn)
+		{}
+
+		virtual bool CanCreateClass(UClass* InObjectClass, bool& bOmitSubObjs) const override
+		{
+			UE_LOG(LogTemp, Error, TEXT("*** CanCreateClass: %s"), *GetNameSafe(InObjectClass));
+			return InObjectClass->IsChildOf(UStateTreeState::StaticClass())
+				|| InObjectClass->IsChildOf(UStateTreeClipboardBindings::StaticClass());
+		}
+
+		virtual void ProcessConstructedObject(UObject* NewObject) override
+		{
+			if (UStateTreeState* State = Cast<UStateTreeState>(NewObject))
+			{
+				States.Add(State);
+			}
+			else if (UStateTreeClipboardBindings* Bindings = Cast<UStateTreeClipboardBindings>(NewObject))
+			{
+				ClipboardBindings = Bindings;
+			}
+		}
+
+	public:
+		TArray<UStateTreeState*> States;
+		UStateTreeClipboardBindings* ClipboardBindings = nullptr;
+	};
+
+
+	void CollectBindingsRecursive(UStateTreeEditorData* TreeData, UStateTreeState* State, TArray<FStateTreePropertyPathBinding>& AllBindings)
+	{
+		if (!State)
+		{
+			return;
+		}
+		
+		TreeData->VisitStateNodes(*State, [TreeData, &AllBindings](const UStateTreeState* State, const FStateTreeBindableStructDesc& Desc, const FStateTreeDataView Value)
+		{
+			TArray<FStateTreePropertyPathBinding> NodeBindings;
+			TreeData->GetPropertyEditorBindings()->GetPropertyBindingsFor(Desc.ID, NodeBindings);
+			AllBindings.Append(NodeBindings);
+			return EStateTreeVisitor::Continue;				
+		});
+
+		for (UStateTreeState* ChildState : State->Children)
+		{
+			CollectBindingsRecursive(TreeData, ChildState, AllBindings);
+		}
+	}
+
+	FString ExportStatesToText(UStateTreeEditorData* TreeData, const TArrayView<UStateTreeState*> States)
+	{
+		if (States.IsEmpty())
+		{
+			return FString();
+		}
+
+		// Clear the mark state for saving.
+		UnMarkAllObjects(EObjectMark(OBJECTMARK_TagExp | OBJECTMARK_TagImp));
+
+		FStringOutputDevice Archive;
+		const FExportObjectInnerContext Context;
+
+		UStateTreeClipboardBindings* ClipboardBindings = NewObject<UStateTreeClipboardBindings>();
+		check(ClipboardBindings);
+
+		for (UStateTreeState* State : States)
+		{
+			UObject* ThisOuter = State->GetOuter();
+			UExporter::ExportToOutputDevice(&Context, State, nullptr, Archive, TEXT("copy"), 0, PPF_ExportsNotFullyQualified | PPF_Copy | PPF_Delimited, false, ThisOuter);
+
+			CollectBindingsRecursive(TreeData, State, ClipboardBindings->Bindings);
+		}
+
+		UExporter::ExportToOutputDevice(&Context, ClipboardBindings, nullptr, Archive, TEXT("copy"), 0, PPF_ExportsNotFullyQualified | PPF_Copy | PPF_Delimited, false);
+
+		return *Archive;
+	}
+
+	void CollectStateLinks(const UStruct* Struct, void* Memory, TArray<FStateTreeStateLink*>& Links)
+	{
+		for (TPropertyValueIterator<FStructProperty> It(Struct, Memory); It; ++It)
+		{
+			if (It->Key->Struct == TBaseStructure<FStateTreeStateLink>::Get())
+			{
+				FStateTreeStateLink* StateLink = static_cast<FStateTreeStateLink*>(const_cast<void*>(It->Value));
+				Links.Add(StateLink);
+			}
+		}
+	}
+
+	void FixNodesAfterDuplication(TArrayView<FStateTreeEditorNode> Nodes, TMap<FGuid, FGuid>& IDsMap, TArray<FStateTreeStateLink*>& Links)
+	{
+		for (FStateTreeEditorNode& Node : Nodes)
+		{
+			const FGuid NewNodeID = FGuid::NewGuid();
+			IDsMap.Emplace(Node.ID, NewNodeID);
+			Node.ID = NewNodeID;
+
+			if (Node.Node.IsValid())
+			{
+				CollectStateLinks(Node.Node.GetScriptStruct(), Node.Node.GetMutableMemory(), Links);
+			}
+			if (Node.Instance.IsValid())
+			{
+				CollectStateLinks(Node.Instance.GetScriptStruct(), Node.Instance.GetMutableMemory(), Links);
+			}
+			if (Node.InstanceObject)
+			{
+				CollectStateLinks(Node.InstanceObject->GetClass(), Node.InstanceObject, Links);
+			}
+		}
+	}
+
+	void FixStateAfterDuplication(UStateTreeState* State, UStateTreeState* NewParentState, TMap<FGuid, FGuid>& IDsMap, TArray<FStateTreeStateLink*>& Links, TArray<UStateTreeState*>& NewStates)
+	{
+		State->Modify();
+
+		const FGuid NewStateID = FGuid::NewGuid();
+		IDsMap.Emplace(State->ID, NewStateID);
+		State->ID = NewStateID;
+		
+		const FGuid NewParametersID = FGuid::NewGuid();
+		IDsMap.Emplace(State->Parameters.ID, NewParametersID);
+		State->Parameters.ID = NewParametersID;
+		
+		State->Parent = NewParentState;
+		NewStates.Add(State);
+		
+		if (State->Type == EStateTreeStateType::Linked)
+		{
+			Links.Emplace(&State->LinkedSubtree);
+		}
+
+		FixNodesAfterDuplication(TArrayView<FStateTreeEditorNode>(&State->SingleTask, 1), IDsMap, Links);
+		FixNodesAfterDuplication(State->Tasks, IDsMap, Links);
+		FixNodesAfterDuplication(State->EnterConditions, IDsMap, Links);
+
+		for (FStateTreeTransition& Transition : State->Transitions)
+		{
+			FixNodesAfterDuplication(Transition.Conditions, IDsMap, Links);
+			Links.Emplace(&Transition.State);
+		}
+
+		for (UStateTreeState* Child : State->Children)
+		{
+			FixStateAfterDuplication(Child, State, IDsMap, Links, NewStates);
+		}
+	}
+
 	// Removes states from the array which are children of any other state.
 	void RemoveContainedChildren(TArray<UStateTreeState*>& States)
 	{
@@ -442,6 +598,113 @@ void FStateTreeViewModel::RemoveSelectedStates()
 	}
 }
 
+void FStateTreeViewModel::CopySelectedStates()
+{
+	UStateTreeEditorData* TreeData = TreeDataWeak.Get();
+	if (TreeData == nullptr)
+	{
+		return;
+	}
+
+	TArray<UStateTreeState*> States;
+	GetSelectedStates(States);
+	UE::StateTree::Editor::RemoveContainedChildren(States);
+	
+	FString ExportedText = UE::StateTree::Editor::ExportStatesToText(TreeData, States);
+	
+	FPlatformApplicationMisc::ClipboardCopy(*ExportedText);
+}
+
+void FStateTreeViewModel::PasteStates(UStateTreeState* AfterState)
+{
+	UStateTreeEditorData* TreeData = TreeDataWeak.Get();
+	if (TreeData == nullptr)
+	{
+		return;
+	}
+	
+	if (AfterState)
+	{
+		const int32 Index = AfterState->Parent ? AfterState->Parent->Children.Find(AfterState) : TreeData->SubTrees.Find(AfterState);
+		if (Index != INDEX_NONE)
+		{
+			PasteStatesAsChildren(AfterState->Parent, Index + 1);
+		}
+	}
+}
+
+void FStateTreeViewModel::PasteStatesAsChildren(UStateTreeState* ParentState, const int32 Index)
+{
+	UStateTreeEditorData* TreeData = TreeDataWeak.Get();
+	if (TreeData == nullptr)
+	{
+		return;
+	}
+
+	FString TextToImport;
+	FPlatformApplicationMisc::ClipboardPaste(TextToImport);
+
+	UObject* Outer = ParentState ? static_cast<UObject*>(ParentState) : static_cast<UObject*>(TreeData);
+	Outer->Modify();
+
+	UE::StateTree::Editor::FStateTreeStateTextFactory Factory;
+	Factory.ProcessBuffer(Outer, RF_Transactional, TextToImport);
+
+	TArray<UStateTreeState*>& ParentArray = ParentState ? ParentState->Children : TreeData->SubTrees;
+	const int32 TargetIndex = (Index == INDEX_NONE) ? ParentArray.Num() : Index;
+	ParentArray.Insert(Factory.States, TargetIndex);
+
+	TArray<FStateTreeStateLink*> Links;
+	TMap<FGuid, FGuid> IDsMap;
+	TArray<UStateTreeState*> NewStates;
+
+	for (UStateTreeState* State : Factory.States)
+	{		
+		UE::StateTree::Editor::FixStateAfterDuplication(State, ParentState, IDsMap, Links, NewStates);
+	}
+
+	// Copy property bindings for the duplicated states.
+	if (Factory.ClipboardBindings)
+	{
+		for (const TPair<FGuid, FGuid>& Entry : IDsMap)
+		{
+			const FGuid OldTargetID = Entry.Key;
+			const FGuid NewTargetID = Entry.Value;
+			
+			for (const FStateTreePropertyPathBinding& Binding : Factory.ClipboardBindings->Bindings)
+			{
+				if (Binding.GetTargetPath().GetStructID() == OldTargetID)
+				{
+					FStateTreePropertyPath TargetPath(Binding.GetTargetPath());
+					TargetPath.SetStructID(NewTargetID);
+					
+					FStateTreePropertyPath SourcePath(Binding.GetSourcePath());
+					if (const FGuid* NewSourceID = IDsMap.Find(Binding.GetSourcePath().GetStructID()))
+					{
+						SourcePath.SetStructID(*NewSourceID);
+					}
+					
+					TreeData->GetPropertyEditorBindings()->AddPropertyBinding(SourcePath, TargetPath);
+				}
+			}
+		}
+	}
+
+	// Patch IDs in state links.
+	for (FStateTreeStateLink* Link : Links)
+	{
+		if (FGuid* NewID = IDsMap.Find(Link->ID))
+		{
+			Link->ID = *NewID;
+		}
+	}
+
+	for (UStateTreeState* State : NewStates)
+	{
+		OnStateAdded.Broadcast(State->Parent, State);
+	}
+}
+
 void FStateTreeViewModel::MoveSelectedStatesBefore(UStateTreeState* TargetState)
 {
 	MoveSelectedStates(TargetState, FStateTreeViewModelInsert::Before);
diff --git a/Engine/Plugins/Runtime/StateTree/Source/StateTreeEditorModule/Private/StateTreeViewModel.h b/Engine/Plugins/Runtime/StateTree/Source/StateTreeEditorModule/Private/StateTreeViewModel.h
--- a/Engine/Plugins/Runtime/StateTree/Source/StateTreeEditorModule/Private/StateTreeViewModel.h
+++ b/Engine/Plugins/Runtime/StateTree/Source/StateTreeEditorModule/Private/StateTreeViewModel.h
@@ -4,6 +4,7 @@
 
 #include "UObject/WeakObjectPtr.h"
 #include "EditorUndoClient.h"
+#include "StateTreeViewModel.generated.h"
 
 struct FPropertyChangedEvent;
 
@@ -11,6 +12,7 @@ class UStateTreeEditorData;
 class UStateTreeState;
 class UStateTree;
 class FMenuBuilder;
+struct FStateTreePropertyPathBinding;
 
 enum class FStateTreeViewModelInsert : uint8
 {
@@ -66,6 +68,9 @@ public:
 	void AddChildState(UStateTreeState* ParentState);
 	void RenameState(UStateTreeState* State, FName NewName);
 	void RemoveSelectedStates();
+	void CopySelectedStates();
+	void PasteStates(UStateTreeState* AfterState);
+	void PasteStatesAsChildren(UStateTreeState* ParentState, const int32 Index = INDEX_NONE);
 	void MoveSelectedStatesBefore(UStateTreeState* TargetState);
 	void MoveSelectedStatesAfter(UStateTreeState* TargetState);
 	void MoveSelectedStatesInto(UStateTreeState* TargetState);
@@ -107,3 +112,13 @@ protected:
 	FOnStatesMoved OnStatesMoved;
 	FOnSelectionChanged OnSelectionChanged;
 };
+
+/** Helper class to allow to copy bindings into clipboard. */
+UCLASS(Hidden)
+class STATETREEEDITORMODULE_API UStateTreeClipboardBindings : public UObject
+{
+	GENERATED_BODY()
+public:
+	UPROPERTY()
+	TArray<FStateTreePropertyPathBinding> Bindings;
+};
\ No newline at end of file
diff --git a/Engine/Plugins/Runtime/StateTree/Source/StateTreeEditorModule/Public/StateTreeCompiler.h b/Engine/Plugins/Runtime/StateTree/Source/StateTreeEditorModule/Public/StateTreeCompiler.h
--- a/Engine/Plugins/Runtime/StateTree/Source/StateTreeEditorModule/Public/StateTreeCompiler.h
+++ b/Engine/Plugins/Runtime/StateTree/Source/StateTreeEditorModule/Public/StateTreeCompiler.h
@@ -38,7 +38,7 @@ private:
 	/** Resolves the state a transition points to. SourceState is nullptr for global tasks. */
 	bool ResolveTransitionState(const UStateTreeState* SourceState, const FStateTreeStateLink& Link, FStateTreeStateHandle& OutTransitionHandle) const;
 	FStateTreeStateHandle GetStateHandle(const FGuid& StateID) const;
-	UStateTreeState* GetState(const FGuid& StateID);
+	UStateTreeState* GetState(const FGuid& StateID) const;
 
 	bool CreateStates();
 	bool CreateStateRecursive(UStateTreeState& State, const FStateTreeStateHandle Parent);
diff --git a/Engine/Plugins/Runtime/StateTree/Source/StateTreeEditorModule/Public/StateTreeEditorData.h b/Engine/Plugins/Runtime/StateTree/Source/StateTreeEditorModule/Public/StateTreeEditorData.h
--- a/Engine/Plugins/Runtime/StateTree/Source/StateTreeEditorModule/Public/StateTreeEditorData.h
+++ b/Engine/Plugins/Runtime/StateTree/Source/StateTreeEditorModule/Public/StateTreeEditorData.h
@@ -36,6 +36,7 @@ public:
 	// ~IStateTreeEditorPropertyBindingsOwner
 
 #if WITH_EDITOR
+	virtual void PostLoad() override;
 	virtual void PostEditChangeChainProperty(FPropertyChangedChainEvent& PropertyChangedEvent) override;
 #endif
 
@@ -102,6 +103,8 @@ public:
 	UE_DEPRECATED(5.3, "Use GetAllStructValues with values instead.")
 	void GetAllStructIDs(TMap<FGuid, const UStruct*>& AllStructs) const;
 
+	void ReparentStates();
+	
 	// StateTree Builder API
 
 	/**
diff --git a/Engine/Plugins/Runtime/StateTree/Source/StateTreeEditorModule/Public/StateTreeEditorStyle.h b/Engine/Plugins/Runtime/StateTree/Source/StateTreeEditorModule/Public/StateTreeEditorStyle.h
--- a/Engine/Plugins/Runtime/StateTree/Source/StateTreeEditorModule/Public/StateTreeEditorStyle.h
+++ b/Engine/Plugins/Runtime/StateTree/Source/StateTreeEditorModule/Public/StateTreeEditorStyle.h
@@ -7,23 +7,17 @@
 class ISlateStyle;
 
 class STATETREEEDITORMODULE_API FStateTreeEditorStyle
+	: public FSlateStyleSet
 {
 public:
+	static FStateTreeEditorStyle& Get();
 
-	static void Initialize();
-	static void Shutdown();
-	static TSharedPtr<ISlateStyle> Get() { return StyleSet; }
-	static FName GetStyleSetName();
-
-	static const FSlateBrush* GetBrush(FName PropertyName, const ANSICHAR* Specifier = NULL)
-	{
-		return StyleSet->GetBrush(PropertyName, Specifier);
-	}
-
-	static FString InContent(const FString& RelativePath, const ANSICHAR* Extension);
+protected:
+	friend class FStateTreeEditorModule;
 
+	static void Register();
+	static void Unregister();
 
 private:
-
-	static TSharedPtr<FSlateStyleSet> StyleSet;
+	FStateTreeEditorStyle();
 };
