StateTree: Fix bad merge in state tree editor styles.

diff --git a/Engine/Plugins/Runtime/StateTree/Source/StateTreeEditorModule/Private/StateTreeEditorStyle.cpp b/Engine/Plugins/Runtime/StateTree/Source/StateTreeEditorModule/Private/StateTreeEditorStyle.cpp
--- a/Engine/Plugins/Runtime/StateTree/Source/StateTreeEditorModule/Private/StateTreeEditorStyle.cpp
+++ b/Engine/Plugins/Runtime/StateTree/Source/StateTreeEditorModule/Private/StateTreeEditorStyle.cpp
@@ -44,6 +44,12 @@ FStateTreeEditorStyle::FStateTreeEditorStyle()
 {
 	const FVector2f Icon8x8(8.0f, 8.0f);
 	
+	const FString EngineEditorSlateDir = FPaths::EngineContentDir() / TEXT("Slate");
+	SetCoreContentRoot(EngineEditorSlateDir);
+
+	const FString StateTreePluginContentDir = FPaths::EnginePluginsDir() / TEXT("Runtime/StateTree/Resources");
+	SetContentRoot(StateTreePluginContentDir);
+
 	const FScrollBarStyle ScrollBar = FAppStyle::GetWidgetStyle<FScrollBarStyle>("ScrollBar");
 	const FTextBlockStyle& NormalText = FAppStyle::Get().GetWidgetStyle<FTextBlockStyle>("NormalText");
 
