StateTree: Fix schema validation when pasting tasks

diff --git a/Engine/Plugins/Runtime/StateTree/Source/StateTreeEditorModule/Private/Customizations/StateTreeEditorNodeDetails.cpp b/Engine/Plugins/Runtime/StateTree/Source/StateTreeEditorModule/Private/Customizations/StateTreeEditorNodeDetails.cpp
--- a/Engine/Plugins/Runtime/StateTree/Source/StateTreeEditorModule/Private/Customizations/StateTreeEditorNodeDetails.cpp
+++ b/Engine/Plugins/Runtime/StateTree/Source/StateTreeEditorModule/Private/Customizations/StateTreeEditorNodeDetails.cpp
@@ -508,10 +508,12 @@ void FStateTreeEditorNodeDetails::OnPasteNode()
 		return;
 	}
 	
-	// Reject nodes that are not allowed by the schema. 
+	// Reject nodes that are not allowed by the schema.
 	const UStateTreeSchema* Schema = EditorData ? EditorData->Schema : nullptr;
-	if (!Schema->IsStructAllowed(TempNode.Node.GetScriptStruct()))
+	if (Schema)
 	{
+		bool bNodeIsAllowed = false;
+		
 		// BP nodes are identified by the instance type.
 		if (NodeTypeStruct->IsChildOf(FStateTreeBlueprintEvaluatorWrapper::StaticStruct())
 			|| NodeTypeStruct->IsChildOf(FStateTreeBlueprintTaskWrapper::StaticStruct())
@@ -519,16 +521,28 @@ void FStateTreeEditorNodeDetails::OnPasteNode()
 		{
 			if (const FStateTreeNodeBase* Node = TempNode.Node.GetPtr<FStateTreeNodeBase>())
 			{
-				NodeTypeStruct = Node->GetInstanceDataType();
+				NodeTypeStruct = Node->GetInstanceDataType(); // Report error with the BP node type, as that is what the user expects to see.
+				if (const UClass* InstanceClass = Cast<UClass>(NodeTypeStruct))
+				{
+					bNodeIsAllowed = Schema->IsClassAllowed(InstanceClass);
+				}
 			}
 		}
+		else
+		{
+			bNodeIsAllowed = Schema->IsStructAllowed(TempNode.Node.GetScriptStruct());
+		}
+		
+		if (!bNodeIsAllowed)
+		{
 
-		FNotificationInfo NotificationInfo(FText::GetEmpty());
-		NotificationInfo.Text = FText::Format(LOCTEXT("NotSupportedBySchema", "Node {0} is not supported by {1} schema."),
-									NodeTypeStruct->GetDisplayNameText(), Schema->GetClass()->GetDisplayNameText());
-		NotificationInfo.ExpireDuration = 5.0f;
-		FSlateNotificationManager::Get().AddNotification(NotificationInfo);
-		return;
+			FNotificationInfo NotificationInfo(FText::GetEmpty());
+			NotificationInfo.Text = FText::Format(LOCTEXT("NotSupportedBySchema", "Node {0} is not supported by {1} schema."),
+										NodeTypeStruct->GetDisplayNameText(), Schema->GetClass()->GetDisplayNameText());
+			NotificationInfo.ExpireDuration = 5.0f;
+			FSlateNotificationManager::Get().AddNotification(NotificationInfo);
+			return;
+		}
 	}
 	
 	FScopedTransaction Transaction(LOCTEXT("PasteNode", "Paste Node"));
