StateTree: Fixed crashes when compiled data is not loaded properly (e.g. some structs not loaded)

diff --git a/Engine/Plugins/Runtime/StateTree/Source/StateTreeModule/Private/StateTree.cpp b/Engine/Plugins/Runtime/StateTree/Source/StateTreeModule/Private/StateTree.cpp
--- a/Engine/Plugins/Runtime/StateTree/Source/StateTreeModule/Private/StateTree.cpp
+++ b/Engine/Plugins/Runtime/StateTree/Source/StateTreeModule/Private/StateTree.cpp
@@ -217,10 +217,33 @@ bool UStateTree::Link()
 	{
 		if (!DefaultInstanceData.IsValid())
 		{
-			UE_LOG(LogStateTree, Error, TEXT("%s: StartTree does not have instance data. Please recompile the StateTree asset."), *GetName());
+			UE_LOG(LogStateTree, Error, TEXT("%s: StartTree does not have instance data. Please recompile the StateTree asset."), *GetFullName());
 			return false;
 		}
 
+		// Check that all nodes are valid.
+		for (FConstStructView Node : Nodes)
+		{
+			if (!Node.IsValid())
+			{
+				UE_LOG(LogStateTree, Error, TEXT("%s: State Tree asset was not properly loaded (missing node). See log for loading failures, or recompile the StateTree asset."), *GetFullName());
+				return false;
+			}
+		}
+
+		if (!DefaultInstanceData.AreAllInstancesValid())
+		{
+			UE_LOG(LogStateTree, Error, TEXT("%s: State Tree asset was not properly loaded (missing instance data). See log for loading failures, or recompile the StateTree asset."), *GetFullName());
+			return false;
+		}
+
+		if (!SharedInstanceData.AreAllInstancesValid())
+		{
+			UE_LOG(LogStateTree, Error, TEXT("%s: State Tree asset was not properly loaded (missing shared instance data). See log for loading failures, or recompile the StateTree asset."), *GetFullName());
+			return false;
+		}
+
+		
 		// Update property bag structs before resolving binding.
 		const TArrayView<FStateTreeBindableStructDesc> SourceStructs = PropertyBindings.SourceStructs;
 		const TArrayView<FStateTreePropertyCopyBatch> CopyBatches = PropertyBindings.CopyBatches;
@@ -259,7 +282,7 @@ bool UStateTree::Link()
 			{
 				if (State.ParameterInstanceIndex.IsValid() == false)
 				{
-					UE_LOG(LogStateTree, Error, TEXT("%s: Data for state '%s' is malformed. Please recompile the StateTree asset."), *GetName(), *State.Name.ToString());
+					UE_LOG(LogStateTree, Error, TEXT("%s: Data for state '%s' is malformed. Please recompile the StateTree asset."), *GetFullName(), *State.Name.ToString());
 					return false;
 				}
 
@@ -278,7 +301,7 @@ bool UStateTree::Link()
 				if (State.ParameterInstanceIndex.IsValid() == false
 					|| LinkedState.ParameterInstanceIndex.IsValid() == false)
 				{
-					UE_LOG(LogStateTree, Error, TEXT("%s: Data for state '%s' is malformed. Please recompile the StateTree asset."), *GetName(), *State.Name.ToString());
+					UE_LOG(LogStateTree, Error, TEXT("%s: Data for state '%s' is malformed. Please recompile the StateTree asset."), *GetFullName(), *State.Name.ToString());
 					return false;
 				}
 
@@ -289,7 +312,7 @@ bool UStateTree::Link()
 
 				if (LinkedStateParams.Parameters.GetPropertyBagStruct() != Params.Parameters.GetPropertyBagStruct())
 				{
-					UE_LOG(LogStateTree, Error, TEXT("%s: The parameters on state '%s' does not match the linked state parameters in state '%s'. Please recompile the StateTree asset."), *GetName(), *State.Name.ToString(), *LinkedState.Name.ToString());
+					UE_LOG(LogStateTree, Error, TEXT("%s: The parameters on state '%s' does not match the linked state parameters in state '%s'. Please recompile the StateTree asset."), *GetFullName(), *State.Name.ToString(), *LinkedState.Name.ToString());
 					return false;
 				}
 
@@ -321,7 +344,7 @@ bool UStateTree::Link()
 			const bool bLinkSucceeded = NodePtr->Link(Linker);
 			if (!bLinkSucceeded || Linker.GetStatus() == EStateTreeLinkerStatus::Failed)
 			{
-				UE_LOG(LogStateTree, Error, TEXT("%s: node '%s' failed to resolve its references."), *GetName(), *NodePtr->StaticStruct()->GetName());
+				UE_LOG(LogStateTree, Error, TEXT("%s: node '%s' failed to resolve its references."), *GetFullName(), *NodePtr->StaticStruct()->GetName());
 				return false;
 			}
 		}
@@ -362,7 +385,10 @@ TArray<FStateTreeMemoryUsage> UStateTree::CalculateEstimatedMemoryUsage() const
 	TArray<FStateTreeMemoryUsage> MemoryUsages;
 	TArray<TPair<int32, int32>> StateLinks;
 
-	if (States.IsEmpty() || !Nodes.IsValid() || !DefaultInstanceData.IsValid())
+	if (!bIsLinked
+		|| States.IsEmpty()
+		|| !Nodes.IsValid()
+		|| !DefaultInstanceData.IsValid())
 	{
 		return MemoryUsages;
 	}
@@ -424,16 +450,18 @@ TArray<FStateTreeMemoryUsage> UStateTree::CalculateEstimatedMemoryUsage() const
 		
 		for (int32 TaskIndex = CompactState.TasksBegin; TaskIndex < (CompactState.TasksBegin + CompactState.TasksNum); TaskIndex++)
 		{
-			const FStateTreeTaskBase& Task = Nodes[TaskIndex].Get<const FStateTreeTaskBase>();
-			if (Task.bInstanceIsObject == false)
-			{
-				MemUsage.NodeCount++;
-				MemUsage.AddUsage(DefaultInstanceData.GetStruct(Task.InstanceIndex.Get()));
-			}
-			else
+			if (const FStateTreeTaskBase* Task = Nodes[TaskIndex].GetPtr<const FStateTreeTaskBase>())
 			{
-				MemUsage.NodeCount++;
-				MemUsage.AddUsage(DefaultInstanceData.GetObject(Task.InstanceIndex.Get()));
+				if (Task->bInstanceIsObject)
+				{
+					MemUsage.NodeCount++;
+					MemUsage.AddUsage(DefaultInstanceData.GetObject(Task->InstanceIndex.Get()));
+				}
+				else
+				{
+					MemUsage.NodeCount++;
+					MemUsage.AddUsage(DefaultInstanceData.GetStruct(Task->InstanceIndex.Get()));
+				}
 			}
 		}
 	}
diff --git a/Engine/Plugins/Runtime/StateTree/Source/StateTreeModule/Private/StateTreeInstanceData.cpp b/Engine/Plugins/Runtime/StateTree/Source/StateTreeModule/Private/StateTreeInstanceData.cpp
--- a/Engine/Plugins/Runtime/StateTree/Source/StateTreeModule/Private/StateTreeInstanceData.cpp
+++ b/Engine/Plugins/Runtime/StateTree/Source/StateTreeModule/Private/StateTreeInstanceData.cpp
@@ -71,6 +71,25 @@ void FStateTreeInstanceStorage::ResetTransitionRequests()
 	TransitionRequests.Reset();
 }
 
+bool FStateTreeInstanceStorage::AreAllInstancesValid() const
+{
+	for (FConstStructView Instance : InstanceStructs)
+	{
+		if (!Instance.IsValid())
+		{
+			return false;
+		}
+	}
+	for (const UObject* Instance : InstanceObjects)
+	{
+		if (!Instance)
+		{
+			return false;
+		}
+	}
+	return true;
+}
+
 
 //----------------------------------------------------------------//
 // FStateTreeInstanceData
@@ -143,6 +162,11 @@ void FStateTreeInstanceData::ResetTransitionRequests()
 	GetMutableStorage().ResetTransitionRequests();
 }
 
+bool FStateTreeInstanceData::AreAllInstancesValid() const
+{
+	return GetStorage().AreAllInstancesValid();
+}
+
 int32 FStateTreeInstanceData::GetEstimatedMemoryUsage() const
 {
 	const FStateTreeInstanceStorage& Storage = GetStorage();
diff --git a/Engine/Plugins/Runtime/StateTree/Source/StateTreeModule/Public/StateTreeInstanceData.h b/Engine/Plugins/Runtime/StateTree/Source/StateTreeModule/Public/StateTreeInstanceData.h
--- a/Engine/Plugins/Runtime/StateTree/Source/StateTreeModule/Public/StateTreeInstanceData.h
+++ b/Engine/Plugins/Runtime/StateTree/Source/StateTreeModule/Public/StateTreeInstanceData.h
@@ -72,7 +72,10 @@ struct STATETREEMODULE_API FStateTreeInstanceStorage
 
 	/** Reset all pending transition requests. */
 	void ResetTransitionRequests();
-	
+
+	/** @return true if all instances are valid. */
+	bool AreAllInstancesValid() const;
+
 protected:
 	/** Struct instances */
 	UPROPERTY()
@@ -174,6 +177,9 @@ struct STATETREEMODULE_API FStateTreeInstanceData
 	/** Reset all pending transition requests. */
 	void ResetTransitionRequests();
 
+	/** @return true if all instances are valid. */
+	bool AreAllInstancesValid() const;
+
 	FStateTreeInstanceStorage& GetMutableStorage();
 	const FStateTreeInstanceStorage& GetStorage() const;
 
