State Tree: fix delayed transitions
- Fixed bug where delayed transitions from events were not triggered properly.

diff --git a/Engine/Plugins/Runtime/StateTree/Source/StateTreeModule/Private/StateTreeExecutionContext.cpp b/Engine/Plugins/Runtime/StateTree/Source/StateTreeModule/Private/StateTreeExecutionContext.cpp
--- a/Engine/Plugins/Runtime/StateTree/Source/StateTreeModule/Private/StateTreeExecutionContext.cpp
+++ b/Engine/Plugins/Runtime/StateTree/Source/StateTreeModule/Private/StateTreeExecutionContext.cpp
@@ -1415,51 +1415,70 @@ bool FStateTreeExecutionContext::TriggerTransitions()
 			{
 				continue;
 			}
-			
-			if (Transition.Trigger == EStateTreeTransitionTrigger::OnTick
-				|| (Transition.Trigger == EStateTreeTransitionTrigger::OnEvent && HasEventToProcess(Transition.EventTag)))
+
+			// Skip completion transitions
+			if (EnumHasAnyFlags(Transition.Trigger, EStateTreeTransitionTrigger::OnStateCompleted))
 			{
-				// If delay has passed, allow to trigger the transition even if conditions do not pass.
-				FStateTreeTransitionDelayedState* DelayedState = Transition.HasDelay() ? Exec.FindDelayedTransition(FStateTreeIndex16(TransitionIndex)) : nullptr;
-				const bool bDelayCompleted = DelayedState != nullptr && DelayedState->TimeLeft <= 0.0f;
+				continue;
+			}
 
-				if (bDelayCompleted || TestAllConditions(Transition.ConditionsBegin, Transition.ConditionsNum))
+			// If a delayed transition has passed the delay, and remove it from the queue, and try trigger it.
+			FStateTreeTransitionDelayedState* DelayedState = nullptr;
+			if (Transition.HasDelay())
+			{
+				DelayedState = Exec.FindDelayedTransition(FStateTreeIndex16(TransitionIndex));
+				if (DelayedState != nullptr && DelayedState->TimeLeft <= 0.0f)
 				{
-					// If the transitions is delayed, set up the delay. 
-					if (Transition.HasDelay())
+					STATETREE_LOG(Verbose, TEXT("Passed delayed transition from '%s' (%s) -> '%s'"),
+						*GetSafeStateName(Exec.ActiveStates.Last()), *State.Name.ToString(), *GetSafeStateName(Transition.State));
+
+					Exec.DelayedTransitions.RemoveAllSwap([TransitionIndex](const FStateTreeTransitionDelayedState& DelayedState)
+						{
+							return DelayedState.TransitionIndex.Get() == TransitionIndex;
+						});
+
+					// Trigger Delayed Transition when the delay has passed.
+					RequestTransition(Transition.State, Transition.Priority);
+					continue;
+				}
+			}
+
+			const bool bShouldTrigger = Transition.Trigger == EStateTreeTransitionTrigger::OnTick
+										|| (Transition.Trigger == EStateTreeTransitionTrigger::OnEvent
+											&& HasEventToProcess(Transition.EventTag));
+				
+			if (bShouldTrigger
+				&& TestAllConditions(Transition.ConditionsBegin, Transition.ConditionsNum))
+			{
+				// If the transitions is delayed, set up the delay. 
+				if (Transition.HasDelay())
+				{
+					if (DelayedState == nullptr)
 					{
-						if (DelayedState == nullptr)
+						// Initialize new delayed transition.
+						const float DelayDuration = Transition.Delay.GetRandomDuration();
+						if (DelayDuration > 0.0f)
 						{
 							DelayedState = &Exec.DelayedTransitions.AddDefaulted_GetRef();
-							
 							DelayedState->TransitionIndex = FStateTreeIndex16(TransitionIndex);
-							DelayedState->TimeLeft = Transition.Delay.GetRandomDuration();
+							DelayedState->TimeLeft = DelayDuration;
 							BeginDelayedTransition(*DelayedState);
 							STATETREE_LOG(Verbose, TEXT("Delayed transition triggered from '%s' (%s) -> '%s' %.1fs"),
 								*GetSafeStateName(Exec.ActiveStates.Last()), *State.Name.ToString(), *GetSafeStateName(Transition.State), DelayedState->TimeLeft);
-							// Fall through to handle 0 delay.
-						}
-						check(DelayedState);
-
-						// We get here if the transitions re-triggers, or delay has completed.
-						// In case of re-trigger, we will just ignore it during the delay. 
-						if (DelayedState->TimeLeft > 0.0f)
-						{
+							
+							// Delay state added, skip requesting the transition.
 							continue;
 						}
-
-						STATETREE_LOG(Verbose, TEXT("Passed delayed transition from '%s' (%s) -> '%s'"),
-							*GetSafeStateName(Exec.ActiveStates.Last()), *State.Name.ToString(), *GetSafeStateName(Transition.State));
-
-						// The transition passed the delay, and remove it from the queue, and try trigger it.
-						Exec.DelayedTransitions.RemoveAllSwap([TransitionIndex](const FStateTreeTransitionDelayedState& DelayedState)
-							{
-								return DelayedState.TransitionIndex.Get() == TransitionIndex;
-							});
+						// Fallthrough to request transition if duration was zero. 
+					}
+					else
+					{
+						// We get here if the transitions re-triggers during the delay, on which case we'll just ignore it.
+						continue;
 					}
-
-					RequestTransition(Transition.State, Transition.Priority);
 				}
+
+				RequestTransition(Transition.State, Transition.Priority);
 			}
 		}
 	}
diff --git a/Engine/Plugins/Runtime/StateTree/Source/StateTreeTestSuite/Private/StateTreeTest.cpp b/Engine/Plugins/Runtime/StateTree/Source/StateTreeTestSuite/Private/StateTreeTest.cpp
--- a/Engine/Plugins/Runtime/StateTree/Source/StateTreeTestSuite/Private/StateTreeTest.cpp
+++ b/Engine/Plugins/Runtime/StateTree/Source/StateTreeTestSuite/Private/StateTreeTest.cpp
@@ -9,6 +9,7 @@
 #include "StateTreeTestTypes.h"
 #include "Engine/World.h"
 #include "Async/ParallelFor.h"
+#include "GameplayTagsManager.h"
 
 #include UE_INLINE_GENERATED_CPP_BY_NAME(StateTreeTest)
 
@@ -31,6 +32,27 @@ namespace UE::StateTree::Tests
 		return *StateTree;
 	}
 
+	// Helper struct to define some test tags
+	struct FNativeGameplayTags : public FGameplayTagNativeAdder
+	{
+		virtual ~FNativeGameplayTags() {}
+		
+		FGameplayTag TestTag;
+
+		virtual void AddTags() override
+		{
+			UGameplayTagsManager& Manager = UGameplayTagsManager::Get();
+			TestTag = Manager.AddNativeGameplayTag(TEXT("Test.StateTree.Tag"));
+		}
+
+		FORCEINLINE static const FNativeGameplayTags& Get()
+		{
+			return StaticInstance;
+		}
+		static FNativeGameplayTags StaticInstance;
+	};
+	FNativeGameplayTags FNativeGameplayTags::StaticInstance;
+
 }
 
 struct FStateTreeTest_MakeAndBakeStateTree : FAITestBase
@@ -692,6 +714,145 @@ struct FStateTreeTest_TransitionGlobalDataView : FAITestBase
 };
 IMPLEMENT_AI_INSTANT_TEST(FStateTreeTest_TransitionGlobalDataView, "System.StateTree.Transition.GlobalDataView");
 
+struct FStateTreeTest_TransitionDelay : FAITestBase
+{
+	virtual bool InstantTest() override
+	{
+		UStateTree& StateTree = UE::StateTree::Tests::NewStateTree(&GetWorld());
+		UStateTreeEditorData& EditorData = *Cast<UStateTreeEditorData>(StateTree.EditorData);
+		const FGameplayTag Tag = UE::StateTree::Tests::FNativeGameplayTags::Get().TestTag;
+
+		UStateTreeState& Root = EditorData.AddSubTree(FName(TEXT("Root")));
+		UStateTreeState& StateA = Root.AddChildState(FName(TEXT("A")));
+		UStateTreeState& StateB = Root.AddChildState(FName(TEXT("B")));
+
+		// State A
+		auto& Task0 = StateA.AddTask<FTestTask_Stand>(FName(TEXT("Task0")));
+		Task0.GetNode().TicksToCompletion = 100;
+		
+		FStateTreeTransition& Transition = StateA.AddTransition(EStateTreeTransitionTrigger::OnEvent, EStateTreeTransitionType::GotoState, &StateB);
+		Transition.bDelayTransition = true;
+		Transition.DelayDuration = 0.15f;
+		Transition.DelayRandomVariance = 0.0f;
+		Transition.EventTag = Tag;
+
+		// State B
+		auto& Task1 = StateB.AddTask<FTestTask_Stand>(FName(TEXT("Task1")));
+		Task1.GetNode().TicksToCompletion = 100;
+
+		FStateTreeCompilerLog Log;
+		FStateTreeCompiler Compiler(Log);
+		const bool bResult = Compiler.Compile(StateTree);
+
+		AITEST_TRUE("StateTree should get compiled", bResult);
+
+		EStateTreeRunStatus Status = EStateTreeRunStatus::Unset;
+		FStateTreeInstanceData InstanceData;
+		FTestStateTreeExecutionContext Exec(StateTree, StateTree, InstanceData);
+		const bool bInitSucceeded = Exec.IsValid();
+		AITEST_TRUE("StateTree should init", bInitSucceeded);
+
+		const FString TickStr(TEXT("Tick"));
+		const FString EnterStateStr(TEXT("EnterState"));
+		const FString ExitStateStr(TEXT("ExitState"));
+		const FString StateCompletedStr(TEXT("StateCompleted"));
+
+		// Start and enter state
+		Status = Exec.Start();
+		AITEST_TRUE("StateTree Task0 should enter state", Exec.Expect(Task0.GetName(), EnterStateStr));
+		Exec.LogClear();
+
+		// This should cause delayed transition.
+		Exec.SendEvent(Tag);
+		
+		Status = Exec.Tick(0.1f);
+		AITEST_TRUE("StateTree Task0 should tick", Exec.Expect(Task0.GetName(), TickStr));
+		Exec.LogClear();
+
+		// Should have delayed transitions
+		const int32 NumDelayedTransitions0 = InstanceData.GetStruct(0).Get<const FStateTreeExecutionState>().DelayedTransitions.Num();
+		AITEST_EQUAL("Should have a delayed transition", NumDelayedTransitions0, 1);
+
+		// Tick and expect a delayed transition. 
+		Status = Exec.Tick(0.1f);
+		AITEST_TRUE("StateTree Task0 should tick", Exec.Expect(Task0.GetName(), TickStr));
+		Exec.LogClear();
+
+		const int32 NumDelayedTransitions1 = InstanceData.GetStruct(0).Get<const FStateTreeExecutionState>().DelayedTransitions.Num();
+		AITEST_EQUAL("Should have a delayed transition", NumDelayedTransitions1, 1);
+
+		// Should complete delayed transition.
+		Status = Exec.Tick(0.1f);
+		AITEST_TRUE("StateTree Task0 should exit state", Exec.Expect(Task0.GetName(), ExitStateStr));
+		AITEST_TRUE("StateTree Task1 should enter state", Exec.Expect(Task1.GetName(), EnterStateStr));
+		Exec.LogClear();
+
+		return true;
+	}
+};
+IMPLEMENT_AI_INSTANT_TEST(FStateTreeTest_TransitionDelay, "System.StateTree.TransitionDelay");
+
+struct FStateTreeTest_TransitionDelayZero : FAITestBase
+{
+	virtual bool InstantTest() override
+	{
+		UStateTree& StateTree = UE::StateTree::Tests::NewStateTree(&GetWorld());
+		UStateTreeEditorData& EditorData = *Cast<UStateTreeEditorData>(StateTree.EditorData);
+		const FGameplayTag Tag = UE::StateTree::Tests::FNativeGameplayTags::Get().TestTag;
+
+		UStateTreeState& Root = EditorData.AddSubTree(FName(TEXT("Root")));
+		UStateTreeState& StateA = Root.AddChildState(FName(TEXT("A")));
+		UStateTreeState& StateB = Root.AddChildState(FName(TEXT("B")));
+
+		// State A
+		auto& Task0 = StateA.AddTask<FTestTask_Stand>(FName(TEXT("Task0")));
+		Task0.GetNode().TicksToCompletion = 100;
+		
+		FStateTreeTransition& Transition = StateA.AddTransition(EStateTreeTransitionTrigger::OnEvent, EStateTreeTransitionType::GotoState, &StateB);
+		Transition.bDelayTransition = true;
+		Transition.DelayDuration = 0.0f;
+		Transition.DelayRandomVariance = 0.0f;
+		Transition.EventTag = Tag;
+
+		// State B
+		auto& Task1 = StateB.AddTask<FTestTask_Stand>(FName(TEXT("Task1")));
+		Task1.GetNode().TicksToCompletion = 100;
+
+		FStateTreeCompilerLog Log;
+		FStateTreeCompiler Compiler(Log);
+		const bool bResult = Compiler.Compile(StateTree);
+
+		AITEST_TRUE("StateTree should get compiled", bResult);
+
+		EStateTreeRunStatus Status = EStateTreeRunStatus::Unset;
+		FStateTreeInstanceData InstanceData;
+		FTestStateTreeExecutionContext Exec(StateTree, StateTree, InstanceData);
+		const bool bInitSucceeded = Exec.IsValid();
+		AITEST_TRUE("StateTree should init", bInitSucceeded);
+
+		const FString TickStr(TEXT("Tick"));
+		const FString EnterStateStr(TEXT("EnterState"));
+		const FString ExitStateStr(TEXT("ExitState"));
+		const FString StateCompletedStr(TEXT("StateCompleted"));
+
+		// Start and enter state
+		Status = Exec.Start();
+		AITEST_TRUE("StateTree Task0 should enter state", Exec.Expect(Task0.GetName(), EnterStateStr));
+		Exec.LogClear();
+
+		// This should cause delayed transition. Because the time is 0, it should happen immediately.
+		Exec.SendEvent(Tag);
+		
+		Status = Exec.Tick(0.1f);
+		AITEST_TRUE("StateTree Task0 should exit state", Exec.Expect(Task0.GetName(), ExitStateStr));
+		AITEST_TRUE("StateTree Task1 should enter state", Exec.Expect(Task1.GetName(), EnterStateStr));
+		Exec.LogClear();
+
+		return true;
+	}
+};
+IMPLEMENT_AI_INSTANT_TEST(FStateTreeTest_TransitionDelayZero, "System.StateTree.TransitionDelayZero");
+
 struct FStateTreeTest_PropertyPathOffset : FAITestBase
 {
 	virtual bool InstantTest() override
