[StateTreeDebugger]
Initial version of the debugger
- currently compiled by StateTreeModule and StateTreeEditorModule using WITH_STATETREE_DEBUGGER
- currently not exposed to UI by StateTree settings 'bUseDebugger'
- in this version only one instance per asset can be debugged
- using Trace services to read events generated by statetree instances.
- can connect to any traces (Editor, Client, Server) as long as the compiled statetree matches

diff --git a/Engine/Plugins/Runtime/StateTree/Source/StateTreeEditorModule/Private/SStateTreeDebuggerView.cpp b/Engine/Plugins/Runtime/StateTree/Source/StateTreeEditorModule/Private/SStateTreeDebuggerView.cpp
new file mode 100644
--- /dev/null
+++ b/Engine/Plugins/Runtime/StateTree/Source/StateTreeEditorModule/Private/SStateTreeDebuggerView.cpp
@@ -0,0 +1,443 @@
+// Copyright Epic Games, Inc. All Rights Reserved.
+
+#if WITH_STATETREE_DEBUGGER
+
+#include "SStateTreeDebuggerView.h"
+
+#include "Editor.h"
+#include "Framework/MultiBox/MultiBoxBuilder.h"
+#include "IMessageLogListing.h"
+#include "Kismet2/DebuggerCommands.h"
+#include "MessageLogModule.h"
+#include "Modules/ModuleManager.h"
+#include "StateTree.h"
+#include "Debugger/StateTreeDebugger.h"
+#include "Debugger/StateTreeTraceTypes.h"
+#include "StateTreeDebuggerCommands.h"
+#include "StateTreeState.h"
+#include "StateTreeTaskBase.h"
+#include "StateTreeViewModel.h"
+#include "Widgets/Input/SComboButton.h"
+#include "Widgets/Layout/SBorder.h"
+
+#define LOCTEXT_NAMESPACE "StateTreeEditor"
+
+SStateTreeDebuggerView::SStateTreeDebuggerView()
+{
+	FEditorDelegates::BeginPIE.AddRaw(this, &SStateTreeDebuggerView::OnPIEStarted);
+	FEditorDelegates::EndPIE.AddRaw(this, &SStateTreeDebuggerView::OnPIEStopped);
+	FEditorDelegates::PausePIE.AddRaw(this, &SStateTreeDebuggerView::OnPIEPaused);
+	FEditorDelegates::ResumePIE.AddRaw(this, &SStateTreeDebuggerView::OnPIEResumed);
+	FEditorDelegates::SingleStepPIE.AddRaw(this, &SStateTreeDebuggerView::OnPIESingleStepped);
+}
+
+SStateTreeDebuggerView::~SStateTreeDebuggerView()
+{
+	FEditorDelegates::BeginPIE.RemoveAll(this);
+	FEditorDelegates::EndPIE.RemoveAll(this);
+	FEditorDelegates::PausePIE.RemoveAll(this);
+	FEditorDelegates::ResumePIE.RemoveAll(this);
+	FEditorDelegates::SingleStepPIE.RemoveAll(this);
+}
+
+void SStateTreeDebuggerView::OnPIEStarted(const bool bIsSimulating) const
+{
+	if (DebuggerTraceListing)
+	{
+		DebuggerTraceListing->ClearMessages();
+	}
+
+	if (!Debugger->IsAnalysisSessionActive())
+	{
+		Debugger->StartLastLiveSessionAnalysis();
+	}
+}
+
+void SStateTreeDebuggerView::OnPIEStopped(const bool bIsSimulating) const
+{
+	Debugger->Unpause();
+}
+
+void SStateTreeDebuggerView::OnPIEPaused(const bool bIsSimulating) const
+{
+	Debugger->Pause();
+}
+
+void SStateTreeDebuggerView::OnPIEResumed(const bool bIsSimulating) const
+{
+	Debugger->Unpause();
+}
+
+void SStateTreeDebuggerView::OnPIESingleStepped(bool bSimulating) const
+{
+	Debugger->SyncToCurrentSessionDuration();
+}
+
+void SStateTreeDebuggerView::Construct(const FArguments& InArgs, const UStateTree* InStateTree, const TSharedRef<FStateTreeViewModel> InStateTreeViewModel, const TSharedRef<FUICommandList> InCommandList)
+{
+	StateTreeViewModel = InStateTreeViewModel;
+	StateTree = InStateTree;
+	
+	Debugger = InStateTreeViewModel->GetDebugger();
+
+	// Put debugger in proper simulation state when view is constructed after PIE/SIE was started
+	if (FPlayWorldCommandCallbacks::HasPlayWorldAndPaused())
+	{
+		Debugger->Pause();
+	}
+
+	BindDebuggerToolbarCommands(InCommandList);
+
+	FMessageLogModule& MessageLogModule = FModuleManager::LoadModuleChecked<FMessageLogModule>("MessageLog");
+	FMessageLogInitializationOptions LogOptions;
+	LogOptions.bShowPages = false;
+	LogOptions.bShowFilters = false;
+	LogOptions.bAllowClear = false;
+	LogOptions.bScrollToBottom = true;
+	LogOptions.MaxPageCount = 1;
+
+	DebuggerTraceListing = MessageLogModule.CreateLogListing("StateTreeDebugger", LogOptions);
+	DebuggerTraces = MessageLogModule.CreateLogListingWidget(DebuggerTraceListing.ToSharedRef());
+
+	const TSharedRef<FUICommandList> CommandList = MakeShareable(new FUICommandList);
+
+	// Register the play world commands
+	CommandList->Append(FPlayWorldCommands::GlobalPlayWorldActions.ToSharedRef());
+	
+	CommandList->MapAction(
+		FStateTreeDebuggerCommands::Get().ToggleBreakpoint,
+		FExecuteAction::CreateSP(this, &SStateTreeDebuggerView::ToggleBreakpoint),
+		FCanExecuteAction::CreateSP(this, &SStateTreeDebuggerView::CanToggleBreakpoint)
+	);
+	
+	const TSharedRef<SWidget> TraceSelectionBox = SNew(SComboButton)
+		.OnGetMenuContent(this, &SStateTreeDebuggerView::OnGetDebuggerTracesMenu)
+		.ButtonContent()
+		[
+			SNew(STextBlock)
+			.ToolTipText(LOCTEXT("SelectTraceSession", "Pick trace session to debug"))
+			.Text_Lambda([Debugger = Debugger]()
+			{
+				return Debugger.IsValid() ? FText::FromString(Debugger->GetDebuggedTraceDescription()) : FText::GetEmpty();
+			})
+		];
+	
+	const TSharedRef<SWidget> InstanceSelectionBox = SNew(SComboButton)
+		.OnGetMenuContent(this, &SStateTreeDebuggerView::OnGetDebuggerInstancesMenu)
+		.ButtonContent()
+		[
+			SNew(STextBlock)
+			.ToolTipText(LOCTEXT("SelectDebugInstance", "Pick instance to debug"))
+			.Text_Lambda([Debugger = Debugger]()
+			{
+				return Debugger.IsValid() ? FText::FromString(Debugger->GetDebuggedInstanceDescription()) : FText::GetEmpty();
+			})
+		];
+	
+	// Build debug toolbar
+	FToolBarBuilder ToolbarBuilder(CommandList, FMultiBoxCustomization::None);
+	ToolbarBuilder.BeginSection(TEXT("Debugging"));
+	{
+		const FPlayWorldCommands& PlayWorldCommand = FPlayWorldCommands::Get();
+		ToolbarBuilder.AddToolBarButton(PlayWorldCommand.RepeatLastPlay);
+		ToolbarBuilder.AddToolBarButton(PlayWorldCommand.PausePlaySession,
+					NAME_None, TAttribute<FText>(), TAttribute<FText>(), FSlateIcon(FAppStyle::GetAppStyleSetName(), "PlayWorld.PausePlaySession.Small"));
+		ToolbarBuilder.AddToolBarButton(PlayWorldCommand.ResumePlaySession,
+					NAME_None, TAttribute<FText>(), TAttribute<FText>(), FSlateIcon(FAppStyle::GetAppStyleSetName(), "PlayWorld.ResumePlaySession.Small"));	
+		ToolbarBuilder.AddToolBarButton(PlayWorldCommand.StopPlaySession,
+					NAME_None, TAttribute<FText>(), TAttribute<FText>(), FSlateIcon(FAppStyle::GetAppStyleSetName(), "PlayWorld.StopPlaySession.Small"));
+		ToolbarBuilder.AddSeparator();
+		ToolbarBuilder.AddToolBarButton(FStateTreeDebuggerCommands::Get().Back);
+		ToolbarBuilder.AddToolBarButton(FStateTreeDebuggerCommands::Get().Forward);
+		ToolbarBuilder.AddSeparator();
+		ToolbarBuilder.AddWidget(TraceSelectionBox);
+		ToolbarBuilder.AddSeparator();
+		ToolbarBuilder.AddWidget(InstanceSelectionBox);
+	}
+	ToolbarBuilder.EndSection();
+
+	InCommandList->Append(CommandList);
+
+	TSharedPtr<FUICommandList> ActionList = InCommandList;
+		Debugger->OnTracesUpdated.BindSP(this, &SStateTreeDebuggerView::OnTracesUpdated);
+	Debugger->OnBreakpointHit.BindSP(this, &SStateTreeDebuggerView::OnBreakpointHit, ActionList);
+
+	ChildSlot
+	[
+		SNew(SBorder)
+		.Padding(4.0f)
+		.BorderImage(FAppStyle::GetBrush("ToolPanel.GroupBorder"))
+		[
+			SNew(SVerticalBox)
+			+ SVerticalBox::Slot()
+			.AutoHeight()
+			.Padding(0.0f, 0.0f, 0.0f, 4.0f)
+			[
+				SNew(SHorizontalBox)
+				+ SHorizontalBox::Slot()
+				.AutoWidth()
+				[
+					ToolbarBuilder.MakeWidget()
+				]
+			]
+			+ SVerticalBox::Slot()
+			.FillHeight(1.0f)
+			.Padding(0.0f, 0.0f, 0.0f, 4.0f)
+			[
+				SNew(SBox)
+				[
+					DebuggerTraces.ToSharedRef()
+				]
+			]
+		]
+	];
+}
+
+void SStateTreeDebuggerView::BindDebuggerToolbarCommands(const TSharedRef<FUICommandList> ToolkitCommands)
+{
+	const FStateTreeDebuggerCommands& Commands = FStateTreeDebuggerCommands::Get();
+	
+	ToolkitCommands->MapAction(
+		Commands.Back,
+		FExecuteAction::CreateSP(this, &SStateTreeDebuggerView::StepBack),
+		FCanExecuteAction::CreateSP(this, &SStateTreeDebuggerView::CanStepBack));
+
+	ToolkitCommands->MapAction(
+		Commands.Forward,
+		FExecuteAction::CreateSP(this, &SStateTreeDebuggerView::StepForward),
+		FCanExecuteAction::CreateSP(this, &SStateTreeDebuggerView::CanStepForward));
+}
+
+bool SStateTreeDebuggerView::CanStepBack() const
+{
+	return Debugger->CanStepBack();
+}
+
+void SStateTreeDebuggerView::StepBack() const
+{
+	Debugger->StepBack();
+}
+
+bool SStateTreeDebuggerView::CanStepForward() const
+{
+	return Debugger->CanStepForward();
+}
+
+void SStateTreeDebuggerView::StepForward() const
+{
+	Debugger->StepForward();
+}
+
+bool SStateTreeDebuggerView::CanToggleBreakpoint() const
+{
+	return (Debugger.IsValid() && StateTreeViewModel.IsValid() && StateTreeViewModel->HasSelection());
+}
+
+void SStateTreeDebuggerView::ToggleBreakpoint() const
+{
+	check(StateTreeViewModel);
+	check(Debugger);
+	check(StateTree.IsValid());
+
+	TArray<UStateTreeState*> States;
+	StateTreeViewModel->GetSelectedStates(States);
+
+	TArray<FStateTreeStateHandle> StateHandles;
+	StateHandles.Reserve(States.Num());
+	for (const UStateTreeState* SelectedState : States)
+	{
+		if (SelectedState->Type != EStateTreeStateType::State || SelectedState->Parent == nullptr)
+		{
+			continue;
+		}
+		
+		FStateTreeStateHandle Handle = StateTree->GetStateHandleFromId(SelectedState->ID);
+		if (Handle.IsValid())
+		{
+			StateHandles.Add(Handle);
+		}
+	}
+	Debugger->ToggleBreakpoints(StateHandles);
+}
+
+void SStateTreeDebuggerView::OnTracesUpdated(
+	const TConstArrayView<const FStateTreeTraceEventVariantType> Events,
+	const TConstArrayView<UE::StateTreeDebugger::FFrameIndexSpan> Spans
+	) const
+{
+	DebuggerTraceListing->ClearMessages();
+	
+	if (Events.IsEmpty())
+	{
+		return;
+	}
+
+	check(Spans.Num());
+
+	TArray<TSharedRef<FTokenizedMessage>> Messages;
+	Messages.Reserve(Events.Num());
+
+	check(StateTree.IsValid());
+
+	int32 SpanIdx = INDEX_NONE;
+	for (int32 EventIdx = 0; EventIdx < Events.Num(); EventIdx++)
+	{
+		const FStateTreeTraceEventVariantType& Event = Events[EventIdx];
+		
+		if (SpanIdx == INDEX_NONE
+			|| (Spans.IsValidIndex(SpanIdx+1) && EventIdx >= Spans[SpanIdx+1].EventIdx))
+		{
+			SpanIdx = SpanIdx == INDEX_NONE ? 0 : SpanIdx+1;
+			Messages.Add(FTokenizedMessage::Create(EMessageSeverity::Info,
+				FText::FromString(FString::Printf(TEXT("---------------- Frame: %s ----------------"), *LexToString(Spans[SpanIdx].FrameIdx)))));
+		}
+		
+		// Add log event messages directly
+		if (const FStateTreeTraceLogEvent* LogEvent = Event.TryGet<FStateTreeTraceLogEvent>())
+		{
+			if (LogEvent->Message.Len())
+			{
+				Messages.Add(FTokenizedMessage::Create(EMessageSeverity::Info, FText::FromString(LogEvent->Message)));	
+			}			
+			continue;	
+		}
+
+		// Add messages for state events 
+		if (const FStateTreeTraceStateEvent* StateEvent = Event.TryGet<FStateTreeTraceStateEvent>())
+		{
+			const FStateTreeStateHandle StateHandle(StateEvent->StateIdx);
+			if (const FCompactStateTreeState* CompactState = StateTree->GetStateFromHandle(StateHandle))
+			{
+				Messages.Add(FTokenizedMessage::Create(EMessageSeverity::Info, FText::FromString(FString::Printf(TEXT("%s State '%s'"),
+								*StaticEnum<EStateTreeTraceNodeEventType>()->GetNameStringByValue((int64)StateEvent->EventType),
+								*CompactState->Name.ToString()))));
+			}
+		}
+
+		if (const FStateTreeTraceTaskEvent* TaskEvent = Event.TryGet<FStateTreeTraceTaskEvent>())
+		{
+			FConstStructView NodeView = StateTree->GetNode(TaskEvent->TaskIdx);
+			const FStateTreeTaskBase* Task = NodeView.GetPtr<const FStateTreeTaskBase>();
+					
+			Messages.Add(FTokenizedMessage::Create(EMessageSeverity::Info, FText::FromString(FString::Printf(TEXT("  %s Task '%s'"),
+				*StaticEnum<EStateTreeTraceNodeEventType>()->GetNameStringByValue((int64)TaskEvent->EventType),
+				Task != nullptr ? *Task->Name.ToString() : *LexToString(TaskEvent->TaskIdx)))));
+		}
+	}
+
+	DebuggerTraceListing->AddMessages(Messages, /*bMirrorToOutputLog*/false);
+}
+
+void SStateTreeDebuggerView::OnBreakpointHit(const FStateTreeInstanceDebugId InstanceId, const FStateTreeStateHandle StateHandle, const TSharedPtr<FUICommandList> ActionList) const
+{
+	if (ActionList.IsValid() && FPlayWorldCommands::Get().PausePlaySession.IsValid())
+	{
+		if (ActionList->CanExecuteAction(FPlayWorldCommands::Get().PausePlaySession.ToSharedRef()))
+		{
+			// Make sure the instance is selected in case the breakpoint was set for all instances 
+			const FStateTreeDebuggerInstanceDesc& DebuggerInstance = Debugger->GetDebuggedInstance();
+			if (DebuggerInstance.Id != InstanceId)
+			{
+				TArray<FStateTreeDebuggerInstanceDesc> MatchingInstances;
+				Debugger->GetActivateInstances(MatchingInstances);
+				const FStateTreeDebuggerInstanceDesc* FoundDesc = MatchingInstances.FindByPredicate([InstanceId](const FStateTreeDebuggerInstanceDesc& InstanceDesc)
+				{
+					return InstanceDesc.Id == InstanceId;
+				});
+
+				if (FoundDesc != nullptr)
+				{
+					Debugger->SetDebuggedInstance(*FoundDesc);
+				}
+			}
+
+			ActionList->ExecuteAction(FPlayWorldCommands::Get().PausePlaySession.ToSharedRef());
+		}
+	}
+}
+
+TSharedRef<SWidget> SStateTreeDebuggerView::OnGetDebuggerTracesMenu() const
+{
+	FMenuBuilder MenuBuilder(/*bInShouldCloseWindowAfterMenuSelection*/true, /*InCommandList*/nullptr);
+
+	if (Debugger.IsValid())
+	{
+		TArray<FStateTreeDebugger::FTraceDescriptor> TraceDescriptors;
+		Debugger->GetLiveTraces(TraceDescriptors);
+
+		for (const FStateTreeDebugger::FTraceDescriptor& TraceDescriptor : TraceDescriptors)
+		{
+			const FText Desc = FText::FromString(Debugger->DescribeTrace(TraceDescriptor));
+
+			FUIAction ItemAction(FExecuteAction::CreateLambda([WeakDebugger = Debugger, TraceDescriptor]()
+			{
+				if (WeakDebugger)
+				{
+					WeakDebugger->StartSessionAnalysis(TraceDescriptor);
+				}
+			}));
+			MenuBuilder.AddMenuEntry(Desc, TAttribute<FText>(), FSlateIcon(), ItemAction);
+		}
+
+		// Failsafe when no match
+		if (TraceDescriptors.Num() == 0)
+		{
+			const FText Desc = LOCTEXT("NoLiveSessions","Can't find live trace sessions");
+			FUIAction ItemAction(FExecuteAction::CreateLambda([WeakDebugger = Debugger]()
+			{
+				if (WeakDebugger)
+				{
+					WeakDebugger->StartSessionAnalysis(FStateTreeDebugger::FTraceDescriptor());
+				}
+			}));
+			MenuBuilder.AddMenuEntry(Desc, TAttribute<FText>(), FSlateIcon(), ItemAction);
+		}
+	}
+
+	return MenuBuilder.MakeWidget();
+}
+
+TSharedRef<SWidget> SStateTreeDebuggerView::OnGetDebuggerInstancesMenu() const
+{
+	FMenuBuilder MenuBuilder(/*bInShouldCloseWindowAfterMenuSelection*/true, /*InCommandList*/nullptr);
+
+	if (Debugger.IsValid())
+	{
+		TArray<FStateTreeDebuggerInstanceDesc> ActiveInstances;
+		Debugger->GetActivateInstances(ActiveInstances); 
+
+		for (const FStateTreeDebuggerInstanceDesc& Instance : ActiveInstances)
+		{
+			const FText Desc = FText::FromString(Debugger->DescribeInstance(Instance));
+
+			FUIAction ItemAction(FExecuteAction::CreateLambda([WeakDebugger = Debugger, Instance]()
+			{
+				if (WeakDebugger)
+				{
+					WeakDebugger->SetDebuggedInstance(Instance);
+				}
+			}));
+			MenuBuilder.AddMenuEntry(Desc, TAttribute<FText>(), FSlateIcon(), ItemAction);
+		}
+
+		// Failsafe when no match
+		if (ActiveInstances.Num() == 0)
+		{
+			const FText Desc = LOCTEXT("NoMatchingInstance","Can't find matching instances");
+			FUIAction ItemAction(FExecuteAction::CreateLambda([WeakDebugger = Debugger]()
+			{
+				if (WeakDebugger)
+				{
+					WeakDebugger->SetDebuggedInstance(FStateTreeDebuggerInstanceDesc());
+				}
+			}));
+			MenuBuilder.AddMenuEntry(Desc, TAttribute<FText>(), FSlateIcon(), ItemAction);
+		}
+	}
+
+	return MenuBuilder.MakeWidget();
+}
+
+#undef LOCTEXT_NAMESPACE
+
+#endif // WITH_STATETREE_DEBUGGER
\ No newline at end of file
diff --git a/Engine/Plugins/Runtime/StateTree/Source/StateTreeEditorModule/Private/SStateTreeDebuggerView.h b/Engine/Plugins/Runtime/StateTree/Source/StateTreeEditorModule/Private/SStateTreeDebuggerView.h
new file mode 100644
--- /dev/null
+++ b/Engine/Plugins/Runtime/StateTree/Source/StateTreeEditorModule/Private/SStateTreeDebuggerView.h
@@ -0,0 +1,83 @@
+// Copyright Epic Games, Inc. All Rights Reserved.
+
+#pragma once
+
+#if WITH_STATETREE_DEBUGGER
+
+#include "Debugger/StateTreeTraceTypes.h"
+#include "Widgets/SCompoundWidget.h"
+
+namespace UE::StateTreeDebugger { struct FFrameIndexSpan; }
+struct FStateTreeDebugger;
+struct FStateTreeInstanceDebugId;
+class FStateTreeEditor;
+class FStateTreeViewModel;
+class FUICommandList;
+class IMessageLogListing;
+class UStateTree;
+
+class SStateTreeDebuggerView : public SCompoundWidget
+{
+public:
+	SLATE_BEGIN_ARGS(SStateTreeDebuggerView) {}
+	SLATE_END_ARGS()
+
+	SStateTreeDebuggerView();
+	~SStateTreeDebuggerView();
+
+	void Construct(const FArguments& InArgs, const UStateTree* InStateTree, TSharedRef<FStateTreeViewModel> InStateTreeViewModel, TSharedRef<FUICommandList> InCommandList);
+
+private:
+	TSharedRef<SWidget> OnGetDebuggerInstancesMenu() const;
+	TSharedRef<SWidget> OnGetDebuggerTracesMenu() const;
+
+	void OnPIEStarted(bool bIsSimulating) const;
+	void OnPIEStopped(bool bIsSimulating) const;
+	void OnPIEPaused(bool bIsSimulating) const;
+	void OnPIEResumed(bool bIsSimulating) const;
+	void OnPIESingleStepped(bool bIsSimulating) const;
+	
+	void OnTracesUpdated(const TConstArrayView<const FStateTreeTraceEventVariantType> Events, const TConstArrayView<UE::StateTreeDebugger::FFrameIndexSpan> Spans) const;
+	void OnBreakpointHit(const FStateTreeInstanceDebugId InstanceId, const FStateTreeStateHandle StateHandle, const TSharedPtr<FUICommandList> ActionList) const;
+
+	void BindDebuggerToolbarCommands(const TSharedRef<FUICommandList> ToolkitCommands);
+
+	/**
+	 * Stores raw log information so it can be processed
+	 */
+	struct FPendingLogMessage
+	{
+		FString Message;
+		ELogVerbosity::Type Verbosity = ELogVerbosity::NoLogging;
+		FName Category;
+
+		FPendingLogMessage()
+		{
+		}
+
+		FPendingLogMessage(const TCHAR* InMessage, const ELogVerbosity::Type InVerbosity, const FName& InCategory)
+			: Message(InMessage)
+			, Verbosity(InVerbosity)
+			, Category(InCategory)
+		{
+		}
+	};
+
+	bool CanStepBack() const;
+	void StepBack() const;
+
+	bool CanStepForward() const;
+	void StepForward() const;
+	
+	bool CanToggleBreakpoint() const;
+	void ToggleBreakpoint() const;
+
+	TSharedPtr<FStateTreeDebugger> Debugger;
+	TSharedPtr<FStateTreeViewModel> StateTreeViewModel;
+	TWeakObjectPtr<const UStateTree> StateTree;
+	
+	TSharedPtr<SWidget> DebuggerTraces;
+	TSharedPtr<IMessageLogListing> DebuggerTraceListing;
+};
+
+#endif // WITH_STATETREE_DEBUGGER
\ No newline at end of file
diff --git a/Engine/Plugins/Runtime/StateTree/Source/StateTreeEditorModule/Private/SStateTreeView.cpp b/Engine/Plugins/Runtime/StateTree/Source/StateTreeEditorModule/Private/SStateTreeView.cpp
--- a/Engine/Plugins/Runtime/StateTree/Source/StateTreeEditorModule/Private/SStateTreeView.cpp
+++ b/Engine/Plugins/Runtime/StateTree/Source/StateTreeEditorModule/Private/SStateTreeView.cpp
@@ -5,9 +5,11 @@
 #include "Widgets/Layout/SScrollBox.h"
 #include "SPositiveActionButton.h"
 #include "SStateTreeViewRow.h"
+#include "StateTreeDebuggerCommands.h"
 #include "StateTreeViewModel.h"
 #include "StateTreeState.h"
 #include "StateTreeEditorCommands.h"
+#include "StateTreeSettings.h"
 #include "Framework/Commands/UICommandList.h"
 
 #define LOCTEXT_NAMESPACE "StateTreeEditor"
@@ -423,6 +425,13 @@ TSharedPtr<SWidget> SStateTreeView::HandleContextMenuOpening()
 	MenuBuilder.AddMenuEntry(FStateTreeEditorCommands::Get().DeleteStates);
 	MenuBuilder.AddMenuEntry(FStateTreeEditorCommands::Get().RenameState);
 
+#if WITH_STATETREE_DEBUGGER
+	if (UStateTreeSettings::Get().bUseDebugger)
+	{
+		MenuBuilder.AddMenuEntry(FStateTreeDebuggerCommands::Get().ToggleBreakpoint);
+	}
+#endif // WITH_STATETREE_DEBUGGER
+	
 	return MenuBuilder.MakeWidget();
 }
 
diff --git a/Engine/Plugins/Runtime/StateTree/Source/StateTreeEditorModule/Private/SStateTreeViewRow.cpp b/Engine/Plugins/Runtime/StateTree/Source/StateTreeEditorModule/Private/SStateTreeViewRow.cpp
--- a/Engine/Plugins/Runtime/StateTree/Source/StateTreeEditorModule/Private/SStateTreeViewRow.cpp
+++ b/Engine/Plugins/Runtime/StateTree/Source/StateTreeEditorModule/Private/SStateTreeViewRow.cpp
@@ -23,12 +23,12 @@ void SStateTreeViewRow::Construct(const FArguments& InArgs, const TSharedRef<STa
 	WeakState = InState;
 
 	STableRow<TWeakObjectPtr<UStateTreeState>>::ConstructInternal(STableRow::FArguments()
-        .Padding(5.0f)
-        .OnDragDetected(this, &SStateTreeViewRow::HandleDragDetected)
-        .OnCanAcceptDrop(this, &SStateTreeViewRow::HandleCanAcceptDrop)
-        .OnAcceptDrop(this, &SStateTreeViewRow::HandleAcceptDrop)
-        .Style(&FStateTreeEditorStyle::Get().GetWidgetStyle<FTableRowStyle>("StateTree.Selection"))
-    , InOwnerTableView);
+		.Padding(5.0f)
+		.OnDragDetected(this, &SStateTreeViewRow::HandleDragDetected)
+		.OnCanAcceptDrop(this, &SStateTreeViewRow::HandleCanAcceptDrop)
+		.OnAcceptDrop(this, &SStateTreeViewRow::HandleAcceptDrop)
+		.Style(&FStateTreeEditorStyle::Get().GetWidgetStyle<FTableRowStyle>("StateTree.Selection"))
+		, InOwnerTableView);
 
 	static const FLinearColor TasksBackground = FLinearColor(FColor(17, 117, 131));
 	static const FLinearColor LinkBackground = FLinearColor(FColor(84, 84, 84));
@@ -37,6 +37,12 @@ void SStateTreeViewRow::Construct(const FArguments& InArgs, const TSharedRef<STa
 	.HAlign(HAlign_Fill)
 	[
 		SNew(SBox)
+	    .MinDesiredWidth_Lambda([ViewBox]()
+			{
+    			// Make the row at least as wide as the view.
+	    		// The -1 is needed or we'll see a scrollbar. 
+				return ViewBox->GetTickSpaceGeometry().GetLocalSize().X - 1;
+			})
 		[
 			SNew(SHorizontalBox)
 			+ SHorizontalBox::Slot()
@@ -50,9 +56,9 @@ void SStateTreeViewRow::Construct(const FArguments& InArgs, const TSharedRef<STa
 				.BaseIndentLevel(0)
 			]
 
-			// State Box
 			+ SHorizontalBox::Slot()
-			.VAlign(VAlign_Center)
+			.VAlign(VAlign_Fill)
+			.HAlign(HAlign_Left)
 			.Padding(FMargin(0.0f, 4.0f))
 			.AutoWidth()
 			[
@@ -62,100 +68,133 @@ void SStateTreeViewRow::Construct(const FArguments& InArgs, const TSharedRef<STa
 				[
 					SNew(SBorder)
 					.BorderImage(FAppStyle::GetBrush("WhiteBrush"))
-					.BorderBackgroundColor(this, &SStateTreeViewRow::GetTitleColor)
-					.Padding(FMargin(16.0f, 0.0f, 16.0f, 0.0f))
+					.BorderBackgroundColor(this, &SStateTreeViewRow::GetActiveStateColor)
+					.Visibility(this, &SStateTreeViewRow::GetActiveStateVisibility)
 					[
-						// Conditions icon
+						// State Box
 						SNew(SHorizontalBox)
 						+ SHorizontalBox::Slot()
 						.VAlign(VAlign_Center)
 						.AutoWidth()
 						[
 							SNew(SBox)
-							.Padding(FMargin(0.0f, 0.0f, 4.0f, 0.0f))
-							.Visibility(this, &SStateTreeViewRow::GetConditionVisibility)
+							.HeightOverride(28.0f)
+							.VAlign(VAlign_Fill)
 							[
-								SNew(SImage)
-								.Image(FAppStyle::Get().GetBrush("Icons.Help"))
+								SNew(SBorder)
+								.BorderImage(FAppStyle::GetBrush("WhiteBrush"))
+								.BorderBackgroundColor(this, &SStateTreeViewRow::GetTitleColor)
+								.Padding(FMargin(16.0f, 0.0f, 16.0f, 0.0f))
+								[
+									// Breakpoint box
+									SNew(SHorizontalBox)
+									+ SHorizontalBox::Slot()
+									.VAlign(VAlign_Top)
+									.HAlign(HAlign_Left)
+									.AutoWidth()
+									[
+										SNew(SBox)
+										.Padding(FMargin(0.0f, 0.0f, 4.0f, 0.0f))
+										[
+											SNew(SImage)
+											.DesiredSizeOverride(FVector2D(12.f, 12.f))
+											.Image(FAppStyle::GetBrush(TEXT("Kismet.DebuggerOverlay.Breakpoint.EnabledAndValid")))
+											.Visibility(this, &SStateTreeViewRow::GetBreakpointVisibility)
+										]
+									]
+
+									// Conditions icon
+									+ SHorizontalBox::Slot()
+									.VAlign(VAlign_Center)
+									.AutoWidth()
+									[
+										SNew(SBox)
+										.Padding(FMargin(0.0f, 0.0f, 4.0f, 0.0f))
+										.Visibility(this, &SStateTreeViewRow::GetConditionVisibility)
+										[
+											SNew(SImage)
+											.Image(FAppStyle::Get().GetBrush("Icons.Help"))
+										]
+									]
+
+									// Selector icon
+									+ SHorizontalBox::Slot()
+									.VAlign(VAlign_Center)
+									.AutoWidth()
+									[
+										SNew(SBox)
+										.Padding(FMargin(0.0f, 0.0f, 4.0f, 0.0f))
+										.Visibility(this, &SStateTreeViewRow::GetSelectorVisibility)
+										[
+											SNew(STextBlock)
+											.Text(this, &SStateTreeViewRow::GetSelectorDesc)
+											.TextStyle(FStateTreeEditorStyle::Get(), "StateTree.Icon")
+										]
+									]
+
+									// State Name
+									+ SHorizontalBox::Slot()
+									.VAlign(VAlign_Center)
+									.AutoWidth()
+									[
+										SAssignNew(NameTextBlock, SInlineEditableTextBlock)
+										.Style(FStateTreeEditorStyle::Get(), "StateTree.State.TitleInlineEditableText")
+										.OnVerifyTextChanged_Lambda([](const FText& NewLabel, FText& OutErrorMessage)
+										{
+											return !NewLabel.IsEmptyOrWhitespace();
+										})
+										.OnTextCommitted(this, &SStateTreeViewRow::HandleNodeLabelTextCommitted)
+										.Text(this, &SStateTreeViewRow::GetStateDesc)
+										.Clipping(EWidgetClipping::ClipToBounds)
+										.IsSelected(this, &SStateTreeViewRow::IsStateSelected)
+									]
+								]
 							]
 						]
 
-						// Selector icon
+						// Tasks Box
 						+ SHorizontalBox::Slot()
 						.VAlign(VAlign_Center)
 						.AutoWidth()
 						[
 							SNew(SBox)
-							.Padding(FMargin(0.0f, 0.0f, 4.0f, 0.0f))
-							.Visibility(this, &SStateTreeViewRow::GetSelectorVisibility)
+							.HeightOverride(28.0f)
+							.VAlign(VAlign_Fill)
+							.Visibility(this, &SStateTreeViewRow::GetTasksVisibility)
 							[
-								SNew(STextBlock)
-								.Text(this, &SStateTreeViewRow::GetSelectorDesc)
-								.TextStyle(FStateTreeEditorStyle::Get(), "StateTree.Icon")
+								SNew(SBorder)
+								.BorderImage(FAppStyle::GetBrush("WhiteBrush"))
+								.BorderBackgroundColor(TasksBackground)
+								.Padding(FMargin(12.0f, 0.0f, 16.0f, 0.0f))
+								[
+									// Task icon
+									SNew(SHorizontalBox)
+									+ SHorizontalBox::Slot()
+									.VAlign(VAlign_Center)
+									.Padding(FMargin(0.0f, 0.0f, 4.0f, 0.0f))
+									.AutoWidth()
+									[
+										SNew(STextBlock)
+										.Text(FEditorFontGlyphs::Paper_Plane)
+										.TextStyle(FStateTreeEditorStyle::Get(), "StateTree.DetailsIcon")
+									]
+
+									// Tasks list
+									+ SHorizontalBox::Slot()
+									.VAlign(VAlign_Center)
+									.AutoWidth()
+									[
+										SNew(STextBlock)
+										.Text(this, &SStateTreeViewRow::GetTasksDesc)
+										.TextStyle(FStateTreeEditorStyle::Get(), "StateTree.Details")
+									]
+								]
 							]
 						]
-
-						// State Name
-						+ SHorizontalBox::Slot()
-						.VAlign(VAlign_Center)
-						.AutoWidth()
-						[
-							SAssignNew(NameTextBlock, SInlineEditableTextBlock)
-							.Style(FStateTreeEditorStyle::Get(), "StateTree.State.TitleInlineEditableText")
-							.OnVerifyTextChanged_Lambda([](const FText& NewLabel, FText& OutErrorMessage)
-							{
-								return !NewLabel.IsEmptyOrWhitespace();
-							})
-							.OnTextCommitted(this, &SStateTreeViewRow::HandleNodeLabelTextCommitted)
-							.Text(this, &SStateTreeViewRow::GetStateDesc)
-							.Clipping(EWidgetClipping::ClipToBounds)
-							.IsSelected(this, &SStateTreeViewRow::IsStateSelected)
-						]
 					]
 				]
 			]
 
-			// Tasks Box
-			+ SHorizontalBox::Slot()
-			.VAlign(VAlign_Center)
-			.Padding(FMargin(0.0f, 4.0f))
-			.AutoWidth()
-			[
-				SNew(SBox)
-				.HeightOverride(28.0f)
-				.VAlign(VAlign_Fill)
-				.Visibility(this, &SStateTreeViewRow::GetTasksVisibility)
-				[
-					SNew(SBorder)
-					.BorderImage(FAppStyle::GetBrush("WhiteBrush"))
-					.BorderBackgroundColor(TasksBackground)
-					.Padding(FMargin(12.0f, 0.0f, 16.0f, 0.0f))
-					[
-						// Task icon
-						SNew(SHorizontalBox)
-						+ SHorizontalBox::Slot()
-						.VAlign(VAlign_Center)
-						.Padding(FMargin(0.0f, 0.0f, 4.0f, 0.0f))
-						.AutoWidth()
-						[
-							SNew(STextBlock)
-							.Text(FEditorFontGlyphs::Paper_Plane)
-							.TextStyle(FStateTreeEditorStyle::Get(), "StateTree.DetailsIcon")
-						]
-
-						// Tasks list
-						+ SHorizontalBox::Slot()
-						.VAlign(VAlign_Center)
-						.AutoWidth()
-						[
-							SNew(STextBlock)
-							.Text(this, &SStateTreeViewRow::GetTasksDesc)
-							.TextStyle(FStateTreeEditorStyle::Get(), "StateTree.Details")
-						]
-					]
-				]
-			]
-			
 			// Linked State box
 			+ SHorizontalBox::Slot()
 			.VAlign(VAlign_Center)
@@ -196,7 +235,7 @@ void SStateTreeViewRow::Construct(const FArguments& InArgs, const TSharedRef<STa
 					]
 				]
 			]
-			
+
 			// Completed transitions
 			+ SHorizontalBox::Slot()
 			.VAlign(VAlign_Center)
@@ -368,6 +407,19 @@ FSlateColor SStateTreeViewRow::GetTitleColor() const
 	return FLinearColor(FColor(31, 151, 167));
 }
 
+FSlateColor SStateTreeViewRow::GetActiveStateColor() const
+{
+	if (const UStateTreeState* State = WeakState.Get())
+	{
+		if (StateTreeViewModel && StateTreeViewModel->IsStateActiveInDebugger(*State))
+		{
+			return FLinearColor::Yellow;
+		}
+	}
+
+	return GetTitleColor();
+}
+
 FText SStateTreeViewRow::GetStateDesc() const
 {
 	if (const UStateTreeState* State = WeakState.Get())
@@ -407,6 +459,31 @@ EVisibility SStateTreeViewRow::GetSelectorVisibility() const
 	return EVisibility::Collapsed;
 }
 
+EVisibility SStateTreeViewRow::GetActiveStateVisibility() const
+{
+	if (const UStateTreeState* State = WeakState.Get())
+	{
+		if (StateTreeViewModel && StateTreeViewModel->IsStateActiveInDebugger(*State))
+		{
+			return EVisibility::Visible;
+		}
+	}
+
+	return EVisibility::Visible;
+}
+
+EVisibility SStateTreeViewRow::GetBreakpointVisibility() const
+{
+	if (const UStateTreeState* State = WeakState.Get())
+	{
+		if (StateTreeViewModel && StateTreeViewModel->DoesStateHaveBreakpoint(*State))
+		{
+			return EVisibility::Visible;
+		}
+	}
+
+	return EVisibility::Hidden;
+}
 
 FText SStateTreeViewRow::GetSelectorDesc() const
 {
@@ -436,7 +513,7 @@ EVisibility SStateTreeViewRow::GetTasksVisibility() const
 		int32 ValidCount = 0;
 		for (int32 i = 0; i < State->Tasks.Num(); i++)
 		{
-			if (const FStateTreeTaskBase* Task = State->Tasks[i].Node.GetPtr<FStateTreeTaskBase>())
+			if (State->Tasks[i].Node.GetPtr<FStateTreeTaskBase>())
 			{
 				ValidCount++;
 			}
@@ -559,11 +636,11 @@ FText SStateTreeViewRow::GetTransitionsIcon(const UStateTreeState& State, const
 	enum EIconType
 	{
 		IconNone = 0,
-        IconRightArrow =	1 << 0,
-        IconDownArrow =		1 << 1,
-        IconLevelUp =		1 << 2,
-        IconWarning =		1 << 3,
-    };
+		IconRightArrow =	1 << 0,
+		IconDownArrow =		1 << 1,
+		IconLevelUp =		1 << 2,
+		IconWarning =		1 << 3,
+	};
 	uint8 IconType = IconNone;
 	
 	for (const FStateTreeTransition& Transition : State.Transitions)
@@ -604,8 +681,8 @@ FText SStateTreeViewRow::GetTransitionsIcon(const UStateTreeState& State, const
 	
 	if (State.Children.Num() == 0
 		&& State.Type == EStateTreeStateType::State
-        && IconType == IconNone
-        && EnumHasAnyFlags(Trigger, EStateTreeTransitionTrigger::OnStateCompleted))
+		&& IconType == IconNone
+		&& EnumHasAnyFlags(Trigger, EStateTreeTransitionTrigger::OnStateCompleted))
 	{
 		if (HasParentTransitionForTrigger(State, Trigger))
 		{
diff --git a/Engine/Plugins/Runtime/StateTree/Source/StateTreeEditorModule/Private/SStateTreeViewRow.h b/Engine/Plugins/Runtime/StateTree/Source/StateTreeEditorModule/Private/SStateTreeViewRow.h
--- a/Engine/Plugins/Runtime/StateTree/Source/StateTreeEditorModule/Private/SStateTreeViewRow.h
+++ b/Engine/Plugins/Runtime/StateTree/Source/StateTreeEditorModule/Private/SStateTreeViewRow.h
@@ -2,8 +2,8 @@
 
 #pragma once
 
-#include "Widgets/Views/ITableRow.h"
 #include "Widgets/Views/STableRow.h"
+#include "Widgets/Views/SListView.h"
 
 class STableViewBase;
 class UStateTreeState;
@@ -30,10 +30,13 @@ public:
 private:
 
 	FSlateColor GetTitleColor() const;
+	FSlateColor GetActiveStateColor() const;
 	FText GetStateDesc() const;
 
 	EVisibility GetConditionVisibility() const;
 	EVisibility GetSelectorVisibility() const;
+	EVisibility GetActiveStateVisibility() const;
+	EVisibility GetBreakpointVisibility() const;
 	FText GetSelectorDesc() const;
 
 	EVisibility GetTasksVisibility() const;
diff --git a/Engine/Plugins/Runtime/StateTree/Source/StateTreeEditorModule/Private/StateTreeCompiler.cpp b/Engine/Plugins/Runtime/StateTree/Source/StateTreeEditorModule/Private/StateTreeCompiler.cpp
--- a/Engine/Plugins/Runtime/StateTree/Source/StateTreeEditorModule/Private/StateTreeCompiler.cpp
+++ b/Engine/Plugins/Runtime/StateTree/Source/StateTreeEditorModule/Private/StateTreeCompiler.cpp
@@ -177,7 +177,13 @@ bool FStateTreeCompiler::Compile(UStateTree& InStateTree)
 		Log.Reportf(EMessageSeverity::Error, TEXT("Unexpected failure to link the StateTree asset. See log for more info."));
 		return false;
 	}
-	
+
+	// Store mapping between state unique ID and state handle. Used for debugging purposes.
+	for (const TPair<FGuid, int>& ToState : IDToState)
+	{
+		StateTree->IDToStateMappings.Emplace(ToState.Key, FStateTreeStateHandle(ToState.Value));
+	}
+
 	return true;
 }
 
diff --git a/Engine/Plugins/Runtime/StateTree/Source/StateTreeEditorModule/Private/StateTreeDebuggerCommands.cpp b/Engine/Plugins/Runtime/StateTree/Source/StateTreeEditorModule/Private/StateTreeDebuggerCommands.cpp
new file mode 100644
--- /dev/null
+++ b/Engine/Plugins/Runtime/StateTree/Source/StateTreeEditorModule/Private/StateTreeDebuggerCommands.cpp
@@ -0,0 +1,32 @@
+// Copyright Epic Games, Inc. All Rights Reserved.
+
+#if WITH_STATETREE_DEBUGGER
+
+#include "StateTreeDebuggerCommands.h"
+
+#include "StateTreeEditorStyle.h"
+#include "Styling/AppStyle.h"
+
+#define LOCTEXT_NAMESPACE "StateTreeDebugger"
+
+FStateTreeDebuggerCommands::FStateTreeDebuggerCommands()
+	: TCommands
+	(
+		"StateTreeEditor.Debugger",								// Context name for fast lookup
+		LOCTEXT("StateTreeDebugger", "StateTree Debugger"),		// Localized context name for displaying
+		NAME_None,												// Parent context name
+		FStateTreeEditorStyle::Get().GetStyleSetName()
+	)
+{
+}
+
+void FStateTreeDebuggerCommands::RegisterCommands()
+{
+	UI_COMMAND(ToggleBreakpoint, "Toggle Breakpoint", "Adds or removes a breakpoint on each selected node", EUserInterfaceActionType::Button, FInputChord(EKeys::F9))
+	UI_COMMAND(Back, "Back", "Show previous recorded state", EUserInterfaceActionType::Button, FInputChord());
+	UI_COMMAND(Forward, "Forward", "Show next recorded state", EUserInterfaceActionType::Button, FInputChord());
+}
+
+#undef LOCTEXT_NAMESPACE
+
+#endif // WITH_STATETREE_DEBUGGER
\ No newline at end of file
diff --git a/Engine/Plugins/Runtime/StateTree/Source/StateTreeEditorModule/Private/StateTreeDebuggerCommands.h b/Engine/Plugins/Runtime/StateTree/Source/StateTreeEditorModule/Private/StateTreeDebuggerCommands.h
new file mode 100644
--- /dev/null
+++ b/Engine/Plugins/Runtime/StateTree/Source/StateTreeEditorModule/Private/StateTreeDebuggerCommands.h
@@ -0,0 +1,25 @@
+// Copyright Epic Games, Inc. All Rights Reserved.
+
+#pragma once
+
+#if WITH_STATETREE_DEBUGGER
+
+#include "Framework/Commands/Commands.h"
+
+/**
+ * StateTree Debugger command set.
+ */
+class FStateTreeDebuggerCommands : public TCommands<FStateTreeDebuggerCommands>
+{
+public:
+	FStateTreeDebuggerCommands();
+
+	// TCommands<> overrides
+	virtual void RegisterCommands() override;
+
+	TSharedPtr<FUICommandInfo> ToggleBreakpoint;
+	TSharedPtr<FUICommandInfo> Back;
+	TSharedPtr<FUICommandInfo> Forward;
+};
+
+#endif // WITH_STATETREE_DEBUGGER
\ No newline at end of file
diff --git a/Engine/Plugins/Runtime/StateTree/Source/StateTreeEditorModule/Private/StateTreeEditor.cpp b/Engine/Plugins/Runtime/StateTree/Source/StateTreeEditorModule/Private/StateTreeEditor.cpp
--- a/Engine/Plugins/Runtime/StateTree/Source/StateTreeEditorModule/Private/StateTreeEditor.cpp
+++ b/Engine/Plugins/Runtime/StateTree/Source/StateTreeEditorModule/Private/StateTreeEditor.cpp
@@ -27,6 +27,8 @@
 #include "ToolMenuEntry.h"
 #include "Widgets/Input/SMultiLineEditableTextBox.h"
 #include "FileHelpers.h"
+#include "SStateTreeDebuggerView.h"
+#include "StateTreeSettings.h"
 #include "ToolMenuSection.h"
 
 #define LOCTEXT_NAMESPACE "StateTreeEditor"
@@ -38,6 +40,9 @@ const FName FStateTreeEditor::AssetDetailsTabId(TEXT("StateTreeEditor_AssetDetai
 const FName FStateTreeEditor::StateTreeViewTabId(TEXT("StateTreeEditor_StateTreeView"));
 const FName FStateTreeEditor::StateTreeStatisticsTabId(TEXT("StateTreeEditor_StateTreeStatistics"));
 const FName FStateTreeEditor::CompilerResultsTabId(TEXT("StateTreeEditor_CompilerResults"));
+#if WITH_STATETREE_DEBUGGER
+const FName FStateTreeEditor::DebuggerTabId(TEXT("StateTreeEditor_Debugger"));
+#endif // WITH_STATETREE_DEBUGGER
 
 void FStateTreeEditor::PostUndo(bool bSuccess)
 {
@@ -76,14 +81,26 @@ void FStateTreeEditor::RegisterTabSpawners(const TSharedRef<class FTabManager>&
 		.SetDisplayName(NSLOCTEXT("StateTreeEditor", "StateTreeViewTab", "StateTree"))
 		.SetGroup(WorkspaceMenuCategoryRef)
 		.SetIcon(FSlateIcon(FAppStyle::GetAppStyleSetName(), "LevelEditor.Tabs.Outliner"));
+	
 	InTabManager->RegisterTabSpawner(StateTreeStatisticsTabId, FOnSpawnTab::CreateSP(this, &FStateTreeEditor::SpawnTab_StateTreeStatistics))
 		.SetDisplayName(NSLOCTEXT("StateTreeEditor", "StatisticsTab", "StateTree Statistics"))
 		.SetGroup(WorkspaceMenuCategoryRef)
 		.SetIcon(FSlateIcon(FAppStyle::GetAppStyleSetName(), "LevelEditor.Tabs.Outliner"));
+	
 	InTabManager->RegisterTabSpawner(CompilerResultsTabId, FOnSpawnTab::CreateSP(this, &FStateTreeEditor::SpawnTab_CompilerResults))
 		.SetDisplayName(NSLOCTEXT("StateTreeEditor", "CompilerResultsTab", "Compiler Results"))
 		.SetGroup(WorkspaceMenuCategoryRef)
 		.SetIcon(FSlateIcon(FAppStyle::GetAppStyleSetName(), "LevelEditor.Tabs.Outliner"));
+
+#if WITH_STATETREE_DEBUGGER
+	if (UStateTreeSettings::Get().bUseDebugger)
+	{
+		InTabManager->RegisterTabSpawner(DebuggerTabId, FOnSpawnTab::CreateSP(this, &FStateTreeEditor::SpawnTab_Debugger))
+		   .SetDisplayName(NSLOCTEXT("StateTreeEditor", "DebuggerTab", "Debugger"))
+		   .SetGroup(WorkspaceMenuCategoryRef)
+		   .SetIcon(FSlateIcon(FAppStyle::GetAppStyleSetName(), "LevelEditor.Tabs.Outliner"));
+	}
+#endif // WITH_STATETREE_DEBUGGER
 }
 
 
@@ -96,6 +113,9 @@ void FStateTreeEditor::UnregisterTabSpawners(const TSharedRef<class FTabManager>
 	InTabManager->UnregisterTabSpawner(StateTreeViewTabId);
 	InTabManager->UnregisterTabSpawner(StateTreeStatisticsTabId);
 	InTabManager->UnregisterTabSpawner(CompilerResultsTabId);
+#if WITH_STATETREE_DEBUGGER
+	InTabManager->UnregisterTabSpawner(DebuggerTabId);
+#endif // WITH_STATETREE_DEBUGGER
 }
 
 void FStateTreeEditor::InitEditor( const EToolkitMode::Type Mode, const TSharedPtr< class IToolkitHost >& InitToolkitHost, UStateTree* InStateTree)
@@ -194,6 +214,9 @@ void FStateTreeEditor::InitEditor( const EToolkitMode::Type Mode, const TSharedP
 					FTabManager::NewStack()
 					->SetSizeCoefficient(0.25f)
 					->AddTab(CompilerResultsTabId, ETabState::ClosedTab)
+#if WITH_STATETREE_DEBUGGER
+					->AddTab(DebuggerTabId, ETabState::ClosedTab)
+#endif // WITH_STATETREE_DEBUGGER
 				)
 			)
 			->Split
@@ -355,6 +378,20 @@ TSharedRef<SDockTab> FStateTreeEditor::SpawnTab_CompilerResults(const FSpawnTabA
 	return SpawnedTab;
 }
 
+#if WITH_STATETREE_DEBUGGER
+TSharedRef<SDockTab> FStateTreeEditor::SpawnTab_Debugger(const FSpawnTabArgs& Args)
+{
+	check(Args.GetTabId() == DebuggerTabId);
+
+	return SNew(SDockTab)
+		.Label(LOCTEXT("DebuggerTitle", "Debugger"))
+		.TabColorScale(GetTabColorScale())
+		[
+			SAssignNew(DebuggerView, SStateTreeDebuggerView, StateTree, StateTreeViewModel.ToSharedRef(), ToolkitCommands)
+		];
+}
+#endif // WITH_STATETREE_DEBUGGER
+
 FText FStateTreeEditor::GetStatisticsText() const
 {
 	if (!StateTree)
diff --git a/Engine/Plugins/Runtime/StateTree/Source/StateTreeEditorModule/Private/StateTreeEditor.h b/Engine/Plugins/Runtime/StateTree/Source/StateTreeEditorModule/Private/StateTreeEditor.h
--- a/Engine/Plugins/Runtime/StateTree/Source/StateTreeEditorModule/Private/StateTreeEditor.h
+++ b/Engine/Plugins/Runtime/StateTree/Source/StateTreeEditorModule/Private/StateTreeEditor.h
@@ -110,6 +110,13 @@ private:
 	
 	TSharedPtr<FStateTreeViewModel> StateTreeViewModel;
 
+#if WITH_STATETREE_DEBUGGER
+	TSharedRef<SDockTab> SpawnTab_Debugger(const FSpawnTabArgs& Args);
+
+	TSharedPtr<class SStateTreeDebuggerView> DebuggerView;
+	static const FName DebuggerTabId;
+#endif // WITH_STATETREE_DEBUGGER
+
 	static const FName StateTreeViewTabId;
 	static const FName SelectionDetailsTabId;
 	static const FName AssetDetailsTabId;
diff --git a/Engine/Plugins/Runtime/StateTree/Source/StateTreeEditorModule/Private/StateTreeEditorModule.cpp b/Engine/Plugins/Runtime/StateTree/Source/StateTreeEditorModule/Private/StateTreeEditorModule.cpp
--- a/Engine/Plugins/Runtime/StateTree/Source/StateTreeEditorModule/Private/StateTreeEditorModule.cpp
+++ b/Engine/Plugins/Runtime/StateTree/Source/StateTreeEditorModule/Private/StateTreeEditorModule.cpp
@@ -16,6 +16,7 @@
 #include "StateTree.h"
 #include "StateTreeCompiler.h"
 #include "StateTreeCompilerLog.h"
+#include "StateTreeDebuggerCommands.h"
 #include "StateTreeDelegates.h"
 #include "StateTreeEditor.h"
 #include "StateTreeEditorCommands.h"
@@ -69,6 +70,10 @@ void FStateTreeEditorModule::StartupModule()
 {
 	UE::StateTree::Delegates::OnRequestCompile.BindStatic(&UE::StateTree::Editor::CompileStateTree);
 
+#if WITH_STATETREE_DEBUGGER
+	FStateTreeDebuggerCommands::Register();
+#endif // WITH_STATETREE_DEBUGGER
+
 	MenuExtensibilityManager = MakeShareable(new FExtensibilityManager);
 	ToolBarExtensibilityManager = MakeShareable(new FExtensibilityManager);
 
@@ -92,6 +97,10 @@ void FStateTreeEditorModule::StartupModule()
 void FStateTreeEditorModule::ShutdownModule()
 {
 	UE::StateTree::Delegates::OnRequestCompile.Unbind();
+
+#if WITH_STATETREE_DEBUGGER
+	FStateTreeDebuggerCommands::Unregister();
+#endif // WITH_STATETREE_DEBUGGER
 	
 	MenuExtensibilityManager.Reset();
 	ToolBarExtensibilityManager.Reset();
diff --git a/Engine/Plugins/Runtime/StateTree/Source/StateTreeEditorModule/Private/StateTreeEditorStyle.cpp b/Engine/Plugins/Runtime/StateTree/Source/StateTreeEditorModule/Private/StateTreeEditorStyle.cpp
--- a/Engine/Plugins/Runtime/StateTree/Source/StateTreeEditorModule/Private/StateTreeEditorStyle.cpp
+++ b/Engine/Plugins/Runtime/StateTree/Source/StateTreeEditorModule/Private/StateTreeEditorStyle.cpp
@@ -42,10 +42,9 @@ private:
 FStateTreeEditorStyle::FStateTreeEditorStyle()
 	: FSlateStyleSet(TEXT("StateTreeEditorStyle"))
 {
-	const FVector2f Icon8x8(8.0f, 8.0f);
-	
-	const FString EngineEditorSlateDir = FPaths::EngineContentDir() / TEXT("Slate");
-	SetCoreContentRoot(EngineEditorSlateDir);
+	const FString EngineSlateContentDir = FPaths::EngineContentDir() / TEXT("Slate");
+	const FString EngineEditorSlateContentDir = FPaths::EngineContentDir() / TEXT("Editor/Slate");
+	SetCoreContentRoot(EngineSlateContentDir);
 
 	const FString StateTreePluginContentDir = FPaths::EnginePluginsDir() / TEXT("Runtime/StateTree/Resources");
 	SetContentRoot(StateTreePluginContentDir);
@@ -112,11 +111,11 @@ FStateTreeEditorStyle::FStateTreeEditorStyle()
 	const FTableRowStyle& NormalTableRowStyle = FAppStyle::Get().GetWidgetStyle<FTableRowStyle>("TableView.Row");
 	Set("StateTree.Selection",
 		FTableRowStyle(NormalTableRowStyle)
-		.SetActiveBrush(IMAGE_BRUSH("Common/Selection", Icon8x8, SelectionColor))
-		.SetActiveHoveredBrush(IMAGE_BRUSH("Common/Selection", Icon8x8, SelectionColor))
-		.SetInactiveBrush(IMAGE_BRUSH("Common/Selection", Icon8x8, SelectionColor))
-		.SetInactiveHoveredBrush(IMAGE_BRUSH("Common/Selection", Icon8x8, SelectionColor))
-		.SetSelectorFocusedBrush(IMAGE_BRUSH("Common/Selection", Icon8x8, SelectionColor))
+		.SetActiveBrush(IMAGE_BRUSH("Common/Selection", CoreStyleConstants::Icon8x8, SelectionColor))
+		.SetActiveHoveredBrush(IMAGE_BRUSH("Common/Selection", CoreStyleConstants::Icon8x8, SelectionColor))
+		.SetInactiveBrush(IMAGE_BRUSH("Common/Selection", CoreStyleConstants::Icon8x8, SelectionColor))
+		.SetInactiveHoveredBrush(IMAGE_BRUSH("Common/Selection", CoreStyleConstants::Icon8x8, SelectionColor))
+		.SetSelectorFocusedBrush(IMAGE_BRUSH("Common/Selection", CoreStyleConstants::Icon8x8, SelectionColor))
 	);
 
 	const FComboButtonStyle& ComboButtonStyle = FCoreStyle::Get().GetWidgetStyle<FComboButtonStyle>("ComboButton");
@@ -167,8 +166,8 @@ FStateTreeEditorStyle::FStateTreeEditorStyle()
 
 	// Command icons
 	{
-		// From generic
-		UE::StateTree::Editor::FContentRootScope Scope(this, EngineEditorSlateDir);
+		// From generic engine
+		UE::StateTree::Editor::FContentRootScope Scope(this, EngineSlateContentDir);
 		Set("StateTreeEditor.CutStates", new IMAGE_BRUSH_SVG("Starship/Common/Cut", CoreStyleConstants::Icon16x16));
 		Set("StateTreeEditor.CopyStates", new IMAGE_BRUSH_SVG("Starship/Common/Copy", CoreStyleConstants::Icon16x16));
 		Set("StateTreeEditor.DuplicateStates", new IMAGE_BRUSH_SVG("Starship/Common/Duplicate", CoreStyleConstants::Icon16x16));
@@ -176,6 +175,13 @@ FStateTreeEditorStyle::FStateTreeEditorStyle()
 		Set("StateTreeEditor.RenameState", new IMAGE_BRUSH_SVG("Starship/Common/Rename", CoreStyleConstants::Icon16x16));
 	}
 
+	{
+		// From generic engine editor
+		UE::StateTree::Editor::FContentRootScope Scope(this, EngineEditorSlateContentDir);
+		Set("StateTreeEditor.Debugger.Back", new IMAGE_BRUSH_SVG("Starship/Blueprints/diff_prev", CoreStyleConstants::Icon20x20));
+		Set("StateTreeEditor.Debugger.Forward", new IMAGE_BRUSH_SVG("Starship/Blueprints/diff_next", CoreStyleConstants::Icon20x20));
+	}
+
 	{
 		// From plugin
 		Set("StateTreeEditor.AddSiblingState", new IMAGE_BRUSH_SVG("Icons/Sibling_State", CoreStyleConstants::Icon16x16));
@@ -201,4 +207,4 @@ FStateTreeEditorStyle& FStateTreeEditorStyle::Get()
 {
 	static FStateTreeEditorStyle Instance;
 	return Instance;
-}
\ No newline at end of file
+}
diff --git a/Engine/Plugins/Runtime/StateTree/Source/StateTreeEditorModule/Private/StateTreeViewModel.cpp b/Engine/Plugins/Runtime/StateTree/Source/StateTreeEditorModule/Private/StateTreeViewModel.cpp
--- a/Engine/Plugins/Runtime/StateTree/Source/StateTreeEditorModule/Private/StateTreeViewModel.cpp
+++ b/Engine/Plugins/Runtime/StateTree/Source/StateTreeEditorModule/Private/StateTreeViewModel.cpp
@@ -4,6 +4,7 @@
 #include "StateTree.h"
 #include "StateTreeEditorData.h"
 #include "StateTreeDelegates.h"
+#include "Debugger/StateTreeDebugger.h"
 #include "Editor.h"
 #include "ScopedTransaction.h"
 #include "HAL/PlatformApplicationMisc.h"
@@ -229,6 +230,9 @@ namespace UE::StateTree::Editor
 
 FStateTreeViewModel::FStateTreeViewModel()
 	: TreeDataWeak(nullptr)
+#if WITH_STATETREE_DEBUGGER
+	, Debugger(MakeShareable(new FStateTreeDebugger))
+#endif // WITH_STATETREE_DEBUGGER
 {
 }
 
@@ -246,13 +250,23 @@ void FStateTreeViewModel::Init(UStateTreeEditorData* InTreeData)
 	GEditor->RegisterForUndo(this);
 
 	UE::StateTree::Delegates::OnIdentifierChanged.AddSP(this, &FStateTreeViewModel::HandleIdentifierChanged);
+
+	BindToDebuggerDelegates();
+}
+
+const UStateTree* FStateTreeViewModel::GetStateTree() const
+{
+	if (const UStateTreeEditorData* TreeData = TreeDataWeak.Get())
+	{
+		return TreeData->GetTypedOuter<UStateTree>();
+	}
+
+	return nullptr;
 }
 
 void FStateTreeViewModel::HandleIdentifierChanged(const UStateTree& StateTree) const
 {
-	UStateTreeEditorData* TreeData = TreeDataWeak.Get();
-	const UStateTree* OuterStateTree = TreeData ? Cast<UStateTree>(TreeData->GetOuter()) : nullptr;
-	if (OuterStateTree == &StateTree)
+	if (GetStateTree() == &StateTree)
 	{
 		OnAssetChanged.Broadcast();
 	}
@@ -864,4 +878,60 @@ void FStateTreeViewModel::MoveSelectedStates(UStateTreeState* TargetState, const
 }
 
 
+void FStateTreeViewModel::BindToDebuggerDelegates()
+{
+#if WITH_STATETREE_DEBUGGER
+	Debugger->OnActiveStatesChanged.BindLambda([this](const TConstArrayView<FStateTreeStateHandle> NewActiveStates)
+	{
+		if (const UStateTree* OuterStateTree = GetStateTree())
+		{
+			ActiveStates.Reset(NewActiveStates.Num());
+
+			for (const FStateTreeStateHandle Handle : NewActiveStates)
+			{
+				ActiveStates.Add(OuterStateTree->GetStateIdFromHandle(Handle));
+			}
+		}
+	});
+
+	Debugger->OnBreakpointsChanged.BindLambda([this](TConstArrayView<FStateTreeStateHandle> NewActiveBreakpoints)
+	{
+		StatesWithBreakpoints.Reset(NewActiveBreakpoints.Num());
+
+		const UStateTreeEditorData* TreeData = TreeDataWeak.Get();
+		if (const UStateTree* OuterStateTree = TreeData ? TreeData->GetTypedOuter<UStateTree>() : nullptr)
+		{
+			TreeData->VisitHierarchy([this, &NewActiveBreakpoints, &OuterStateTree](const UStateTreeState& State, UStateTreeState* /*ParentState*/)
+			{
+				const FStateTreeStateHandle Handle = OuterStateTree->GetStateHandleFromId(State.ID);
+				if (NewActiveBreakpoints.Contains(Handle))
+				{
+					StatesWithBreakpoints.Emplace(&State);
+				}
+
+				return EStateTreeVisitor::Continue;
+			});
+		}
+	});
+#endif // WITH_STATETREE_DEBUGGER
+}
+
+bool FStateTreeViewModel::DoesStateHaveBreakpoint(const UStateTreeState& State) const
+{
+#if WITH_STATETREE_DEBUGGER
+	return StatesWithBreakpoints.Contains(&State);
+#else
+	return false;
+#endif // WITH_STATETREE_DEBUGGER
+}
+
+bool FStateTreeViewModel::IsStateActiveInDebugger(const UStateTreeState& State) const
+{
+#if WITH_STATETREE_DEBUGGER
+	return ActiveStates.Contains(State.ID);
+#else
+	return false;
+#endif // WITH_STATETREE_DEBUGGER
+}
+
 #undef LOCTEXT_NAMESPACE
diff --git a/Engine/Plugins/Runtime/StateTree/Source/StateTreeEditorModule/Private/StateTreeViewModel.h b/Engine/Plugins/Runtime/StateTree/Source/StateTreeEditorModule/Private/StateTreeViewModel.h
--- a/Engine/Plugins/Runtime/StateTree/Source/StateTreeEditorModule/Private/StateTreeViewModel.h
+++ b/Engine/Plugins/Runtime/StateTree/Source/StateTreeEditorModule/Private/StateTreeViewModel.h
@@ -6,6 +6,7 @@
 #include "EditorUndoClient.h"
 #include "StateTreeViewModel.generated.h"
 
+struct FStateTreeDebugger;
 struct FPropertyChangedEvent;
 
 class UStateTreeEditorData;
@@ -54,6 +55,9 @@ public:
 	void GetSelectedStates(TArray<TWeakObjectPtr<UStateTreeState>>& OutSelectedStates);
 	bool HasSelection() const;
 
+	// Returns associated state tree asset.
+	const UStateTree* GetStateTree() const;
+	
 	// Returns array of subtrees to edit.
 	TArray<UStateTreeState*>* GetSubTrees() const;
 	int32 GetSubTreeCount() const;
@@ -80,6 +84,14 @@ public:
 	void NotifyAssetChangedExternally() const;
 	void NotifyStatesChangedExternally(const TSet<UStateTreeState*>& ChangedStates, const FPropertyChangedEvent& PropertyChangedEvent) const;
 
+	// Debugging
+#if WITH_STATETREE_DEBUGGER
+	TSharedRef<FStateTreeDebugger> GetDebugger() const { return Debugger; }
+#endif // WITH_STATETREE_DEBUGGER
+
+	bool IsStateActiveInDebugger(const UStateTreeState& State) const;
+	bool DoesStateHaveBreakpoint(const UStateTreeState& State) const;
+
 	// Called when the whole asset is updated (i.e. undo/redo).
 	FOnAssetChanged& GetOnAssetChanged() { return OnAssetChanged; }
 	
@@ -104,8 +116,17 @@ protected:
 
 	void HandleIdentifierChanged(const UStateTree& StateTree) const;
 
+	void BindToDebuggerDelegates();
+	
 	TWeakObjectPtr<UStateTreeEditorData> TreeDataWeak;
 	TSet<TWeakObjectPtr<UStateTreeState>> SelectedStates;
+
+#if WITH_STATETREE_DEBUGGER
+	TSharedRef<FStateTreeDebugger> Debugger;
+	TArray<TWeakObjectPtr<const UStateTreeState>> StatesWithBreakpoints;
+	TArray<FGuid> ActiveStates;
+#endif // WITH_STATETREE_DEBUGGER
+	
 	FOnAssetChanged OnAssetChanged;
 	FOnStatesChanged OnStatesChanged;
 	FOnStateAdded OnStateAdded;
diff --git a/Engine/Plugins/Runtime/StateTree/Source/StateTreeEditorModule/StateTreeEditorModule.Build.cs b/Engine/Plugins/Runtime/StateTree/Source/StateTreeEditorModule/StateTreeEditorModule.Build.cs
--- a/Engine/Plugins/Runtime/StateTree/Source/StateTreeEditorModule/StateTreeEditorModule.Build.cs
+++ b/Engine/Plugins/Runtime/StateTree/Source/StateTreeEditorModule/StateTreeEditorModule.Build.cs
@@ -57,6 +57,8 @@ namespace UnrealBuildTool.Rules
 			PrivateIncludePaths.AddRange(new string[] {
 				System.IO.Path.Combine(GetModuleDirectory("PropertyEditor"), "Private"),
 			});
+		
+			PublicDefinitions.Add("WITH_STATETREE_DEBUGGER=1");
 		}
 	}
 }
diff --git a/Engine/Plugins/Runtime/StateTree/Source/StateTreeModule/Private/Debugger/StateTreeDebugger.cpp b/Engine/Plugins/Runtime/StateTree/Source/StateTreeModule/Private/Debugger/StateTreeDebugger.cpp
new file mode 100644
--- /dev/null
+++ b/Engine/Plugins/Runtime/StateTree/Source/StateTreeModule/Private/Debugger/StateTreeDebugger.cpp
@@ -0,0 +1,646 @@
+// Copyright Epic Games, Inc. All Rights Reserved.
+
+#if WITH_STATETREE_DEBUGGER
+
+#include "Debugger/StateTreeDebugger.h"
+#include "Algo/RemoveIf.h"
+#include "Debugger/IStateTreeTraceProvider.h"
+#include "Debugger/StateTreeTraceProvider.h"
+#include "Debugger/StateTreeTraceTypes.h"
+#include "Misc/Paths.h"
+#include "Modules/ModuleManager.h"
+#include "StateTree.h"
+#include "StateTreeModule.h"
+#include "Trace/StoreClient.h"
+#include "TraceServices/AnalysisService.h"
+#include "TraceServices/ITraceServicesModule.h"
+#include "TraceServices/Model/AnalysisSession.h"
+#include "TraceServices/Model/Frames.h"
+#include "Trace/Analyzer.h"
+#include "Trace/Analysis.h"
+#include "TraceServices/Model/Diagnostics.h"
+#include "GenericPlatform/GenericPlatformMisc.h"
+
+//----------------------------------------------------------------//
+// UE::StateTreeDebugger
+//----------------------------------------------------------------//
+namespace UE::StateTreeDebugger
+{
+	struct FDiagnosticsSessionAnalyzer : public UE::Trace::IAnalyzer
+	{
+		virtual void OnAnalysisBegin(const FOnAnalysisContext& Context) override
+		{
+			auto& Builder = Context.InterfaceBuilder;
+			Builder.RouteEvent(RouteId_Session2, "Diagnostics", "Session2");
+		}
+
+		virtual bool OnEvent(const uint16 RouteId, EStyle, const FOnEventContext& Context) override
+		{
+			const FEventData& EventData = Context.EventData;
+
+			switch (RouteId)
+			{
+			case RouteId_Session2:
+				{
+					EventData.GetString("Platform", SessionInfo.Platform);
+					EventData.GetString("AppName", SessionInfo.AppName);
+					EventData.GetString("CommandLine", SessionInfo.CommandLine);
+					EventData.GetString("Branch", SessionInfo.Branch);
+					EventData.GetString("BuildVersion", SessionInfo.BuildVersion);
+					SessionInfo.Changelist = EventData.GetValue<uint32>("Changelist", 0);
+					SessionInfo.ConfigurationType = (EBuildConfiguration) EventData.GetValue<uint8>("ConfigurationType");
+					SessionInfo.TargetType = (EBuildTargetType) EventData.GetValue<uint8>("TargetType");
+
+					return false;
+				}
+			default: ;
+			}
+
+			return true;
+		}
+
+		enum : uint16
+		{
+			RouteId_Session2,
+		};
+
+		TraceServices::FSessionInfo SessionInfo;
+	};
+
+} // UE::StateTreeDebugger
+
+
+//----------------------------------------------------------------//
+// FStateTreeDebuggerInstanceDesc
+//----------------------------------------------------------------//
+FStateTreeDebuggerInstanceDesc::FStateTreeDebuggerInstanceDesc(const UStateTree* InStateTree, const FStateTreeInstanceDebugId InId, const FString& InName)
+	: StateTree(InStateTree)
+	, Name(InName)
+	, Id(InId)	
+{ 
+}
+
+bool FStateTreeDebuggerInstanceDesc::IsValid() const
+{
+	return StateTree.IsValid() && Name.Len() && Id.IsValid();
+}
+
+
+//----------------------------------------------------------------//
+// FStateTreeDebugger
+//----------------------------------------------------------------//
+FStateTreeDebugger::FStateTreeDebugger()
+	: StateTreeModule(FModuleManager::GetModuleChecked<IStateTreeModule>("StateTreeModule"))
+{
+}
+
+FStateTreeDebugger::~FStateTreeDebugger()
+{
+	StopAnalysis();
+}
+
+void FStateTreeDebugger::Tick(float DeltaTime)
+{
+	if (!bPaused)
+	{
+		TRACE_CPUPROFILER_EVENT_SCOPE(FStateTreeDebugger::Tick);
+		SyncToCurrentSessionDuration();
+	}
+
+	if (DebuggedInstance.IsValid())
+	{
+		TArray<FStateTreeDebuggerInstanceDesc> ActiveInstances;
+		GetActivateInstances(ActiveInstances);
+		if (!ActiveInstances.Contains(DebuggedInstance))
+		{
+			SetDebuggedInstance(FStateTreeDebuggerInstanceDesc());
+		}
+	}
+}
+
+bool FStateTreeDebugger::IsTickable() const
+{
+	return AnalysisSession.IsValid();
+}
+
+void FStateTreeDebugger::StopAnalysis()
+{
+	if (const TraceServices::IAnalysisSession* Session = GetAnalysisSession())
+	{
+		Session->Stop(true);
+		AnalysisSession.Reset();	
+	}
+
+	ResetAnalysisData();
+}
+
+void FStateTreeDebugger::Pause()
+{
+	bPaused = true;
+
+	// Force a refresh to make sure most recent traces are processed
+	SyncToCurrentSessionDuration();
+}
+
+void FStateTreeDebugger::Unpause()
+{
+	// unpause and let next tick read the traces
+	bPaused = false;	
+}
+
+void FStateTreeDebugger::SyncToCurrentSessionDuration()
+{
+	if (const TraceServices::IAnalysisSession* Session = GetAnalysisSession())
+	{
+		TraceServices::FAnalysisSessionReadScope SessionReadScope(*Session);
+		SetCurrentScrubTime(Session->GetDurationSeconds());	
+	}
+}
+
+void FStateTreeDebugger::SetCurrentScrubTime(const double Time)
+{
+	ReadTrace(Time);
+}
+
+const FStateTreeDebuggerInstanceDesc& FStateTreeDebugger::GetDebuggedInstance() const
+{
+	return DebuggedInstance;
+}
+
+FString FStateTreeDebugger::GetDebuggedInstanceDescription() const
+{
+	return DescribeInstance(DebuggedInstance);
+}
+
+void FStateTreeDebugger::SetDebuggedInstance(const FStateTreeDebuggerInstanceDesc& SelectedInstance)
+{
+	if (DebuggedInstance != SelectedInstance)
+	{
+		DebuggedInstance = SelectedInstance;
+		
+		// Reparse traces
+		ResetAnalysisData();
+
+		bProcessingInitialEvents = true;
+		SyncToCurrentSessionDuration();
+		bProcessingInitialEvents = false;
+	}
+}
+
+void FStateTreeDebugger::GetActivateInstances(TArray<FStateTreeDebuggerInstanceDesc>& OutInstances) const
+{
+	if (const TraceServices::IAnalysisSession* Session = GetAnalysisSession())
+	{
+		TraceServices::FAnalysisSessionReadScope SessionReadScope(*Session);
+
+		if (const IStateTreeTraceProvider* Provider = Session->ReadProvider<IStateTreeTraceProvider>(FStateTreeTraceProvider::ProviderName))
+		{
+			Provider->GetActivateInstances(OutInstances);
+		}
+	}
+}
+
+void FStateTreeDebugger::StartLastLiveSessionAnalysis()
+{
+	TArray<FTraceDescriptor> Traces;
+	GetLiveTraces(Traces);
+
+	if (Traces.Num())
+	{
+		StartSessionAnalysis(Traces.Last());
+	}
+}
+
+void FStateTreeDebugger::StartSessionAnalysis(const FTraceDescriptor& TraceDescriptor)
+{
+	if (ActiveSessionTraceDescriptor == TraceDescriptor)
+	{
+		return;
+	}
+	
+	UE::Trace::FStoreClient* StoreClient = GetStoreClient();
+	if (StoreClient == nullptr)
+	{
+		return;
+	}
+
+	// Make sure any active analysis is stopped
+	StopAnalysis();
+
+	const uint32 TraceId = TraceDescriptor.TraceId;
+	
+	// Make sure it is still live
+	const UE::Trace::FStoreClient::FSessionInfo* SessionInfo = StoreClient->GetSessionInfoByTraceId(TraceId);
+	if (SessionInfo != nullptr)
+	{
+		UE::Trace::FStoreClient::FTraceData TraceData = StoreClient->ReadTrace(TraceId);
+		if (!TraceData)
+		{
+			return;
+		}
+
+		FString TraceName(StoreClient->GetStatus()->GetStoreDir());
+		const UE::Trace::FStoreClient::FTraceInfo* TraceInfo = StoreClient->GetTraceInfoById(TraceId);
+		if (TraceInfo != nullptr)
+		{
+			FString Name(TraceInfo->GetName());
+			if (!Name.EndsWith(TEXT(".utrace")))
+			{
+				Name += TEXT(".utrace");
+			}
+			TraceName = FPaths::Combine(TraceName, Name);
+			FPaths::NormalizeFilename(TraceName);
+		}
+		
+		ITraceServicesModule& TraceServicesModule = FModuleManager::LoadModuleChecked<ITraceServicesModule>("TraceServices");
+		if (const TSharedPtr<TraceServices::IAnalysisService> TraceAnalysisService = TraceServicesModule.GetAnalysisService())
+		{
+			checkf(!AnalysisSession.IsValid(), TEXT("Must make sure that current session was properly stopped before starting a new one otherwise it can cause threading issues"));
+			AnalysisSession = TraceAnalysisService->StartAnalysis(TraceId, *TraceName, MoveTemp(TraceData));
+			
+			bProcessingInitialEvents = true;
+			SyncToCurrentSessionDuration();
+			bProcessingInitialEvents = false;
+		}
+
+		ActiveSessionTraceDescriptor = AnalysisSession.IsValid() ? TraceDescriptor : FTraceDescriptor();
+	}
+}
+
+void FStateTreeDebugger::GetLiveTraces(TArray<FTraceDescriptor>& OutTraceDescriptors) const
+{
+	UE::Trace::FStoreClient* StoreClient = GetStoreClient();
+	if (StoreClient == nullptr)
+	{
+		return;
+	}
+
+	OutTraceDescriptors.Reset();
+
+	const uint32 SessionCount = StoreClient->GetSessionCount();
+	for (uint32 SessionIndex = 0; SessionIndex < SessionCount; ++SessionIndex)
+	{
+		const UE::Trace::FStoreClient::FSessionInfo* SessionInfo = StoreClient->GetSessionInfo(SessionIndex);
+		if (SessionInfo != nullptr)
+		{
+			const uint32 TraceId = SessionInfo->GetTraceId();
+			const UE::Trace::FStoreClient::FTraceInfo* TraceInfo = StoreClient->GetTraceInfoById(TraceId);
+			if (TraceInfo != nullptr)
+			{
+				FTraceDescriptor& Trace = OutTraceDescriptors.AddDefaulted_GetRef();
+				Trace.TraceId = TraceId;
+				Trace.Name = FString(TraceInfo->GetName());
+				UpdateMetadata(Trace);
+			}
+		}
+	}
+}
+
+void FStateTreeDebugger::UpdateMetadata(FTraceDescriptor& TraceDescriptor) const
+{
+	UE::Trace::FStoreClient* StoreClient = GetStoreClient();
+	if (StoreClient == nullptr)
+	{
+		return;
+	}
+
+	const UE::Trace::FStoreClient::FTraceData TraceData = StoreClient->ReadTrace(TraceDescriptor.TraceId);
+	if (!TraceData)
+	{
+		return;
+	}
+
+	// inspired from FStoreBrowser
+	struct FDataStream : public UE::Trace::IInDataStream
+	{
+		enum class EReadStatus
+		{
+			Ready = 0,
+			StoppedByReadSizeLimit
+		};
+
+		virtual int32 Read(void* Data, const uint32 Size) override
+		{
+			if (BytesRead >= 1024 * 1024)
+			{
+				Status = EReadStatus::StoppedByReadSizeLimit;
+				return 0;
+			}
+			const int32 InnerBytesRead = Inner->Read(Data, Size);
+			BytesRead += InnerBytesRead;
+
+			return InnerBytesRead;
+		}
+
+		virtual void Close() override
+		{
+			Inner->Close();
+		}
+
+		IInDataStream* Inner = nullptr;
+		int32 BytesRead = 0;
+		EReadStatus Status = EReadStatus::Ready;
+	};
+
+	FDataStream DataStream;
+	DataStream.Inner = TraceData.Get();
+
+	UE::StateTreeDebugger::FDiagnosticsSessionAnalyzer Analyzer;
+	UE::Trace::FAnalysisContext Context;
+	Context.AddAnalyzer(Analyzer);
+	Context.Process(DataStream).Wait();
+
+	TraceDescriptor.SessionInfo = Analyzer.SessionInfo;
+}
+
+FString FStateTreeDebugger::GetDebuggedTraceDescription() const
+{
+	if (ActiveSessionTraceDescriptor.IsValid())
+	{
+		return DescribeTrace(ActiveSessionTraceDescriptor);
+	}
+
+	return TEXT("No trace selected");
+}
+
+FString FStateTreeDebugger::DescribeTrace(const FTraceDescriptor& TraceDescriptor)
+{
+	if (TraceDescriptor.IsValid())
+	{
+		const TraceServices::FSessionInfo& SessionInfo = TraceDescriptor.SessionInfo;
+
+		return FString::Printf(TEXT("%s-%s-%s-%s-%s"),
+			*LexToString(TraceDescriptor.TraceId),
+			*SessionInfo.Platform,
+			*SessionInfo.AppName,
+			LexToString(SessionInfo.ConfigurationType),
+			LexToString(SessionInfo.TargetType));
+	}
+
+	return TEXT("Invalid");
+}
+
+FString FStateTreeDebugger::DescribeInstance(const FStateTreeDebuggerInstanceDesc& InstanceDesc)
+{
+	if (FStateTreeDebuggerInstanceDesc() == InstanceDesc)
+	{
+		return TEXT("No instance selected");
+	}
+	return LexToString(InstanceDesc);
+}
+
+void FStateTreeDebugger::SetActiveStates(const TConstArrayView<FStateTreeStateHandle> NewActiveStates)
+{
+	ActiveStates = NewActiveStates;
+	OnActiveStatesChanged.ExecuteIfBound(ActiveStates);
+}
+
+bool FStateTreeDebugger::CanStepBack() const
+{
+	return bPaused && CurrentFrameIndex != INDEX_NONE
+		&& FramesWithEvents.ContainsByPredicate([CurrentIndex = CurrentFrameIndex](const UE::StateTreeDebugger::FFrameIndexSpan& IteratedFrame)
+		{
+			return IteratedFrame.FrameIdx < CurrentIndex;
+		});
+}
+
+void FStateTreeDebugger::StepBack()
+{
+	const int32 PrevFrameIndex = FramesWithEvents.FindLastByPredicate([CurrentIndex = CurrentFrameIndex](const UE::StateTreeDebugger::FFrameIndexSpan& IteratedFrame)
+	{
+		return IteratedFrame.FrameIdx < CurrentIndex;
+	});
+	
+	if (PrevFrameIndex != INDEX_NONE)
+	{
+		// Rebuild from beginning up to previous frame with events
+		Events.Reset();
+		LastTraceReadTime = UnsetTime;
+		ReadTrace(FramesWithEvents[PrevFrameIndex].FrameIdx);
+		check(CurrentFrameIndex == FramesWithEvents[PrevFrameIndex].FrameIdx);
+	}
+}
+
+bool FStateTreeDebugger::CanStepForward() const
+{
+	return bPaused && CurrentFrameIndex != INDEX_NONE
+		&& FramesWithEvents.ContainsByPredicate([CurrentIndex = CurrentFrameIndex](const UE::StateTreeDebugger::FFrameIndexSpan& IteratedFrame)
+		{
+			return IteratedFrame.FrameIdx > CurrentIndex;
+		});
+}
+
+void FStateTreeDebugger::StepForward()
+{
+	const UE::StateTreeDebugger::FFrameIndexSpan* NextFrame = FramesWithEvents.FindByPredicate([CurrentIndex = CurrentFrameIndex](const UE::StateTreeDebugger::FFrameIndexSpan& IteratedFrame)
+	{
+		return IteratedFrame.FrameIdx > CurrentIndex;
+	});
+	
+	if (NextFrame != nullptr)
+	{
+		const uint64 FrameIndex = (*NextFrame).FrameIdx; 
+		ReadTrace(FrameIndex);
+		check(CurrentFrameIndex == FrameIndex);
+	}
+}
+
+void FStateTreeDebugger::ToggleBreakpoints(const TConstArrayView<FStateTreeStateHandle> SelectedStates)
+{
+	TArray<FStateTreeStateHandle> CurrentBreakpoints = StatesWithBreakpoint;	
+	TArray<FStateTreeStateHandle> NewBreakpoints(SelectedStates);
+
+	// remove from the selected states any state that already has a breakpoint
+	const int32 ExistingBreakpointsStartIndex = Algo::RemoveIf
+	(NewBreakpoints,
+		[&CurrentBreakpoints](const FStateTreeStateHandle BreakpointToToggle)
+		{
+			return CurrentBreakpoints.Contains(BreakpointToToggle);
+		});
+
+	for (int32 i = ExistingBreakpointsStartIndex; i < NewBreakpoints.Num(); i++)
+	{
+		CurrentBreakpoints.RemoveSingleSwap(NewBreakpoints[i], /*bAllowShrinking*/false);
+	}
+
+	NewBreakpoints.SetNum(ExistingBreakpointsStartIndex, /*bAllowShrinking*/false);
+
+	// Both lists were reduced and can be merged as the new complete list of breakpoints.
+	CurrentBreakpoints.Append(NewBreakpoints);
+
+	StatesWithBreakpoint = CurrentBreakpoints;
+
+	OnBreakpointsChanged.ExecuteIfBound(StatesWithBreakpoint);
+}
+
+const TraceServices::IAnalysisSession* FStateTreeDebugger::GetAnalysisSession() const
+{
+	return AnalysisSession.Get();
+}
+
+UE::Trace::FStoreClient* FStateTreeDebugger::GetStoreClient() const
+{
+	return StateTreeModule.GetStoreClient();
+}
+
+void FStateTreeDebugger::ResetAnalysisData()
+{
+	Events.Reset();
+	HitBreakpointStateIndex = INDEX_NONE;
+	HitBreakpointInstanceId.Invalidate();
+	LastTraceReadTime = UnsetTime;
+	CurrentScrubTime = UnsetTime;
+	CurrentFrameIndex = INDEX_NONE;
+	FramesWithEvents.Reset();
+
+	SetActiveStates({});
+}
+
+void FStateTreeDebugger::ReadTrace(const uint64 FrameIndex)
+{
+	if (const TraceServices::IAnalysisSession* Session = GetAnalysisSession())
+	{
+		TraceServices::FAnalysisSessionReadScope SessionReadScope(*Session);
+
+		const TraceServices::IFrameProvider& FrameProvider = ReadFrameProvider(*Session);
+
+		if (const TraceServices::FFrame* TargetFrame = FrameProvider.GetFrame(TraceFrameType_Game, FrameIndex))
+		{
+			CurrentScrubTime = TargetFrame->EndTime;
+			ReadTrace(*Session, FrameProvider, *TargetFrame);
+		}
+	}
+}
+
+void FStateTreeDebugger::ReadTrace(const double ScrubTime)
+{
+	if (const TraceServices::IAnalysisSession* Session = GetAnalysisSession())
+	{
+		TraceServices::FAnalysisSessionReadScope SessionReadScope(*Session);
+
+		const TraceServices::IFrameProvider& FrameProvider = ReadFrameProvider(*Session);
+
+		TraceServices::FFrame TargetFrame;
+		if (FrameProvider.GetFrameFromTime(TraceFrameType_Game, ScrubTime, TargetFrame))
+		{
+			CurrentScrubTime = ScrubTime;
+			ReadTrace(*Session, FrameProvider, TargetFrame);						
+		}
+	}
+}
+
+void FStateTreeDebugger::ReadTrace(
+	const TraceServices::IAnalysisSession& Session,
+	const TraceServices::IFrameProvider& FrameProvider,
+	const TraceServices::FFrame& Frame
+	)
+{
+	TraceServices::FFrame LastReadFrame;
+	const bool bValidLastReadFrame = FrameProvider.GetFrameFromTime(TraceFrameType_Game, LastTraceReadTime, LastReadFrame);
+	if (LastTraceReadTime == UnsetTime || (bValidLastReadFrame && Frame.Index > LastReadFrame.Index))
+	{
+		if (const IStateTreeTraceProvider* Provider = Session.ReadProvider<IStateTreeTraceProvider>(FStateTreeTraceProvider::ProviderName))
+		{
+			const int32 PrevNumEvents = Events.Num();
+
+			AddEvents(LastTraceReadTime, Frame.EndTime, FrameProvider, *Provider);
+			LastTraceReadTime = CurrentScrubTime;
+
+			const int32 NumEvents = Events.Num();
+			if (NumEvents != PrevNumEvents)
+			{
+				OnTracesUpdated.ExecuteIfBound(Events, FramesWithEvents);
+			}
+
+			if (HitBreakpointStateIndex != INDEX_NONE)
+			{
+				check(HitBreakpointInstanceId.IsValid());
+				OnBreakpointHit.ExecuteIfBound(HitBreakpointInstanceId, StatesWithBreakpoint[HitBreakpointStateIndex]);
+
+				HitBreakpointStateIndex = INDEX_NONE;
+				HitBreakpointInstanceId.Invalidate();
+			}
+		}
+	}
+
+	if (FramesWithEvents.ContainsByPredicate([CurrentIndex = CurrentFrameIndex](const UE::StateTreeDebugger::FFrameIndexSpan& IteratedFrame)
+		{
+			return IteratedFrame.FrameIdx == CurrentIndex;
+		}))
+	{
+		CurrentFrameIndex = Frame.Index;	
+	}
+	else if (FramesWithEvents.Num())
+	{
+		CurrentFrameIndex = FramesWithEvents.Last().FrameIdx;
+	}
+}
+
+bool FStateTreeDebugger::ProcessEvent(const FStateTreeInstanceDebugId InstanceId, const TraceServices::FFrame& Frame, const FStateTreeTraceEventVariantType& Event)
+{
+	if (DebuggedInstance.Id == InstanceId && Event.IsType<FStateTreeTraceActiveStatesEvent>())
+	{
+		// Active states
+		SetActiveStates(Event.Get<FStateTreeTraceActiveStatesEvent>().ActiveStates);
+	}
+	else
+	{
+		// Breakpoints 
+		if (bPaused == false // ignored when scrubbing a paused session
+			&& bProcessingInitialEvents == false // only processed for events added after the initial parsing
+			&& HitBreakpointStateIndex == INDEX_NONE // stop on first hit breakpoint
+			&& StatesWithBreakpoint.Num()
+			&& (DebuggedInstance.Id == InstanceId || !DebuggedInstance.IsValid())) // allow breakpoints on any instances if not specified
+		{
+			const FStateTreeTraceStateEvent* StateEvent = Event.TryGet<FStateTreeTraceStateEvent>();
+			if (StateEvent != nullptr && StateEvent->EventType == EStateTreeTraceNodeEventType::OnEnter)
+			{
+				HitBreakpointStateIndex = StatesWithBreakpoint.Find(FStateTreeStateHandle(StateEvent->StateIdx));
+				if (HitBreakpointStateIndex != INDEX_NONE)
+				{
+					HitBreakpointInstanceId = InstanceId;
+				}
+			}
+		}
+
+		// Store events for currently debugged entry 
+		if (DebuggedInstance.Id == InstanceId)
+		{
+			if (FramesWithEvents.IsEmpty() || FramesWithEvents.Last().FrameIdx < Frame.Index)
+			{
+				FramesWithEvents.Add(UE::StateTreeDebugger::FFrameIndexSpan(Frame.Index, Events.Num()));
+			}
+			
+			Events.Emplace(Event);
+		}
+	}
+	
+	return /*bKeepProcessing*/true;
+}
+
+void FStateTreeDebugger::AddEvents(float StartTime, float EndTime, const TraceServices::IFrameProvider& FrameProvider, const IStateTreeTraceProvider& StateTreeTraceProvider)
+{
+	StateTreeTraceProvider.ReadTimelines(DebuggedInstance.Id,
+		[this, StartTime, EndTime, &FrameProvider](const FStateTreeInstanceDebugId InstanceId, const IStateTreeTraceProvider::FEventsTimeline& TimelineData)
+		{
+			// Keep track of the frames containing events. Starting with an invalid frame.
+			TraceServices::FFrame Frame;
+			Frame.Index = INDEX_NONE;
+			
+			TimelineData.EnumerateEvents(StartTime,	EndTime,
+				[this, InstanceId, &FrameProvider, &Frame](const double EventStartTime, const double EventEndTime, uint32 InDepth, const FStateTreeTraceEventVariantType& Event)
+				{
+					// Fetch frame when not set yet or if events no longer part of the current one
+					if (Frame.Index == INDEX_NONE ||
+						(EventEndTime < Frame.StartTime || Frame.EndTime < EventStartTime))
+					{
+						FrameProvider.GetFrameFromTime(TraceFrameType_Game, EventStartTime, Frame);	
+					}
+					
+					const bool bKeepProcessing = ProcessEvent(InstanceId, Frame, Event);
+					return bKeepProcessing ? TraceServices::EEventEnumerate::Continue : TraceServices::EEventEnumerate::Stop;
+				});
+		});
+}
+
+#endif // WITH_STATETREE_DEBUGGER
\ No newline at end of file
diff --git a/Engine/Plugins/Runtime/StateTree/Source/StateTreeModule/Private/Debugger/StateTreeTrace.cpp b/Engine/Plugins/Runtime/StateTree/Source/StateTreeModule/Private/Debugger/StateTreeTrace.cpp
new file mode 100644
--- /dev/null
+++ b/Engine/Plugins/Runtime/StateTree/Source/StateTreeModule/Private/Debugger/StateTreeTrace.cpp
@@ -0,0 +1,125 @@
+// Copyright Epic Games, Inc. All Rights Reserved.
+
+#if WITH_STATETREE_DEBUGGER
+
+#include "Debugger/StateTreeTrace.h"
+#include "Debugger/StateTreeDebugger.h"
+#include "StateTree.h"
+#include "UObject/Package.h"
+#include "Trace/Trace.inl"
+
+UE_TRACE_CHANNEL_DEFINE(StateTreeDebugChannel)
+
+UE_TRACE_EVENT_BEGIN(StateTreeDebugger, InstanceEvent)
+	UE_TRACE_EVENT_FIELD(uint64, Cycle)
+	UE_TRACE_EVENT_FIELD(UE::Trace::WideString, TreeName)
+	UE_TRACE_EVENT_FIELD(UE::Trace::WideString, TreePath)
+	UE_TRACE_EVENT_FIELD(uint32, CompiledDataHash)
+	UE_TRACE_EVENT_FIELD(uint32, InstanceId)
+	UE_TRACE_EVENT_FIELD(uint32, InstanceSerial)
+	UE_TRACE_EVENT_FIELD(UE::Trace::WideString, InstanceName)
+	UE_TRACE_EVENT_FIELD(uint8, EventType)
+UE_TRACE_EVENT_END()
+
+UE_TRACE_EVENT_BEGIN(StateTreeDebugger, LogEvent)
+	UE_TRACE_EVENT_FIELD(uint64, Cycle)
+	UE_TRACE_EVENT_FIELD(uint32, InstanceId)
+	UE_TRACE_EVENT_FIELD(uint32, InstanceSerial)
+	UE_TRACE_EVENT_FIELD(UE::Trace::WideString, Message)
+UE_TRACE_EVENT_END()
+
+UE_TRACE_EVENT_BEGIN(StateTreeDebugger, StateEvent)
+	UE_TRACE_EVENT_FIELD(uint64, Cycle)
+	UE_TRACE_EVENT_FIELD(uint32, InstanceId)
+	UE_TRACE_EVENT_FIELD(uint32, InstanceSerial)
+	UE_TRACE_EVENT_FIELD(uint16, StateIdx)
+	UE_TRACE_EVENT_FIELD(uint8, EventType)
+UE_TRACE_EVENT_END()
+
+UE_TRACE_EVENT_BEGIN(StateTreeDebugger, TaskEvent)
+	UE_TRACE_EVENT_FIELD(uint64, Cycle)
+	UE_TRACE_EVENT_FIELD(uint32, InstanceId)
+	UE_TRACE_EVENT_FIELD(uint32, InstanceSerial)
+	UE_TRACE_EVENT_FIELD(uint16, TaskIdx)
+	UE_TRACE_EVENT_FIELD(uint8, EventType)
+UE_TRACE_EVENT_END()
+
+UE_TRACE_EVENT_BEGIN(StateTreeDebugger, ActiveStatesEvent)
+	UE_TRACE_EVENT_FIELD(uint64, Cycle)
+	UE_TRACE_EVENT_FIELD(uint32, InstanceId)
+	UE_TRACE_EVENT_FIELD(uint32, InstanceSerial)
+	UE_TRACE_EVENT_FIELD(uint16[], ActiveStates)
+UE_TRACE_EVENT_END()
+
+void FStateTreeTrace::OutputInstanceLifetimeEvent(
+	const FStateTreeInstanceDebugId InstanceId,
+	const UStateTree* StateTree,
+	const TCHAR* InstanceName,
+	const EStateTreeTraceInstanceEventType EventType)
+{
+	if (StateTree == nullptr)
+	{
+		return;
+	}
+	const FString ObjectName = StateTree->GetName();
+	const FString ObjectPackageName = StateTree->GetPackage()->GetName();
+
+	UE_TRACE_LOG(StateTreeDebugger, InstanceEvent, StateTreeDebugChannel)
+		<< InstanceEvent.Cycle(FPlatformTime::Cycles64())
+		<< InstanceEvent.TreeName(*ObjectName, ObjectName.Len())
+		<< InstanceEvent.TreePath(*ObjectPackageName, ObjectPackageName.Len())
+		<< InstanceEvent.CompiledDataHash(StateTree->LastCompiledEditorDataHash)
+		<< InstanceEvent.InstanceId(InstanceId.Id)
+		<< InstanceEvent.InstanceSerial(InstanceId.SerialNumber)
+		<< InstanceEvent.InstanceName(InstanceName)
+		<< InstanceEvent.EventType(static_cast<uint8>(EventType));
+}
+
+void FStateTreeTrace::OutputLogEventTrace(const FStateTreeInstanceDebugId InstanceId, const TCHAR* Fmt, ...)
+{
+	static TCHAR TraceStaticBuffer[8192];
+	GET_VARARGS(TraceStaticBuffer, UE_ARRAY_COUNT(TraceStaticBuffer), UE_ARRAY_COUNT(TraceStaticBuffer) - 1, Fmt, Fmt);
+
+	UE_TRACE_LOG(StateTreeDebugger, LogEvent, StateTreeDebugChannel)
+		<< LogEvent.Cycle(FPlatformTime::Cycles64())
+		<< LogEvent.InstanceId(InstanceId.Id)
+		<< LogEvent.InstanceSerial(InstanceId.SerialNumber)
+		<< LogEvent.Message(TraceStaticBuffer);
+}
+
+void FStateTreeTrace::OutputStateEventTrace(const FStateTreeInstanceDebugId InstanceId, const uint16 StateIdx, const EStateTreeTraceNodeEventType EventType)
+{
+	UE_TRACE_LOG(StateTreeDebugger, StateEvent, StateTreeDebugChannel)
+		<< StateEvent.Cycle(FPlatformTime::Cycles64())
+		<< StateEvent.InstanceId(InstanceId.Id)
+		<< StateEvent.InstanceSerial(InstanceId.SerialNumber)
+		<< StateEvent.StateIdx(StateIdx)
+		<< StateEvent.EventType(static_cast<uint8>(EventType));
+}
+
+void FStateTreeTrace::OutputTaskEventTrace(const FStateTreeInstanceDebugId InstanceId, const uint16 TaskIdx, const EStateTreeTraceNodeEventType EventType)
+{
+	UE_TRACE_LOG(StateTreeDebugger, TaskEvent, StateTreeDebugChannel)
+		<< TaskEvent.Cycle(FPlatformTime::Cycles64())
+		<< TaskEvent.InstanceId(InstanceId.Id)
+		<< TaskEvent.InstanceSerial(InstanceId.SerialNumber)
+		<< TaskEvent.TaskIdx(TaskIdx)
+		<< TaskEvent.EventType(static_cast<uint8>(EventType));
+}
+
+void FStateTreeTrace::OutputActiveStatesEventTrace(const FStateTreeInstanceDebugId InstanceId, const FStateTreeActiveStates& ActiveStates)
+{
+	TArray<uint16, TInlineAllocator<FStateTreeActiveStates::MaxStates>> StatesIndices;
+	for (int32 i = 0; i < ActiveStates.Num(); i++)
+	{
+		StatesIndices.Add(ActiveStates[i].Index);
+	}
+
+	UE_TRACE_LOG(StateTreeDebugger, ActiveStatesEvent, StateTreeDebugChannel)
+		<< ActiveStatesEvent.Cycle(FPlatformTime::Cycles64())
+		<< ActiveStatesEvent.InstanceId(InstanceId.Id)
+		<< ActiveStatesEvent.InstanceSerial(InstanceId.SerialNumber)
+		<< ActiveStatesEvent.ActiveStates(StatesIndices.GetData(), StatesIndices.Num());
+}
+
+#endif // WITH_STATETREE_DEBUGGER
diff --git a/Engine/Plugins/Runtime/StateTree/Source/StateTreeModule/Private/Debugger/StateTreeTraceAnalyzer.cpp b/Engine/Plugins/Runtime/StateTree/Source/StateTreeModule/Private/Debugger/StateTreeTraceAnalyzer.cpp
new file mode 100644
--- /dev/null
+++ b/Engine/Plugins/Runtime/StateTree/Source/StateTreeModule/Private/Debugger/StateTreeTraceAnalyzer.cpp
@@ -0,0 +1,150 @@
+// Copyright Epic Games, Inc. All Rights Reserved.
+
+#if WITH_STATETREE_DEBUGGER
+
+#include "Debugger/StateTreeTraceAnalyzer.h"
+#include "Debugger/StateTreeDebugger.h"
+#include "Debugger/StateTreeTraceProvider.h"
+#include "Debugger/StateTreeTraceTypes.h"
+#include "TraceServices/Model/AnalysisSession.h"
+
+FStateTreeTraceAnalyzer::FStateTreeTraceAnalyzer(TraceServices::IAnalysisSession& InSession, FStateTreeTraceProvider& InProvider)
+	: Session(InSession)
+	, Provider(InProvider)
+{
+}
+
+void FStateTreeTraceAnalyzer::OnAnalysisBegin(const FOnAnalysisContext& Context)
+{
+	auto& Builder = Context.InterfaceBuilder;
+
+	Builder.RouteEvent(RouteId_Instance, "StateTreeDebugger", "InstanceEvent");
+	Builder.RouteEvent(RouteId_LogMessage, "StateTreeDebugger", "LogEvent");
+	Builder.RouteEvent(RouteId_State, "StateTreeDebugger", "StateEvent");
+	Builder.RouteEvent(RouteId_Task, "StateTreeDebugger", "TaskEvent");
+	Builder.RouteEvent(RouteId_ActiveStates, "StateTreeDebugger", "ActiveStatesEvent");
+}
+
+bool FStateTreeTraceAnalyzer::OnEvent(const uint16 RouteId, EStyle Style, const FOnEventContext& Context)
+{
+	LLM_SCOPE_BYNAME(TEXT("Insights/FStateTreeAnalyzer"));
+
+	TraceServices::FAnalysisSessionEditScope _(Session);
+
+	const auto& EventData = Context.EventData;
+	switch (RouteId)
+	{
+	case RouteId_Instance:
+		{
+			FString ObjectName, ObjectPathName;
+			EventData.GetString("TreeName", ObjectName);
+			EventData.GetString("TreePath", ObjectPathName);
+
+			const FTopLevelAssetPath Path((FName)ObjectPathName, (FName)ObjectName);
+			TWeakObjectPtr<const UStateTree> WeakStateTree;
+			{
+				// This might not work when using a debugger on a client but should be fine in Editor as lock as
+				// we are not trying to find the object during GC. We might not currently be in the game thread.  
+				// @todo STDBG: eventually errors should be reported in the UI
+				FGCScopeGuard Guard;
+				WeakStateTree = FindObject<UStateTree>(Path);
+			}
+			
+			if (const UStateTree* StateTree = WeakStateTree.Get())
+			{
+				const uint32 CompiledDataHash = EventData.GetValue<uint32>("CompiledDataHash");
+				if (StateTree->LastCompiledEditorDataHash == CompiledDataHash)
+				{
+					const FStateTreeInstanceDebugId InstanceId(
+						EventData.GetValue<uint32>("InstanceId"),
+						EventData.GetValue<uint32>("InstanceSerial"));
+					FString InstanceName;
+					EventData.GetString("InstanceName", InstanceName);
+					const EStateTreeTraceInstanceEventType EventType = EventData.GetValue<EStateTreeTraceInstanceEventType>("EventType");
+					Provider.AppendInstance(StateTree, InstanceId, *InstanceName, EventType);
+				}
+				else
+				{
+					UE_LOG(LogStateTree, Warning, TEXT("Traces are not using the same StateTree asset version as the current asset."));
+				}
+			}
+			else
+			{
+				UE_LOG(LogStateTree, Warning, TEXT("Unable to find StateTree asset: %s : %s"), *ObjectPathName, *ObjectName);
+			}
+			break;
+		}
+	case RouteId_LogMessage:
+		{
+			const uint64 Cycle = EventData.GetValue<uint64>("Cycle");
+			const FStateTreeInstanceDebugId InstanceId(
+				EventData.GetValue<uint32>("InstanceId"),
+				EventData.GetValue<uint32>("InstanceSerial"));
+
+			FStateTreeTraceLogEvent Event;
+			EventData.GetString("Message", Event.Message);
+
+			Provider.AppendEvent(
+				InstanceId,
+				Context.EventTime.AsSeconds(Cycle),
+				FStateTreeTraceEventVariantType(TInPlaceType<FStateTreeTraceLogEvent>(), Event));
+			break;
+		}
+	case RouteId_State:
+		{
+			const uint64 Cycle = EventData.GetValue<uint64>("Cycle");
+			const FStateTreeInstanceDebugId InstanceId(
+				EventData.GetValue<uint32>("InstanceId"),
+				EventData.GetValue<uint32>("InstanceSerial"));
+
+			FStateTreeTraceStateEvent Event;
+			Event.StateIdx = EventData.GetValue<uint16>("StateIdx");
+			Event.EventType = EventData.GetValue<EStateTreeTraceNodeEventType>("EventType");
+
+			Provider.AppendEvent(
+				InstanceId,
+				Context.EventTime.AsSeconds(Cycle),
+				FStateTreeTraceEventVariantType(TInPlaceType<FStateTreeTraceStateEvent>(), Event));
+			break;
+		}
+	case RouteId_Task:
+		{
+			const uint64 Cycle = EventData.GetValue<uint64>("Cycle");
+			const FStateTreeInstanceDebugId InstanceId(
+				EventData.GetValue<uint32>("InstanceId"),
+				EventData.GetValue<uint32>("InstanceSerial"));
+
+			FStateTreeTraceTaskEvent Event;
+			Event.TaskIdx = EventData.GetValue<uint16>("TaskIdx");
+			Event.EventType = EventData.GetValue<EStateTreeTraceNodeEventType>("EventType");
+
+			Provider.AppendEvent(
+				InstanceId,
+				Context.EventTime.AsSeconds(Cycle),
+				FStateTreeTraceEventVariantType(TInPlaceType<FStateTreeTraceTaskEvent>(), Event));
+			break;
+		}
+	case RouteId_ActiveStates:
+		{
+			const uint64 Cycle = EventData.GetValue<uint64>("Cycle");
+			const FStateTreeInstanceDebugId InstanceId(
+				EventData.GetValue<uint32>("InstanceId"),
+				EventData.GetValue<uint32>("InstanceSerial"));
+
+			FStateTreeTraceActiveStatesEvent Event;
+			Event.ActiveStates = EventData.GetArrayView<uint16>("ActiveStates");
+
+			Provider.AppendEvent(
+				InstanceId,
+				Context.EventTime.AsSeconds(Cycle),
+				FStateTreeTraceEventVariantType(TInPlaceType<FStateTreeTraceActiveStatesEvent>(), Event));
+			break;
+		}
+	default:
+		ensureMsgf(false, TEXT("Unhandle route id: %s"), RouteId);
+	}
+
+	return true;
+}
+
+#endif // WITH_STATETREE_DEBUGGER
\ No newline at end of file
diff --git a/Engine/Plugins/Runtime/StateTree/Source/StateTreeModule/Private/Debugger/StateTreeTraceModule.cpp b/Engine/Plugins/Runtime/StateTree/Source/StateTreeModule/Private/Debugger/StateTreeTraceModule.cpp
new file mode 100644
--- /dev/null
+++ b/Engine/Plugins/Runtime/StateTree/Source/StateTreeModule/Private/Debugger/StateTreeTraceModule.cpp
@@ -0,0 +1,32 @@
+// Copyright Epic Games, Inc. All Rights Reserved.
+
+#if WITH_STATETREE_DEBUGGER
+
+#include "Debugger/StateTreeTraceModule.h"
+#include "Debugger/StateTreeTraceProvider.h"
+#include "Debugger/StateTreeTraceAnalyzer.h"
+#include "Debugger/StateTreeDebugger.h"
+#include "TraceServices/Model/AnalysisSession.h"
+
+FName FStateTreeTraceModule::ModuleName("TraceModule_StateTree");
+
+void FStateTreeTraceModule::GetModuleInfo(TraceServices::FModuleInfo& OutModuleInfo)
+{
+	OutModuleInfo.Name = ModuleName;
+	OutModuleInfo.DisplayName = TEXT("StateTreee");
+}
+
+void FStateTreeTraceModule::OnAnalysisBegin(TraceServices::IAnalysisSession& InSession)
+{
+	const TSharedPtr<FStateTreeTraceProvider> Provider = MakeShared<FStateTreeTraceProvider>(InSession);
+	InSession.AddProvider(FStateTreeTraceProvider::ProviderName, Provider);
+
+	InSession.AddAnalyzer(new FStateTreeTraceAnalyzer(InSession, *Provider));
+}
+
+void FStateTreeTraceModule::GetLoggers(TArray<const TCHAR *>& OutLoggers)
+{
+	OutLoggers.Add(TEXT("StateTreee"));
+}
+
+#endif // WITH_STATETREE_DEBUGGER
\ No newline at end of file
diff --git a/Engine/Plugins/Runtime/StateTree/Source/StateTreeModule/Private/Debugger/StateTreeTraceProvider.cpp b/Engine/Plugins/Runtime/StateTree/Source/StateTreeModule/Private/Debugger/StateTreeTraceProvider.cpp
new file mode 100644
--- /dev/null
+++ b/Engine/Plugins/Runtime/StateTree/Source/StateTreeModule/Private/Debugger/StateTreeTraceProvider.cpp
@@ -0,0 +1,95 @@
+// Copyright Epic Games, Inc. All Rights Reserved.
+
+#if WITH_STATETREE_DEBUGGER
+
+#include "Debugger/StateTreeTraceProvider.h"
+#include "StateTreeTypes.h"
+#include "Debugger/StateTreeDebugger.h"
+#include "Algo/RemoveIf.h"
+
+FName FStateTreeTraceProvider::ProviderName("StateTreeDebuggerProvider");
+
+#define LOCTEXT_NAMESPACE "StateTreeDebuggerProvider"
+
+FStateTreeTraceProvider::FStateTreeTraceProvider(TraceServices::IAnalysisSession& InSession)
+	: Session(InSession)
+{
+}
+
+bool FStateTreeTraceProvider::ReadTimelines(const FStateTreeInstanceDebugId InstanceId, const TFunctionRef<void(const FStateTreeInstanceDebugId ProcessedInstanceId, const FEventsTimeline&)> Callback) const
+{
+	Session.ReadAccessCheck();
+
+	// Read specific timeline if specified
+	if (InstanceId.IsValid())
+	{
+		const uint32* IndexPtr = InstanceIdToDebuggerEntryTimelines.Find(InstanceId);
+		if (IndexPtr != nullptr && EventsTimelines.IsValidIndex(*IndexPtr))
+		{
+			Callback(InstanceId, *EventsTimelines[*IndexPtr]);
+			return true;
+		}	
+	}
+	else
+	{
+		for(auto It = InstanceIdToDebuggerEntryTimelines.CreateConstIterator(); It; ++It)
+		{
+			if (EventsTimelines.IsValidIndex(It.Value()))
+			{
+				Callback(It.Key(), *EventsTimelines[It.Value()]);
+			}
+		}
+
+		return EventsTimelines.Num() > 0;	
+	}
+
+	return false;
+}
+
+void FStateTreeTraceProvider::AppendEvent(const FStateTreeInstanceDebugId InInstanceId, const double InTime, const FStateTreeTraceEventVariantType& Event)
+{
+	Session.WriteAccessCheck();
+
+	TSharedPtr<TraceServices::TPointTimeline<FStateTreeTraceEventVariantType>> Timeline;
+	const uint32* IndexPtr = InstanceIdToDebuggerEntryTimelines.Find(InInstanceId);
+	if (IndexPtr != nullptr)
+	{
+		Timeline = EventsTimelines[*IndexPtr];
+	}
+	else
+	{
+		InstanceIdToDebuggerEntryTimelines.Add(InInstanceId, EventsTimelines.Num());
+		Timeline = EventsTimelines.Emplace_GetRef(MakeShared<TraceServices::TPointTimeline<FStateTreeTraceEventVariantType>>(Session.GetLinearAllocator()));
+	}
+
+	Timeline->AppendEvent(InTime, Event);
+	Session.UpdateDurationSeconds(InTime);
+}
+
+void FStateTreeTraceProvider::AppendInstance(
+	const UStateTree* InStateTree,
+	const FStateTreeInstanceDebugId InInstanceId,
+	const TCHAR* InInstanceName,
+	const EStateTreeTraceInstanceEventType EventType)
+{
+	if (EventType == EStateTreeTraceInstanceEventType::Started)
+	{
+		ActiveInstances.Emplace(InStateTree, InInstanceId, InInstanceName);	
+	}
+	else if (EventType == EStateTreeTraceInstanceEventType::Stopped)
+	{
+		ActiveInstances.SetNum(Algo::RemoveIf(ActiveInstances, [InInstanceId](const FStateTreeDebuggerInstanceDesc& Instance)
+			{
+				return Instance.Id == InInstanceId;
+			}));
+	}
+}
+
+void FStateTreeTraceProvider::GetActivateInstances(TArray<FStateTreeDebuggerInstanceDesc>& OutInstances) const
+{
+	OutInstances = ActiveInstances;
+}
+
+#undef LOCTEXT_NAMESPACE
+
+#endif // WITH_STATETREE_DEBUGGER
diff --git a/Engine/Plugins/Runtime/StateTree/Source/StateTreeModule/Private/StateTree.cpp b/Engine/Plugins/Runtime/StateTree/Source/StateTreeModule/Private/StateTree.cpp
--- a/Engine/Plugins/Runtime/StateTree/Source/StateTreeModule/Private/StateTree.cpp
+++ b/Engine/Plugins/Runtime/StateTree/Source/StateTreeModule/Private/StateTree.cpp
@@ -21,6 +21,28 @@ bool UStateTree::IsReadyToRun() const
 	return States.Num() > 0 && bIsLinked;
 }
 
+FConstStructView UStateTree::GetNode(const int32 NodeIndex) const
+{
+	return Nodes.IsValidIndex(NodeIndex) ? Nodes[NodeIndex] : FConstStructView();	
+}
+
+const FCompactStateTreeState* UStateTree::GetStateFromHandle(const FStateTreeStateHandle StateHandle) const
+{
+	return States.IsValidIndex(StateHandle.Index) ? &States[StateHandle.Index] : nullptr;
+}
+
+FStateTreeStateHandle UStateTree::GetStateHandleFromId(const FGuid Id) const
+{
+	const FStateTreeStateIdToHandle* Entry = IDToStateMappings.FindByPredicate([Id](const FStateTreeStateIdToHandle& Entry){ return Entry.Id == Id; });
+	return Entry != nullptr ? Entry->Handle : FStateTreeStateHandle::Invalid;
+}
+
+FGuid UStateTree::GetStateIdFromHandle(const FStateTreeStateHandle Handle) const
+{
+	const FStateTreeStateIdToHandle* Entry = IDToStateMappings.FindByPredicate([Handle](const FStateTreeStateIdToHandle& Entry){ return Entry.Handle == Handle; });
+	return Entry != nullptr ? Entry->Id : FGuid();
+}
+
 TSharedPtr<FStateTreeInstanceData> UStateTree::GetSharedInstanceData() const
 {
 	// Create a unique index for each thread.
@@ -73,6 +95,7 @@ void UStateTree::ResetCompiled()
 	ContextDataDescs.Reset();
 	PropertyBindings.Reset();
 	Parameters.Reset();
+	IDToStateMappings.Reset();
 
 	ParametersDataViewIndex = FStateTreeIndex8::Invalid;
 	
diff --git a/Engine/Plugins/Runtime/StateTree/Source/StateTreeModule/Private/StateTreeExecutionContext.cpp b/Engine/Plugins/Runtime/StateTree/Source/StateTreeModule/Private/StateTreeExecutionContext.cpp
--- a/Engine/Plugins/Runtime/StateTree/Source/StateTreeModule/Private/StateTreeExecutionContext.cpp
+++ b/Engine/Plugins/Runtime/StateTree/Source/StateTreeModule/Private/StateTreeExecutionContext.cpp
@@ -5,6 +5,8 @@
 #include "StateTreeEvaluatorBase.h"
 #include "StateTreeConditionBase.h"
 #include "Containers/StaticArray.h"
+#include "Debugger/StateTreeTrace.h"
+#include "Debugger/StateTreeTraceTypes.h"
 #include "VisualLogger/VisualLogger.h"
 #include "ProfilingDebugging/CsvProfiler.h"
 #include "Logging/LogScopedVerbosityOverride.h"
@@ -179,6 +181,8 @@ EStateTreeRunStatus FStateTreeExecutionContext::Start()
 		return EStateTreeRunStatus::Failed;
 	}
 
+	TRACE_STATETREE_INSTANCE_EVENT(GetInstanceDebugId(), GetStateTree(), *GetInstanceDescription(), EStateTreeTraceInstanceEventType::Started);
+	
 	// Start evaluators and global tasks. Fail the execution if any global task fails.
 	FStateTreeIndex16 LastInitializedTaskIndex;
 	if (!StartEvaluatorsAndGlobalTasks(LastInitializedTaskIndex))
@@ -286,6 +290,10 @@ EStateTreeRunStatus FStateTreeExecutionContext::Stop()
 	// Stop evaluators and global tasks.
 	StopEvaluatorsAndGlobalTasks();
 
+	// Trace before resetting the instance data since it is required to provide all the event information
+	TRACE_STATETREE_ACTIVE_STATES_EVENT(GetInstanceDebugId(), FStateTreeActiveStates());
+	TRACE_STATETREE_INSTANCE_EVENT(GetInstanceDebugId(), GetStateTree(), *GetInstanceDescription(), EStateTreeTraceInstanceEventType::Stopped);
+
 	// Destruct all allocated instance data (does not shrink the buffer). This will invalidate Exec too.
 	InstanceData.Reset();
 
@@ -460,6 +468,7 @@ void FStateTreeExecutionContext::SendEvent(const FGameplayTag Tag, const FConstS
 	}
 
 	STATETREE_LOG(Verbose, TEXT("Send Event '%s'"), *Tag.ToString());
+	TRACE_STATETREE_LOG_EVENT(GetInstanceDebugId(), TEXT("Send Event '%s'"), *Tag.ToString());
 
 	FStateTreeEventQueue& EventQueue = InstanceData.GetMutableEventQueue();
 	EventQueue.SendEvent(&Owner, Tag, Payload, Origin);
@@ -497,6 +506,25 @@ void FStateTreeExecutionContext::RequestTransition(const FStateTreeTransitionReq
 	}
 }
 
+#if WITH_STATETREE_DEBUGGER
+FStateTreeInstanceDebugId FStateTreeExecutionContext::GetInstanceDebugId() const
+{
+	FStateTreeInstanceDebugId& InstanceDebugId = GetExecState().InstanceDebugId; 
+	if (!InstanceDebugId.IsValid())
+	{
+		const uint32 Id = GetTypeHash(GetInstanceDescription());
+		static TMap<uint32, uint32> DebugIdMappings;
+		uint32* SerialNumber = DebugIdMappings.Find(Id);
+		if (SerialNumber == nullptr)
+		{
+			SerialNumber = &DebugIdMappings.Add(Id, 0);
+		}
+		InstanceDebugId = FStateTreeInstanceDebugId(Id, (*SerialNumber)++); 
+	}
+	return InstanceDebugId;
+}
+#endif // WITH_STATETREE_DEBUGGER
+
 void FStateTreeExecutionContext::UpdateInstanceData(const FStateTreeActiveStates& CurrentActiveStates, const FStateTreeActiveStates& NextActiveStates)
 {
 	// Find common section of states at start.
@@ -661,6 +689,7 @@ EStateTreeRunStatus FStateTreeExecutionContext::EnterState(const FStateTreeTrans
 	}
 	
 	STATETREE_LOG(Log, TEXT("Enter state '%s' (%d)"), *DebugGetStatePath(Transition.NextActiveStates), Exec.StateChangeCount);
+	TRACE_STATETREE_STATE_EVENT(GetInstanceDebugId(), Transition.NextActiveStates.Last().Index, EStateTreeTraceNodeEventType::OnEnter);
 
 	for (int32 Index = 0; Index < Transition.NextActiveStates.Num() && Result != EStateTreeRunStatus::Failed; Index++)
 	{
@@ -712,6 +741,7 @@ EStateTreeRunStatus FStateTreeExecutionContext::EnterState(const FStateTreeTrans
 			if (bIsEnteringState && bShouldCallStateChange)
 			{
 				STATETREE_LOG(Verbose, TEXT("%*s  Task '%s'"), Index*UE::StateTree::DebugIndentSize, TEXT(""), *Task.Name.ToString());
+				TRACE_STATETREE_TASK_EVENT(GetInstanceDebugId(), TaskIndex, EStateTreeTraceNodeEventType::OnEnter);
 				QUICK_SCOPE_CYCLE_COUNTER(StateTree_Task_EnterState);
 				CSV_SCOPED_TIMING_STAT_EXCLUSIVE(StateTree_Task_EnterState);
 				const EStateTreeRunStatus Status = Task.EnterState(*this, CurrentTransition);
@@ -736,6 +766,7 @@ EStateTreeRunStatus FStateTreeExecutionContext::EnterState(const FStateTreeTrans
 		}
 	}
 
+	TRACE_STATETREE_ACTIVE_STATES_EVENT(GetInstanceDebugId(), Exec.ActiveStates);
 	return Result;
 }
 
@@ -831,6 +862,7 @@ void FStateTreeExecutionContext::ExitState(const FStateTreeTransitionResult& Tra
 		CurrentTransition.ChangeType = ExitedStateChangeType[Index];
 
 		STATETREE_LOG(Log, TEXT("%*sState '%s' %s"), Index*UE::StateTree::DebugIndentSize, TEXT(""), *DebugGetStatePath(Transition.CurrentActiveStates, ExitedStateActiveIndex[Index]), *UEnum::GetDisplayValueAsText(CurrentTransition.ChangeType).ToString());
+		TRACE_STATETREE_STATE_EVENT(GetInstanceDebugId(), CurrentHandle.Index, EStateTreeTraceNodeEventType::OnExit);
 
 		// Tasks
 		for (int32 TaskIndex = (State.TasksBegin + State.TasksNum) - 1; TaskIndex >= State.TasksBegin; TaskIndex--)
@@ -848,6 +880,7 @@ void FStateTreeExecutionContext::ExitState(const FStateTreeTransitionResult& Tra
 				if (bShouldCallStateChange)
 				{
 					STATETREE_LOG(Verbose, TEXT("%*s  Task '%s'"), Index*UE::StateTree::DebugIndentSize, TEXT(""), *Task.Name.ToString());
+					TRACE_STATETREE_TASK_EVENT(GetInstanceDebugId(), TaskIndex, EStateTreeTraceNodeEventType::OnExit);
 					{
 						QUICK_SCOPE_CYCLE_COUNTER(StateTree_Task_ExitState);
 						CSV_SCOPED_TIMING_STAT_EXCLUSIVE(StateTree_Task_ExitState);
@@ -882,6 +915,7 @@ void FStateTreeExecutionContext::StateCompleted()
 		FCurrentlyProcessedStateScope StateScope(*this, CurrentHandle);
 		
 		STATETREE_LOG(Verbose, TEXT("%*sState '%s'"), Index*UE::StateTree::DebugIndentSize, TEXT(""), *DebugGetStatePath(Exec.ActiveStates, Index));
+		TRACE_STATETREE_STATE_EVENT(GetInstanceDebugId(), CurrentHandle.Index, EStateTreeTraceNodeEventType::OnCompleted);
 
 		// Notify Tasks
 		for (int32 TaskIndex = (State.TasksBegin + State.TasksNum) - 1; TaskIndex >= State.TasksBegin; TaskIndex--)
@@ -894,6 +928,7 @@ void FStateTreeExecutionContext::StateCompleted()
 				const FStateTreeTaskBase& Task = StateTree.Nodes[TaskIndex].Get<const FStateTreeTaskBase>();
 
 				STATETREE_LOG(Verbose, TEXT("%*s  Task '%s'"), Index*UE::StateTree::DebugIndentSize, TEXT(""), *Task.Name.ToString());
+				TRACE_STATETREE_TASK_EVENT(GetInstanceDebugId(), TaskIndex, EStateTreeTraceNodeEventType::OnCompleted);
 				Task.StateCompleted(*this, Exec.LastTickStatus, Exec.ActiveStates);
 			}
 		}
@@ -904,7 +939,7 @@ EStateTreeRunStatus FStateTreeExecutionContext::TickEvaluatorsAndGlobalTasks(con
 {
 	CSV_SCOPED_TIMING_STAT_EXCLUSIVE(StateTree_TickEvaluators);
 
-	STATETREE_CLOG(StateTree.EvaluatorsNum > 0, Verbose, TEXT("Ticking Evaluators"));
+	STATETREE_CLOG(StateTree.EvaluatorsNum > 0, VeryVerbose, TEXT("Ticking Evaluators"));
 
 	// Tick evaluators
 	int32 InstanceStructIndex = 1; // Exec is at index 0
@@ -920,7 +955,7 @@ EStateTreeRunStatus FStateTreeExecutionContext::TickEvaluatorsAndGlobalTasks(con
 		{
 			StateTree.PropertyBindings.CopyTo(DataViews, Eval.BindingsBatch, DataViews[Eval.DataViewIndex.Get()]);
 		}
-		STATETREE_LOG(Verbose, TEXT("  Tick: '%s'"), *Eval.Name.ToString());
+		STATETREE_LOG(VeryVerbose, TEXT("  Tick: '%s'"), *Eval.Name.ToString());
 		{
 			QUICK_SCOPE_CYCLE_COUNTER(StateTree_Eval_Tick);
 			Eval.Tick(*this, DeltaTime);
@@ -1559,7 +1594,7 @@ bool FStateTreeExecutionContext::TriggerTransitions()
 		if (ParentLinkedState.IsValid())
 		{
 			const EStateTreeRunStatus RunStatus = NextTransition.TargetState.ToCompletionStatus(); 
-			STATETREE_LOG(Verbose, TEXT("Compled subtree '%s' from state '%s' (%s): %s"),
+			STATETREE_LOG(Verbose, TEXT("Completed subtree '%s' from state '%s' (%s): %s"),
 				*GetSafeStateName(ParentLinkedState), *GetSafeStateName(Exec.ActiveStates.Last()), *GetSafeStateName(NextTransition.SourceState), *UEnum::GetDisplayValueAsText(RunStatus).ToString());
 			Exec.CompletedStateHandle = ParentLinkedState;
 			Exec.LastTickStatus = RunStatus;
diff --git a/Engine/Plugins/Runtime/StateTree/Source/StateTreeModule/Private/StateTreeModule.cpp b/Engine/Plugins/Runtime/StateTree/Source/StateTreeModule/Private/StateTreeModule.cpp
--- a/Engine/Plugins/Runtime/StateTree/Source/StateTreeModule/Private/StateTreeModule.cpp
+++ b/Engine/Plugins/Runtime/StateTree/Source/StateTreeModule/Private/StateTreeModule.cpp
@@ -1,6 +1,24 @@
 // Copyright Epic Games, Inc. All Rights Reserved.
 
 #include "StateTreeModule.h"
+#include "StateTreeTypes.h"
+
+#if WITH_STATETREE_DEBUGGER
+#include "Debugger/StateTreeTraceModule.h"
+#include "Features/IModularFeatures.h"
+#include "Misc/CoreDelegates.h"
+#include "ProfilingDebugging/TraceAuxiliary.h"
+#include "StateTreeSettings.h"
+#include "Trace/StoreClient.h"
+#include "Trace/StoreService.h"
+#include "TraceServices/AnalysisService.h"
+#include "TraceServices/ITraceServicesModule.h"
+
+#if WITH_TRACE_STORE
+#include "Misc/Paths.h"
+#endif // WITH_TRACE_STORE
+
+#endif // WITH_STATETREE_DEBUGGER
 
 #define LOCTEXT_NAMESPACE "StateTree"
 
@@ -9,16 +27,141 @@ class FStateTreeModule : public IStateTreeModule
 	// Begin IModuleInterface
 	virtual void StartupModule() override;
 	virtual void ShutdownModule() override;
+
+	void StartTraces();
+	void StopTraces();
+
+#if WITH_STATETREE_DEBUGGER
+	/**
+	 * Gets the store client.
+	 */
+	virtual UE::Trace::FStoreClient* GetStoreClient() override  { return StoreClient.Get(); }
+	
+	TSharedPtr<TraceServices::IAnalysisService> TraceAnalysisService;
+	TSharedPtr<TraceServices::IModuleService> TraceModuleService;
+
+	/** The location of the trace files managed by the trace store. */
+	FString StoreDir;
+
+	/** The client used to connect to the trace store. */
+	TUniquePtr<UE::Trace::FStoreClient> StoreClient;
+
+#if WITH_TRACE_STORE
+	TUniquePtr<UE::Trace::FStoreService> StoreService;
+#endif // WITH_TRACE_STORE
+	
+	FStateTreeTraceModule StateTreeTraceModule;
+	FDelegateHandle StoreServiceHandle;
+#endif // WITH_STATETREE_DEBUGGER
 };
 
 IMPLEMENT_MODULE(FStateTreeModule, StateTreeModule)
 
 void FStateTreeModule::StartupModule()
 {
+	StartTraces();
 }
 
 void FStateTreeModule::ShutdownModule()
 {
+	StopTraces();
+}
+
+void FStateTreeModule::StartTraces()
+{
+#if WITH_STATETREE_DEBUGGER
+	if (!UStateTreeSettings::Get().bUseDebugger || IsRunningCommandlet())
+	{
+		return;
+	}
+
+	ITraceServicesModule& TraceServicesModule = FModuleManager::LoadModuleChecked<ITraceServicesModule>("TraceServices");
+	TraceAnalysisService = TraceServicesModule.GetAnalysisService();
+	TraceModuleService = TraceServicesModule.GetModuleService();
+
+	IModularFeatures::Get().RegisterModularFeature(TraceServices::ModuleFeatureName, &StateTreeTraceModule);
+
+	if (!FTraceAuxiliary::IsConnected())
+	{
+		// cpu tracing is enabled by default, but not connected.
+		// when not connected, disable all channels initially to avoid a whole bunch of extra cpu, memory, and disk overhead from processing all the extra default trace channels
+		////////////////////////////////////////////////////////////////////////////////
+		UE::Trace::EnumerateChannels([](const ANSICHAR* ChannelName, const bool bEnabled, void*)
+			{
+				if (bEnabled)
+				{
+					FString ChannelNameFString(ChannelName);
+					UE::Trace::ToggleChannel(ChannelNameFString.GetCharArray().GetData(), false);
+				}
+			}
+			, nullptr);
+	}
+
+	UE::Trace::ToggleChannel(TEXT("StateTreeDebugChannel"), true);
+	UE::Trace::ToggleChannel(TEXT("FrameChannel"), true);
+
+	// FTraceAuxiliary::FOptions Options;
+	// Options.bExcludeTail = true;
+	// FTraceAuxiliary::Start(FTraceAuxiliary::EConnectionType::Network, TEXT("127.0.0.1"), TEXT(""), &Options, LogStateTree);
+	FTraceAuxiliary::Start(FTraceAuxiliary::EConnectionType::Network, TEXT("127.0.0.1"), TEXT(""), nullptr, LogStateTree);
+
+	// Conditionally create local store service after engine init (if someone doesn't beat us to it).
+	// This is temp until a more formal local server is done by the insights system.
+	StoreServiceHandle = FCoreDelegates::OnFEngineLoopInitComplete.AddLambda([this]
+	{
+		LLM_SCOPE_BYNAME(TEXT("StateTree"));
+
+		if (!GetStoreClient())
+		{
+#if WITH_TRACE_STORE
+			UE_LOG(LogCore, Display, TEXT("StateTree module auto-connecting to internal trace server..."));
+
+			// Create the Store Service.
+			StoreDir = FPaths::ProjectSavedDir() / TEXT("TraceSessions");
+			UE::Trace::FStoreService::FDesc StoreServiceDesc;
+			StoreServiceDesc.StoreDir = *StoreDir;
+			StoreServiceDesc.RecorderPort = 0; // Let system decide port
+			StoreServiceDesc.ThreadCount = 2;
+			StoreService = TUniquePtr<UE::Trace::FStoreService>(UE::Trace::FStoreService::Create(StoreServiceDesc));
+			if (StoreService.IsValid())
+			{
+				// Connect to our newly created store
+				StoreClient = TUniquePtr<UE::Trace::FStoreClient>(UE::Trace::FStoreClient::Connect(TEXT("localhost")));
+				UE::Trace::SendTo(TEXT("localhost"), StoreService->GetRecorderPort());
+				
+				FCoreDelegates::OnPreExit.AddLambda([this]() { StoreService.Reset(); });
+			}
+#else
+			UE_LOG(LogCore, Display, TEXT("StateTree module auto-connecting to local trace server..."));
+
+			StoreClient = TUniquePtr<UE::Trace::FStoreClient>(UE::Trace::FStoreClient::Connect(TEXT("localhost")));
+			if (StoreClient.IsValid())
+			{
+				const UE::Trace::FStoreClient::FStatus* Status = StoreClient->GetStatus();
+				StoreDir = FString(Status->GetStoreDir());
+			}
+#endif // WITH_TRACE_STORE
+		}
+	});
+#endif // WITH_STATETREE_DEBUGGER
+}
+
+void FStateTreeModule::StopTraces()
+{
+#if WITH_STATETREE_DEBUGGER
+	if (StoreServiceHandle.IsValid())
+	{
+		if (FTraceAuxiliary::IsConnected())
+		{
+			FTraceAuxiliary::Stop();
+		}
+		
+		IModularFeatures::Get().UnregisterModularFeature(TraceServices::ModuleFeatureName, &StateTreeTraceModule);
+		
+		FCoreDelegates::OnFEngineLoopInitComplete.Remove(StoreServiceHandle);
+		StoreServiceHandle.Reset();
+	}
+#endif // WITH_STATETREE_DEBUGGER
 }
 
 #undef LOCTEXT_NAMESPACE
diff --git a/Engine/Plugins/Runtime/StateTree/Source/StateTreeModule/Private/StateTreeTypes.cpp b/Engine/Plugins/Runtime/StateTree/Source/StateTreeModule/Private/StateTreeTypes.cpp
--- a/Engine/Plugins/Runtime/StateTree/Source/StateTreeModule/Private/StateTreeTypes.cpp
+++ b/Engine/Plugins/Runtime/StateTree/Source/StateTreeModule/Private/StateTreeTypes.cpp
@@ -17,6 +17,9 @@ const FStateTreeIndex8 FStateTreeIndex8::Invalid = FStateTreeIndex8();
 
 const FStateTreeExternalDataHandle FStateTreeExternalDataHandle::Invalid = FStateTreeExternalDataHandle();
 
+#if WITH_STATETREE_DEBUGGER
+const FStateTreeInstanceDebugId FStateTreeInstanceDebugId::Invalid = FStateTreeInstanceDebugId();
+#endif // WITH_STATETREE_DEBUGGER
 
 bool FStateTreeIndex16::SerializeFromMismatchedTag(const FPropertyTag& Tag, FStructuredArchive::FSlot Slot)
 {
diff --git a/Engine/Plugins/Runtime/StateTree/Source/StateTreeModule/Public/Debugger/IStateTreeTraceProvider.h b/Engine/Plugins/Runtime/StateTree/Source/StateTreeModule/Public/Debugger/IStateTreeTraceProvider.h
new file mode 100644
--- /dev/null
+++ b/Engine/Plugins/Runtime/StateTree/Source/StateTreeModule/Public/Debugger/IStateTreeTraceProvider.h
@@ -0,0 +1,30 @@
+// Copyright Epic Games, Inc. All Rights Reserved.
+
+#pragma once
+
+#if WITH_STATETREE_DEBUGGER
+
+#include "StateTreeTraceTypes.h"
+#include "TraceServices/Model/AnalysisSession.h"
+#include "TraceServices/Containers/Timelines.h"
+
+struct FStateTreeDebuggerInstanceDesc;
+struct FStateTreeInstanceDebugId;
+
+class IStateTreeTraceProvider : public TraceServices::IProvider
+{
+public:
+	typedef TraceServices::ITimeline<FStateTreeTraceEventVariantType> FEventsTimeline;
+
+	virtual void GetActivateInstances(TArray<FStateTreeDebuggerInstanceDesc>& OutInstances) const = 0;
+
+	/**
+	 * Execute given function receiving an event timeline for a given instance or all timelines if instance not specified.  
+	 * @param InstanceId Id of a specific instance to get the timeline for; could be an FStateTreeInstanceDebugId::Invalid to go through all timelines
+	 * @param Callback Function called for timeline(s) matching the provided Id 
+	 * @return True if the specified instance was found for a given Id or at least one timeline was found when no specific Id is provided. 
+	 */
+	virtual bool ReadTimelines(const FStateTreeInstanceDebugId InstanceId, TFunctionRef<void(const FStateTreeInstanceDebugId ProcessedInstanceId, const FEventsTimeline&)> Callback) const = 0;
+};
+
+#endif // WITH_STATETREE_DEBUGGER
diff --git a/Engine/Plugins/Runtime/StateTree/Source/StateTreeModule/Public/Debugger/StateTreeDebugger.h b/Engine/Plugins/Runtime/StateTree/Source/StateTreeModule/Public/Debugger/StateTreeDebugger.h
new file mode 100644
--- /dev/null
+++ b/Engine/Plugins/Runtime/StateTree/Source/StateTreeModule/Public/Debugger/StateTreeDebugger.h
@@ -0,0 +1,194 @@
+// Copyright Epic Games, Inc. All Rights Reserved.
+
+#pragma once
+
+#if WITH_STATETREE_DEBUGGER
+
+#include "IStateTreeTraceProvider.h"
+#include "StateTreeTypes.h"
+#include "StateTree.h"
+#include "Tickable.h"
+#include "TraceServices/Model/Diagnostics.h"
+#include "TraceServices/Model/AnalysisSession.h"
+#include "TraceServices/Model/Frames.h"
+
+namespace UE::Trace
+{
+	class FStoreClient;
+}
+
+class IStateTreeModule;
+class UStateTreeState;
+class UStateTree;
+
+namespace UE::StateTreeDebugger
+{
+	struct FFrameIndexSpan
+	{
+		FFrameIndexSpan() = default;
+		FFrameIndexSpan(const uint64 FrameIdx, const uint32 EventIdx)
+			: FrameIdx(FrameIdx)
+			, EventIdx(EventIdx)
+		{
+		}
+
+		/** Frame index in the analysis session */
+		uint64 FrameIdx = INDEX_NONE;
+
+		/** Index of the first event for that Frame index */
+		int32 EventIdx = INDEX_NONE;
+	};
+}
+
+struct STATETREEMODULE_API FStateTreeDebuggerInstanceDesc
+{
+	FStateTreeDebuggerInstanceDesc() = default;
+	FStateTreeDebuggerInstanceDesc(const UStateTree* InStateTree, const FStateTreeInstanceDebugId InId, const FString& InName);
+
+	bool IsValid() const;
+
+	bool operator==(const FStateTreeDebuggerInstanceDesc& Other) const
+	{
+		return StateTree == Other.StateTree
+			&& Id == Other.Id;
+	}
+
+	bool operator!=(const FStateTreeDebuggerInstanceDesc& Other) const
+	{
+		return !(*this == Other);
+	}
+
+	friend FString LexToString(const FStateTreeDebuggerInstanceDesc& InstanceDesc)
+	{
+		return FString::Printf(TEXT("%s | %s | %s"),
+			*GetNameSafe(InstanceDesc.StateTree.Get()),
+			*LexToString(InstanceDesc.Id),
+			*InstanceDesc.Name);
+	}
+
+	TWeakObjectPtr<const UStateTree> StateTree = nullptr;
+	FString Name;
+	FStateTreeInstanceDebugId Id = FStateTreeInstanceDebugId::Invalid;
+};
+
+
+DECLARE_DELEGATE_TwoParams(FOnStateTreeDebuggerTracesUpdated, TConstArrayView<FStateTreeTraceEventVariantType>, TConstArrayView<UE::StateTreeDebugger::FFrameIndexSpan>);
+DECLARE_DELEGATE_TwoParams(FOnStateTreeDebuggerBreakpointHit, FStateTreeInstanceDebugId InstanceId, FStateTreeStateHandle StateHandle);
+DECLARE_DELEGATE_OneParam(FOnStateTreeDebuggerBreakpointsChanged, TConstArrayView<FStateTreeStateHandle> Breakpoints);
+DECLARE_DELEGATE_OneParam(FOnStateTreeDebuggerActiveStatesChanges, TConstArrayView<FStateTreeStateHandle> ActiveStates);
+
+struct STATETREEMODULE_API FStateTreeDebugger : FTickableGameObject
+{
+	struct FTraceDescriptor
+	{
+		FTraceDescriptor() = default;
+		FTraceDescriptor(const FString& Name, const uint32 Id) : Name(Name), TraceId(Id) {}
+		
+		bool operator==(const FTraceDescriptor& Other) const { return Other.TraceId == TraceId; }
+		bool operator!=(const FTraceDescriptor& Other) const { return !(Other == *this); }
+		bool IsValid() const { return TraceId != INDEX_NONE; }
+
+		FString Name;
+		uint32 TraceId = INDEX_NONE;
+
+		TraceServices::FSessionInfo SessionInfo;
+	};
+	
+	FStateTreeDebugger();
+	virtual ~FStateTreeDebugger() override;
+
+	/** Stops reading traces every frame to preserve current state */
+	void Pause();
+
+	/** Resumes reading traces every frame */
+	void Unpause();
+
+	/** Forces a single refresh to latest state. Useful when simulation is paused. */
+	void SyncToCurrentSessionDuration();
+	
+	bool CanStepBack() const;
+	bool CanStepForward() const;
+	void StepBack();
+	void StepForward();
+	
+	bool IsAnalysisSessionActive() const { return GetAnalysisSession() != nullptr; }
+
+	void GetActivateInstances(TArray<FStateTreeDebuggerInstanceDesc>& OutInstances) const;
+	const FStateTreeDebuggerInstanceDesc& GetDebuggedInstance() const;
+	void SetDebuggedInstance(const FStateTreeDebuggerInstanceDesc& SelectedInstance);
+	FString GetDebuggedInstanceDescription() const;
+	static FString DescribeTrace(const FTraceDescriptor& TraceDescriptor);
+	static FString DescribeInstance(const FStateTreeDebuggerInstanceDesc& StateTreeInstanceDesc);
+
+	void GetLiveTraces(TArray<FTraceDescriptor>& OutTraceDescriptors) const;
+	void StartLastLiveSessionAnalysis();
+	void StartSessionAnalysis(const FTraceDescriptor& TraceDescriptor);
+	FTraceDescriptor GetDebuggedTraceDescriptor() const { return ActiveSessionTraceDescriptor; }
+	FString GetDebuggedTraceDescription() const;
+	
+	void ToggleBreakpoints(const TConstArrayView<FStateTreeStateHandle> SelectedStates);
+
+	FOnStateTreeDebuggerTracesUpdated OnTracesUpdated;
+	FOnStateTreeDebuggerBreakpointHit OnBreakpointHit;
+	FOnStateTreeDebuggerBreakpointsChanged OnBreakpointsChanged;
+	FOnStateTreeDebuggerActiveStatesChanges OnActiveStatesChanged;
+
+protected:
+	virtual void Tick(float DeltaTime) override;
+	virtual bool IsTickable() const override;
+	virtual bool IsTickableWhenPaused() const override { return true; }
+	virtual bool IsTickableInEditor() const override { return true; }
+	virtual TStatId GetStatId() const override { RETURN_QUICK_DECLARE_CYCLE_STAT(FStateTreeDebugger, STATGROUP_Tickables); }
+	
+private:
+	void ResetAnalysisData();
+	
+	void ReadTrace(double ScrubTime);
+	void ReadTrace(uint64 FrameIndex);
+	void ReadTrace(
+		const TraceServices::IAnalysisSession& Session,
+		const TraceServices::IFrameProvider& FrameProvider,
+		const TraceServices::FFrame& Frame
+		);
+	
+	void StopAnalysis();
+
+	void SetCurrentScrubTime(double RecordingDuration);	
+	void SetActiveStates(const TConstArrayView<FStateTreeStateHandle> NewActiveStates);
+	
+	const TraceServices::IAnalysisSession* GetAnalysisSession() const;
+	UE::Trace::FStoreClient* GetStoreClient() const;
+
+	bool ProcessEvent(const FStateTreeInstanceDebugId InstanceId, const TraceServices::FFrame& Frame, const FStateTreeTraceEventVariantType& Event);
+	void AddEvents(float StartTime, float EndTime, const TraceServices::IFrameProvider& FrameProvider, const IStateTreeTraceProvider& StateTreeTraceProvider);
+	void UpdateMetadata(FTraceDescriptor& TraceDescriptor) const;
+
+	IStateTreeModule& StateTreeModule;
+	
+	/** The trace analysis session. */
+	TSharedPtr<const TraceServices::IAnalysisSession> AnalysisSession;
+	
+	FTraceDescriptor ActiveSessionTraceDescriptor;
+
+	FStateTreeDebuggerInstanceDesc DebuggedInstance;
+	
+	TArray<FStateTreeStateHandle> StatesWithBreakpoint;
+	TArray<FStateTreeStateHandle> ActiveStates;
+	TArray<FStateTreeTraceEventVariantType> Events;
+	
+	FStateTreeInstanceDebugId HitBreakpointInstanceId = FStateTreeInstanceDebugId::Invalid;
+	int32 HitBreakpointStateIndex = INDEX_NONE;
+
+	inline static constexpr double UnsetTime = -1;
+	double RecordingDuration = UnsetTime;
+	double CurrentScrubTime = UnsetTime;
+	double LastTraceReadTime = UnsetTime;
+
+	uint64 CurrentFrameIndex = INDEX_NONE;
+	TArray<UE::StateTreeDebugger::FFrameIndexSpan> FramesWithEvents;
+
+	bool bProcessingInitialEvents = false;
+	bool bPaused = false;
+};
+
+#endif // WITH_STATETREE_DEBUGGER
diff --git a/Engine/Plugins/Runtime/StateTree/Source/StateTreeModule/Public/Debugger/StateTreeTrace.h b/Engine/Plugins/Runtime/StateTree/Source/StateTreeModule/Public/Debugger/StateTreeTrace.h
new file mode 100644
--- /dev/null
+++ b/Engine/Plugins/Runtime/StateTree/Source/StateTreeModule/Public/Debugger/StateTreeTrace.h
@@ -0,0 +1,64 @@
+// Copyright Epic Games, Inc. All Rights Reserved.
+
+#pragma once
+
+#if WITH_STATETREE_DEBUGGER
+
+#include "Trace/Trace.h"
+
+class UStateTree;
+struct FStateTreeActiveStates;
+struct FStateTreeInstanceDebugId;
+enum class EStateTreeTraceNodeEventType : uint8;
+enum class EStateTreeTraceInstanceEventType : uint8;
+
+UE_TRACE_CHANNEL_EXTERN(StateTreeDebugChannel, STATETREEMODULE_API)
+
+struct FStateTreeTrace
+{
+	static void OutputInstanceLifetimeEvent(const FStateTreeInstanceDebugId InstanceId, const UStateTree* StateTree, const TCHAR* InstanceName, const EStateTreeTraceInstanceEventType EventType);
+	static void OutputLogEventTrace(const FStateTreeInstanceDebugId InstanceId, const TCHAR* Fmt, ...);
+	static void OutputStateEventTrace(const FStateTreeInstanceDebugId InstanceId, const uint16 StateIdx, const EStateTreeTraceNodeEventType EventType);
+	static void OutputTaskEventTrace(const FStateTreeInstanceDebugId InstanceId, const uint16 TaskIdx, const EStateTreeTraceNodeEventType EventType);
+	static void OutputActiveStatesEventTrace(const FStateTreeInstanceDebugId InstanceId, const FStateTreeActiveStates& ActiveStates);
+};
+
+#define TRACE_STATETREE_INSTANCE_EVENT(InstanceID, StateTree, InstanceName, EventType) \
+	if (UE_TRACE_CHANNELEXPR_IS_ENABLED(StateTreeDebugChannel)) \
+	{ \
+		FStateTreeTrace::OutputInstanceLifetimeEvent(InstanceID, StateTree, InstanceName, EventType); \
+	}
+
+#define TRACE_STATETREE_LOG_EVENT(InstanceId, Format, ...) \
+	if (UE_TRACE_CHANNELEXPR_IS_ENABLED(StateTreeDebugChannel)) \
+	{ \
+		FStateTreeTrace::OutputLogEventTrace(InstanceId, Format, ##__VA_ARGS__); \
+	}
+
+#define TRACE_STATETREE_STATE_EVENT(InstanceId, StateIdx, EventType) \
+	if (UE_TRACE_CHANNELEXPR_IS_ENABLED(StateTreeDebugChannel)) \
+	{ \
+		FStateTreeTrace::OutputStateEventTrace(InstanceId, StateIdx, EventType); \
+	}
+
+#define TRACE_STATETREE_TASK_EVENT(InstanceId, TaskIdx, EventType) \
+	if (UE_TRACE_CHANNELEXPR_IS_ENABLED(StateTreeDebugChannel)) \
+	{ \
+		FStateTreeTrace::OutputTaskEventTrace(InstanceId, TaskIdx, EventType); \
+	}
+
+#define TRACE_STATETREE_ACTIVE_STATES_EVENT(InstanceId, ActivateStates) \
+	if (UE_TRACE_CHANNELEXPR_IS_ENABLED(StateTreeDebugChannel)) \
+	{ \
+		FStateTreeTrace::OutputActiveStatesEventTrace(InstanceId, ActivateStates); \
+	}
+
+#else //STATETREE_DEBUG_TRACE_ENABLED
+
+#define TRACE_STATETREE_INSTANCE_EVENT(InstanceID, StateTree, InstanceName, EventType)
+#define TRACE_STATETREE_LOG_EVENT(InstanceId, Format, ...)
+#define TRACE_STATETREE_STATE_EVENT(InstanceId, StateIdx, EventType)
+#define TRACE_STATETREE_TASK_EVENT(InstanceId, TaskIdx, EventType)
+#define TRACE_STATETREE_ACTIVE_STATES_EVENT(InstanceId, ActivateStates)
+
+#endif // STATETREE_DEBUG_TRACE_ENABLED
diff --git a/Engine/Plugins/Runtime/StateTree/Source/StateTreeModule/Public/Debugger/StateTreeTraceAnalyzer.h b/Engine/Plugins/Runtime/StateTree/Source/StateTreeModule/Public/Debugger/StateTreeTraceAnalyzer.h
new file mode 100644
--- /dev/null
+++ b/Engine/Plugins/Runtime/StateTree/Source/StateTreeModule/Public/Debugger/StateTreeTraceAnalyzer.h
@@ -0,0 +1,34 @@
+// Copyright Epic Games, Inc. All Rights Reserved.
+
+#pragma once
+
+#if WITH_STATETREE_DEBUGGER
+
+#include "Trace/Analyzer.h"
+
+class FStateTreeTraceProvider;
+namespace TraceServices { class IAnalysisSession; }
+
+class FStateTreeTraceAnalyzer : public UE::Trace::IAnalyzer
+{
+public:
+	FStateTreeTraceAnalyzer(TraceServices::IAnalysisSession& InSession, FStateTreeTraceProvider& InProvider);
+
+	virtual void OnAnalysisBegin(const FOnAnalysisContext& Context) override;
+	virtual void OnAnalysisEnd() override {}
+	virtual bool OnEvent(uint16 RouteId, EStyle Style, const FOnEventContext& Context) override;
+
+private:
+	enum : uint16
+	{
+		RouteId_Instance,
+		RouteId_LogMessage,
+		RouteId_State,
+		RouteId_Task,
+		RouteId_ActiveStates
+	};
+
+	TraceServices::IAnalysisSession& Session;
+	FStateTreeTraceProvider& Provider;
+};
+#endif // WITH_STATETREE_DEBUGGER
\ No newline at end of file
diff --git a/Engine/Plugins/Runtime/StateTree/Source/StateTreeModule/Public/Debugger/StateTreeTraceModule.h b/Engine/Plugins/Runtime/StateTree/Source/StateTreeModule/Public/Debugger/StateTreeTraceModule.h
new file mode 100644
--- /dev/null
+++ b/Engine/Plugins/Runtime/StateTree/Source/StateTreeModule/Public/Debugger/StateTreeTraceModule.h
@@ -0,0 +1,21 @@
+// Copyright Epic Games, Inc. All Rights Reserved.
+
+#pragma once
+
+#if WITH_STATETREE_DEBUGGER
+
+#include "TraceServices/ModuleService.h"
+
+class FStateTreeTraceModule : public TraceServices::IModule
+{
+public:
+	// TraceServices::IModule interface
+	virtual void GetModuleInfo(TraceServices::FModuleInfo& OutModuleInfo) override;
+	virtual void OnAnalysisBegin(TraceServices::IAnalysisSession& Session) override;
+	virtual void GetLoggers(TArray<const TCHAR *>& OutLoggers) override;
+
+private:
+	static FName ModuleName;
+};
+
+#endif // WITH_STATETREE_DEBUGGER
\ No newline at end of file
diff --git a/Engine/Plugins/Runtime/StateTree/Source/StateTreeModule/Public/Debugger/StateTreeTraceProvider.h b/Engine/Plugins/Runtime/StateTree/Source/StateTreeModule/Public/Debugger/StateTreeTraceProvider.h
new file mode 100644
--- /dev/null
+++ b/Engine/Plugins/Runtime/StateTree/Source/StateTreeModule/Public/Debugger/StateTreeTraceProvider.h
@@ -0,0 +1,40 @@
+// Copyright Epic Games, Inc. All Rights Reserved.
+
+#pragma once
+
+#if WITH_STATETREE_DEBUGGER
+
+#include "IStateTreeTraceProvider.h"
+#include "Model/PointTimeline.h"
+#include "StateTreeTypes.h" // required to compile TMap<FStateTreeInstanceDebugId, ...>
+
+namespace TraceServices { class IAnalysisSession; }
+enum class EStateTreeTraceInstanceEventType : uint8;
+class UStateTree;
+
+class FStateTreeTraceProvider : public IStateTreeTraceProvider
+{
+public:
+	static FName ProviderName;
+
+	explicit FStateTreeTraceProvider(TraceServices::IAnalysisSession& InSession);
+	
+	void AppendEvent(FStateTreeInstanceDebugId InInstanceId, double InTime, const FStateTreeTraceEventVariantType& Event);
+	void AppendInstance(
+		const UStateTree* InStateTree,
+		const FStateTreeInstanceDebugId InInstanceId,
+		const TCHAR* InInstanceName,
+		const EStateTreeTraceInstanceEventType EventType);
+
+protected:
+	/** IStateTreeDebuggerProvider interface */
+	virtual void GetActivateInstances(TArray<FStateTreeDebuggerInstanceDesc>& OutInstances) const override;
+	virtual bool ReadTimelines(const FStateTreeInstanceDebugId InstanceId, TFunctionRef<void(const FStateTreeInstanceDebugId ProcessedInstanceId, const FEventsTimeline&)> Callback) const override;
+
+private:
+	TraceServices::IAnalysisSession& Session;
+	TArray<FStateTreeDebuggerInstanceDesc> ActiveInstances;
+	TMap<FStateTreeInstanceDebugId, uint32> InstanceIdToDebuggerEntryTimelines;
+	TArray<TSharedRef<TraceServices::TPointTimeline<FStateTreeTraceEventVariantType>>> EventsTimelines;
+};
+#endif // WITH_STATETREE_DEBUGGER
\ No newline at end of file
diff --git a/Engine/Plugins/Runtime/StateTree/Source/StateTreeModule/Public/Debugger/StateTreeTraceTypes.h b/Engine/Plugins/Runtime/StateTree/Source/StateTreeModule/Public/Debugger/StateTreeTraceTypes.h
new file mode 100644
--- /dev/null
+++ b/Engine/Plugins/Runtime/StateTree/Source/StateTreeModule/Public/Debugger/StateTreeTraceTypes.h
@@ -0,0 +1,53 @@
+﻿// Copyright Epic Games, Inc. All Rights Reserved.
+
+#pragma once
+
+#include "Misc/TVariant.h"
+#include "StateTreeTraceTypes.generated.h"
+
+struct FStateTreeStateHandle;
+
+UENUM()
+enum class EStateTreeTraceInstanceEventType : uint8
+{
+	Started,
+	Stopped
+};
+
+UENUM()
+enum class EStateTreeTraceNodeEventType : uint8
+{
+	Unset,
+	OnEnter,
+	OnExit,
+	OnCompleted
+};
+
+#if WITH_STATETREE_DEBUGGER
+
+struct FStateTreeTraceLogEvent
+{
+	FString Message;
+};
+
+struct FStateTreeTraceStateEvent
+{
+	int16 StateIdx = INDEX_NONE;
+	EStateTreeTraceNodeEventType EventType = EStateTreeTraceNodeEventType::Unset;
+};
+
+struct FStateTreeTraceTaskEvent
+{
+	int16 TaskIdx = INDEX_NONE;
+	EStateTreeTraceNodeEventType EventType = EStateTreeTraceNodeEventType::Unset;
+};
+
+struct FStateTreeTraceActiveStatesEvent
+{
+	TArray<FStateTreeStateHandle> ActiveStates;
+};
+
+/** Type aliases for statetree trace events */
+using FStateTreeTraceEventVariantType = TVariant<FStateTreeTraceLogEvent, FStateTreeTraceStateEvent, FStateTreeTraceTaskEvent, FStateTreeTraceActiveStatesEvent>;
+
+#endif // WITH_STATETREE_DEBUGGER
diff --git a/Engine/Plugins/Runtime/StateTree/Source/StateTreeModule/Public/StateTree.h b/Engine/Plugins/Runtime/StateTree/Source/StateTreeModule/Public/StateTree.h
--- a/Engine/Plugins/Runtime/StateTree/Source/StateTreeModule/Public/StateTree.h
+++ b/Engine/Plugins/Runtime/StateTree/Source/StateTreeModule/Public/StateTree.h
@@ -103,6 +103,21 @@ public:
 	/** @return schema that was used to compile the StateTree. */
 	const UStateTreeSchema* GetSchema() const { return Schema; }
 
+	/** @return Pointer to a state or null if state not found */ 
+	const FCompactStateTreeState* GetStateFromHandle(const FStateTreeStateHandle StateHandle) const;
+
+	/** @return State handle matching a given Id; invalid handle if state not found. */
+	FStateTreeStateHandle GetStateHandleFromId(const FGuid Id) const;
+
+	/** @return Id of the state matching a given state handle; invalid Id if state not found. */
+	FGuid GetStateIdFromHandle(FStateTreeStateHandle Handle) const;
+
+	/** @return Struct view of the node matching a given state index; invalid view if state not found. */
+	FConstStructView GetNode(int32 NodeIndex) const;
+
+	/** @return View of all states. */
+	TConstArrayView<FCompactStateTreeState> GetStates() const { return States; }
+
 #if WITH_EDITOR
 	/** Resets the compiled data to empty. */
 	void ResetCompiled();
@@ -115,11 +130,11 @@ public:
 	/** Edit time data for the StateTree, instance of UStateTreeEditorData */
 	UPROPERTY()
 	TObjectPtr<UObject> EditorData;
+#endif
 
-	/** Hash of the editor data from last compile. */
+	/** Hash of the editor data from last compile. Also used to detect mismatching events from recorded traces. */
 	UPROPERTY()
 	uint32 LastCompiledEditorDataHash = 0;
-#endif
 
 protected:
 	
@@ -186,6 +201,10 @@ private:
 	UPROPERTY()
 	FStateTreePropertyBindings PropertyBindings;
 
+	/** Mapping of state guid for the Editor and state handles, created at compilation. */
+	UPROPERTY()
+	TArray<FStateTreeStateIdToHandle> IDToStateMappings;
+	
 	/**
 	 * Parameters that could be used for bindings within the Tree.
 	 * Default values are stored within the asset but StateTreeReference can be used to parameterized the tree.
diff --git a/Engine/Plugins/Runtime/StateTree/Source/StateTreeModule/Public/StateTreeExecutionContext.h b/Engine/Plugins/Runtime/StateTree/Source/StateTreeModule/Public/StateTreeExecutionContext.h
--- a/Engine/Plugins/Runtime/StateTree/Source/StateTreeModule/Public/StateTreeExecutionContext.h
+++ b/Engine/Plugins/Runtime/StateTree/Source/StateTreeModule/Public/StateTreeExecutionContext.h
@@ -14,6 +14,7 @@ struct FStateTreeTaskBase;
 struct FStateTreeConditionBase;
 struct FStateTreeEvent;
 struct FStateTreeTransitionRequest;
+struct FStateTreeInstanceDebugId;
 
 /**
  * StateTree Execution Context is a helper that is used to update and access StateTree instance data.
@@ -160,7 +161,7 @@ public:
 	/** @return Pointer to a State or null if state not found */ 
 	const FCompactStateTreeState* GetStateFromHandle(const FStateTreeStateHandle StateHandle) const
 	{
-		return StateTree.States.IsValidIndex(StateHandle.Index) ? &StateTree.States[StateHandle.Index] : nullptr;
+		return StateTree.GetStateFromHandle(StateHandle);
 	}
 
 	/** @return Array view to external data descriptors associated with this context. Note: Init() must be called before calling this method. */
@@ -266,7 +267,11 @@ public:
 
 protected:
 
-	/** @return Prefix that will be used by STATETREE_LOG and STATETREE_CLOG, empty by default. */
+#if WITH_STATETREE_DEBUGGER
+	FStateTreeInstanceDebugId GetInstanceDebugId() const;
+#endif // WITH_STATETREE_DEBUGGER
+
+	/** @return Prefix that will be used by STATETREE_LOG and STATETREE_CLOG, Owner name by default. */
 	virtual FString GetInstanceDescription() const;
 
 	UE_DEPRECATED(5.2, "Use BeginDelayedTransition() instead.")
diff --git a/Engine/Plugins/Runtime/StateTree/Source/StateTreeModule/Public/StateTreeModule.h b/Engine/Plugins/Runtime/StateTree/Source/StateTreeModule/Public/StateTreeModule.h
--- a/Engine/Plugins/Runtime/StateTree/Source/StateTreeModule/Public/StateTreeModule.h
+++ b/Engine/Plugins/Runtime/StateTree/Source/StateTreeModule/Public/StateTreeModule.h
@@ -4,6 +4,11 @@
 
 #include "Modules/ModuleManager.h"
 
+namespace UE::Trace
+{
+	class FStoreClient;
+}
+
 /**
 * The public interface to this module
 */
@@ -18,7 +23,7 @@ public:
 	*
 	* @return Returns singleton instance, loading the module on demand if needed
 	*/
-	static inline IStateTreeModule& Get()
+	static IStateTreeModule& Get()
 	{
 		return FModuleManager::LoadModuleChecked<IStateTreeModule>("StateTreeModule");
 	}
@@ -28,10 +33,17 @@ public:
 	*
 	* @return True if the module is loaded and ready to use
 	*/
-	static inline bool IsAvailable()
+	static bool IsAvailable()
 	{
 		return FModuleManager::Get().IsModuleLoaded("StateTreeModule");
 	}
+
+#if WITH_STATETREE_DEBUGGER
+	/**
+	 * Gets the store client.
+	 */
+	virtual UE::Trace::FStoreClient* GetStoreClient() = 0;
+#endif // WITH_STATETREE_DEBUGGER
 };
 
 
diff --git a/Engine/Plugins/Runtime/StateTree/Source/StateTreeModule/Public/StateTreeSettings.h b/Engine/Plugins/Runtime/StateTree/Source/StateTreeModule/Public/StateTreeSettings.h
new file mode 100644
--- /dev/null
+++ b/Engine/Plugins/Runtime/StateTree/Source/StateTreeModule/Public/StateTreeSettings.h
@@ -0,0 +1,21 @@
+﻿// Copyright Epic Games, Inc. All Rights Reserved.
+
+#pragma once
+
+#include "Engine/DeveloperSettings.h"
+#include "StateTreeSettings.generated.h"
+
+/**
+ * Default StateTree settings
+ */
+UCLASS(config = StateTree, defaultconfig, DisplayName = "StateTree")
+class STATETREEMODULE_API UStateTreeSettings : public UDeveloperSettings
+{
+	GENERATED_BODY()
+
+public:
+	static UStateTreeSettings& Get() { return *CastChecked<UStateTreeSettings>(UStateTreeSettings::StaticClass()->GetDefaultObject()); }
+	
+	UPROPERTY(EditDefaultsOnly, Category = StateTree, config)
+	bool bUseDebugger = false;
+};
\ No newline at end of file
diff --git a/Engine/Plugins/Runtime/StateTree/Source/StateTreeModule/Public/StateTreeTypes.h b/Engine/Plugins/Runtime/StateTree/Source/StateTreeModule/Public/StateTreeTypes.h
--- a/Engine/Plugins/Runtime/StateTree/Source/StateTreeModule/Public/StateTreeTypes.h
+++ b/Engine/Plugins/Runtime/StateTree/Source/StateTreeModule/Public/StateTreeTypes.h
@@ -175,6 +175,7 @@ struct STATETREEMODULE_API FStateTreeStateHandle
 	explicit FStateTreeStateHandle(uint16 InIndex) : Index(InIndex) {}
 
 	bool IsValid() const { return Index != InvalidIndex; }
+	void Invalidate() { Index = InvalidIndex; }
 	bool IsCompletionState() const { return Index == SucceededIndex || Index == FailedIndex; }
 	EStateTreeRunStatus ToCompletionStatus() const
 	{
@@ -811,6 +812,27 @@ enum class EStateTreePropertyUsage : uint8
 	Output,
 };
 
+/**
+ * Pair of state guid and its associated state handle created at compilation.
+ */
+USTRUCT()
+struct STATETREEMODULE_API FStateTreeStateIdToHandle
+{
+	GENERATED_BODY()
+
+	FStateTreeStateIdToHandle() = default;
+	explicit FStateTreeStateIdToHandle(const FGuid& Id, const FStateTreeStateHandle Handle)
+		: Id(Id)
+		, Handle(Handle)
+	{
+	}
+
+	UPROPERTY();
+	FGuid Id;
+
+	UPROPERTY();
+	FStateTreeStateHandle Handle;
+};
 
 /**
  * Describes an external data. The data can point to a struct or object.
@@ -1085,6 +1107,46 @@ struct STATETREEMODULE_API FStateTreeTransitionDelayedState
 	float TimeLeft = 0.0f;
 };
 
+#if WITH_STATETREE_DEBUGGER
+struct FStateTreeInstanceDebugId
+{
+	FStateTreeInstanceDebugId() = default;
+	FStateTreeInstanceDebugId(const uint32 InstanceId, const uint32 SerialNumber)
+		: Id(InstanceId), SerialNumber(SerialNumber)
+	{
+	}
+	
+	bool IsValid() const { return Id != INDEX_NONE && SerialNumber != INDEX_NONE; }
+	bool IsInvalid() const { return !IsValid(); }
+	void Invalidate() { *this = Invalid; }
+
+	bool operator==(const FStateTreeInstanceDebugId& Other) const
+	{
+		return Id == Other.Id && SerialNumber == Other.SerialNumber;
+	}
+
+	bool operator!=(const FStateTreeInstanceDebugId& Other) const
+	{
+		return !(*this == Other);
+	}
+
+	friend uint32 GetTypeHash(const FStateTreeInstanceDebugId InstanceDebugId)
+	{
+		return HashCombine(InstanceDebugId.Id, InstanceDebugId.SerialNumber);
+	}
+
+	friend FString LexToString(const FStateTreeInstanceDebugId InstanceDebugId)
+	{
+		return FString::Printf(TEXT("0x%x | %d"), InstanceDebugId.Id, InstanceDebugId.SerialNumber);
+	}
+
+	static const FStateTreeInstanceDebugId Invalid;
+	
+	uint32 Id = INDEX_NONE;
+	uint32 SerialNumber = INDEX_NONE;
+};
+#endif // WITH_STATETREE_DEBUGGER
+
 USTRUCT()
 struct STATETREEMODULE_API FStateTreeExecutionState
 {
@@ -1099,6 +1161,11 @@ struct STATETREEMODULE_API FStateTreeExecutionState
 	/** Currently active states */
 	FStateTreeActiveStates ActiveStates;
 
+#if WITH_STATETREE_DEBUGGER
+	/** Id for the active instance used for debugging. */
+	mutable FStateTreeInstanceDebugId InstanceDebugId;
+#endif
+
 	/** Index of the first task struct in the currently initialized instance data. */
 	FStateTreeIndex16 FirstTaskStructIndex = FStateTreeIndex16::Invalid;
 	
diff --git a/Engine/Plugins/Runtime/StateTree/Source/StateTreeModule/StateTreeModule.Build.cs b/Engine/Plugins/Runtime/StateTree/Source/StateTreeModule/StateTreeModule.Build.cs
--- a/Engine/Plugins/Runtime/StateTree/Source/StateTreeModule/StateTreeModule.Build.cs
+++ b/Engine/Plugins/Runtime/StateTree/Source/StateTreeModule/StateTreeModule.Build.cs
@@ -15,6 +15,7 @@ namespace UnrealBuildTool.Rules
 				new string[] {
 					"Core",
 					"CoreUObject",
+					"DeveloperSettings",
 					"Engine",
 					"AIModule",
 					"GameplayTags",
@@ -39,7 +40,22 @@ namespace UnrealBuildTool.Rules
 				);
 			}
 
+			if (Target.Configuration != UnrealTargetConfiguration.Shipping || Target.bBuildEditor)
+			{
+				PublicDefinitions.Add("WITH_STATETREE_DEBUGGER=1");
+				PublicDependencyModuleNames.AddRange(
+					new string[]
+					{
+						"TraceLog",
+						"TraceServices",
+						"TraceAnalysis"
+					}
+				);
+			}
+			else
+			{
+				PublicDefinitions.Add("WITH_STATETREE_DEBUGGER=0");
+			}
 		}
-
 	}
 }
diff --git a/Engine/Source/Editor/UnrealEd/Public/Kismet2/DebuggerCommands.h b/Engine/Source/Editor/UnrealEd/Public/Kismet2/DebuggerCommands.h
--- a/Engine/Source/Editor/UnrealEd/Public/Kismet2/DebuggerCommands.h
+++ b/Engine/Source/Editor/UnrealEd/Public/Kismet2/DebuggerCommands.h
@@ -34,7 +34,7 @@ public:
 //////////////////////////////////////////////////////////////////////////
 // FPlayWorldCommands
 
-class FPlayWorldCommands : public TCommands<FPlayWorldCommands>
+class UNREALED_API FPlayWorldCommands : public TCommands<FPlayWorldCommands>
 {
 private:
 
@@ -54,7 +54,7 @@ public:
 	static void BindGlobalPlayWorldCommands();
 
 	/** Populates a toolbar with the menu commands for play-world control (pause/resume/stop/possess/eject/step/show current loc) */
-	UNREALED_API static void BuildToolbar( FToolMenuSection& InSection, bool bIncludeLaunchButtonAndOptions = false );
+	static void BuildToolbar( FToolMenuSection& InSection, bool bIncludeLaunchButtonAndOptions = false );
 
 	/**
 	* Return the active widget that processes play world actions for PIE
@@ -73,7 +73,7 @@ public:
 	/** 
 	 * A command list that can be passed around and isn't bound to an instance of any tool or editor. 
 	 */
-	UNREALED_API static TSharedPtr<FUICommandList> GlobalPlayWorldActions;
+	static TSharedPtr<FUICommandList> GlobalPlayWorldActions;
 
 public:
 
