Tentative CIS fix by not compiling StateTree debugger for non win64 platforms

diff --git a/Engine/Plugins/Runtime/StateTree/Source/StateTreeModule/StateTreeModule.Build.cs b/Engine/Plugins/Runtime/StateTree/Source/StateTreeModule/StateTreeModule.Build.cs
--- a/Engine/Plugins/Runtime/StateTree/Source/StateTreeModule/StateTreeModule.Build.cs
+++ b/Engine/Plugins/Runtime/StateTree/Source/StateTreeModule/StateTreeModule.Build.cs
@@ -40,7 +40,8 @@ namespace UnrealBuildTool.Rules
 				);
 			}
 
-			if (Target.Configuration != UnrealTargetConfiguration.Shipping || Target.bBuildEditor)
+			if (Target.Platform == UnrealTargetPlatform.Win64 && 
+				(Target.Configuration != UnrealTargetConfiguration.Shipping || Target.bBuildEditor))
 			{
 				PublicDefinitions.Add("WITH_STATETREE_DEBUGGER=1");
 				PublicDependencyModuleNames.AddRange(
