Tentative CIS fix by not compiling StateTree debugger for non win64 platforms (part 2)

diff --git a/Engine/Plugins/Runtime/StateTree/Source/StateTreeEditorModule/StateTreeEditorModule.Build.cs b/Engine/Plugins/Runtime/StateTree/Source/StateTreeEditorModule/StateTreeEditorModule.Build.cs
--- a/Engine/Plugins/Runtime/StateTree/Source/StateTreeEditorModule/StateTreeEditorModule.Build.cs
+++ b/Engine/Plugins/Runtime/StateTree/Source/StateTreeEditorModule/StateTreeEditorModule.Build.cs
@@ -58,7 +58,14 @@ namespace UnrealBuildTool.Rules
 				System.IO.Path.Combine(GetModuleDirectory("PropertyEditor"), "Private"),
 			});
 		
-			PublicDefinitions.Add("WITH_STATETREE_DEBUGGER=1");
+			if (Target.Platform == UnrealTargetPlatform.Win64)
+			{
+				PublicDefinitions.Add("WITH_STATETREE_DEBUGGER=1");
+			}
+			else
+			{
+				PublicDefinitions.Add("WITH_STATETREE_DEBUGGER=0");
+			}
 		}
 	}
 }
