StateTree: Fix core content references in editor style.

diff --git a/Engine/Plugins/Runtime/StateTree/Source/StateTreeEditorModule/Private/StateTreeEditorStyle.cpp b/Engine/Plugins/Runtime/StateTree/Source/StateTreeEditorModule/Private/StateTreeEditorStyle.cpp
--- a/Engine/Plugins/Runtime/StateTree/Source/StateTreeEditorModule/Private/StateTreeEditorStyle.cpp
+++ b/Engine/Plugins/Runtime/StateTree/Source/StateTreeEditorModule/Private/StateTreeEditorStyle.cpp
@@ -72,10 +72,10 @@ FStateTreeEditorStyle::FStateTreeEditorStyle()
 		FEditableTextBoxStyle StateTitleEditableText = FEditableTextBoxStyle()
 			.SetTextStyle(NormalText)
 			.SetFont(CORE_FONT("Fonts/Roboto-Bold", 12))
-			.SetBackgroundImageNormal(BOX_BRUSH("Common/TextBox", FMargin(4.0f / 16.0f)))
-			.SetBackgroundImageHovered(BOX_BRUSH("Common/TextBox_Hovered", FMargin(4.0f / 16.0f)))
-			.SetBackgroundImageFocused(BOX_BRUSH("Common/TextBox_Hovered", FMargin(4.0f / 16.0f)))
-			.SetBackgroundImageReadOnly(BOX_BRUSH("Common/TextBox_ReadOnly", FMargin(4.0f / 16.0f)))
+			.SetBackgroundImageNormal(CORE_BOX_BRUSH("Common/TextBox", FMargin(4.0f / 16.0f)))
+			.SetBackgroundImageHovered(CORE_BOX_BRUSH("Common/TextBox_Hovered", FMargin(4.0f / 16.0f)))
+			.SetBackgroundImageFocused(CORE_BOX_BRUSH("Common/TextBox_Hovered", FMargin(4.0f / 16.0f)))
+			.SetBackgroundImageReadOnly(CORE_BOX_BRUSH("Common/TextBox_ReadOnly", FMargin(4.0f / 16.0f)))
 			.SetBackgroundColor(FLinearColor(0,0,0,0.1f))
 			.SetPadding(FMargin(0))
 			.SetScrollBarStyle(ScrollBar);
@@ -111,11 +111,11 @@ FStateTreeEditorStyle::FStateTreeEditorStyle()
 	const FTableRowStyle& NormalTableRowStyle = FAppStyle::Get().GetWidgetStyle<FTableRowStyle>("TableView.Row");
 	Set("StateTree.Selection",
 		FTableRowStyle(NormalTableRowStyle)
-		.SetActiveBrush(IMAGE_BRUSH("Common/Selection", CoreStyleConstants::Icon8x8, SelectionColor))
-		.SetActiveHoveredBrush(IMAGE_BRUSH("Common/Selection", CoreStyleConstants::Icon8x8, SelectionColor))
-		.SetInactiveBrush(IMAGE_BRUSH("Common/Selection", CoreStyleConstants::Icon8x8, SelectionColor))
-		.SetInactiveHoveredBrush(IMAGE_BRUSH("Common/Selection", CoreStyleConstants::Icon8x8, SelectionColor))
-		.SetSelectorFocusedBrush(IMAGE_BRUSH("Common/Selection", CoreStyleConstants::Icon8x8, SelectionColor))
+		.SetActiveBrush(CORE_IMAGE_BRUSH("Common/Selection", CoreStyleConstants::Icon8x8, SelectionColor))
+		.SetActiveHoveredBrush(CORE_IMAGE_BRUSH("Common/Selection", CoreStyleConstants::Icon8x8, SelectionColor))
+		.SetInactiveBrush(CORE_IMAGE_BRUSH("Common/Selection", CoreStyleConstants::Icon8x8, SelectionColor))
+		.SetInactiveHoveredBrush(CORE_IMAGE_BRUSH("Common/Selection", CoreStyleConstants::Icon8x8, SelectionColor))
+		.SetSelectorFocusedBrush(CORE_IMAGE_BRUSH("Common/Selection", CoreStyleConstants::Icon8x8, SelectionColor))
 	);
 
 	const FComboButtonStyle& ComboButtonStyle = FCoreStyle::Get().GetWidgetStyle<FComboButtonStyle>("ComboButton");
@@ -191,6 +191,7 @@ FStateTreeEditorStyle::FStateTreeEditorStyle()
 		Set("StateTreeEditor.PasteStatesAsChildren", new IMAGE_BRUSH_SVG("Icons/Child_State", CoreStyleConstants::Icon16x16));
 	}
 
+
 }
 
 void FStateTreeEditorStyle::Register()
