StateTree: collapse instead of hide invisible breakpoints

diff --git a/Engine/Plugins/Runtime/StateTree/Source/StateTreeEditorModule/Private/SStateTreeViewRow.cpp b/Engine/Plugins/Runtime/StateTree/Source/StateTreeEditorModule/Private/SStateTreeViewRow.cpp
--- a/Engine/Plugins/Runtime/StateTree/Source/StateTreeEditorModule/Private/SStateTreeViewRow.cpp
+++ b/Engine/Plugins/Runtime/StateTree/Source/StateTreeEditorModule/Private/SStateTreeViewRow.cpp
@@ -482,7 +482,7 @@ EVisibility SStateTreeViewRow::GetBreakpointVisibility() const
 		}
 	}
 
-	return EVisibility::Hidden;
+	return EVisibility::Collapsed;
 }
 
 FText SStateTreeViewRow::GetSelectorDesc() const
