Fix CIS Issue 275930
#rnx

diff --git a/Engine/Plugins/Runtime/StateTree/Source/StateTreeModule/Public/StateTreeTypes.h b/Engine/Plugins/Runtime/StateTree/Source/StateTreeModule/Public/StateTreeTypes.h
--- a/Engine/Plugins/Runtime/StateTree/Source/StateTreeModule/Public/StateTreeTypes.h
+++ b/Engine/Plugins/Runtime/StateTree/Source/StateTreeModule/Public/StateTreeTypes.h
@@ -1108,7 +1108,7 @@ struct STATETREEMODULE_API FStateTreeTransitionDelayedState
 };
 
 #if WITH_STATETREE_DEBUGGER
-struct FStateTreeInstanceDebugId
+struct STATETREEMODULE_API FStateTreeInstanceDebugId
 {
 	FStateTreeInstanceDebugId() = default;
 	FStateTreeInstanceDebugId(const uint32 InstanceId, const uint32 SerialNumber)
