[StateTree] trivial changes from review
#rnx

diff --git a/Engine/Plugins/Runtime/StateTree/Source/StateTreeEditorModule/Private/SStateTreeViewRow.cpp b/Engine/Plugins/Runtime/StateTree/Source/StateTreeEditorModule/Private/SStateTreeViewRow.cpp
--- a/Engine/Plugins/Runtime/StateTree/Source/StateTreeEditorModule/Private/SStateTreeViewRow.cpp
+++ b/Engine/Plugins/Runtime/StateTree/Source/StateTreeEditorModule/Private/SStateTreeViewRow.cpp
@@ -69,7 +69,6 @@ void SStateTreeViewRow::Construct(const FArguments& InArgs, const TSharedRef<STa
 					SNew(SBorder)
 					.BorderImage(FAppStyle::GetBrush("WhiteBrush"))
 					.BorderBackgroundColor(this, &SStateTreeViewRow::GetActiveStateColor)
-					.Visibility(this, &SStateTreeViewRow::GetActiveStateVisibility)
 					[
 						// State Box
 						SNew(SHorizontalBox)
@@ -459,19 +458,6 @@ EVisibility SStateTreeViewRow::GetSelectorVisibility() const
 	return EVisibility::Collapsed;
 }
 
-EVisibility SStateTreeViewRow::GetActiveStateVisibility() const
-{
-	if (const UStateTreeState* State = WeakState.Get())
-	{
-		if (StateTreeViewModel && StateTreeViewModel->IsStateActiveInDebugger(*State))
-		{
-			return EVisibility::Visible;
-		}
-	}
-
-	return EVisibility::Visible;
-}
-
 EVisibility SStateTreeViewRow::GetBreakpointVisibility() const
 {
 	if (const UStateTreeState* State = WeakState.Get())
diff --git a/Engine/Plugins/Runtime/StateTree/Source/StateTreeEditorModule/Private/SStateTreeViewRow.h b/Engine/Plugins/Runtime/StateTree/Source/StateTreeEditorModule/Private/SStateTreeViewRow.h
--- a/Engine/Plugins/Runtime/StateTree/Source/StateTreeEditorModule/Private/SStateTreeViewRow.h
+++ b/Engine/Plugins/Runtime/StateTree/Source/StateTreeEditorModule/Private/SStateTreeViewRow.h
@@ -35,7 +35,6 @@ private:
 
 	EVisibility GetConditionVisibility() const;
 	EVisibility GetSelectorVisibility() const;
-	EVisibility GetActiveStateVisibility() const;
 	EVisibility GetBreakpointVisibility() const;
 	FText GetSelectorDesc() const;
 
diff --git a/Engine/Plugins/Runtime/StateTree/Source/StateTreeEditorModule/Private/StateTreeViewModel.cpp b/Engine/Plugins/Runtime/StateTree/Source/StateTreeEditorModule/Private/StateTreeViewModel.cpp
--- a/Engine/Plugins/Runtime/StateTree/Source/StateTreeEditorModule/Private/StateTreeViewModel.cpp
+++ b/Engine/Plugins/Runtime/StateTree/Source/StateTreeEditorModule/Private/StateTreeViewModel.cpp
@@ -906,7 +906,7 @@ void FStateTreeViewModel::BindToDebuggerDelegates()
 				const FStateTreeStateHandle Handle = OuterStateTree->GetStateHandleFromId(State.ID);
 				if (NewActiveBreakpoints.Contains(Handle))
 				{
-					StatesWithBreakpoints.Emplace(&State);
+					StatesWithBreakpoints.Emplace(State.ID);
 				}
 
 				return EStateTreeVisitor::Continue;
@@ -919,7 +919,7 @@ void FStateTreeViewModel::BindToDebuggerDelegates()
 bool FStateTreeViewModel::DoesStateHaveBreakpoint(const UStateTreeState& State) const
 {
 #if WITH_STATETREE_DEBUGGER
-	return StatesWithBreakpoints.Contains(&State);
+	return StatesWithBreakpoints.Contains(State.ID);
 #else
 	return false;
 #endif // WITH_STATETREE_DEBUGGER
diff --git a/Engine/Plugins/Runtime/StateTree/Source/StateTreeEditorModule/Private/StateTreeViewModel.h b/Engine/Plugins/Runtime/StateTree/Source/StateTreeEditorModule/Private/StateTreeViewModel.h
--- a/Engine/Plugins/Runtime/StateTree/Source/StateTreeEditorModule/Private/StateTreeViewModel.h
+++ b/Engine/Plugins/Runtime/StateTree/Source/StateTreeEditorModule/Private/StateTreeViewModel.h
@@ -123,7 +123,7 @@ protected:
 
 #if WITH_STATETREE_DEBUGGER
 	TSharedRef<FStateTreeDebugger> Debugger;
-	TArray<TWeakObjectPtr<const UStateTreeState>> StatesWithBreakpoints;
+	TArray<FGuid> StatesWithBreakpoints;
 	TArray<FGuid> ActiveStates;
 #endif // WITH_STATETREE_DEBUGGER
 	
diff --git a/Engine/Plugins/Runtime/StateTree/Source/StateTreeModule/Private/Debugger/StateTreeDebugger.cpp b/Engine/Plugins/Runtime/StateTree/Source/StateTreeModule/Private/Debugger/StateTreeDebugger.cpp
--- a/Engine/Plugins/Runtime/StateTree/Source/StateTreeModule/Private/Debugger/StateTreeDebugger.cpp
+++ b/Engine/Plugins/Runtime/StateTree/Source/StateTreeModule/Private/Debugger/StateTreeDebugger.cpp
@@ -487,7 +487,7 @@ void FStateTreeDebugger::ResetAnalysisData()
 {
 	Events.Reset();
 	HitBreakpointStateIndex = INDEX_NONE;
-	HitBreakpointInstanceId.Invalidate();
+	HitBreakpointInstanceId.Reset();
 	LastTraceReadTime = UnsetTime;
 	CurrentScrubTime = UnsetTime;
 	CurrentFrameIndex = INDEX_NONE;
@@ -558,7 +558,7 @@ void FStateTreeDebugger::ReadTrace(
 				OnBreakpointHit.ExecuteIfBound(HitBreakpointInstanceId, StatesWithBreakpoint[HitBreakpointStateIndex]);
 
 				HitBreakpointStateIndex = INDEX_NONE;
-				HitBreakpointInstanceId.Invalidate();
+				HitBreakpointInstanceId.Reset();
 			}
 		}
 	}
diff --git a/Engine/Plugins/Runtime/StateTree/Source/StateTreeModule/Private/Debugger/StateTreeTrace.cpp b/Engine/Plugins/Runtime/StateTree/Source/StateTreeModule/Private/Debugger/StateTreeTrace.cpp
--- a/Engine/Plugins/Runtime/StateTree/Source/StateTreeModule/Private/Debugger/StateTreeTrace.cpp
+++ b/Engine/Plugins/Runtime/StateTree/Source/StateTreeModule/Private/Debugger/StateTreeTrace.cpp
@@ -51,7 +51,7 @@ UE_TRACE_EVENT_BEGIN(StateTreeDebugger, ActiveStatesEvent)
 	UE_TRACE_EVENT_FIELD(uint16[], ActiveStates)
 UE_TRACE_EVENT_END()
 
-void FStateTreeTrace::OutputInstanceLifetimeEvent(
+void UE::StateTreeTrace::OutputInstanceLifetimeEvent(
 	const FStateTreeInstanceDebugId InstanceId,
 	const UStateTree* StateTree,
 	const TCHAR* InstanceName,
@@ -75,7 +75,7 @@ void FStateTreeTrace::OutputInstanceLifetimeEvent(
 		<< InstanceEvent.EventType(static_cast<uint8>(EventType));
 }
 
-void FStateTreeTrace::OutputLogEventTrace(const FStateTreeInstanceDebugId InstanceId, const TCHAR* Fmt, ...)
+void UE::StateTreeTrace::OutputLogEventTrace(const FStateTreeInstanceDebugId InstanceId, const TCHAR* Fmt, ...)
 {
 	static TCHAR TraceStaticBuffer[8192];
 	GET_VARARGS(TraceStaticBuffer, UE_ARRAY_COUNT(TraceStaticBuffer), UE_ARRAY_COUNT(TraceStaticBuffer) - 1, Fmt, Fmt);
@@ -87,7 +87,7 @@ void FStateTreeTrace::OutputLogEventTrace(const FStateTreeInstanceDebugId Instan
 		<< LogEvent.Message(TraceStaticBuffer);
 }
 
-void FStateTreeTrace::OutputStateEventTrace(const FStateTreeInstanceDebugId InstanceId, const uint16 StateIdx, const EStateTreeTraceNodeEventType EventType)
+void UE::StateTreeTrace::OutputStateEventTrace(const FStateTreeInstanceDebugId InstanceId, const uint16 StateIdx, const EStateTreeTraceNodeEventType EventType)
 {
 	UE_TRACE_LOG(StateTreeDebugger, StateEvent, StateTreeDebugChannel)
 		<< StateEvent.Cycle(FPlatformTime::Cycles64())
@@ -97,7 +97,7 @@ void FStateTreeTrace::OutputStateEventTrace(const FStateTreeInstanceDebugId Inst
 		<< StateEvent.EventType(static_cast<uint8>(EventType));
 }
 
-void FStateTreeTrace::OutputTaskEventTrace(const FStateTreeInstanceDebugId InstanceId, const uint16 TaskIdx, const EStateTreeTraceNodeEventType EventType)
+void UE::StateTreeTrace::OutputTaskEventTrace(const FStateTreeInstanceDebugId InstanceId, const uint16 TaskIdx, const EStateTreeTraceNodeEventType EventType)
 {
 	UE_TRACE_LOG(StateTreeDebugger, TaskEvent, StateTreeDebugChannel)
 		<< TaskEvent.Cycle(FPlatformTime::Cycles64())
@@ -107,7 +107,7 @@ void FStateTreeTrace::OutputTaskEventTrace(const FStateTreeInstanceDebugId Insta
 		<< TaskEvent.EventType(static_cast<uint8>(EventType));
 }
 
-void FStateTreeTrace::OutputActiveStatesEventTrace(const FStateTreeInstanceDebugId InstanceId, const FStateTreeActiveStates& ActiveStates)
+void UE::StateTreeTrace::OutputActiveStatesEventTrace(const FStateTreeInstanceDebugId InstanceId, const FStateTreeActiveStates& ActiveStates)
 {
 	TArray<uint16, TInlineAllocator<FStateTreeActiveStates::MaxStates>> StatesIndices;
 	for (int32 i = 0; i < ActiveStates.Num(); i++)
diff --git a/Engine/Plugins/Runtime/StateTree/Source/StateTreeModule/Private/Debugger/StateTreeTraceAnalyzer.cpp b/Engine/Plugins/Runtime/StateTree/Source/StateTreeModule/Private/Debugger/StateTreeTraceAnalyzer.cpp
--- a/Engine/Plugins/Runtime/StateTree/Source/StateTreeModule/Private/Debugger/StateTreeTraceAnalyzer.cpp
+++ b/Engine/Plugins/Runtime/StateTree/Source/StateTreeModule/Private/Debugger/StateTreeTraceAnalyzer.cpp
@@ -43,7 +43,7 @@ bool FStateTreeTraceAnalyzer::OnEvent(const uint16 RouteId, EStyle Style, const
 			const FTopLevelAssetPath Path((FName)ObjectPathName, (FName)ObjectName);
 			TWeakObjectPtr<const UStateTree> WeakStateTree;
 			{
-				// This might not work when using a debugger on a client but should be fine in Editor as lock as
+				// This might not work when using a debugger on a client but should be fine in Editor as long as
 				// we are not trying to find the object during GC. We might not currently be in the game thread.  
 				// @todo STDBG: eventually errors should be reported in the UI
 				FGCScopeGuard Guard;
diff --git a/Engine/Plugins/Runtime/StateTree/Source/StateTreeModule/Public/Debugger/StateTreeTrace.h b/Engine/Plugins/Runtime/StateTree/Source/StateTreeModule/Public/Debugger/StateTreeTrace.h
--- a/Engine/Plugins/Runtime/StateTree/Source/StateTreeModule/Public/Debugger/StateTreeTrace.h
+++ b/Engine/Plugins/Runtime/StateTree/Source/StateTreeModule/Public/Debugger/StateTreeTrace.h
@@ -14,43 +14,43 @@ enum class EStateTreeTraceInstanceEventType : uint8;
 
 UE_TRACE_CHANNEL_EXTERN(StateTreeDebugChannel, STATETREEMODULE_API)
 
-struct FStateTreeTrace
+namespace UE::StateTreeTrace
 {
-	static void OutputInstanceLifetimeEvent(const FStateTreeInstanceDebugId InstanceId, const UStateTree* StateTree, const TCHAR* InstanceName, const EStateTreeTraceInstanceEventType EventType);
-	static void OutputLogEventTrace(const FStateTreeInstanceDebugId InstanceId, const TCHAR* Fmt, ...);
-	static void OutputStateEventTrace(const FStateTreeInstanceDebugId InstanceId, const uint16 StateIdx, const EStateTreeTraceNodeEventType EventType);
-	static void OutputTaskEventTrace(const FStateTreeInstanceDebugId InstanceId, const uint16 TaskIdx, const EStateTreeTraceNodeEventType EventType);
-	static void OutputActiveStatesEventTrace(const FStateTreeInstanceDebugId InstanceId, const FStateTreeActiveStates& ActiveStates);
+	void OutputInstanceLifetimeEvent(const FStateTreeInstanceDebugId InstanceId, const UStateTree* StateTree, const TCHAR* InstanceName, const EStateTreeTraceInstanceEventType EventType);
+	void OutputLogEventTrace(const FStateTreeInstanceDebugId InstanceId, const TCHAR* Fmt, ...);
+	void OutputStateEventTrace(const FStateTreeInstanceDebugId InstanceId, const uint16 StateIdx, const EStateTreeTraceNodeEventType EventType);
+	void OutputTaskEventTrace(const FStateTreeInstanceDebugId InstanceId, const uint16 TaskIdx, const EStateTreeTraceNodeEventType EventType);
+	void OutputActiveStatesEventTrace(const FStateTreeInstanceDebugId InstanceId, const FStateTreeActiveStates& ActiveStates);
 };
 
 #define TRACE_STATETREE_INSTANCE_EVENT(InstanceID, StateTree, InstanceName, EventType) \
 	if (UE_TRACE_CHANNELEXPR_IS_ENABLED(StateTreeDebugChannel)) \
 	{ \
-		FStateTreeTrace::OutputInstanceLifetimeEvent(InstanceID, StateTree, InstanceName, EventType); \
+		UE::StateTreeTrace::OutputInstanceLifetimeEvent(InstanceID, StateTree, InstanceName, EventType); \
 	}
 
 #define TRACE_STATETREE_LOG_EVENT(InstanceId, Format, ...) \
 	if (UE_TRACE_CHANNELEXPR_IS_ENABLED(StateTreeDebugChannel)) \
 	{ \
-		FStateTreeTrace::OutputLogEventTrace(InstanceId, Format, ##__VA_ARGS__); \
+		UE::StateTreeTrace::OutputLogEventTrace(InstanceId, Format, ##__VA_ARGS__); \
 	}
 
 #define TRACE_STATETREE_STATE_EVENT(InstanceId, StateIdx, EventType) \
 	if (UE_TRACE_CHANNELEXPR_IS_ENABLED(StateTreeDebugChannel)) \
 	{ \
-		FStateTreeTrace::OutputStateEventTrace(InstanceId, StateIdx, EventType); \
+		UE::StateTreeTrace::OutputStateEventTrace(InstanceId, StateIdx, EventType); \
 	}
 
 #define TRACE_STATETREE_TASK_EVENT(InstanceId, TaskIdx, EventType) \
 	if (UE_TRACE_CHANNELEXPR_IS_ENABLED(StateTreeDebugChannel)) \
 	{ \
-		FStateTreeTrace::OutputTaskEventTrace(InstanceId, TaskIdx, EventType); \
+		UE::StateTreeTrace::OutputTaskEventTrace(InstanceId, TaskIdx, EventType); \
 	}
 
 #define TRACE_STATETREE_ACTIVE_STATES_EVENT(InstanceId, ActivateStates) \
 	if (UE_TRACE_CHANNELEXPR_IS_ENABLED(StateTreeDebugChannel)) \
 	{ \
-		FStateTreeTrace::OutputActiveStatesEventTrace(InstanceId, ActivateStates); \
+		UE::StateTreeTrace::OutputActiveStatesEventTrace(InstanceId, ActivateStates); \
 	}
 
 #else //STATETREE_DEBUG_TRACE_ENABLED
diff --git a/Engine/Plugins/Runtime/StateTree/Source/StateTreeModule/Public/StateTree.h b/Engine/Plugins/Runtime/StateTree/Source/StateTreeModule/Public/StateTree.h
--- a/Engine/Plugins/Runtime/StateTree/Source/StateTreeModule/Public/StateTree.h
+++ b/Engine/Plugins/Runtime/StateTree/Source/StateTreeModule/Public/StateTree.h
@@ -110,10 +110,10 @@ public:
 	FStateTreeStateHandle GetStateHandleFromId(const FGuid Id) const;
 
 	/** @return Id of the state matching a given state handle; invalid Id if state not found. */
-	FGuid GetStateIdFromHandle(FStateTreeStateHandle Handle) const;
+	FGuid GetStateIdFromHandle(const FStateTreeStateHandle Handle) const;
 
-	/** @return Struct view of the node matching a given state index; invalid view if state not found. */
-	FConstStructView GetNode(int32 NodeIndex) const;
+	/** @return Struct view of the node matching a given node index; invalid view if state not found. */
+	FConstStructView GetNode(const int32 NodeIndex) const;
 
 	/** @return View of all states. */
 	TConstArrayView<FCompactStateTreeState> GetStates() const { return States; }
diff --git a/Engine/Plugins/Runtime/StateTree/Source/StateTreeModule/Public/StateTreeSettings.h b/Engine/Plugins/Runtime/StateTree/Source/StateTreeModule/Public/StateTreeSettings.h
--- a/Engine/Plugins/Runtime/StateTree/Source/StateTreeModule/Public/StateTreeSettings.h
+++ b/Engine/Plugins/Runtime/StateTree/Source/StateTreeModule/Public/StateTreeSettings.h
@@ -15,7 +15,8 @@ class STATETREEMODULE_API UStateTreeSettings : public UDeveloperSettings
 
 public:
 	static UStateTreeSettings& Get() { return *CastChecked<UStateTreeSettings>(UStateTreeSettings::StaticClass()->GetDefaultObject()); }
-	
+
+	/** The debugger is still an experimental feature, hence not active by default. */
 	UPROPERTY(EditDefaultsOnly, Category = StateTree, config)
 	bool bUseDebugger = false;
 };
\ No newline at end of file
diff --git a/Engine/Plugins/Runtime/StateTree/Source/StateTreeModule/Public/StateTreeTypes.h b/Engine/Plugins/Runtime/StateTree/Source/StateTreeModule/Public/StateTreeTypes.h
--- a/Engine/Plugins/Runtime/StateTree/Source/StateTreeModule/Public/StateTreeTypes.h
+++ b/Engine/Plugins/Runtime/StateTree/Source/StateTreeModule/Public/StateTreeTypes.h
@@ -1118,7 +1118,7 @@ struct STATETREEMODULE_API FStateTreeInstanceDebugId
 	
 	bool IsValid() const { return Id != INDEX_NONE && SerialNumber != INDEX_NONE; }
 	bool IsInvalid() const { return !IsValid(); }
-	void Invalidate() { *this = Invalid; }
+	void Reset() { *this = Invalid; }
 
 	bool operator==(const FStateTreeInstanceDebugId& Other) const
 	{
