[StateTreeDebugger] replaced debug serial number mappings by single global counter.
#rnx

diff --git a/Engine/Plugins/Runtime/StateTree/Source/StateTreeModule/Private/StateTreeExecutionContext.cpp b/Engine/Plugins/Runtime/StateTree/Source/StateTreeModule/Private/StateTreeExecutionContext.cpp
--- a/Engine/Plugins/Runtime/StateTree/Source/StateTreeModule/Private/StateTreeExecutionContext.cpp
+++ b/Engine/Plugins/Runtime/StateTree/Source/StateTreeModule/Private/StateTreeExecutionContext.cpp
@@ -512,14 +512,8 @@ FStateTreeInstanceDebugId FStateTreeExecutionContext::GetInstanceDebugId() const
 	FStateTreeInstanceDebugId& InstanceDebugId = GetExecState().InstanceDebugId; 
 	if (!InstanceDebugId.IsValid())
 	{
-		const uint32 Id = GetTypeHash(GetInstanceDescription());
-		static TMap<uint32, uint32> DebugIdMappings;
-		uint32* SerialNumber = DebugIdMappings.Find(Id);
-		if (SerialNumber == nullptr)
-		{
-			SerialNumber = &DebugIdMappings.Add(Id, 0);
-		}
-		InstanceDebugId = FStateTreeInstanceDebugId(Id, (*SerialNumber)++); 
+		static std::atomic<uint32> SerialNumber = 0;
+		InstanceDebugId = FStateTreeInstanceDebugId(GetTypeHash(GetInstanceDescription()), ++SerialNumber); 
 	}
 	return InstanceDebugId;
 }
