[StateTreeDebugger]
- process breakpoints only for instances matching the debugged asset
#rnx

diff --git a/Engine/Plugins/Runtime/StateTree/Source/StateTreeEditorModule/Private/SStateTreeDebuggerView.cpp b/Engine/Plugins/Runtime/StateTree/Source/StateTreeEditorModule/Private/SStateTreeDebuggerView.cpp
--- a/Engine/Plugins/Runtime/StateTree/Source/StateTreeEditorModule/Private/SStateTreeDebuggerView.cpp
+++ b/Engine/Plugins/Runtime/StateTree/Source/StateTreeEditorModule/Private/SStateTreeDebuggerView.cpp
@@ -338,9 +338,8 @@ void SStateTreeDebuggerView::OnBreakpointHit(const FStateTreeInstanceDebugId Ins
 			const FStateTreeDebuggerInstanceDesc& DebuggerInstance = Debugger->GetDebuggedInstance();
 			if (DebuggerInstance.Id != InstanceId)
 			{
-				TArray<FStateTreeDebuggerInstanceDesc> MatchingInstances;
-				Debugger->GetActivateInstances(MatchingInstances);
-				const FStateTreeDebuggerInstanceDesc* FoundDesc = MatchingInstances.FindByPredicate([InstanceId](const FStateTreeDebuggerInstanceDesc& InstanceDesc)
+				const TConstArrayView<FStateTreeDebuggerInstanceDesc> ActiveInstances = Debugger->GetActiveInstances();
+				const FStateTreeDebuggerInstanceDesc* FoundDesc = ActiveInstances.FindByPredicate([InstanceId](const FStateTreeDebuggerInstanceDesc& InstanceDesc)
 				{
 					return InstanceDesc.Id == InstanceId;
 				});
@@ -401,15 +400,20 @@ TSharedRef<SWidget> SStateTreeDebuggerView::OnGetDebuggerInstancesMenu() const
 {
 	FMenuBuilder MenuBuilder(/*bInShouldCloseWindowAfterMenuSelection*/true, /*InCommandList*/nullptr);
 
-	if (Debugger.IsValid())
+	if (Debugger.IsValid() && StateTreeViewModel.IsValid())
 	{
-		TArray<FStateTreeDebuggerInstanceDesc> ActiveInstances;
-		Debugger->GetActivateInstances(ActiveInstances); 
+		const TConstArrayView<FStateTreeDebuggerInstanceDesc> ActiveInstances = Debugger->GetActiveInstances();
+		const UStateTree* CurrentStateTreeAsset = StateTreeViewModel->GetStateTree();
 
 		for (const FStateTreeDebuggerInstanceDesc& Instance : ActiveInstances)
 		{
+			// Do not list instances not matching the debugged asset
+			if (Instance.StateTree.Get() != CurrentStateTreeAsset)
+			{
+				continue;
+			}
+			
 			const FText Desc = FText::FromString(Debugger->DescribeInstance(Instance));
-
 			FUIAction ItemAction(FExecuteAction::CreateLambda([WeakDebugger = Debugger, Instance]()
 			{
 				if (WeakDebugger)
diff --git a/Engine/Plugins/Runtime/StateTree/Source/StateTreeEditorModule/Private/StateTreeViewModel.cpp b/Engine/Plugins/Runtime/StateTree/Source/StateTreeEditorModule/Private/StateTreeViewModel.cpp
--- a/Engine/Plugins/Runtime/StateTree/Source/StateTreeEditorModule/Private/StateTreeViewModel.cpp
+++ b/Engine/Plugins/Runtime/StateTree/Source/StateTreeEditorModule/Private/StateTreeViewModel.cpp
@@ -251,7 +251,10 @@ void FStateTreeViewModel::Init(UStateTreeEditorData* InTreeData)
 
 	UE::StateTree::Delegates::OnIdentifierChanged.AddSP(this, &FStateTreeViewModel::HandleIdentifierChanged);
 
+#if WITH_STATETREE_DEBUGGER
+	Debugger->SetAsset(GetStateTree());
 	BindToDebuggerDelegates();
+#endif // WITH_STATETREE_DEBUGGER	
 }
 
 const UStateTree* FStateTreeViewModel::GetStateTree() const
diff --git a/Engine/Plugins/Runtime/StateTree/Source/StateTreeModule/Private/Debugger/StateTreeDebugger.cpp b/Engine/Plugins/Runtime/StateTree/Source/StateTreeModule/Private/Debugger/StateTreeDebugger.cpp
--- a/Engine/Plugins/Runtime/StateTree/Source/StateTreeModule/Private/Debugger/StateTreeDebugger.cpp
+++ b/Engine/Plugins/Runtime/StateTree/Source/StateTreeModule/Private/Debugger/StateTreeDebugger.cpp
@@ -101,6 +101,8 @@ FStateTreeDebugger::~FStateTreeDebugger()
 
 void FStateTreeDebugger::Tick(float DeltaTime)
 {
+	GetSessionActiveInstances(ActiveInstances);
+
 	if (!bPaused)
 	{
 		TRACE_CPUPROFILER_EVENT_SCOPE(FStateTreeDebugger::Tick);
@@ -109,8 +111,6 @@ void FStateTreeDebugger::Tick(float DeltaTime)
 
 	if (DebuggedInstance.IsValid())
 	{
-		TArray<FStateTreeDebuggerInstanceDesc> ActiveInstances;
-		GetActivateInstances(ActiveInstances);
 		if (!ActiveInstances.Contains(DebuggedInstance))
 		{
 			SetDebuggedInstance(FStateTreeDebuggerInstanceDesc());
@@ -152,14 +152,14 @@ void FStateTreeDebugger::SyncToCurrentSessionDuration()
 {
 	if (const TraceServices::IAnalysisSession* Session = GetAnalysisSession())
 	{
-		TraceServices::FAnalysisSessionReadScope SessionReadScope(*Session);
-		SetCurrentScrubTime(Session->GetDurationSeconds());	
-	}
-}
+		double SessionDuration = 0.;
+		{
+			TraceServices::FAnalysisSessionReadScope SessionReadScope(*Session);
+			SessionDuration = Session->GetDurationSeconds();
+		}
 
-void FStateTreeDebugger::SetCurrentScrubTime(const double Time)
-{
-	ReadTrace(Time);
+		ReadTrace(SessionDuration);
+	}
 }
 
 const FStateTreeDebuggerInstanceDesc& FStateTreeDebugger::GetDebuggedInstance() const
@@ -187,7 +187,12 @@ void FStateTreeDebugger::SetDebuggedInstance(const FStateTreeDebuggerInstanceDes
 	}
 }
 
-void FStateTreeDebugger::GetActivateInstances(TArray<FStateTreeDebuggerInstanceDesc>& OutInstances) const
+TConstArrayView<FStateTreeDebuggerInstanceDesc> FStateTreeDebugger::GetActiveInstances() const
+{
+	return ActiveInstances;
+}
+
+void FStateTreeDebugger::GetSessionActiveInstances(TArray<FStateTreeDebuggerInstanceDesc>& OutInstances) const
 {
 	if (const TraceServices::IAnalysisSession* Session = GetAnalysisSession())
 	{
@@ -195,7 +200,7 @@ void FStateTreeDebugger::GetActivateInstances(TArray<FStateTreeDebuggerInstanceD
 
 		if (const IStateTreeTraceProvider* Provider = Session->ReadProvider<IStateTreeTraceProvider>(FStateTreeTraceProvider::ProviderName))
 		{
-			Provider->GetActivateInstances(OutInstances);
+			Provider->GetActiveInstances(OutInstances);
 		}
 	}
 }
@@ -510,6 +515,9 @@ void FStateTreeDebugger::ReadTrace(const uint64 FrameIndex)
 			ReadTrace(*Session, FrameProvider, *TargetFrame);
 		}
 	}
+
+	// Notify outside session read scope
+	SendNotifications();
 }
 
 void FStateTreeDebugger::ReadTrace(const double ScrubTime)
@@ -527,6 +535,28 @@ void FStateTreeDebugger::ReadTrace(const double ScrubTime)
 			ReadTrace(*Session, FrameProvider, TargetFrame);						
 		}
 	}
+
+	// Notify outside session read scope
+	SendNotifications();
+}
+
+void FStateTreeDebugger::SendNotifications()
+{
+	if (bNewEvents)
+	{
+		OnTracesUpdated.ExecuteIfBound(Events, FramesWithEvents);
+		bNewEvents = false;
+	}
+
+	if (HitBreakpointStateIndex != INDEX_NONE)
+	{
+		check(HitBreakpointInstanceId.IsValid());
+		check(StatesWithBreakpoint.IsValidIndex(HitBreakpointStateIndex));
+		OnBreakpointHit.ExecuteIfBound(HitBreakpointInstanceId, StatesWithBreakpoint[HitBreakpointStateIndex]);
+
+		HitBreakpointStateIndex = INDEX_NONE;
+		HitBreakpointInstanceId.Reset();
+	}
 }
 
 void FStateTreeDebugger::ReadTrace(
@@ -541,25 +571,8 @@ void FStateTreeDebugger::ReadTrace(
 	{
 		if (const IStateTreeTraceProvider* Provider = Session.ReadProvider<IStateTreeTraceProvider>(FStateTreeTraceProvider::ProviderName))
 		{
-			const int32 PrevNumEvents = Events.Num();
-
 			AddEvents(LastTraceReadTime, Frame.EndTime, FrameProvider, *Provider);
 			LastTraceReadTime = CurrentScrubTime;
-
-			const int32 NumEvents = Events.Num();
-			if (NumEvents != PrevNumEvents)
-			{
-				OnTracesUpdated.ExecuteIfBound(Events, FramesWithEvents);
-			}
-
-			if (HitBreakpointStateIndex != INDEX_NONE)
-			{
-				check(HitBreakpointInstanceId.IsValid());
-				OnBreakpointHit.ExecuteIfBound(HitBreakpointInstanceId, StatesWithBreakpoint[HitBreakpointStateIndex]);
-
-				HitBreakpointStateIndex = INDEX_NONE;
-				HitBreakpointInstanceId.Reset();
-			}
 		}
 	}
 
@@ -588,6 +601,7 @@ bool FStateTreeDebugger::ProcessEvent(const FStateTreeInstanceDebugId InstanceId
 		// Breakpoints 
 		if (bPaused == false // ignored when scrubbing a paused session
 			&& bProcessingInitialEvents == false // only processed for events added after the initial parsing
+			&& DebuggedAsset != nullptr // asset is required to properly match state handles
 			&& HitBreakpointStateIndex == INDEX_NONE // stop on first hit breakpoint
 			&& StatesWithBreakpoint.Num()
 			&& (DebuggedInstance.Id == InstanceId || !DebuggedInstance.IsValid())) // allow breakpoints on any instances if not specified
@@ -595,10 +609,29 @@ bool FStateTreeDebugger::ProcessEvent(const FStateTreeInstanceDebugId InstanceId
 			const FStateTreeTraceStateEvent* StateEvent = Event.TryGet<FStateTreeTraceStateEvent>();
 			if (StateEvent != nullptr && StateEvent->EventType == EStateTreeTraceNodeEventType::OnEnter)
 			{
-				HitBreakpointStateIndex = StatesWithBreakpoint.Find(FStateTreeStateHandle(StateEvent->StateIdx));
-				if (HitBreakpointStateIndex != INDEX_NONE)
+				const UStateTree* InstanceStateTree = nullptr;
+				if (!DebuggedInstance.IsValid())
 				{
-					HitBreakpointInstanceId = InstanceId;
+					// No specific instance selected yet, find matching descriptor from id to extract the associated StateTree asset
+					const FStateTreeDebuggerInstanceDesc* FoundInstance = ActiveInstances.FindByPredicate(
+						[InstanceId](const FStateTreeDebuggerInstanceDesc& InstanceDesc)
+						{
+							return InstanceDesc.Id == InstanceId;
+						});
+
+					if (FoundInstance != nullptr)
+					{
+						InstanceStateTree = FoundInstance->StateTree.Get();
+					}
+				}
+
+				if (DebuggedInstance.IsValid() || InstanceStateTree == DebuggedAsset)
+				{
+					HitBreakpointStateIndex = StatesWithBreakpoint.Find(FStateTreeStateHandle(StateEvent->StateIdx));
+					if (HitBreakpointStateIndex != INDEX_NONE)
+					{
+						HitBreakpointInstanceId = InstanceId;
+					}
 				}
 			}
 		}
@@ -610,8 +643,9 @@ bool FStateTreeDebugger::ProcessEvent(const FStateTreeInstanceDebugId InstanceId
 			{
 				FramesWithEvents.Add(UE::StateTreeDebugger::FFrameIndexSpan(Frame.Index, Events.Num()));
 			}
-			
+
 			Events.Emplace(Event);
+			bNewEvents = true;
 		}
 	}
 	
diff --git a/Engine/Plugins/Runtime/StateTree/Source/StateTreeModule/Private/Debugger/StateTreeTraceProvider.cpp b/Engine/Plugins/Runtime/StateTree/Source/StateTreeModule/Private/Debugger/StateTreeTraceProvider.cpp
--- a/Engine/Plugins/Runtime/StateTree/Source/StateTreeModule/Private/Debugger/StateTreeTraceProvider.cpp
+++ b/Engine/Plugins/Runtime/StateTree/Source/StateTreeModule/Private/Debugger/StateTreeTraceProvider.cpp
@@ -85,7 +85,7 @@ void FStateTreeTraceProvider::AppendInstance(
 	}
 }
 
-void FStateTreeTraceProvider::GetActivateInstances(TArray<FStateTreeDebuggerInstanceDesc>& OutInstances) const
+void FStateTreeTraceProvider::GetActiveInstances(TArray<FStateTreeDebuggerInstanceDesc>& OutInstances) const
 {
 	OutInstances = ActiveInstances;
 }
diff --git a/Engine/Plugins/Runtime/StateTree/Source/StateTreeModule/Public/Debugger/IStateTreeTraceProvider.h b/Engine/Plugins/Runtime/StateTree/Source/StateTreeModule/Public/Debugger/IStateTreeTraceProvider.h
--- a/Engine/Plugins/Runtime/StateTree/Source/StateTreeModule/Public/Debugger/IStateTreeTraceProvider.h
+++ b/Engine/Plugins/Runtime/StateTree/Source/StateTreeModule/Public/Debugger/IStateTreeTraceProvider.h
@@ -16,7 +16,7 @@ class IStateTreeTraceProvider : public TraceServices::IProvider
 public:
 	typedef TraceServices::ITimeline<FStateTreeTraceEventVariantType> FEventsTimeline;
 
-	virtual void GetActivateInstances(TArray<FStateTreeDebuggerInstanceDesc>& OutInstances) const = 0;
+	virtual void GetActiveInstances(TArray<FStateTreeDebuggerInstanceDesc>& OutInstances) const = 0;
 
 	/**
 	 * Execute given function receiving an event timeline for a given instance or all timelines if instance not specified.  
diff --git a/Engine/Plugins/Runtime/StateTree/Source/StateTreeModule/Public/Debugger/StateTreeDebugger.h b/Engine/Plugins/Runtime/StateTree/Source/StateTreeModule/Public/Debugger/StateTreeDebugger.h
--- a/Engine/Plugins/Runtime/StateTree/Source/StateTreeModule/Public/Debugger/StateTreeDebugger.h
+++ b/Engine/Plugins/Runtime/StateTree/Source/StateTreeModule/Public/Debugger/StateTreeDebugger.h
@@ -97,6 +97,8 @@ struct STATETREEMODULE_API FStateTreeDebugger : FTickableGameObject
 	FStateTreeDebugger();
 	virtual ~FStateTreeDebugger() override;
 
+	void SetAsset(const UStateTree* StateTreeAsset) { DebuggedAsset = StateTreeAsset; }
+	
 	/** Stops reading traces every frame to preserve current state */
 	void Pause();
 
@@ -113,7 +115,7 @@ struct STATETREEMODULE_API FStateTreeDebugger : FTickableGameObject
 	
 	bool IsAnalysisSessionActive() const { return GetAnalysisSession() != nullptr; }
 
-	void GetActivateInstances(TArray<FStateTreeDebuggerInstanceDesc>& OutInstances) const;
+	TConstArrayView<FStateTreeDebuggerInstanceDesc> GetActiveInstances() const;
 	const FStateTreeDebuggerInstanceDesc& GetDebuggedInstance() const;
 	void SetDebuggedInstance(const FStateTreeDebuggerInstanceDesc& SelectedInstance);
 	FString GetDebuggedInstanceDescription() const;
@@ -141,8 +143,9 @@ protected:
 	virtual TStatId GetStatId() const override { RETURN_QUICK_DECLARE_CYCLE_STAT(FStateTreeDebugger, STATGROUP_Tickables); }
 	
 private:
+	void StopAnalysis();
 	void ResetAnalysisData();
-	
+
 	void ReadTrace(double ScrubTime);
 	void ReadTrace(uint64 FrameIndex);
 	void ReadTrace(
@@ -150,19 +153,21 @@ private:
 		const TraceServices::IFrameProvider& FrameProvider,
 		const TraceServices::FFrame& Frame
 		);
-	
-	void StopAnalysis();
 
-	void SetCurrentScrubTime(double RecordingDuration);	
+	void SendNotifications();
+
 	void SetActiveStates(const TConstArrayView<FStateTreeStateHandle> NewActiveStates);
 	
 	const TraceServices::IAnalysisSession* GetAnalysisSession() const;
 	UE::Trace::FStoreClient* GetStoreClient() const;
+	void GetSessionActiveInstances(TArray<FStateTreeDebuggerInstanceDesc>& OutInstances) const;
 
 	bool ProcessEvent(const FStateTreeInstanceDebugId InstanceId, const TraceServices::FFrame& Frame, const FStateTreeTraceEventVariantType& Event);
 	void AddEvents(float StartTime, float EndTime, const TraceServices::IFrameProvider& FrameProvider, const IStateTreeTraceProvider& StateTreeTraceProvider);
 	void UpdateMetadata(FTraceDescriptor& TraceDescriptor) const;
 
+	TWeakObjectPtr<const UStateTree> DebuggedAsset;
+	
 	IStateTreeModule& StateTreeModule;
 	
 	/** The trace analysis session. */
@@ -170,6 +175,7 @@ private:
 	
 	FTraceDescriptor ActiveSessionTraceDescriptor;
 
+	TArray<FStateTreeDebuggerInstanceDesc> ActiveInstances;
 	FStateTreeDebuggerInstanceDesc DebuggedInstance;
 	
 	TArray<FStateTreeStateHandle> StatesWithBreakpoint;
@@ -188,6 +194,7 @@ private:
 	TArray<UE::StateTreeDebugger::FFrameIndexSpan> FramesWithEvents;
 
 	bool bProcessingInitialEvents = false;
+	bool bNewEvents = false;
 	bool bPaused = false;
 };
 
diff --git a/Engine/Plugins/Runtime/StateTree/Source/StateTreeModule/Public/Debugger/StateTreeTraceProvider.h b/Engine/Plugins/Runtime/StateTree/Source/StateTreeModule/Public/Debugger/StateTreeTraceProvider.h
--- a/Engine/Plugins/Runtime/StateTree/Source/StateTreeModule/Public/Debugger/StateTreeTraceProvider.h
+++ b/Engine/Plugins/Runtime/StateTree/Source/StateTreeModule/Public/Debugger/StateTreeTraceProvider.h
@@ -28,7 +28,7 @@ public:
 
 protected:
 	/** IStateTreeDebuggerProvider interface */
-	virtual void GetActivateInstances(TArray<FStateTreeDebuggerInstanceDesc>& OutInstances) const override;
+	virtual void GetActiveInstances(TArray<FStateTreeDebuggerInstanceDesc>& OutInstances) const override;
 	virtual bool ReadTimelines(const FStateTreeInstanceDebugId InstanceId, TFunctionRef<void(const FStateTreeInstanceDebugId ProcessedInstanceId, const FEventsTimeline&)> Callback) const override;
 
 private:
