PR #10275: Remove state tree debug LogTemp messages (Contributed by [REDACTED_PERSON])

diff --git a/Engine/Plugins/Runtime/StateTree/Source/StateTreeEditorModule/Private/Customizations/StateTreeEditorNodeDetails.cpp b/Engine/Plugins/Runtime/StateTree/Source/StateTreeEditorModule/Private/Customizations/StateTreeEditorNodeDetails.cpp
--- a/Engine/Plugins/Runtime/StateTree/Source/StateTreeEditorModule/Private/Customizations/StateTreeEditorNodeDetails.cpp
+++ b/Engine/Plugins/Runtime/StateTree/Source/StateTreeEditorModule/Private/Customizations/StateTreeEditorNodeDetails.cpp
@@ -1578,15 +1578,12 @@ int32 FStateTreeEditorNodeDetails::FilterNodeTypesChildren(const TArray<FString>
 		{
 			if (ItemName.Contains(Filter))
 			{
-				UE_LOG(LogTemp, Error, TEXT("%s contains %s"), *ItemName, *Filter);
 				return true;
 			}
 		}
 		return false;
 	};
 
-	UE_LOG(LogTemp, Error, TEXT("FilterNodeTypesChildren"));
-	
 	for (const TSharedPtr<FStateTreeNodeTypeItem>& SourceItem : SourceArray)
 	{
 		// Check if our name matches the filters
