StateTree: Fixed tree update on State Tree view

diff --git a/Engine/Plugins/Runtime/StateTree/Source/StateTreeEditorModule/Private/SStateTreeView.cpp b/Engine/Plugins/Runtime/StateTree/Source/StateTreeEditorModule/Private/SStateTreeView.cpp
--- a/Engine/Plugins/Runtime/StateTree/Source/StateTreeEditorModule/Private/SStateTreeView.cpp
+++ b/Engine/Plugins/Runtime/StateTree/Source/StateTreeEditorModule/Private/SStateTreeView.cpp
@@ -274,6 +274,8 @@ void SStateTreeView::UpdateTree(bool bExpandPersistent)
 	TreeView->ClearSelection();
 	TreeView->SetItemSelection(SelectedStates, true);
 
+	TreeView->RequestTreeRefresh();
+
 	bItemsDirty = false;
 }
 
