StateTree: Allow property binding to copy pointers to components

diff --git a/Engine/Plugins/Runtime/StateTree/Source/StateTreeEditorModule/Private/Customizations/StateTreeBindingExtension.cpp b/Engine/Plugins/Runtime/StateTree/Source/StateTreeEditorModule/Private/Customizations/StateTreeBindingExtension.cpp
--- a/Engine/Plugins/Runtime/StateTree/Source/StateTreeEditorModule/Private/Customizations/StateTreeBindingExtension.cpp
+++ b/Engine/Plugins/Runtime/StateTree/Source/StateTreeEditorModule/Private/Customizations/StateTreeBindingExtension.cpp
@@ -203,7 +203,7 @@ FOnStateTreePropertyBindingChanged STATETREEEDITORMODULE_API OnStateTreeProperty
 bool FStateTreeBindingExtension::IsPropertyExtendable(const UClass* InObjectClass, const IPropertyHandle& PropertyHandle) const
 {
 	const FProperty* Property = PropertyHandle.GetProperty();
-	if (Property->HasAnyPropertyFlags(CPF_PersistentInstance | CPF_EditorOnly | CPF_Config | CPF_InstancedReference))
+	if (Property->HasAnyPropertyFlags(CPF_PersistentInstance | CPF_EditorOnly | CPF_Config))
 	{
 		return false;
 	}
@@ -222,7 +222,7 @@ bool FStateTreeBindingExtension::IsPropertyExtendable(const UClass* InObjectClas
 		return true;
 	}
 	
-	return false;  
+	return false;
 }
 
 
