StateTree: Fix indentation handling in FStateTreeCompiler::CreateConditions() which lead to a crash an uncorrect handling of nested expressions.

diff --git a/Engine/Plugins/Runtime/StateTree/Source/StateTreeEditorModule/Private/StateTreeCompiler.cpp b/Engine/Plugins/Runtime/StateTree/Source/StateTreeEditorModule/Private/StateTreeCompiler.cpp
--- a/Engine/Plugins/Runtime/StateTree/Source/StateTreeEditorModule/Private/StateTreeCompiler.cpp
+++ b/Engine/Plugins/Runtime/StateTree/Source/StateTreeEditorModule/Private/StateTreeCompiler.cpp
@@ -311,7 +311,8 @@ bool FStateTreeCompiler::CreateConditions(UStateTreeState& State, TConstArrayVie
 		// First indent must be 0 to make the parentheses calculation match.
 		const int32 CurrIndent = bIsFirst ? 0 : FMath::Clamp((int32)CondNode.ConditionIndent, 0, UE::StateTree::MaxConditionIndent);
 		// Next indent, or terminate at zero.
-		const int32 NextIndent = Conditions.IsValidIndex(Index + 1) ? FMath::Clamp((int32)Conditions[Index].ConditionIndent, 0, UE::StateTree::MaxConditionIndent) : 0;
+		const int32 NextIndent = Conditions.IsValidIndex(Index + 1) ? FMath::Clamp((int32)Conditions[Index + 1].ConditionIndent, 0, UE::StateTree::MaxConditionIndent) : 0;
+		
 		const int32 DeltaIndent = NextIndent - CurrIndent;
 
 		if (!CreateCondition(State, CondNode, Operand, (int8)DeltaIndent))
diff --git a/Engine/Plugins/Runtime/StateTree/Source/StateTreeTestSuite/Private/StateTreeTest.cpp b/Engine/Plugins/Runtime/StateTree/Source/StateTreeTestSuite/Private/StateTreeTest.cpp
--- a/Engine/Plugins/Runtime/StateTree/Source/StateTreeTestSuite/Private/StateTreeTest.cpp
+++ b/Engine/Plugins/Runtime/StateTree/Source/StateTreeTestSuite/Private/StateTreeTest.cpp
@@ -655,6 +655,59 @@ struct FStateTreeTest_TransitionPriorityEnterState : FAITestBase
 };
 IMPLEMENT_AI_INSTANT_TEST(FStateTreeTest_TransitionPriorityEnterState, "System.StateTree.Transition.PriorityEnterState");
 
+struct FStateTreeTest_LastConditionWithIndent : FAITestBase
+{
+	virtual bool InstantTest() override
+	{
+		UStateTree& StateTree = UE::StateTree::Tests::NewStateTree(&GetWorld());
+		UStateTreeEditorData& EditorData = *Cast<UStateTreeEditorData>(StateTree.EditorData);
+		
+		UStateTreeState& Root = EditorData.AddSubTree(FName(TEXT("Root")));
+		UStateTreeState& State1 = Root.AddChildState(FName(TEXT("State1")));
+
+		auto& Task1 = State1.AddTask<FTestTask_Stand>(FName(TEXT("Task1")));
+		State1.AddEnterCondition<FStateTreeTestCondition>();
+		auto& LastCondition = State1.AddEnterCondition<FStateTreeTestCondition>();
+
+		// Last condition has Indent
+		LastCondition.ConditionIndent = 1;
+		
+		State1.AddTransition(EStateTreeTransitionTrigger::OnStateCompleted, EStateTreeTransitionType::Succeeded);
+
+		FStateTreeCompilerLog Log;
+		FStateTreeCompiler Compiler(Log);
+		const bool bResult = Compiler.Compile(StateTree);
+		AITEST_TRUE("StateTree should get compiled", bResult);
+
+		EStateTreeRunStatus Status = EStateTreeRunStatus::Unset;
+		FStateTreeInstanceData InstanceData;
+		FTestStateTreeExecutionContext Exec(StateTree, StateTree, InstanceData);
+		const bool bInitSucceeded = Exec.IsValid();
+		AITEST_TRUE("StateTree should init", bInitSucceeded);
+
+		const FString TickStr(TEXT("Tick"));
+		const FString EnterStateStr(TEXT("EnterState"));
+		const FString ExitStateStr(TEXT("ExitState"));
+
+		Status = Exec.Start();
+		AITEST_TRUE("StateTree Task1 should enter state", Exec.Expect(Task1.GetName(), EnterStateStr));
+		AITEST_FALSE("StateTree Task1 should not tick", Exec.Expect(Task1.GetName(), TickStr));
+		Exec.LogClear();
+
+		Status = Exec.Tick(0.1f);
+		AITEST_TRUE("StateTree Task1 should tick, and exit state", Exec.Expect(Task1.GetName(), TickStr).Then(Task1.GetName(), ExitStateStr));
+		AITEST_TRUE("StateTree should be completed", Status == EStateTreeRunStatus::Succeeded);
+		Exec.LogClear();
+
+		Status = Exec.Tick(0.1f);
+		AITEST_FALSE("StateTree Task1 should not tick", Exec.Expect(Task1.GetName(), TickStr));
+		Exec.LogClear();
+
+		return true;
+	}
+};
+IMPLEMENT_AI_INSTANT_TEST(FStateTreeTest_LastConditionWithIndent, "System.StateTree.LastConditionWithIndent");
+
 struct FStateTreeTest_TransitionGlobalDataView : FAITestBase
 {
 	// Tests that the global eval and task dataviews are kept up to date when transitioning from  
