StateTree: Moving OnBindingChanged methods from FStateTreeConditionBase to FStateTreeNodeBase

diff --git a/Engine/Plugins/Runtime/StateTree/Source/StateTreeEditorModule/Private/Customizations/StateTreeEditorNodeDetails.cpp b/Engine/Plugins/Runtime/StateTree/Source/StateTreeEditorModule/Private/Customizations/StateTreeEditorNodeDetails.cpp
--- a/Engine/Plugins/Runtime/StateTree/Source/StateTreeEditorModule/Private/Customizations/StateTreeEditorNodeDetails.cpp
+++ b/Engine/Plugins/Runtime/StateTree/Source/StateTreeEditorModule/Private/Customizations/StateTreeEditorNodeDetails.cpp
@@ -736,6 +736,11 @@ void FStateTreeEditorNodeDetails::OnBindingChanged(const FStateTreePropertyPath&
 {
 	check(StructProperty);
 
+	if (!EditorData)
+	{
+		return;
+	}
+
 	TArray<UObject*> OuterObjects;
 	StructProperty->GetOuterObjects(OuterObjects);
 
@@ -747,21 +752,21 @@ void FStateTreeEditorNodeDetails::OnBindingChanged(const FStateTreePropertyPath&
 		return;
 	}
 
+	const FStateTreeBindingLookup BindingLookup(EditorData);
+
 	for (int32 i = 0; i < OuterObjects.Num(); i++)
 	{
-		FStateTreeEditorNode* Node = static_cast<FStateTreeEditorNode*>(RawNodeData[i]);
+		FStateTreeEditorNode* EditorNode = static_cast<FStateTreeEditorNode*>(RawNodeData[i]);
 		UObject* OuterObject = OuterObjects[i]; // Immediate outer, i.e StateTreeState
-		if (Node != nullptr && EditorData != nullptr && Node->Node.IsValid() && Node->Instance.IsValid())
+		if (EditorNode && OuterObject && EditorNode->ID == TargetPath.GetStructID())
 		{
-			if (Node->ID == TargetPath.GetStructID())
-			{
-				if (FStateTreeConditionBase* Condition = Node->Node.GetMutablePtr<FStateTreeConditionBase>())
-				{
-					const FStateTreeBindingLookup BindingLookup(EditorData);
+			FStateTreeNodeBase* Node = EditorNode->Node.GetMutablePtr<FStateTreeNodeBase>();
+			FStateTreeDataView InstanceView = EditorNode->GetInstance(); 
 
-					OuterObject->Modify();
-					Condition->OnBindingChanged(Node->ID, FStateTreeDataView(Node->Instance), SourcePath, TargetPath, BindingLookup);
-				}
+			if (Node && InstanceView.IsValid())
+			{
+				OuterObject->Modify();
+				Node->OnBindingChanged(EditorNode->ID, InstanceView, SourcePath, TargetPath, BindingLookup);
 			}
 		}
 	}
diff --git a/Engine/Plugins/Runtime/StateTree/Source/StateTreeModule/Public/StateTreeConditionBase.h b/Engine/Plugins/Runtime/StateTree/Source/StateTreeModule/Public/StateTreeConditionBase.h
--- a/Engine/Plugins/Runtime/StateTree/Source/StateTreeModule/Public/StateTreeConditionBase.h
+++ b/Engine/Plugins/Runtime/StateTree/Source/StateTreeModule/Public/StateTreeConditionBase.h
@@ -5,10 +5,7 @@
 #include "StateTreeNodeBase.h"
 #include "StateTreeConditionBase.generated.h"
 
-struct FStateTreeEditorPropertyPath;
-struct FStateTreePropertyPath;
 struct FStateTreeExecutionContext;
-struct IStateTreeBindingLookup;
 
 enum class EStateTreeCompare : uint8
 {
@@ -23,22 +20,6 @@ USTRUCT(meta = (Hidden))
 struct STATETREEMODULE_API FStateTreeConditionBase : public FStateTreeNodeBase
 {
 	GENERATED_BODY()
-
-#if WITH_EDITOR
-PRAGMA_DISABLE_DEPRECATION_WARNINGS	
-	UE_DEPRECATED(5.3, "Use version with FStateTreePropertyPath instead.")
-	virtual void OnBindingChanged(const FGuid& ID, FStateTreeDataView InstanceData, const FStateTreeEditorPropertyPath& SourcePath, const FStateTreeEditorPropertyPath& TargetPath, const IStateTreeBindingLookup& BindingLookup) final {}
-PRAGMA_ENABLE_DEPRECATION_WARNINGS
-	/**
-	 * Called when binding of any of the properties in the condition changes.
-	 * @param ID ID of the item, can be used make property paths to this item.
-	 * @param InstanceData view to the instance data, can be struct or class. 
-	 * @param SourcePath Source path of the new binding.
-	 * @param TargetPath Target path of the new binding (the property in the condition).
-	 * @param BindingLookup Reference to binding lookup which can be used to reason about property paths.
-	 */
-	virtual void OnBindingChanged(const FGuid& ID, FStateTreeDataView InstanceData, const FStateTreePropertyPath& SourcePath, const FStateTreePropertyPath& TargetPath, const IStateTreeBindingLookup& BindingLookup) {}
-#endif
 	
 	/** @return True if the condition passes. */
 	virtual bool TestCondition(FStateTreeExecutionContext& Context) const { return false; }
diff --git a/Engine/Plugins/Runtime/StateTree/Source/StateTreeModule/Public/StateTreeNodeBase.h b/Engine/Plugins/Runtime/StateTree/Source/StateTreeModule/Public/StateTreeNodeBase.h
--- a/Engine/Plugins/Runtime/StateTree/Source/StateTreeModule/Public/StateTreeNodeBase.h
+++ b/Engine/Plugins/Runtime/StateTree/Source/StateTreeModule/Public/StateTreeNodeBase.h
@@ -6,6 +6,9 @@
 #include "StateTreeNodeBase.generated.h"
 
 struct FStateTreeLinker;
+struct FStateTreeEditorPropertyPath;
+struct FStateTreePropertyPath;
+struct IStateTreeBindingLookup;
 
 /**
  * Base struct of StateTree Conditions, Evaluators, and Tasks.
@@ -39,6 +42,22 @@ struct STATETREEMODULE_API FStateTreeNodeBase
 	 */
 	virtual EDataValidationResult Compile(FStateTreeDataView InstanceDataView, TArray<FText>& ValidationMessages) { return EDataValidationResult::NotValidated; }
 	
+#if WITH_EDITOR
+	PRAGMA_DISABLE_DEPRECATION_WARNINGS
+	UE_DEPRECATED(5.3, "Use version with FStateTreePropertyPath instead.")
+	virtual void OnBindingChanged(const FGuid& ID, FStateTreeDataView InstanceData, const FStateTreeEditorPropertyPath& SourcePath, const FStateTreeEditorPropertyPath& TargetPath, const IStateTreeBindingLookup& BindingLookup) final {}
+	PRAGMA_ENABLE_DEPRECATION_WARNINGS
+	/**
+	 * Called when binding of any of the properties in the node changes.
+	 * @param ID ID of the item, can be used make property paths to this item.
+	 * @param InstanceData view to the instance data, can be struct or class.
+	 * @param SourcePath Source path of the new binding.
+	 * @param TargetPath Target path of the new binding (the property in the condition).
+	 * @param BindingLookup Reference to binding lookup which can be used to reason about property paths.
+	 */
+	virtual void OnBindingChanged(const FGuid& ID, FStateTreeDataView InstanceData, const FStateTreePropertyPath& SourcePath, const FStateTreePropertyPath& TargetPath, const IStateTreeBindingLookup& BindingLookup) {}
+#endif
+
 	/** Name of the node. */
 	UPROPERTY(EditDefaultsOnly, Category = "", meta=(EditCondition = "false", EditConditionHides))
 	FName Name;
