StateTree: Handle nodes with instance subobjects.

diff --git a/Engine/Plugins/Runtime/StateTree/Source/StateTreeEditorModule/Private/StateTreeCompiler.cpp b/Engine/Plugins/Runtime/StateTree/Source/StateTreeEditorModule/Private/StateTreeCompiler.cpp
--- a/Engine/Plugins/Runtime/StateTree/Source/StateTreeEditorModule/Private/StateTreeCompiler.cpp
+++ b/Engine/Plugins/Runtime/StateTree/Source/StateTreeEditorModule/Private/StateTreeCompiler.cpp
@@ -8,9 +8,74 @@
 #include "StateTreeEvaluatorBase.h"
 #include "StateTreeTaskBase.h"
 #include "StateTreeConditionBase.h"
+#include "Serialization/ArchiveUObject.h"
+
+
+class FStateTreeArchiveCheckOuters : public FArchiveUObject
+{
+	
+};
+
 
 namespace UE::StateTree::Compiler
 {
+	// Helper archive that checks that the all instanced sub-objects have correct outer. 
+	class FCheckOutersArchive : public FArchiveUObject
+	{
+		using Super = FArchiveUObject;
+		const UStateTree& StateTree;
+		const UStateTreeEditorData& EditorData;
+		FStateTreeCompilerLog& Log;
+	public:
+
+		FCheckOutersArchive(const UStateTree& InStateTree, const UStateTreeEditorData& InEditorData, FStateTreeCompilerLog& InLog)
+			: StateTree(InStateTree)
+			, EditorData(InEditorData)
+			, Log(InLog)
+		{
+			Super::SetIsSaving(true);
+			Super::SetIsPersistent(true);
+		}
+
+		virtual bool ShouldSkipProperty(const FProperty* InProperty) const
+		{
+			// Skip editor data.
+			if (const FObjectProperty* ObjectProperty = CastField<FObjectProperty>(InProperty))
+			{
+				if (ObjectProperty->PropertyClass == UStateTreeEditorData::StaticClass())
+				{
+					return true;
+				}
+			}
+			return false;
+		}
+
+		virtual FArchive& operator<<(UObject*& Object) override
+		{
+			if (Object)
+			{
+				if (const FProperty* Property = GetSerializedProperty())
+				{
+					if (Property->HasAnyPropertyFlags(CPF_InstancedReference))
+					{
+						if (!Object->IsInOuter(&StateTree))
+						{
+							Log.Reportf(EMessageSeverity::Error, TEXT("Compiled StateTree contains instanced object %s (%s), which does not belong to the StateTree. This is due to error in the State Tree node implementation."),
+								*GetFullNameSafe(Object), *GetFullNameSafe(Object->GetClass()));
+						}
+
+						if (Object->IsInOuter(&EditorData))
+						{
+							Log.Reportf(EMessageSeverity::Error, TEXT("Compiled StateTree contains instanced object %s (%s), which still belongs to the Editor data. This is due to error in the State Tree node implementation."),
+								*GetFullNameSafe(Object), *GetFullNameSafe(Object->GetClass()));
+						}
+					}
+				}
+			}
+			return *this;
+		}
+	};
+
 	void FValidationResult::Log(FStateTreeCompilerLog& Log, const TCHAR* ContextText, const FStateTreeBindableStructDesc& ContextStruct) const
 	{
 		Log.Reportf(EMessageSeverity::Error, ContextStruct, TEXT("The StateTree is too complex. Compact index %s out of range %d/%d."), ContextText, Value, MaxValue);
@@ -184,6 +249,9 @@ bool FStateTreeCompiler::Compile(UStateTree& InStateTree)
 		StateTree->IDToStateMappings.Emplace(ToState.Key, FStateTreeStateHandle(ToState.Value));
 	}
 
+	UE::StateTree::Compiler::FCheckOutersArchive CheckOuters(*StateTree, *EditorData, Log);
+	StateTree->Serialize(CheckOuters);
+	
 	return true;
 }
 
@@ -774,6 +842,8 @@ bool FStateTreeCompiler::CreateCondition(UStateTreeState& State, const FStateTre
 
 	// Copy the condition
 	FInstancedStruct& Node = Nodes.Add_GetRef(CondNode.Node);
+	InstantiateStructSubobjects(Node);
+
 	FStateTreeConditionBase& Cond = Node.GetMutable<FStateTreeConditionBase>();
 
 	Cond.Operand = Operand;
@@ -785,7 +855,8 @@ bool FStateTreeCompiler::CreateCondition(UStateTreeState& State, const FStateTre
 	{
 		// Struct instance
 		const int32 InstanceIndex = SharedInstanceStructs.Add(CondNode.Instance);
-	
+		InstantiateStructSubobjects(SharedInstanceStructs[InstanceIndex]);
+
 		// Create binding source struct descriptor.
 		StructDesc.Struct = CondNode.Instance.GetScriptStruct();
 		StructDesc.Name = Cond.Name;
@@ -941,6 +1012,8 @@ bool FStateTreeCompiler::CreateTask(UStateTreeState* State, const FStateTreeEdit
 
 	// Copy the task
 	FInstancedStruct& Node = Nodes.Add_GetRef(TaskNode.Node);
+	InstantiateStructSubobjects(Node);
+	
 	FStateTreeTaskBase& Task = Node.GetMutable<FStateTreeTaskBase>();
 	FStateTreeDataView InstanceDataView;
 
@@ -948,6 +1021,7 @@ bool FStateTreeCompiler::CreateTask(UStateTreeState* State, const FStateTreeEdit
 	{
 		// Struct Instance
 		const int32 InstanceIndex = InstanceStructs.Add(TaskNode.Instance);
+		InstantiateStructSubobjects(InstanceStructs[InstanceIndex]);
 
 		// Create binding source struct descriptor.
 		StructDesc.Struct = TaskNode.Instance.GetScriptStruct();
@@ -1052,6 +1126,8 @@ bool FStateTreeCompiler::CreateEvaluator(const FStateTreeEditorNode& EvalNode)
 
 	// Copy the evaluator
 	FInstancedStruct& Node = Nodes.Add_GetRef(EvalNode.Node);
+	InstantiateStructSubobjects(Node);
+	
 	FStateTreeEvaluatorBase& Eval = Node.GetMutable<FStateTreeEvaluatorBase>();
 	FStateTreeDataView InstanceDataView;
 	
@@ -1059,6 +1135,7 @@ bool FStateTreeCompiler::CreateEvaluator(const FStateTreeEditorNode& EvalNode)
 	{
 		// Struct Instance
 		const int32 InstanceIndex = InstanceStructs.Add(EvalNode.Instance);
+		InstantiateStructSubobjects(InstanceStructs[InstanceIndex]);
 
 		// Create binding source struct descriptor.
 		StructDesc.Struct = EvalNode.Instance.GetScriptStruct();
@@ -1444,3 +1521,47 @@ bool FStateTreeCompiler::GetAndValidateBindings(const FStateTreeBindableStructDe
 
 	return bResult;
 }
+
+void FStateTreeCompiler::InstantiateStructSubobjects(FStructView Struct)
+{
+	check(StateTree);
+	check(EditorData);
+	
+	// Empty struct, nothing to do.
+	if (!Struct.IsValid())
+	{
+		return;
+	}
+
+	for (TPropertyValueIterator<FProperty> It(Struct.GetScriptStruct(), Struct.GetMemory()); It; ++It)
+	{
+		if (const FObjectProperty* ObjectProperty = CastField<FObjectProperty>(It->Key))
+		{
+			// Duplicate instanced objects.
+			if (ObjectProperty->HasAnyPropertyFlags(CPF_InstancedReference))
+			{
+				if (UObject* Object = ObjectProperty->GetObjectPropertyValue(It->Value))
+				{
+					UObject* OuterObject = Object->GetOuter();
+					// If the instanced object was created as Editor Data as outer,
+					// change the outer to State Tree to prevent references to editor only data.
+					if (Object->IsInOuter(EditorData))
+					{
+						OuterObject = StateTree;
+					}
+					UObject* DuplicatedObject = DuplicateObject(Object, OuterObject);
+					ObjectProperty->SetObjectPropertyValue(const_cast<void*>(It->Value), DuplicatedObject);
+				}
+			}
+		}
+		if (const FStructProperty* StructProperty = CastField<FStructProperty>(It->Key))
+		{
+			// If we encounter instanced struct, recursively handle it too.
+			if (StructProperty->Struct == TBaseStructure<FInstancedStruct>::Get())
+			{
+				FInstancedStruct& InstancedStruct = *static_cast<FInstancedStruct*>(const_cast<void*>(It->Value));
+				InstantiateStructSubobjects(InstancedStruct);
+			}
+		}
+	}
+}
diff --git a/Engine/Plugins/Runtime/StateTree/Source/StateTreeEditorModule/Public/StateTreeCompiler.h b/Engine/Plugins/Runtime/StateTree/Source/StateTreeEditorModule/Public/StateTreeCompiler.h
--- a/Engine/Plugins/Runtime/StateTree/Source/StateTreeEditorModule/Public/StateTreeCompiler.h
+++ b/Engine/Plugins/Runtime/StateTree/Source/StateTreeEditorModule/Public/StateTreeCompiler.h
@@ -58,6 +58,8 @@ private:
 							const FStateTreeBindableStructDesc& TargetStruct, FStateTreePropertyPath TargetPath) const;
 	bool CompileAndValidateNode(const UStateTreeState* SourceState, const FStateTreeBindableStructDesc& NodeDesc, FStructView NodeView, const FStateTreeDataView InstanceData);
 
+	void InstantiateStructSubobjects(FStructView Struct);
+	
 	FStateTreeDataView GetBindingSourceValue(const int32 SourceIndex);
 	
 	FStateTreeCompilerLog& Log;
