State Tree: Added new selection behaviors for states
- Added selection behavior None (not selectable)
- Added selection behavior TryEnterState (enters the specific state, does not consider child states)
- Added selection behavior TryFollowTransitions (follows transitions when state is tried to be selected)
- Brushed up a bit the state representation in editor (added icons, cleaned up layout)

diff --git a/Engine/Plugins/Runtime/StateTree/Source/StateTreeEditorModule/Private/Customizations/StateTreeStateDetails.cpp b/Engine/Plugins/Runtime/StateTree/Source/StateTreeEditorModule/Private/Customizations/StateTreeStateDetails.cpp
--- a/Engine/Plugins/Runtime/StateTree/Source/StateTreeEditorModule/Private/Customizations/StateTreeStateDetails.cpp
+++ b/Engine/Plugins/Runtime/StateTree/Source/StateTreeEditorModule/Private/Customizations/StateTreeStateDetails.cpp
@@ -43,6 +43,7 @@ void FStateTreeStateDetails::CustomizeDetails(IDetailLayoutBuilder& DetailBuilde
 	TSharedPtr<IPropertyHandle> TypeProperty = DetailBuilder.GetProperty(TEXT("Type"));
 	TSharedPtr<IPropertyHandle> LinkedSubtreeProperty = DetailBuilder.GetProperty(TEXT("LinkedSubtree"));
 	TSharedPtr<IPropertyHandle> ParametersProperty = DetailBuilder.GetProperty(TEXT("Parameters"));
+	TSharedPtr<IPropertyHandle> SelectionBehaviorProperty = DetailBuilder.GetProperty(TEXT("SelectionBehavior"));
 
 	uint8 StateTypeValue = 0;
 	TypeProperty->GetValue(StateTypeValue);
@@ -60,6 +61,11 @@ void FStateTreeStateDetails::CustomizeDetails(IDetailLayoutBuilder& DetailBuilde
 	{
 		ParametersProperty->MarkHiddenByCustomization();
 	}
+
+	if (StateType == EStateTreeStateType::Linked)
+	{
+		SelectionBehaviorProperty->MarkHiddenByCustomization();
+	}
 	
 	const FName EnterConditionsCategoryName(TEXT("Enter Conditions"));
 	if (Schema && Schema->AllowEnterConditions())
diff --git a/Engine/Plugins/Runtime/StateTree/Source/StateTreeEditorModule/Private/Customizations/StateTreeStateLinkDetails.cpp b/Engine/Plugins/Runtime/StateTree/Source/StateTreeEditorModule/Private/Customizations/StateTreeStateLinkDetails.cpp
--- a/Engine/Plugins/Runtime/StateTree/Source/StateTreeEditorModule/Private/Customizations/StateTreeStateLinkDetails.cpp
+++ b/Engine/Plugins/Runtime/StateTree/Source/StateTreeEditorModule/Private/Customizations/StateTreeStateLinkDetails.cpp
@@ -106,6 +106,11 @@ void FStateTreeStateLinkDetails::CacheStates(const UStateTreeState* State)
 		bShouldAdd = false;
 	}
 
+	if (State->SelectionBehavior == EStateTreeStateSelectionBehavior::None)
+	{
+		bShouldAdd = false;
+	}
+	
 	if (bShouldAdd)
 	{
 		CachedNames.Add(State->Name);
diff --git a/Engine/Plugins/Runtime/StateTree/Source/StateTreeEditorModule/Private/SStateTreeViewRow.cpp b/Engine/Plugins/Runtime/StateTree/Source/StateTreeEditorModule/Private/SStateTreeViewRow.cpp
--- a/Engine/Plugins/Runtime/StateTree/Source/StateTreeEditorModule/Private/SStateTreeViewRow.cpp
+++ b/Engine/Plugins/Runtime/StateTree/Source/StateTreeEditorModule/Private/SStateTreeViewRow.cpp
@@ -32,6 +32,8 @@ void SStateTreeViewRow::Construct(const FArguments& InArgs, const TSharedRef<STa
 
 	static const FLinearColor TasksBackground = FLinearColor(FColor(17, 117, 131));
 	static const FLinearColor LinkBackground = FLinearColor(FColor(84, 84, 84));
+	static const FLinearColor IconTint = FLinearColor(1, 1, 1, 0.5f);
+	static const FLinearColor TypeIconTint = FLinearColor(0, 0, 0, 0.5f);
 
 	this->ChildSlot
 	.HAlign(HAlign_Fill)
@@ -70,8 +72,24 @@ void SStateTreeViewRow::Construct(const FArguments& InArgs, const TSharedRef<STa
 					.BorderImage(FAppStyle::GetBrush("WhiteBrush"))
 					.BorderBackgroundColor(this, &SStateTreeViewRow::GetActiveStateColor)
 					[
-						// State Box
 						SNew(SHorizontalBox)
+
+						// Sub tree marker
+						+ SHorizontalBox::Slot()
+						.VAlign(VAlign_Center)
+						.AutoWidth()
+						[
+							SNew(SBox)
+							.WidthOverride(4.0f)
+							.HeightOverride(28.0f)
+							[
+								SNew(SBorder)
+								.BorderImage(FAppStyle::GetBrush("WhiteBrush"))
+								.BorderBackgroundColor(this, &SStateTreeViewRow::GetSubTreeMarkerColor)
+							]
+						]
+
+						// State Box
 						+ SHorizontalBox::Slot()
 						.VAlign(VAlign_Center)
 						.AutoWidth()
@@ -83,10 +101,11 @@ void SStateTreeViewRow::Construct(const FArguments& InArgs, const TSharedRef<STa
 								SNew(SBorder)
 								.BorderImage(FAppStyle::GetBrush("WhiteBrush"))
 								.BorderBackgroundColor(this, &SStateTreeViewRow::GetTitleColor)
-								.Padding(FMargin(16.0f, 0.0f, 16.0f, 0.0f))
+								.Padding(FMargin(4.0f, 0.0f, 12.0f, 0.0f))
 								[
-									// Breakpoint box
 									SNew(SHorizontalBox)
+
+									// Breakpoint box
 									+ SHorizontalBox::Slot()
 									.VAlign(VAlign_Top)
 									.HAlign(HAlign_Left)
@@ -112,7 +131,9 @@ void SStateTreeViewRow::Construct(const FArguments& InArgs, const TSharedRef<STa
 										.Visibility(this, &SStateTreeViewRow::GetConditionVisibility)
 										[
 											SNew(SImage)
-											.Image(FAppStyle::Get().GetBrush("Icons.Help"))
+											.ColorAndOpacity(IconTint)
+											.Image(FStateTreeEditorStyle::Get().GetBrush("StateTreeEditor.Conditions"))
+											.ToolTipText(LOCTEXT("StateHasEnterConditions", "State selection is guarded with enter conditions."))
 										]
 									]
 
@@ -123,11 +144,11 @@ void SStateTreeViewRow::Construct(const FArguments& InArgs, const TSharedRef<STa
 									[
 										SNew(SBox)
 										.Padding(FMargin(0.0f, 0.0f, 4.0f, 0.0f))
-										.Visibility(this, &SStateTreeViewRow::GetSelectorVisibility)
 										[
-											SNew(STextBlock)
-											.Text(this, &SStateTreeViewRow::GetSelectorDesc)
-											.TextStyle(FStateTreeEditorStyle::Get(), "StateTree.Icon")
+											SNew(SImage)
+											.Image(this, &SStateTreeViewRow::GetSelectorIcon)
+											.ColorAndOpacity(IconTint)
+											.ToolTipText( this, &SStateTreeViewRow::GetSelectorTooltip)
 										]
 									]
 
@@ -144,6 +165,7 @@ void SStateTreeViewRow::Construct(const FArguments& InArgs, const TSharedRef<STa
 										})
 										.OnTextCommitted(this, &SStateTreeViewRow::HandleNodeLabelTextCommitted)
 										.Text(this, &SStateTreeViewRow::GetStateDesc)
+										.ToolTipText(this, &SStateTreeViewRow::GetStateTypeTooltip)
 										.Clipping(EWidgetClipping::ClipToBounds)
 										.IsSelected(this, &SStateTreeViewRow::IsStateSelected)
 									]
@@ -164,7 +186,7 @@ void SStateTreeViewRow::Construct(const FArguments& InArgs, const TSharedRef<STa
 								SNew(SBorder)
 								.BorderImage(FAppStyle::GetBrush("WhiteBrush"))
 								.BorderBackgroundColor(TasksBackground)
-								.Padding(FMargin(12.0f, 0.0f, 16.0f, 0.0f))
+								.Padding(FMargin(6.0f, 0.0f, 12.0f, 0.0f))
 								[
 									// Task icon
 									SNew(SHorizontalBox)
@@ -190,46 +212,45 @@ void SStateTreeViewRow::Construct(const FArguments& InArgs, const TSharedRef<STa
 								]
 							]
 						]
-					]
-				]
-			]
-
-			// Linked State box
-			+ SHorizontalBox::Slot()
-			.VAlign(VAlign_Center)
-			.Padding(FMargin(0.0f, 4.0f))
-			.AutoWidth()
-			[
-				SNew(SBox)
-				.HeightOverride(28.0f)
-				.VAlign(VAlign_Fill)
-				.Visibility(this, &SStateTreeViewRow::GetLinkedStateVisibility)
-				[
-					SNew(SBorder)
-					.BorderImage(FAppStyle::GetBrush("WhiteBrush"))
-					.BorderBackgroundColor(LinkBackground)
-					.Padding(FMargin(12.0f, 0.0f, 16.0f, 0.0f))
-					[
-						// Link icon
-						SNew(SHorizontalBox)
+						
+						// Linked State box
 						+ SHorizontalBox::Slot()
-						.VAlign(VAlign_Center)
-						.Padding(FMargin(0.0f, 0.0f, 4.0f, 0.0f))
+						.VAlign(VAlign_Fill)
 						.AutoWidth()
 						[
-							SNew(STextBlock)
-							.Text(FEditorFontGlyphs::Link)
-							.TextStyle(FStateTreeEditorStyle::Get(), "StateTree.DetailsIcon")
-						]
+							SNew(SBox)
+							.HeightOverride(28.0f)
+							.VAlign(VAlign_Fill)
+							.Visibility(this, &SStateTreeViewRow::GetLinkedStateVisibility)
+							[
+								SNew(SBorder)
+								.BorderImage(FAppStyle::GetBrush("WhiteBrush"))
+								.BorderBackgroundColor(LinkBackground)
+								.Padding(FMargin(6.0f, 0.0f, 12.0f, 0.0f))
+								[
+									// Link icon
+									SNew(SHorizontalBox)
+									+ SHorizontalBox::Slot()
+									.VAlign(VAlign_Center)
+									.AutoWidth()
+									.Padding(FMargin(0.0f, 0.0f, 4.0f, 0.0f))
+									[
+										SNew(SImage)
+										.ColorAndOpacity(IconTint)
+										.Image(FStateTreeEditorStyle::Get().GetBrush("StateTreeEditor.StateLinked"))
+									]
 
-						// Linked State
-						+ SHorizontalBox::Slot()
-						.VAlign(VAlign_Center)
-						.AutoWidth()
-						[
-							SNew(STextBlock)
-							.Text(this, &SStateTreeViewRow::GetLinkedStateDesc)
-							.TextStyle(FStateTreeEditorStyle::Get(), "StateTree.Details")
+									// Linked State
+									+ SHorizontalBox::Slot()
+									.VAlign(VAlign_Center)
+									.AutoWidth()
+									[
+										SNew(STextBlock)
+										.Text(this, &SStateTreeViewRow::GetLinkedStateDesc)
+										.TextStyle(FStateTreeEditorStyle::Get(), "StateTree.Details")
+									]
+								]
+							]
 						]
 					]
 				]
@@ -393,10 +414,6 @@ FSlateColor SStateTreeViewRow::GetTitleColor() const
 {
 	if (const UStateTreeState* State = WeakState.Get())
 	{
-		if (StateTreeViewModel && StateTreeViewModel->IsSelected(State))
-		{
-			return FLinearColor(FColor(236, 134, 39));
-		}
 		if (IsRootState() || State->Type == EStateTreeStateType::Subtree)
 		{
 			return FLinearColor(FColor(17, 117, 131));
@@ -414,6 +431,24 @@ FSlateColor SStateTreeViewRow::GetActiveStateColor() const
 		{
 			return FLinearColor::Yellow;
 		}
+		if (StateTreeViewModel && StateTreeViewModel->IsSelected(State))
+		{
+			return FLinearColor(FColor(236, 134, 39));
+		}
+	}
+
+	return FLinearColor::Transparent;
+}
+
+FSlateColor SStateTreeViewRow::GetSubTreeMarkerColor() const
+{
+	// Show color for subtree.
+	if (const UStateTreeState* State = WeakState.Get())
+	{
+		if (IsRootState() || State->Type == EStateTreeStateType::Subtree)
+		{
+			return FLinearColor(FColor(136, 186, 193));
+		}
 	}
 
 	return GetTitleColor();
@@ -437,61 +472,104 @@ EVisibility SStateTreeViewRow::GetConditionVisibility() const
 	return EVisibility::Collapsed;
 }
 
-EVisibility SStateTreeViewRow::GetSelectorVisibility() const
+EVisibility SStateTreeViewRow::GetBreakpointVisibility() const
 {
 	if (const UStateTreeState* State = WeakState.Get())
 	{
-		if (State->Type == EStateTreeStateType::State || State->Type == EStateTreeStateType::Group)
-		{
-			return State->Children.Num() > 0 ? EVisibility::Visible : EVisibility::Collapsed;
-		}
-		else if (IsRootState() || State->Type == EStateTreeStateType::Subtree)
-		{
-			return EVisibility::Visible;
-		}
-		else if (State->Type == EStateTreeStateType::Linked)
+		if (StateTreeViewModel && StateTreeViewModel->DoesStateHaveBreakpoint(*State))
 		{
 			return EVisibility::Visible;
 		}
 	}
-	
+
 	return EVisibility::Collapsed;
 }
 
-EVisibility SStateTreeViewRow::GetBreakpointVisibility() const
+const FSlateBrush* SStateTreeViewRow::GetSelectorIcon() const
 {
 	if (const UStateTreeState* State = WeakState.Get())
 	{
-		if (StateTreeViewModel && StateTreeViewModel->DoesStateHaveBreakpoint(*State))
+		if (State->SelectionBehavior == EStateTreeStateSelectionBehavior::None)
 		{
-			return EVisibility::Visible;
+			return FStateTreeEditorStyle::Get().GetBrush("StateTreeEditor.SelectNone");
+		}
+		else if (State->SelectionBehavior == EStateTreeStateSelectionBehavior::TryEnterState)
+		{
+			return FStateTreeEditorStyle::Get().GetBrush("StateTreeEditor.TryEnterState");			
+		}
+		else if (State->SelectionBehavior == EStateTreeStateSelectionBehavior::TrySelectChildrenInOrder)
+		{
+			if (State->Children.IsEmpty())
+			{
+				// Backwards compatible behavior
+				return FStateTreeEditorStyle::Get().GetBrush("StateTreeEditor.TryEnterState");			
+			}
+			else
+			{
+				return FStateTreeEditorStyle::Get().GetBrush("StateTreeEditor.TrySelectChildrenInOrder");
+			}
+		}
+		else if (State->SelectionBehavior == EStateTreeStateSelectionBehavior::TryFollowTransitions)
+		{
+			return FStateTreeEditorStyle::Get().GetBrush("StateTreeEditor.TryFollowTransitions");
 		}
 	}
 
-	return EVisibility::Collapsed;
+	return nullptr;
 }
 
-FText SStateTreeViewRow::GetSelectorDesc() const
+FText SStateTreeViewRow::GetSelectorTooltip() const
 {
 	if (const UStateTreeState* State = WeakState.Get())
 	{
-		if (State->Type == EStateTreeStateType::State || State->Type == EStateTreeStateType::Group)
+		const UEnum* Enum = StaticEnum<EStateTreeStateSelectionBehavior>();
+		check(Enum);
+		const int32 Index = Enum->GetIndexByValue((int64)State->SelectionBehavior);
+		
+		if (State->SelectionBehavior == EStateTreeStateSelectionBehavior::None)
 		{
-			return FEditorFontGlyphs::Level_Down; // Selector
+			return Enum->GetToolTipTextByIndex(Index);
 		}
-		else if (IsRootState() || State->Type == EStateTreeStateType::Subtree)
+		else if (State->SelectionBehavior == EStateTreeStateSelectionBehavior::TryEnterState)
 		{
-			return FEditorFontGlyphs::Cogs;
+			return Enum->GetToolTipTextByIndex(Index);
 		}
-		else if (State->Type == EStateTreeStateType::Linked)
+		else if (State->SelectionBehavior == EStateTreeStateSelectionBehavior::TrySelectChildrenInOrder)
 		{
-			return FEditorFontGlyphs::Link;
+			if (State->Children.IsEmpty())
+			{
+				const int32 EnterStateIndex = Enum->GetIndexByValue((int64)EStateTreeStateSelectionBehavior::TryEnterState);
+				return FText::Format(LOCTEXT("ConvertedToEnterState", "{0}\nAutomatically converted from '{1}' becase the State has no child States."),
+					Enum->GetToolTipTextByIndex(EnterStateIndex), UEnum::GetDisplayValueAsText(State->SelectionBehavior));
+			}
+			else
+			{
+				return Enum->GetToolTipTextByIndex(Index);
+			}
+		}
+		else if (State->SelectionBehavior == EStateTreeStateSelectionBehavior::TryFollowTransitions)
+		{
+			return Enum->GetToolTipTextByIndex(Index);
 		}
 	}
 
-	return FText::GetEmpty(); 
+	return FText::GetEmpty();
 }
 
+FText SStateTreeViewRow::GetStateTypeTooltip() const
+{
+	if (const UStateTreeState* State = WeakState.Get())
+	{
+		const UEnum* Enum = StaticEnum<EStateTreeStateType>();
+		check(Enum);
+		const int32 Index = Enum->GetIndexByValue((int64)State->Type);
+		return Enum->GetToolTipTextByIndex(Index);
+	}
+
+	return FText::GetEmpty();
+}
+
+
 EVisibility SStateTreeViewRow::GetTasksVisibility() const
 {
 	if (const UStateTreeState* State = WeakState.Get())
diff --git a/Engine/Plugins/Runtime/StateTree/Source/StateTreeEditorModule/Private/SStateTreeViewRow.h b/Engine/Plugins/Runtime/StateTree/Source/StateTreeEditorModule/Private/SStateTreeViewRow.h
--- a/Engine/Plugins/Runtime/StateTree/Source/StateTreeEditorModule/Private/SStateTreeViewRow.h
+++ b/Engine/Plugins/Runtime/StateTree/Source/StateTreeEditorModule/Private/SStateTreeViewRow.h
@@ -31,12 +31,15 @@ private:
 
 	FSlateColor GetTitleColor() const;
 	FSlateColor GetActiveStateColor() const;
+	FSlateColor GetSubTreeMarkerColor() const;
 	FText GetStateDesc() const;
 
 	EVisibility GetConditionVisibility() const;
-	EVisibility GetSelectorVisibility() const;
 	EVisibility GetBreakpointVisibility() const;
-	FText GetSelectorDesc() const;
+	
+	const FSlateBrush* GetSelectorIcon() const;
+	FText GetSelectorTooltip() const;
+	FText GetStateTypeTooltip() const;
 
 	EVisibility GetTasksVisibility() const;
 	FText GetTasksDesc() const;
diff --git a/Engine/Plugins/Runtime/StateTree/Source/StateTreeEditorModule/Private/StateTreeCompiler.cpp b/Engine/Plugins/Runtime/StateTree/Source/StateTreeEditorModule/Private/StateTreeCompiler.cpp
--- a/Engine/Plugins/Runtime/StateTree/Source/StateTreeEditorModule/Private/StateTreeCompiler.cpp
+++ b/Engine/Plugins/Runtime/StateTree/Source/StateTreeEditorModule/Private/StateTreeCompiler.cpp
@@ -11,12 +11,6 @@
 #include "Serialization/ArchiveUObject.h"
 
 
-class FStateTreeArchiveCheckOuters : public FArchiveUObject
-{
-	
-};
-
-
 namespace UE::StateTree::Compiler
 {
 	// Helper archive that checks that the all instanced sub-objects have correct outer. 
@@ -332,7 +326,8 @@ bool FStateTreeCompiler::CreateStateRecursive(UStateTreeState& State, const FSta
 	CompactState.Name = State.Name;
 	CompactState.Parent = Parent;
 
-	CompactState.Type = State.Type; 
+	CompactState.Type = State.Type;
+	CompactState.SelectionBehavior = State.SelectionBehavior;
 	
 	SourceStates.Add(&State);
 	IDToState.Add(State.ID, StateIdx);
@@ -345,7 +340,7 @@ bool FStateTreeCompiler::CreateStateRecursive(UStateTreeState& State, const FSta
 		return false;
 	}
 	CompactState.ChildrenBegin = uint16(ChildrenBegin);
-	
+
 	for (UStateTreeState* Child : State.Children)
 	{
 		if (Child != nullptr && Child->Type != EStateTreeStateType::Subtree)
@@ -363,8 +358,8 @@ bool FStateTreeCompiler::CreateStateRecursive(UStateTreeState& State, const FSta
 		Validation.Log(Log, TEXT("ChildrenEnd"));
 		return false;
 	}
-	StateTree->States[StateIdx].ChildrenEnd = uint16(ChildrenEnd);
-
+	StateTree->States[StateIdx].ChildrenEnd = uint16(ChildrenEnd); // Not using CompactState here because the array may have changed.
+	
 	return true;
 }
 
@@ -690,10 +685,19 @@ bool FStateTreeCompiler::CreateStateTransitions()
 			CompactTransition.Trigger = Transition.Trigger;
 			CompactTransition.Priority = Transition.Priority;
 			CompactTransition.EventTag = Transition.EventTag;
+			
 			if (Transition.bDelayTransition)
 			{
 				CompactTransition.Delay.Set(Transition.DelayDuration, Transition.DelayRandomVariance);
 			}
+			
+			if (CompactState.SelectionBehavior == EStateTreeStateSelectionBehavior::TryFollowTransitions
+				&& Transition.bDelayTransition)
+			{
+				Log.Reportf(EMessageSeverity::Warning,
+					TEXT("Transition to '%s' with delay will be ignored during state selection."),
+					*Transition.State.Name.ToString());
+			}
 
 			if (EnumHasAnyFlags(Transition.Trigger, EStateTreeTransitionTrigger::OnStateCompleted))
 			{
@@ -769,7 +773,15 @@ bool FStateTreeCompiler::ResolveTransitionState(const UStateTreeState* SourceSta
 				Log.Reportf(EMessageSeverity::Warning,
 					TEXT("Target state '%s' is in different subtree. Verify that this is intentional."),
 					*Link.Name.ToString());
-			}			
+			}
+
+			if (TargetState->SelectionBehavior == EStateTreeStateSelectionBehavior::None)
+			{
+				Log.Reportf(EMessageSeverity::Error,
+					TEXT("The target State '%s' is not selectable, it's selection behavior is set to None."),
+					*Link.Name.ToString());
+				return false;
+			}
 		}
 		
 		OutTransitionHandle = GetStateHandle(Link.ID);
@@ -784,11 +796,11 @@ bool FStateTreeCompiler::ResolveTransitionState(const UStateTreeState* SourceSta
 	else if (Link.LinkType == EStateTreeTransitionType::NextState)
 	{
 		// Find next state.
-		const UStateTreeState* NextState = SourceState ? SourceState->GetNextSiblingState() : nullptr;
+		const UStateTreeState* NextState = SourceState ? SourceState->GetNextSelectableSiblingState() : nullptr;
 		if (NextState == nullptr)
 		{
 			Log.Reportf(EMessageSeverity::Error,
-				TEXT("Failed to resolve transition, there's no next state."));
+				TEXT("Failed to resolve transition, there's no selectable next state."));
 			return false;
 		}
 		OutTransitionHandle = GetStateHandle(NextState->ID);
diff --git a/Engine/Plugins/Runtime/StateTree/Source/StateTreeEditorModule/Private/StateTreeEditorStyle.cpp b/Engine/Plugins/Runtime/StateTree/Source/StateTreeEditorModule/Private/StateTreeEditorStyle.cpp
--- a/Engine/Plugins/Runtime/StateTree/Source/StateTreeEditorModule/Private/StateTreeEditorStyle.cpp
+++ b/Engine/Plugins/Runtime/StateTree/Source/StateTreeEditorModule/Private/StateTreeEditorStyle.cpp
@@ -189,6 +189,16 @@ FStateTreeEditorStyle::FStateTreeEditorStyle()
 
 		Set("StateTreeEditor.PasteStatesAsSiblings", new IMAGE_BRUSH_SVG("Icons/Sibling_State", CoreStyleConstants::Icon16x16));
 		Set("StateTreeEditor.PasteStatesAsChildren", new IMAGE_BRUSH_SVG("Icons/Child_State", CoreStyleConstants::Icon16x16));
+
+		Set("StateTreeEditor.SelectNone", new IMAGE_BRUSH_SVG("Icons/Select_None", CoreStyleConstants::Icon16x16));
+		Set("StateTreeEditor.TryEnterState", new IMAGE_BRUSH_SVG("Icons/Try_Enter_State", CoreStyleConstants::Icon16x16));
+		Set("StateTreeEditor.TrySelectChildrenInOrder", new IMAGE_BRUSH_SVG("Icons/Try_Select_Children_In_Order", CoreStyleConstants::Icon16x16));
+		Set("StateTreeEditor.TryFollowTransitions", new IMAGE_BRUSH_SVG("Icons/Try_Follow_Transitions", CoreStyleConstants::Icon16x16));
+
+		Set("StateTreeEditor.Conditions", new IMAGE_BRUSH_SVG("Icons/Conditions", CoreStyleConstants::Icon16x16));
+
+		Set("StateTreeEditor.StateSubtree", new IMAGE_BRUSH_SVG("Icons/State_Subtree", CoreStyleConstants::Icon16x16));
+		Set("StateTreeEditor.StateLinked", new IMAGE_BRUSH_SVG("Icons/State_Linked", CoreStyleConstants::Icon16x16));
 	}
 
 
diff --git a/Engine/Plugins/Runtime/StateTree/Source/StateTreeEditorModule/Private/StateTreeState.cpp b/Engine/Plugins/Runtime/StateTree/Source/StateTreeEditorModule/Private/StateTreeState.cpp
--- a/Engine/Plugins/Runtime/StateTree/Source/StateTreeEditorModule/Private/StateTreeState.cpp
+++ b/Engine/Plugins/Runtime/StateTree/Source/StateTreeEditorModule/Private/StateTreeState.cpp
@@ -62,6 +62,7 @@ void UStateTreeState::PostEditChangeChainProperty(FPropertyChangedChainEvent& Pr
 	
 	static const FStateTreeEditPropertyPath StateNamePath(UStateTreeState::StaticClass(), TEXT("Name"));
 	static const FStateTreeEditPropertyPath StateTypePath(UStateTreeState::StaticClass(), TEXT("Type"));
+	static const FStateTreeEditPropertyPath SelectionBehaviorPath(UStateTreeState::StaticClass(), TEXT("SelectionBehavior"));
 	static const FStateTreeEditPropertyPath StateLinkedSubtreePath(UStateTreeState::StaticClass(), TEXT("LinkedSubtree"));
 	static const FStateTreeEditPropertyPath StateParametersPath(UStateTreeState::StaticClass(), TEXT("Parameters"));
 	static const FStateTreeEditPropertyPath StateTasksPath(UStateTreeState::StaticClass(), TEXT("Tasks"));
@@ -80,6 +81,16 @@ void UStateTreeState::PostEditChangeChainProperty(FPropertyChangedChainEvent& Pr
 		}
 	}
 
+	// Broadcast selection type changes so that the UI can update.
+	if (SelectionBehaviorPath.IsPathExact(StateTypePath))
+	{
+		const UStateTree* StateTree = GetTypedOuter<UStateTree>();
+		if (ensure(StateTree))
+		{
+			UE::StateTree::Delegates::OnIdentifierChanged.Broadcast(*StateTree);
+		}
+	}
+	
 	if (ChangePropertyPath.IsPathExact(StateTypePath))
 	{
 		// Remove any tasks and evaluators when they are not used.
@@ -99,6 +110,7 @@ void UStateTreeState::PostEditChangeChainProperty(FPropertyChangedChainEvent& Pr
 			// Linked parameter layout is fixed, and copied from the linked target state.
 			Parameters.bFixedLayout = true;
 			UpdateParametersFromLinkedSubtree();
+			SelectionBehavior = EStateTreeStateSelectionBehavior::TrySelectChildrenInOrder;
 		}
 		else if (Type == EStateTreeStateType::Subtree)
 		{
@@ -313,6 +325,30 @@ const UStateTreeState* UStateTreeState::GetNextSiblingState() const
 	return nullptr;
 }
 
+const UStateTreeState* UStateTreeState::GetNextSelectableSiblingState() const
+{
+	if (!Parent)
+	{
+		return nullptr;
+	}
+
+	const int32 StartChildIndex = Parent->Children.IndexOfByKey(this);
+	if (StartChildIndex == INDEX_NONE)
+	{
+		return nullptr;
+	}
+	
+	for (int32 ChildIdx = StartChildIndex + 1; ChildIdx < Parent->Children.Num(); ChildIdx++)
+	{
+		if (Parent->Children[ChildIdx]->SelectionBehavior != EStateTreeStateSelectionBehavior::None)
+		{
+			return Parent->Children[ChildIdx];
+		}
+	}
+	
+	return nullptr;
+}
+
 FStateTreeStateLink UStateTreeState::GetLinkToState() const
 {
 	FStateTreeStateLink Link(EStateTreeTransitionType::GotoState);
diff --git a/Engine/Plugins/Runtime/StateTree/Source/StateTreeEditorModule/Public/StateTreeState.h b/Engine/Plugins/Runtime/StateTree/Source/StateTreeEditorModule/Public/StateTreeState.h
--- a/Engine/Plugins/Runtime/StateTree/Source/StateTreeEditorModule/Public/StateTreeState.h
+++ b/Engine/Plugins/Runtime/StateTree/Source/StateTreeEditorModule/Public/StateTreeState.h
@@ -110,7 +110,8 @@ public:
 
 	const UStateTreeState* GetRootState() const;
 	const UStateTreeState* GetNextSiblingState() const;
-
+	const UStateTreeState* GetNextSelectableSiblingState() const;
+	
 	// StateTree Builder API
 	/** @return state link to this state. */
 	FStateTreeStateLink GetLinkToState() const;
@@ -186,6 +187,9 @@ public:
 	UPROPERTY(EditDefaultsOnly, Category = "State")
 	EStateTreeStateType Type = EStateTreeStateType::State;
 
+	UPROPERTY(EditDefaultsOnly, Category = "State")
+	EStateTreeStateSelectionBehavior SelectionBehavior = EStateTreeStateSelectionBehavior::TrySelectChildrenInOrder;
+	
 	UPROPERTY(EditDefaultsOnly, Category = "State", Meta=(DirectStatesOnly, SubtreesOnly))
 	FStateTreeStateLink LinkedSubtree;
 
diff --git a/Engine/Plugins/Runtime/StateTree/Source/StateTreeModule/Private/StateTreeExecutionContext.cpp b/Engine/Plugins/Runtime/StateTree/Source/StateTreeModule/Private/StateTreeExecutionContext.cpp
--- a/Engine/Plugins/Runtime/StateTree/Source/StateTreeModule/Private/StateTreeExecutionContext.cpp
+++ b/Engine/Plugins/Runtime/StateTree/Source/StateTreeModule/Private/StateTreeExecutionContext.cpp
@@ -35,8 +35,8 @@ FStateTreeExecutionContext::FStateTreeExecutionContext(UObject& InOwner, const U
 	}
 	else
 	{
-		STATETREE_LOG(Warning, TEXT("%s: StateTree asset is not valid ('%s' using StateTree '%s')"),
-			ANSI_TO_TCHAR(__FUNCTION__), *GetNameSafe(&Owner), *GetFullNameSafe(&StateTree));
+		STATETREE_LOG(Warning, TEXT("%hs: StateTree asset is not valid ('%s' using StateTree '%s')"),
+			__FUNCTION__, *GetNameSafe(&Owner), *GetFullNameSafe(&StateTree));
 	}
 }
 
@@ -162,8 +162,8 @@ EStateTreeRunStatus FStateTreeExecutionContext::Start()
 
 	if (!IsValid())
 	{
-		STATETREE_LOG(Warning, TEXT("%s: StateTree context is not initialized properly ('%s' using StateTree '%s')"),
-			ANSI_TO_TCHAR(__FUNCTION__), *GetNameSafe(&Owner), *GetFullNameSafe(&StateTree));
+		STATETREE_LOG(Warning, TEXT("%hs: StateTree context is not initialized properly ('%s' using StateTree '%s')"),
+			__FUNCTION__, *GetNameSafe(&Owner), *GetFullNameSafe(&StateTree));
 		return EStateTreeRunStatus::Failed;
 	}
 
@@ -176,8 +176,8 @@ EStateTreeRunStatus FStateTreeExecutionContext::Start()
 	UpdateInstanceData(Empty, Empty);
 	if (!InstanceData.IsValid())
 	{
-		STATETREE_LOG(Warning, TEXT("%s: Failed to initialize instance data on '%s' using StateTree '%s'. Try to recompile the StateTree asset."),
-			ANSI_TO_TCHAR(__FUNCTION__), *GetNameSafe(&Owner), *GetFullNameSafe(&StateTree));
+		STATETREE_LOG(Warning, TEXT("%hs: Failed to initialize instance data on '%s' using StateTree '%s'. Try to recompile the StateTree asset."),
+			__FUNCTION__, *GetNameSafe(&Owner), *GetFullNameSafe(&StateTree));
 		return EStateTreeRunStatus::Failed;
 	}
 
@@ -188,8 +188,8 @@ EStateTreeRunStatus FStateTreeExecutionContext::Start()
 	if (!StartEvaluatorsAndGlobalTasks(LastInitializedTaskIndex))
 	{
 		StopEvaluatorsAndGlobalTasks(LastInitializedTaskIndex);
-		STATETREE_LOG(Warning, TEXT("%s: Failed to start evaluators and global tasks '%s' using StateTree '%s'. Try to recompile the StateTree asset."),
-			ANSI_TO_TCHAR(__FUNCTION__), *GetNameSafe(&Owner), *GetFullNameSafe(&StateTree));
+		STATETREE_LOG(Warning, TEXT("%hs: Failed to start evaluators and global tasks '%s' using StateTree '%s'. Try to recompile the StateTree asset."),
+			__FUNCTION__, *GetNameSafe(&Owner), *GetFullNameSafe(&StateTree));
 		return EStateTreeRunStatus::Failed;
 	}
 
@@ -206,13 +206,14 @@ EStateTreeRunStatus FStateTreeExecutionContext::Start()
 	static const FStateTreeStateHandle RootState = FStateTreeStateHandle(0);
 
 	FStateTreeActiveStates NextActiveStates;
-	if (SelectState(RootState, NextActiveStates))
+	FStateTreeActiveStates VisitedStates;
+	if (SelectState(RootState, NextActiveStates, VisitedStates))
 	{
 		if (NextActiveStates.Last().IsCompletionState())
 		{
 			// Transition to a terminal state (succeeded/failed), or default transition failed.
-			STATETREE_LOG(Warning, TEXT("%s: Tree %s at StateTree start on '%s' using StateTree '%s'."),
-				ANSI_TO_TCHAR(__FUNCTION__), NextActiveStates.Last() == FStateTreeStateHandle::Succeeded ? TEXT("succeeded") : TEXT("failed"), *GetNameSafe(&Owner), *GetFullNameSafe(&StateTree));
+			STATETREE_LOG(Warning, TEXT("%hs: Tree %s at StateTree start on '%s' using StateTree '%s'."),
+				__FUNCTION__, NextActiveStates.Last() == FStateTreeStateHandle::Succeeded ? TEXT("succeeded") : TEXT("failed"), *GetNameSafe(&Owner), *GetFullNameSafe(&StateTree));
 			Exec->TreeRunStatus = NextActiveStates.Last().ToCompletionStatus();
 		}
 		else
@@ -240,8 +241,8 @@ EStateTreeRunStatus FStateTreeExecutionContext::Start()
 	if (Exec->ActiveStates.IsEmpty())
 	{
 		// Should not happen. This may happen if initial state could not be selected.
-		STATETREE_LOG(Error, TEXT("%s: Failed to select initial state on '%s' using StateTree '%s'. This should not happen, check that the StateTree logic can always select a state at start."),
-			ANSI_TO_TCHAR(__FUNCTION__), *GetNameSafe(&Owner), *GetFullNameSafe(&StateTree));
+		STATETREE_LOG(Error, TEXT("%hs: Failed to select initial state on '%s' using StateTree '%s'. This should not happen, check that the StateTree logic can always select a state at start."),
+			__FUNCTION__, *GetNameSafe(&Owner), *GetFullNameSafe(&StateTree));
 		Exec->TreeRunStatus = EStateTreeRunStatus::Failed;
 	}
 
@@ -254,8 +255,8 @@ EStateTreeRunStatus FStateTreeExecutionContext::Stop()
 
 	if (!IsValid())
 	{
-		STATETREE_LOG(Warning, TEXT("%s: StateTree context is not initialized properly ('%s' using StateTree '%s')"),
-			ANSI_TO_TCHAR(__FUNCTION__), *GetNameSafe(&Owner), *GetFullNameSafe(&StateTree));
+		STATETREE_LOG(Warning, TEXT("%hs: StateTree context is not initialized properly ('%s' using StateTree '%s')"),
+			__FUNCTION__, *GetNameSafe(&Owner), *GetFullNameSafe(&StateTree));
 		return EStateTreeRunStatus::Failed;
 	}
 	
@@ -306,15 +307,15 @@ EStateTreeRunStatus FStateTreeExecutionContext::Tick(const float DeltaTime)
 
 	if (!IsValid())
 	{
-		STATETREE_LOG(Warning, TEXT("%s: StateTree context is not initialized properly ('%s' using StateTree '%s')"),
-			ANSI_TO_TCHAR(__FUNCTION__), *GetNameSafe(&Owner), *GetFullNameSafe(&StateTree));
+		STATETREE_LOG(Warning, TEXT("%hs: StateTree context is not initialized properly ('%s' using StateTree '%s')"),
+			__FUNCTION__, *GetNameSafe(&Owner), *GetFullNameSafe(&StateTree));
 		return EStateTreeRunStatus::Failed;
 	}
 
 	if (!InstanceData.IsValid())
 	{
-		STATETREE_LOG(Error, TEXT("%s: Tick called on %s using StateTree %s with invalid instance data. Start() must be called before Tick()."),
-			ANSI_TO_TCHAR(__FUNCTION__), *GetNameSafe(&Owner), *GetFullNameSafe(&StateTree));
+		STATETREE_LOG(Error, TEXT("%hs: Tick called on %s using StateTree %s with invalid instance data. Start() must be called before Tick()."),
+			__FUNCTION__, *GetNameSafe(&Owner), *GetFullNameSafe(&StateTree));
 		return EStateTreeRunStatus::Failed;
 	}
 
@@ -416,8 +417,8 @@ EStateTreeRunStatus FStateTreeExecutionContext::Tick(const float DeltaTime)
 	if (Exec->ActiveStates.IsEmpty())
 	{
 		// Should not happen. This may happen if a state completion transition could not be selected. 
-		STATETREE_LOG(Error, TEXT("%s: Failed to select state on '%s' using StateTree '%s'. This should not happen, state completion transition is likely missing."),
-			ANSI_TO_TCHAR(__FUNCTION__), *GetNameSafe(&Owner), *GetFullNameSafe(&StateTree));
+		STATETREE_LOG(Error, TEXT("%hs: Failed to select state on '%s' using StateTree '%s'. This should not happen, state completion transition is likely missing."),
+			__FUNCTION__, *GetNameSafe(&Owner), *GetFullNameSafe(&StateTree));
 		Exec->TreeRunStatus = EStateTreeRunStatus::Failed;
 		return Exec->TreeRunStatus;
 	}
@@ -431,8 +432,8 @@ EStateTreeRunStatus FStateTreeExecutionContext::GetStateTreeRunStatus() const
 {
 	if (!IsValid())
 	{
-		STATETREE_LOG(Warning, TEXT("%s: StateTree context is not initialized properly ('%s' using StateTree '%s')"),
-			ANSI_TO_TCHAR(__FUNCTION__), *GetNameSafe(&Owner), *GetFullNameSafe(&StateTree));
+		STATETREE_LOG(Warning, TEXT("%hs: StateTree context is not initialized properly ('%s' using StateTree '%s')"),
+			__FUNCTION__, *GetNameSafe(&Owner), *GetFullNameSafe(&StateTree));
 		return EStateTreeRunStatus::Failed;
 	}
 
@@ -455,15 +456,15 @@ void FStateTreeExecutionContext::SendEvent(const FGameplayTag Tag, const FConstS
 
 	if (!IsValid())
 	{
-		STATETREE_LOG(Warning, TEXT("%s: StateTree context is not initialized properly ('%s' using StateTree '%s')"),
-			ANSI_TO_TCHAR(__FUNCTION__), *GetNameSafe(&Owner), *GetFullNameSafe(&StateTree));
+		STATETREE_LOG(Warning, TEXT("%hs: StateTree context is not initialized properly ('%s' using StateTree '%s')"),
+			__FUNCTION__, *GetNameSafe(&Owner), *GetFullNameSafe(&StateTree));
 		return;
 	}
 
 	if (!InstanceData.IsValid())
 	{
-		STATETREE_LOG(Error, TEXT("%s: SendEvent called on %s using StateTree %s with invalid instance data. Start() must be called before sending events."),
-			ANSI_TO_TCHAR(__FUNCTION__), *GetNameSafe(&Owner), *GetFullNameSafe(&StateTree));
+		STATETREE_LOG(Error, TEXT("%hs: SendEvent called on %s using StateTree %s with invalid instance data. Start() must be called before sending events."),
+			__FUNCTION__, *GetNameSafe(&Owner), *GetFullNameSafe(&StateTree));
 		return;
 	}
 
@@ -480,15 +481,15 @@ void FStateTreeExecutionContext::RequestTransition(const FStateTreeTransitionReq
 
 	if (!IsValid())
 	{
-		STATETREE_LOG(Warning, TEXT("%s: StateTree context is not initialized properly ('%s' using StateTree '%s')"),
-			ANSI_TO_TCHAR(__FUNCTION__), *GetNameSafe(&Owner), *GetFullNameSafe(&StateTree));
+		STATETREE_LOG(Warning, TEXT("%hs: StateTree context is not initialized properly ('%s' using StateTree '%s')"),
+			__FUNCTION__, *GetNameSafe(&Owner), *GetFullNameSafe(&StateTree));
 		return;
 	}
 
 	if (!InstanceData.IsValid())
 	{
-		STATETREE_LOG(Error, TEXT("%s: RequestTransition called on %s using StateTree %s with invalid instance data. Start() must be called before requesting transition."),
-			ANSI_TO_TCHAR(__FUNCTION__), *GetNameSafe(&Owner), *GetFullNameSafe(&StateTree));
+		STATETREE_LOG(Error, TEXT("%hs: RequestTransition called on %s using StateTree %s with invalid instance data. Start() must be called before requesting transition."),
+			__FUNCTION__, *GetNameSafe(&Owner), *GetFullNameSafe(&StateTree));
 		return;
 	}
 
@@ -693,8 +694,8 @@ EStateTreeRunStatus FStateTreeExecutionContext::EnterState(const FStateTreeTrans
 
 		if (!Exec.ActiveStates.Push(CurrentHandle))
 		{
-			STATETREE_LOG(Error, TEXT("%s: Reached max execution depth when trying to enter state '%s'.  '%s' using StateTree '%s'."),
-				ANSI_TO_TCHAR(__FUNCTION__), *GetStateStatusString(Exec), *GetNameSafe(&Owner), *GetFullNameSafe(&StateTree));
+			STATETREE_LOG(Error, TEXT("%hs: Reached max execution depth when trying to enter state '%s'.  '%s' using StateTree '%s'."),
+				__FUNCTION__, *GetStateStatusString(Exec), *GetNameSafe(&Owner), *GetFullNameSafe(&StateTree));
 			break;
 		}
 
@@ -1369,7 +1370,8 @@ bool FStateTreeExecutionContext::RequestTransition(const FStateTreeStateHandle N
 	}
 
 	FStateTreeActiveStates NewActiveState;
-	if (SelectState(NextState, NewActiveState))
+	FStateTreeActiveStates VisitedStates;
+	if (SelectState(NextState, NewActiveState, VisitedStates))
 	{
 		NextTransition.CurrentActiveStates = Exec.ActiveStates;
 		NextTransition.CurrentRunStatus = Exec.TreeRunStatus;
@@ -1644,7 +1646,7 @@ FStateTreeStateHandle FStateTreeExecutionContext::GetParentLinkedStateHandle(con
 	return FStateTreeStateHandle();
 }
 
-bool FStateTreeExecutionContext::SelectState(const FStateTreeStateHandle NextState, FStateTreeActiveStates& OutNewActiveState)
+bool FStateTreeExecutionContext::SelectState(const FStateTreeStateHandle NextState, FStateTreeActiveStates& OutNewActiveState, FStateTreeActiveStates& VisitedStates)
 {
 	const FStateTreeExecutionState& Exec = GetExecState();
 
@@ -1675,8 +1677,8 @@ bool FStateTreeExecutionContext::SelectState(const FStateTreeStateHandle NextSta
 		}
 		if (NumInBetweenStates == InBetweenStates.Num())
 		{
-			STATETREE_LOG(Error, TEXT("%s: Too many parent states when selecting state '%s' from '%s'.  '%s' using StateTree '%s'."),
-				ANSI_TO_TCHAR(__FUNCTION__), *GetSafeStateName(NextState), *GetStateStatusString(Exec), *GetNameSafe(&Owner), *GetFullNameSafe(&StateTree));
+			STATETREE_LOG(Error, TEXT("%hs: Too many parent states when selecting state '%s' from '%s'.  '%s' using StateTree '%s'."),
+				__FUNCTION__, *GetSafeStateName(NextState), *GetStateStatusString(Exec), *GetNameSafe(&Owner), *GetFullNameSafe(&StateTree));
 			return false;
 		}
 
@@ -1695,15 +1697,15 @@ bool FStateTreeExecutionContext::SelectState(const FStateTreeStateHandle NextSta
 
 	if (bActiveStatesOverflow)
 	{
-		STATETREE_LOG(Error, TEXT("%s: Reached max execution depth when trying to select state %s from '%s'.  '%s' using StateTree '%s'."),
-			ANSI_TO_TCHAR(__FUNCTION__), *GetSafeStateName(NextState), *GetStateStatusString(Exec), *GetNameSafe(&Owner), *GetFullNameSafe(&StateTree));
+		STATETREE_LOG(Error, TEXT("%hs: Reached max execution depth when trying to select state %s from '%s'.  '%s' using StateTree '%s'."),
+			__FUNCTION__, *GetSafeStateName(NextState), *GetStateStatusString(Exec), *GetNameSafe(&Owner), *GetFullNameSafe(&StateTree));
 		return false;
 	}
-
-	return SelectStateInternal(NextState, OutNewActiveState);
+	
+	return SelectStateInternal(NextState, OutNewActiveState, VisitedStates);
 }
 
-bool FStateTreeExecutionContext::SelectStateInternal(const FStateTreeStateHandle NextState, FStateTreeActiveStates& OutNewActiveState)
+bool FStateTreeExecutionContext::SelectStateInternal(const FStateTreeStateHandle NextState, FStateTreeActiveStates& OutNewActiveState, FStateTreeActiveStates& VisitedStates)
 {
 	CSV_SCOPED_TIMING_STAT_EXCLUSIVE(StateTree_SelectState);
 
@@ -1712,51 +1714,126 @@ bool FStateTreeExecutionContext::SelectStateInternal(const FStateTreeStateHandle
 	if (!NextState.IsValid())
 	{
 		// Trying to select non-existing state.
-		STATETREE_LOG(Error, TEXT("%s: Trying to select invalid state from '%s'.  '%s' using StateTree '%s'."),
-            ANSI_TO_TCHAR(__FUNCTION__), *GetStateStatusString(Exec), *GetNameSafe(&Owner), *GetFullNameSafe(&StateTree));
+		STATETREE_LOG(Error, TEXT("%hs: Trying to select invalid state from '%s'.  '%s' using StateTree '%s'."),
+            __FUNCTION__, *GetStateStatusString(Exec), *GetNameSafe(&Owner), *GetFullNameSafe(&StateTree));
 		return false;
 	}
 
 	const FCompactStateTreeState& State = StateTree.States[NextState.Index];
 
 	// Check that the state can be entered
-	if (TestAllConditions(State.EnterConditionsBegin, State.EnterConditionsNum))
+	if (State.SelectionBehavior != EStateTreeStateSelectionBehavior::None
+		&& TestAllConditions(State.EnterConditionsBegin, State.EnterConditionsNum))
 	{
 		if (!OutNewActiveState.Push(NextState))
 		{
-			STATETREE_LOG(Error, TEXT("%s: Reached max execution depth when trying to select state %s from '%s'.  '%s' using StateTree '%s'."),
-				ANSI_TO_TCHAR(__FUNCTION__), *GetSafeStateName(NextState), *GetStateStatusString(Exec), *GetNameSafe(&Owner), *GetFullNameSafe(&StateTree));
+			STATETREE_LOG(Error, TEXT("%hs: Reached max execution depth when trying to select state %s from '%s'.  '%s' using StateTree '%s'."),
+				__FUNCTION__, *GetSafeStateName(NextState), *GetStateStatusString(Exec), *GetNameSafe(&Owner), *GetFullNameSafe(&StateTree));
+			return false;
+		}
+		if (!VisitedStates.Push(NextState))
+		{
+			STATETREE_LOG(Error, TEXT("%hs: Reached max visited state depth when trying to select state %s from '%s'.  '%s' using StateTree '%s'."),
+				__FUNCTION__, *GetSafeStateName(NextState), *GetStateStatusString(Exec), *GetNameSafe(&Owner), *GetFullNameSafe(&StateTree));
 			return false;
 		}
 		
 		if (State.LinkedState.IsValid())
 		{
 			// If State is linked, proceed to the linked state.
-			if (SelectStateInternal(State.LinkedState, OutNewActiveState))
+			if (SelectStateInternal(State.LinkedState, OutNewActiveState, VisitedStates))
 			{
 				// Selection succeeded
 				return true;
 			}
 		}
-		else if (State.HasChildren())
+		else if (State.SelectionBehavior == EStateTreeStateSelectionBehavior::TryEnterState)
 		{
-			// If the state has children, proceed to select children.
-			for (uint16 ChildState = State.ChildrenBegin; ChildState < State.ChildrenEnd; ChildState = StateTree.States[ChildState].GetNextSibling())
+			// Select this state.
+			return true;
+		}
+		else if (State.SelectionBehavior == EStateTreeStateSelectionBehavior::TryFollowTransitions)
+		{
+			EStateTreeTransitionPriority CurrentPriority = EStateTreeTransitionPriority::None;
+			
+			for (uint8 i = 0; i < State.TransitionsNum; i++)
 			{
-				if (SelectStateInternal(FStateTreeStateHandle(ChildState), OutNewActiveState))
+				const int16 TransitionIndex = State.TransitionsBegin + i;
+				const FCompactStateTransition& Transition = StateTree.Transitions[TransitionIndex];
+
+				// No need to test the transition if same or higher priority transition has already been processed.
+				if (Transition.Priority <= CurrentPriority)
 				{
-					// Selection succeeded
-					return true;
+					continue;
 				}
+
+				// Skip completion transitions
+				if (EnumHasAnyFlags(Transition.Trigger, EStateTreeTransitionTrigger::OnStateCompleted))
+				{
+					continue;
+				}
+
+				// Cannot follow transitions with delay.
+				if (Transition.HasDelay())
+				{
+					continue;
+				}
+
+				// Try to prevent (infinite) loops in the selection.
+				if (VisitedStates.Contains(Transition.State))
+				{
+					STATETREE_LOG(Error, TEXT("%hs: Loop detected when trying to select state %s from '%s'. Prior states: %s.  '%s' using StateTree '%s'."),
+						__FUNCTION__, *GetSafeStateName(NextState), *GetStateStatusString(Exec), *DebugGetStatePath(VisitedStates), *GetNameSafe(&Owner), *GetFullNameSafe(&StateTree));
+					continue;
+				}
+
+				const bool bShouldTrigger = Transition.Trigger == EStateTreeTransitionTrigger::OnTick
+											|| (Transition.Trigger == EStateTreeTransitionTrigger::OnEvent
+												&& HasEventToProcess(Transition.EventTag));
+				
+				if (bShouldTrigger
+					&& TestAllConditions(Transition.ConditionsBegin, Transition.ConditionsNum))
+				{
+					// Using SelectState() instead of SelectStateInternal to treat the transitions the same way as regular transitions,
+					// e.g. it may jump to a completely different branch.
+					FStateTreeActiveStates NewActiveState;
+					if (SelectState(Transition.State, NewActiveState, VisitedStates))
+					{
+						// Selection succeeded
+						OutNewActiveState = NewActiveState;
+						CurrentPriority = Transition.Priority;
+					}
+				}
+			}
+
+			if (CurrentPriority != EStateTreeTransitionPriority::None)
+			{
+				return true;
 			}
 		}
-		else
+		else if (State.SelectionBehavior == EStateTreeStateSelectionBehavior::TrySelectChildrenInOrder)
 		{
-			// Select this state.
-			return true;
+			if (State.HasChildren())
+			{
+				// If the state has children, proceed to select children.
+				for (uint16 ChildState = State.ChildrenBegin; ChildState < State.ChildrenEnd; ChildState = StateTree.States[ChildState].GetNextSibling())
+				{
+					if (SelectStateInternal(FStateTreeStateHandle(ChildState), OutNewActiveState, VisitedStates))
+					{
+						// Selection succeeded
+						return true;
+					}
+				}
+			}
+			else
+			{
+				// Select this state (For backwards compatibility)
+				return true;
+			}
 		}
 		
 		OutNewActiveState.Pop();
+		VisitedStates.Pop();
 	}
 
 	// Nothing got selected.
diff --git a/Engine/Plugins/Runtime/StateTree/Source/StateTreeModule/Public/StateTreeExecutionContext.h b/Engine/Plugins/Runtime/StateTree/Source/StateTreeModule/Public/StateTreeExecutionContext.h
--- a/Engine/Plugins/Runtime/StateTree/Source/StateTreeModule/Public/StateTreeExecutionContext.h
+++ b/Engine/Plugins/Runtime/StateTree/Source/StateTreeModule/Public/StateTreeExecutionContext.h
@@ -359,14 +359,15 @@ protected:
 	 * If NextState is a leaf state, the active states leading from root to the leaf are returned.
 	 * @param NextState The state which we try to select next.
 	 * @param OutNewActiveStates Active states that got selected.
+	 * @param VisitedStates States visited so far during selection (used for detecting selection loops)
 	 * @return True if succeeded to select new active states.
 	 */
-	bool SelectState(const FStateTreeStateHandle NextState, FStateTreeActiveStates& OutNewActiveStates);
+	bool SelectState(const FStateTreeStateHandle NextState, FStateTreeActiveStates& OutNewActiveStates, FStateTreeActiveStates& VisitedStates);
 
 	/**
 	 * Used internally to do the recursive part of the SelectState().
 	 */
-	bool SelectStateInternal(const FStateTreeStateHandle NextState, FStateTreeActiveStates& OutNewActiveStates);
+	bool SelectStateInternal(const FStateTreeStateHandle NextState, FStateTreeActiveStates& OutNewActiveStates, FStateTreeActiveStates& VisitedStates);
 
 	/** @return StateTree execution state from the instance storage. */
 	FStateTreeExecutionState& GetExecState()
diff --git a/Engine/Plugins/Runtime/StateTree/Source/StateTreeModule/Public/StateTreeTypes.h b/Engine/Plugins/Runtime/StateTree/Source/StateTreeModule/Public/StateTreeTypes.h
--- a/Engine/Plugins/Runtime/StateTree/Source/StateTreeModule/Public/StateTreeTypes.h
+++ b/Engine/Plugins/Runtime/StateTree/Source/StateTreeModule/Public/StateTreeTypes.h
@@ -90,10 +90,10 @@ enum class EStateTreeConditionOperand : uint8
 UENUM()
 enum class EStateTreeStateType : uint8
 {
-	/** A State containing tasks and evaluators. */
+	/** A State containing tasks and child states. */
 	State,
 	
-	/** A State containing just sub states. */
+	/** A State containing just child states. */
 	Group,
 	
 	/** A State that is linked to another state in the tree (the execution continues on the linked state). */
@@ -103,6 +103,22 @@ enum class EStateTreeStateType : uint8
 	Subtree,
 };
 
+UENUM()
+enum class EStateTreeStateSelectionBehavior : uint8
+{
+	/** The State cannot be directly selected. */
+	None,
+
+	/** When state is considered for selection, it is selected even if it has child states. */
+	TryEnterState,
+
+	/** When state is considered for selection, try to selects the first child state (in order they appear in the child list). If no child states are present, behaves like SelectState. */
+	TrySelectChildrenInOrder,
+	
+	/** When state is considered for selection, try to trigger the transitions instead. */
+	TryFollowTransitions,
+};
+
 
 /** Transitions trigger. */
 UENUM()
@@ -724,10 +740,14 @@ struct STATETREEMODULE_API FCompactStateTreeState
 	UPROPERTY()
 	uint8 TasksNum = 0;
 
-	/** Type of the sate */
+	/** Type of the state */
 	UPROPERTY()
 	EStateTreeStateType Type = EStateTreeStateType::State;
 
+	/** What to do when the state is considered for selection. */
+	UPROPERTY()
+	EStateTreeStateSelectionBehavior SelectionBehavior = EStateTreeStateSelectionBehavior::None;
+
 	/** True if the state contains tasks that should be called during transition handling. */
 	UPROPERTY()
 	uint8 bHasTransitionTasks : 1;
diff --git a/Engine/Plugins/Runtime/StateTree/Source/StateTreeTestSuite/Private/StateTreeTest.cpp b/Engine/Plugins/Runtime/StateTree/Source/StateTreeTestSuite/Private/StateTreeTest.cpp
--- a/Engine/Plugins/Runtime/StateTree/Source/StateTreeTestSuite/Private/StateTreeTest.cpp
+++ b/Engine/Plugins/Runtime/StateTree/Source/StateTreeTestSuite/Private/StateTreeTest.cpp
@@ -1322,6 +1322,167 @@ struct FStateTreeTest_BindingsCompiler : FAITestBase
 };
 IMPLEMENT_AI_INSTANT_TEST(FStateTreeTest_BindingsCompiler, "System.StateTree.BindingsCompiler");
 
+struct FStateTreeTest_FollowTransitions : FAITestBase
+{
+	virtual bool InstantTest() override
+	{
+		UStateTree& StateTree = UE::StateTree::Tests::NewStateTree(&GetWorld());
+		UStateTreeEditorData& EditorData = *Cast<UStateTreeEditorData>(StateTree.EditorData);
+
+		EditorData.RootParameters.Parameters.AddProperty(FName(TEXT("Int")), EPropertyBagPropertyType::Int32);
+		EditorData.RootParameters.Parameters.SetValueInt32(FName(TEXT("Int")), 1);
+		
+		UStateTreeState& Root = EditorData.AddSubTree(FName(TEXT("Root")));
+		UStateTreeState& StateTrans = Root.AddChildState(FName(TEXT("Trans")));
+		UStateTreeState& StateA = Root.AddChildState(FName(TEXT("A")));
+		UStateTreeState& StateB = Root.AddChildState(FName(TEXT("B")));
+		UStateTreeState& StateC = Root.AddChildState(FName(TEXT("C")));
+
+		// Root
+
+		// Trans
+		{
+			StateTrans.SelectionBehavior = EStateTreeStateSelectionBehavior::TryFollowTransitions;
+
+			{
+				// This transition should be skipped due to the condition
+				FStateTreeTransition& TransA = StateTrans.AddTransition(EStateTreeTransitionTrigger::OnTick, EStateTreeTransitionType::GotoState, &StateA);
+				TStateTreeEditorNode<FStateTreeCompareIntCondition>& TransIntCond = TransA.AddCondition<FStateTreeCompareIntCondition>(EGenericAICheck::Equal);
+				TransIntCond.GetInstanceData().Right = 0;
+				EditorData.AddPropertyBinding(
+					FStateTreePropertyPath(EditorData.RootParameters.ID, TEXT("Int")),
+					FStateTreePropertyPath(TransIntCond.ID, TEXT("Left")));
+			}
+
+			{
+				// This transition leads to selection, but will be overridden.
+				FStateTreeTransition& TransB = StateTrans.AddTransition(EStateTreeTransitionTrigger::OnTick, EStateTreeTransitionType::GotoState, &StateB);
+				TransB.Priority = EStateTreeTransitionPriority::Normal;
+				TStateTreeEditorNode<FStateTreeCompareIntCondition>& TransIntCond = TransB.AddCondition<FStateTreeCompareIntCondition>(EGenericAICheck::Equal);
+				TransIntCond.GetInstanceData().Right = 1;
+				EditorData.AddPropertyBinding(
+					FStateTreePropertyPath(EditorData.RootParameters.ID, TEXT("Int")),
+					FStateTreePropertyPath(TransIntCond.ID, TEXT("Left")));
+			}
+
+			{
+				// This transition is selected, should override previous one due to priority.
+				FStateTreeTransition& TransC = StateTrans.AddTransition(EStateTreeTransitionTrigger::OnTick, EStateTreeTransitionType::GotoState, &StateC);
+				TransC.Priority = EStateTreeTransitionPriority::High;
+				TStateTreeEditorNode<FStateTreeCompareIntCondition>& TransIntCond = TransC.AddCondition<FStateTreeCompareIntCondition>(EGenericAICheck::Equal);
+				TransIntCond.GetInstanceData().Right = 1;
+				EditorData.AddPropertyBinding(
+					FStateTreePropertyPath(EditorData.RootParameters.ID, TEXT("Int")),
+					FStateTreePropertyPath(TransIntCond.ID, TEXT("Left")));
+			}
+		}
+
+		auto& TaskA = StateA.AddTask<FTestTask_Stand>(FName(TEXT("TaskA")));
+		auto& TaskB = StateB.AddTask<FTestTask_Stand>(FName(TEXT("TaskB")));
+		auto& TaskC = StateC.AddTask<FTestTask_Stand>(FName(TEXT("TaskC")));
+
+		FStateTreeCompilerLog Log;
+		FStateTreeCompiler Compiler(Log);
+		const bool bResult = Compiler.Compile(StateTree);
+
+		AITEST_TRUE("StateTree should get compiled", bResult);
+
+		EStateTreeRunStatus Status = EStateTreeRunStatus::Unset;
+		FStateTreeInstanceData InstanceData;
+		FTestStateTreeExecutionContext Exec(StateTree, StateTree, InstanceData);
+		const bool bInitSucceeded = Exec.IsValid();
+		AITEST_TRUE("StateTree should init", bInitSucceeded);
+
+		const FString TickStr(TEXT("Tick"));
+		const FString EnterStateStr(TEXT("EnterState"));
+		const FString ExitStateStr(TEXT("ExitState"));
+
+		Status = Exec.Start();
+		AITEST_FALSE("StateTree TaskA should not enter state", Exec.Expect(TaskA.GetName(), EnterStateStr));
+		AITEST_FALSE("StateTree TaskB should not enter state", Exec.Expect(TaskB.GetName(), EnterStateStr));
+		AITEST_TRUE("StateTree TaskC should enter state", Exec.Expect(TaskC.GetName(), EnterStateStr));
+		Exec.LogClear();
+
+		return true;
+	}
+};
+IMPLEMENT_AI_INSTANT_TEST(FStateTreeTest_FollowTransitions, "System.StateTree.FollowTransitions");
+
+struct FStateTreeTest_InfiniteLoop : FAITestBase
+{
+	virtual bool InstantTest() override
+	{
+		UStateTree& StateTree = UE::StateTree::Tests::NewStateTree(&GetWorld());
+		UStateTreeEditorData& EditorData = *Cast<UStateTreeEditorData>(StateTree.EditorData);
+
+		EditorData.RootParameters.Parameters.AddProperty(FName(TEXT("Int")), EPropertyBagPropertyType::Int32);
+		EditorData.RootParameters.Parameters.SetValueInt32(FName(TEXT("Int")), 1);
+		
+		UStateTreeState& Root = EditorData.AddSubTree(FName(TEXT("Root")));
+		UStateTreeState& StateA = Root.AddChildState(FName(TEXT("A")));
+		UStateTreeState& StateB = StateA.AddChildState(FName(TEXT("B")));
+
+		// Root
+
+		// State A
+		{
+			StateA.SelectionBehavior = EStateTreeStateSelectionBehavior::TryFollowTransitions;
+			{
+				// A -> B
+				FStateTreeTransition& Trans = StateA.AddTransition(EStateTreeTransitionTrigger::OnTick, EStateTreeTransitionType::GotoState, &StateB);
+				TStateTreeEditorNode<FStateTreeCompareIntCondition>& TransIntCond = Trans.AddCondition<FStateTreeCompareIntCondition>(EGenericAICheck::Equal);
+				TransIntCond.GetInstanceData().Right = 1;
+				EditorData.AddPropertyBinding(
+					FStateTreePropertyPath(EditorData.RootParameters.ID, TEXT("Int")),
+					FStateTreePropertyPath(TransIntCond.ID, TEXT("Left")));
+			}
+		}
+
+		// State B
+		{
+			StateB.SelectionBehavior = EStateTreeStateSelectionBehavior::TryFollowTransitions;
+			{
+				// B -> A
+				FStateTreeTransition& Trans = StateB.AddTransition(EStateTreeTransitionTrigger::OnTick, EStateTreeTransitionType::GotoState, &StateA);
+				TStateTreeEditorNode<FStateTreeCompareIntCondition>& TransIntCond = Trans.AddCondition<FStateTreeCompareIntCondition>(EGenericAICheck::Equal);
+				TransIntCond.GetInstanceData().Right = 1;
+				EditorData.AddPropertyBinding(
+					FStateTreePropertyPath(EditorData.RootParameters.ID, TEXT("Int")),
+					FStateTreePropertyPath(TransIntCond.ID, TEXT("Left")));
+			}
+		}
+		
+		auto& TaskA = StateA.AddTask<FTestTask_Stand>(FName(TEXT("TaskA")));
+		auto& TaskB = StateB.AddTask<FTestTask_Stand>(FName(TEXT("TaskB")));
+
+		FStateTreeCompilerLog Log;
+		FStateTreeCompiler Compiler(Log);
+		const bool bResult = Compiler.Compile(StateTree);
+
+		AITEST_TRUE("StateTree should get compiled", bResult);
+
+		EStateTreeRunStatus Status = EStateTreeRunStatus::Unset;
+		FStateTreeInstanceData InstanceData;
+		FTestStateTreeExecutionContext Exec(StateTree, StateTree, InstanceData);
+		const bool bInitSucceeded = Exec.IsValid();
+		AITEST_TRUE("StateTree should init", bInitSucceeded);
+
+		const FString TickStr(TEXT("Tick"));
+		const FString EnterStateStr(TEXT("EnterState"));
+		const FString ExitStateStr(TEXT("ExitState"));
+
+		GetTestRunner().AddExpectedError(TEXT("Loop detected when trying to select state"), EAutomationExpectedErrorFlags::Contains, 1);
+		GetTestRunner().AddExpectedError(TEXT("Failed to select initial state"), EAutomationExpectedErrorFlags::Contains, 1);
+
+		Status = Exec.Start();
+		AITEST_EQUAL("Start should fail", Status, EStateTreeRunStatus::Failed);
+		Exec.LogClear();
+
+		return true;
+	}
+};
+IMPLEMENT_AI_INSTANT_TEST(FStateTreeTest_InfiniteLoop, "System.StateTree.InfiniteLoop");
+
 
 UE_ENABLE_OPTIMIZATION_SHIP
 
