StateTree: Fix default value of selection behavior to keep old assets working.

diff --git a/Engine/Plugins/Runtime/StateTree/Source/StateTreeModule/Public/StateTreeTypes.h b/Engine/Plugins/Runtime/StateTree/Source/StateTreeModule/Public/StateTreeTypes.h
--- a/Engine/Plugins/Runtime/StateTree/Source/StateTreeModule/Public/StateTreeTypes.h
+++ b/Engine/Plugins/Runtime/StateTree/Source/StateTreeModule/Public/StateTreeTypes.h
@@ -746,7 +746,7 @@ struct STATETREEMODULE_API FCompactStateTreeState
 
 	/** What to do when the state is considered for selection. */
 	UPROPERTY()
-	EStateTreeStateSelectionBehavior SelectionBehavior = EStateTreeStateSelectionBehavior::None;
+	EStateTreeStateSelectionBehavior SelectionBehavior = EStateTreeStateSelectionBehavior::TrySelectChildrenInOrder;
 
 	/** True if the state contains tasks that should be called during transition handling. */
 	UPROPERTY()
