StateTree: Link state tree when a BP recompiled
- Fixes a crash in property bindings which tried to access data based on the old BP class layout

diff --git a/Engine/Plugins/Runtime/StateTree/Source/StateTreeModule/Private/StateTree.cpp b/Engine/Plugins/Runtime/StateTree/Source/StateTreeModule/Private/StateTree.cpp
--- a/Engine/Plugins/Runtime/StateTree/Source/StateTreeModule/Private/StateTree.cpp
+++ b/Engine/Plugins/Runtime/StateTree/Source/StateTreeModule/Private/StateTree.cpp
@@ -109,6 +109,77 @@ void UStateTree::ResetCompiled()
 	ResetLinked();
 }
 
+void UStateTree::OnObjectsReinstanced(const FReplacementObjectMap& ObjectMap)
+{
+	if (ObjectMap.IsEmpty())
+	{
+		return;
+	}
+
+	// If the asset is not linked yet (or has failed), no need to link.
+	if (!bIsLinked)
+	{
+		return;
+	}
+
+	// Relink only if the reinstantiated object belongs to this asset,
+	// or anything from the property binding refers to the classes of the reinstantiated object.
+
+	bool bShouldRelink = false;
+
+	for (TMap<UObject*, UObject*>::TConstIterator It(ObjectMap); It; ++It)
+	{
+		if (UObject* ObjectToBeReplaced = It->Value)
+		{
+			if (ObjectToBeReplaced->IsInOuter(this))
+			{
+				bShouldRelink = true;
+				break;
+			}
+		}			
+	}
+
+	if (!bShouldRelink)
+	{
+		TSet<const UStruct*> Structs;
+		for (TMap<UObject*, UObject*>::TConstIterator It(ObjectMap); It; ++It)
+		{
+			if (UObject* ObjectToBeReplaced = It->Value)
+			{
+				Structs.Add(ObjectToBeReplaced->GetClass());
+			}
+		}
+
+		bShouldRelink |= PropertyBindings.ContainsAnyStruct(Structs);
+	}
+
+	if (bShouldRelink)
+	{
+		if (!Link())
+		{
+			UE_LOG(LogStateTree, Error, TEXT("%s failed to link after Object reinstantiation. Take a look at the asset for any errors. Asset will not be usable at runtime."), *GetFullName());
+		}
+	}
+}
+
+void UStateTree::PostInitProperties()
+{
+	Super::PostInitProperties();
+	
+	OnObjectsReinstancedHandle = FCoreUObjectDelegates::OnObjectsReinstanced.AddUObject(this, &UStateTree::OnObjectsReinstanced);
+}
+
+void UStateTree::BeginDestroy()
+{
+	if (OnObjectsReinstancedHandle.IsValid())
+	{
+		FCoreUObjectDelegates::OnObjectsReinstanced.Remove(OnObjectsReinstancedHandle);
+		OnObjectsReinstancedHandle.Reset();
+	}
+	
+	Super::BeginDestroy();
+}
+
 void UStateTree::GetAssetRegistryTags(TArray<FAssetRegistryTag>& OutTags) const
 {
 	const FString SchemaClassName = Schema ? Schema->GetClass()->GetPathName() : TEXT("");
@@ -267,9 +338,9 @@ bool UStateTree::Link()
 		}
 
 		
-		// Update property bag structs before resolving binding.
 		const TArrayView<FStateTreeBindableStructDesc> SourceStructs = PropertyBindings.SourceStructs;
 		const TArrayView<FStateTreePropertyCopyBatch> CopyBatches = PropertyBindings.CopyBatches;
+		const TArrayView<FStateTreePropertyPathBinding> PropertyPathBindings = PropertyBindings.PropertyPathBindings;
 
 		// Reconcile out of date classes.
 		for (FStateTreeBindableStructDesc& SourceStruct : SourceStructs)
@@ -293,7 +364,27 @@ bool UStateTree::Link()
 			}
 		}
 
+		auto PatchPropertyPath = [](FStateTreePropertyPath& PropertyPath)
+		{
+			for (FStateTreePropertyPathSegment& Segment : PropertyPath.GetMutableSegments())
+			{
+				if (const UClass* InstanceStruct = Cast<UClass>(Segment.GetInstanceStruct()))
+				{
+					if (InstanceStruct->HasAnyClassFlags(CLASS_NewerVersionExists))
+					{
+						Segment.SetInstanceStruct(InstanceStruct->GetAuthoritativeClass());
+					}
+				}
+			}
+		};
+		
+		for (FStateTreePropertyPathBinding& PropertyPathBinding : PropertyPathBindings)
+		{
+			PatchPropertyPath(PropertyPathBinding.GetMutableSourcePath());
+			PatchPropertyPath(PropertyPathBinding.GetMutableTargetPath());
+		}
 
+		// Update property bag structs before resolving binding.
 		if (ParametersDataViewIndex.IsValid() && SourceStructs.IsValidIndex(ParametersDataViewIndex.Get()))
 		{
 			SourceStructs[ParametersDataViewIndex.Get()].Struct = Parameters.GetPropertyBagStruct();
diff --git a/Engine/Plugins/Runtime/StateTree/Source/StateTreeModule/Private/StateTreePropertyBindings.cpp b/Engine/Plugins/Runtime/StateTree/Source/StateTreeModule/Private/StateTreePropertyBindings.cpp
--- a/Engine/Plugins/Runtime/StateTree/Source/StateTreeModule/Private/StateTreePropertyBindings.cpp
+++ b/Engine/Plugins/Runtime/StateTree/Source/StateTreeModule/Private/StateTreePropertyBindings.cpp
@@ -1035,6 +1035,50 @@ bool FStateTreePropertyBindings::ResetObjects(const FStateTreeIndex16 TargetBatc
 	return bResult;
 }
 
+bool FStateTreePropertyBindings::ContainsAnyStruct(const TSet<const UStruct*>& Structs)
+{
+	for (FStateTreeBindableStructDesc& SourceStruct : SourceStructs)
+	{
+		if (Structs.Contains(SourceStruct.Struct))
+		{
+			return true;
+		}
+	}
+
+	for (FStateTreePropertyCopyBatch& CopyBatch : CopyBatches)
+	{
+		if (Structs.Contains(CopyBatch.TargetStruct.Struct))
+		{
+			return true;
+		}
+	}
+
+	auto PathContainsStruct = [&Structs](const FStateTreePropertyPath& PropertyPath)
+	{
+		for (const FStateTreePropertyPathSegment& Segment : PropertyPath.GetSegments())
+		{
+			if (Structs.Contains(Segment.GetInstanceStruct()))
+			{
+				return true;
+			}
+		}
+		return false;
+	};
+	
+	for (FStateTreePropertyPathBinding& PropertyPathBinding : PropertyPathBindings)
+	{
+		if (PathContainsStruct(PropertyPathBinding.GetSourcePath()))
+		{
+			return true;
+		}
+		if (PathContainsStruct(PropertyPathBinding.GetTargetPath()))
+		{
+			return true;
+		}
+	}
+	return false;
+}
+
 void FStateTreePropertyBindings::DebugPrintInternalLayout(FString& OutString) const
 {
 	/** Array of expected source structs. */
diff --git a/Engine/Plugins/Runtime/StateTree/Source/StateTreeModule/Public/StateTree.h b/Engine/Plugins/Runtime/StateTree/Source/StateTreeModule/Public/StateTree.h
--- a/Engine/Plugins/Runtime/StateTree/Source/StateTreeModule/Public/StateTree.h
+++ b/Engine/Plugins/Runtime/StateTree/Source/StateTreeModule/Public/StateTree.h
@@ -81,7 +81,6 @@ class STATETREEMODULE_API UStateTree : public UDataAsset
 	GENERATED_BODY()
 
 public:
-
 	/** @return Shared instance data. */
 	TSharedPtr<FStateTreeInstanceData> GetSharedInstanceData() const;
 
@@ -130,6 +129,8 @@ public:
 	/** Edit time data for the StateTree, instance of UStateTreeEditorData */
 	UPROPERTY()
 	TObjectPtr<UObject> EditorData;
+
+	FDelegateHandle OnObjectsReinstancedHandle;
 #endif
 
 	/** Hash of the editor data from last compile. Also used to detect mismatching events from recorded traces. */
@@ -151,6 +152,10 @@ protected:
 	virtual void Serialize(FStructuredArchiveRecord Record) override;
 	
 #if WITH_EDITOR
+	using FReplacementObjectMap = TMap<UObject*, UObject*>;
+	void OnObjectsReinstanced(const FReplacementObjectMap& ObjectMap);
+	virtual void PostInitProperties() override;
+	virtual void BeginDestroy() override;
 	virtual void GetAssetRegistryTags(TArray<FAssetRegistryTag>& OutTags) const override;
 	virtual void PostLoadAssetRegistryTags(const FAssetData& InAssetData, TArray<FAssetRegistryTag>& OutTagsAndValuesToUpdate) const;
 #endif
diff --git a/Engine/Plugins/Runtime/StateTree/Source/StateTreeModule/Public/StateTreePropertyBindings.h b/Engine/Plugins/Runtime/StateTree/Source/StateTreeModule/Public/StateTreePropertyBindings.h
--- a/Engine/Plugins/Runtime/StateTree/Source/StateTreeModule/Public/StateTreePropertyBindings.h
+++ b/Engine/Plugins/Runtime/StateTree/Source/StateTreeModule/Public/StateTreePropertyBindings.h
@@ -725,6 +725,11 @@ struct STATETREEMODULE_API FStateTreePropertyBindings
 	 */
 	bool ResetObjects(const FStateTreeIndex16 TargetBatchIndex, FStateTreeDataView TargetStructView) const;
 
+	/**
+	 * @return true if any of the elements in the property bindings contains any of the structs in the set.
+	 */
+	bool ContainsAnyStruct(const TSet<const UStruct*>& Structs);
+	
 	void DebugPrintInternalLayout(FString& OutString) const;
 
 	/** @return how properties are compatible for copying. */
