StateTree: Added simple asset validation which runs Link().

diff --git a/Engine/Plugins/Runtime/StateTree/Source/StateTreeModule/Private/StateTree.cpp b/Engine/Plugins/Runtime/StateTree/Source/StateTreeModule/Private/StateTree.cpp
--- a/Engine/Plugins/Runtime/StateTree/Source/StateTreeModule/Private/StateTree.cpp
+++ b/Engine/Plugins/Runtime/StateTree/Source/StateTreeModule/Private/StateTree.cpp
@@ -9,6 +9,7 @@
 #include "Misc/ScopeRWLock.h"
 #include "StateTreeDelegates.h"
 #include "Logging/LogScopedVerbosityOverride.h"
+#include "Misc/DataValidation.h"
 
 #include UE_INLINE_GENERATED_CPP_BY_NAME(StateTree)
 
@@ -204,6 +205,17 @@ void UStateTree::PostLoadAssetRegistryTags(const FAssetData& InAssetData, TArray
 	}
 }
 
+EDataValidationResult UStateTree::IsDataValid(FDataValidationContext& Context)
+{
+	if (!Link())
+	{
+		Context.AddError(FText::FromString(FString::Printf(TEXT("%s failed to link. Please recompile the State Tree for more details errors."), *GetFullName())));
+		return EDataValidationResult::Invalid;
+	}
+
+	return Super::IsDataValid(Context);
+}
+
 #endif // WITH_EDITOR
 
 void UStateTree::AddReferencedObjects(UObject* InThis, FReferenceCollector& Collector)
diff --git a/Engine/Plugins/Runtime/StateTree/Source/StateTreeModule/Public/StateTree.h b/Engine/Plugins/Runtime/StateTree/Source/StateTreeModule/Public/StateTree.h
--- a/Engine/Plugins/Runtime/StateTree/Source/StateTreeModule/Public/StateTree.h
+++ b/Engine/Plugins/Runtime/StateTree/Source/StateTreeModule/Public/StateTree.h
@@ -158,6 +158,7 @@ protected:
 	virtual void BeginDestroy() override;
 	virtual void GetAssetRegistryTags(TArray<FAssetRegistryTag>& OutTags) const override;
 	virtual void PostLoadAssetRegistryTags(const FAssetData& InAssetData, TArray<FAssetRegistryTag>& OutTagsAndValuesToUpdate) const;
+	virtual EDataValidationResult IsDataValid(class FDataValidationContext& Context) override;
 #endif
 
 	static void AddReferencedObjects(UObject* InThis, FReferenceCollector& Collector);
