[StateTree] debugger updates:
- replaced log listing by treeview (as a prototype but will be replaced by a timeline)
- added event for conditions
- added details view for Condition and Tasks instance data

diff --git a/Engine/Plugins/Runtime/StateTree/Source/StateTreeEditorModule/Private/SStateTreeDebuggerView.cpp b/Engine/Plugins/Runtime/StateTree/Source/StateTreeEditorModule/Private/SStateTreeDebuggerView.cpp
--- a/Engine/Plugins/Runtime/StateTree/Source/StateTreeEditorModule/Private/SStateTreeDebuggerView.cpp
+++ b/Engine/Plugins/Runtime/StateTree/Source/StateTreeEditorModule/Private/SStateTreeDebuggerView.cpp
@@ -3,25 +3,148 @@
 #if WITH_STATETREE_DEBUGGER
 
 #include "SStateTreeDebuggerView.h"
-
+#include "Debugger/StateTreeDebugger.h"
+#include "DetailsViewArgs.h"
 #include "Editor.h"
+#include "Factories.h"
 #include "Framework/MultiBox/MultiBoxBuilder.h"
-#include "IMessageLogListing.h"
+#include "IStructureDetailsView.h"
 #include "Kismet2/DebuggerCommands.h"
-#include "MessageLogModule.h"
 #include "Modules/ModuleManager.h"
+#include "PropertyEditorModule.h"
 #include "StateTree.h"
-#include "Debugger/StateTreeDebugger.h"
-#include "Debugger/StateTreeTraceTypes.h"
 #include "StateTreeDebuggerCommands.h"
+#include "StateTreeExecutionContext.h"
 #include "StateTreeState.h"
-#include "StateTreeTaskBase.h"
 #include "StateTreeViewModel.h"
 #include "Widgets/Input/SComboButton.h"
 #include "Widgets/Layout/SBorder.h"
+#include "Widgets/Layout/SSpacer.h"
 
 #define LOCTEXT_NAMESPACE "StateTreeEditor"
 
+
+//----------------------------------------------------------------------//
+// FStateTreeTraceTextObjectFactory
+//----------------------------------------------------------------------//
+struct FStateTreeTraceTextObjectFactory : FCustomizableTextObjectFactory
+{
+	UObject* NodeInstanceObject = nullptr;
+	FStateTreeTraceTextObjectFactory() : FCustomizableTextObjectFactory(GWarn) {}
+
+protected:
+	virtual bool CanCreateClass(UClass* ObjectClass, bool& bOmitSubObjs) const override
+	{
+		return true;
+	}
+
+	virtual void ProcessConstructedObject(UObject* CreatedObject) override
+	{
+		NodeInstanceObject = CreatedObject;
+	}
+};
+
+
+//----------------------------------------------------------------------//
+// SStateTreeDebuggerTableRow
+//----------------------------------------------------------------------//
+class SStateTreeDebuggerTableRow : public SMultiColumnTableRow<TSharedPtr<FTreeElement>>
+{
+public:
+	void Construct(const FArguments& InArgs, const TSharedPtr<STableViewBase>& InOwnerTableView, const TSharedPtr<FTreeElement>& InElement, const TSharedRef<FStateTreeViewModel>& InStateTreeViewModel)
+	{
+		Item = InElement;
+		StateTreeViewModel = InStateTreeViewModel.ToSharedPtr();
+		SMultiColumnTableRow::Construct(InArgs, InOwnerTableView.ToSharedRef());
+	}
+
+	// SMultiColumnTableRow overrides
+	virtual TSharedRef<SWidget> GenerateWidgetForColumn(const FName& ColumnName) override
+	{
+		const TSharedPtr<SHorizontalBox> Contents = SNew(SHorizontalBox);
+
+		Contents->AddSlot()
+			.VAlign(VAlign_Fill)
+			.HAlign(HAlign_Left)
+			.AutoWidth()
+			[
+				SNew(SExpanderArrow, SharedThis(this))
+				.ShouldDrawWires(true)
+				.IndentAmount(32)
+				.BaseIndentLevel(0)
+			];
+		
+		if (ColumnName == FName("Desc"))
+		{
+			Contents->AddSlot()
+				.Padding(5, 0)
+				.FillWidth(1.f)
+				[
+					SNew(STextBlock)
+					.Font(FCoreStyle::GetDefaultFontStyle("Mono", 9))
+					.Text_Lambda([&Event=Item->Event, &StateTreeViewModel=StateTreeViewModel]()
+					{
+						const UStateTree* StateTree = StateTreeViewModel->GetStateTree();
+						
+						// Use log event messages directly
+						if (const FStateTreeTraceLogEvent* LogEvent = Event.TryGet<FStateTreeTraceLogEvent>())
+						{
+							if (LogEvent->Message.Len())
+							{
+								return FText::FromString(*LogEvent->Message);
+							}							
+						}
+						// Process state events (index has a different meaning)
+						else if (const FStateTreeTraceStateEvent* StateEvent = Event.TryGet<FStateTreeTraceStateEvent>())
+						{
+							const FStateTreeStateHandle StateHandle(StateEvent->Idx);
+							if (const FCompactStateTreeState* CompactState = StateTree->GetStateFromHandle(StateHandle))
+							{
+								return FText::FromString(FString::Printf(TEXT("%s State '%s'"),
+												*StaticEnum<EStateTreeTraceNodeEventType>()->GetNameStringByValue((int64)StateEvent->EventType),
+												*CompactState->Name.ToString()));
+							}
+						}
+						// Process Tasks events
+						else if (const FStateTreeTraceTaskEvent* TaskEvent = Event.TryGet<FStateTreeTraceTaskEvent>())
+						{
+							const FConstStructView NodeView = StateTree->GetNode(TaskEvent->Idx);
+							const FStateTreeNodeBase* Node = NodeView.GetPtr<const FStateTreeNodeBase>();
+							
+							return FText::FromString(FString::Printf(TEXT("%s:%s %s '%s'"),
+									*StaticEnum<EStateTreeTraceNodeEventType>()->GetNameStringByValue((int64)TaskEvent->EventType),
+									*StaticEnum<EStateTreeRunStatus>()->GetNameStringByValue((int64)TaskEvent->Status),
+									*NodeView.GetScriptStruct()->GetName(),
+									Node != nullptr ? *Node->Name.ToString() : *LexToString(TaskEvent->Idx)));
+						}
+						// Process Conditions events
+						else if (const FStateTreeTraceConditionEvent* ConditionEvent = Event.TryGet<FStateTreeTraceConditionEvent>())
+						{
+							const FConstStructView NodeView = StateTree->GetNode(ConditionEvent->Idx);
+							const FStateTreeNodeBase* Node = NodeView.GetPtr<const FStateTreeNodeBase>();
+							
+							return FText::FromString(FString::Printf(TEXT("%s %s '%s'"),
+									*StaticEnum<EStateTreeTraceNodeEventType>()->GetNameStringByValue((int64)ConditionEvent->EventType),
+									*NodeView.GetScriptStruct()->GetName(),
+									Node != nullptr ? *Node->Name.ToString() : *LexToString(ConditionEvent->Idx)));
+						}
+						
+						return FText(); 
+					})
+				];
+		}
+
+		return Contents.ToSharedRef();
+	}
+
+	TSharedPtr<FStateTreeViewModel> StateTreeViewModel;
+	TSharedPtr<FTreeElement> Item;
+};
+
+
+//----------------------------------------------------------------------//
+// SStateTreeDebuggerView
+//----------------------------------------------------------------------//
 SStateTreeDebuggerView::SStateTreeDebuggerView()
 {
 	FEditorDelegates::BeginPIE.AddRaw(this, &SStateTreeDebuggerView::OnPIEStarted);
@@ -33,6 +156,18 @@ SStateTreeDebuggerView::SStateTreeDebuggerView()
 
 SStateTreeDebuggerView::~SStateTreeDebuggerView()
 {
+	if (SelectedNodeDataObject.IsValid())
+	{
+		SelectedNodeDataObject->RemoveFromRoot();
+	}
+
+	if (Debugger)
+	{
+		Debugger->OnTracesUpdated.Unbind();
+		Debugger->OnBreakpointHit.Unbind();
+		Debugger->OnDebuggedInstanceSet.Unbind();
+	}
+
 	FEditorDelegates::BeginPIE.RemoveAll(this);
 	FEditorDelegates::EndPIE.RemoveAll(this);
 	FEditorDelegates::PausePIE.RemoveAll(this);
@@ -42,11 +177,6 @@ SStateTreeDebuggerView::~SStateTreeDebuggerView()
 
 void SStateTreeDebuggerView::OnPIEStarted(const bool bIsSimulating) const
 {
-	if (DebuggerTraceListing)
-	{
-		DebuggerTraceListing->ClearMessages();
-	}
-
 	if (!Debugger->IsAnalysisSessionActive())
 	{
 		Debugger->StartLastLiveSessionAnalysis();
@@ -73,7 +203,7 @@ void SStateTreeDebuggerView::OnPIESingleStepped(bool bSimulating) const
 	Debugger->SyncToCurrentSessionDuration();
 }
 
-void SStateTreeDebuggerView::Construct(const FArguments& InArgs, const UStateTree* InStateTree, const TSharedRef<FStateTreeViewModel> InStateTreeViewModel, const TSharedRef<FUICommandList> InCommandList)
+void SStateTreeDebuggerView::Construct(const FArguments& InArgs, const UStateTree* InStateTree, const TSharedRef<FStateTreeViewModel>& InStateTreeViewModel, const TSharedRef<FUICommandList>& InCommandList)
 {
 	StateTreeViewModel = InStateTreeViewModel;
 	StateTree = InStateTree;
@@ -88,17 +218,6 @@ void SStateTreeDebuggerView::Construct(const FArguments& InArgs, const UStateTre
 
 	BindDebuggerToolbarCommands(InCommandList);
 
-	FMessageLogModule& MessageLogModule = FModuleManager::LoadModuleChecked<FMessageLogModule>("MessageLog");
-	FMessageLogInitializationOptions LogOptions;
-	LogOptions.bShowPages = false;
-	LogOptions.bShowFilters = false;
-	LogOptions.bAllowClear = false;
-	LogOptions.bScrollToBottom = true;
-	LogOptions.MaxPageCount = 1;
-
-	DebuggerTraceListing = MessageLogModule.CreateLogListing("StateTreeDebugger", LogOptions);
-	DebuggerTraces = MessageLogModule.CreateLogListingWidget(DebuggerTraceListing.ToSharedRef());
-
 	const TSharedRef<FUICommandList> CommandList = MakeShareable(new FUICommandList);
 
 	// Register the play world commands
@@ -159,9 +278,124 @@ void SStateTreeDebuggerView::Construct(const FArguments& InArgs, const UStateTre
 	InCommandList->Append(CommandList);
 
 	TSharedPtr<FUICommandList> ActionList = InCommandList;
-		Debugger->OnTracesUpdated.BindSP(this, &SStateTreeDebuggerView::OnTracesUpdated);
+
+	Debugger->OnTracesUpdated.BindSP(this, &SStateTreeDebuggerView::OnTracesUpdated);
 	Debugger->OnBreakpointHit.BindSP(this, &SStateTreeDebuggerView::OnBreakpointHit, ActionList);
+	Debugger->OnDebuggedInstanceSet.BindSP(this, &SStateTreeDebuggerView::OnDebuggedInstanceSet);
+
+	PropertiesBorder = SNew(SBorder);
+	
+	// TreeView
+	TreeView = SNew(STreeView<TSharedPtr<FTreeElement>>)
+		.HeaderRow(SNew(SHeaderRow)
+			+SHeaderRow::Column("Desc")
+			.DefaultLabel(LOCTEXT("DescriptionColumnHeader", "Description")))
+		    .OnGenerateRow_Lambda([this](const TSharedPtr<FTreeElement>& InElement, const TSharedRef<STableViewBase>& InOwnerTableView)
+		    {
+			    return SNew(SStateTreeDebuggerTableRow, InOwnerTableView, InElement, StateTreeViewModel.ToSharedRef());
+		    } )
+		    .OnGetChildren_Lambda([](const TSharedPtr<const FTreeElement>& InParent, TArray<TSharedPtr<FTreeElement>>& OutChildren)
+		    {
+			    if (const FTreeElement* Parent = InParent.Get())
+			    {
+				    OutChildren.Append(Parent->Children);
+			    }
+		    })
+		.TreeItemsSource(&TreeElements)
+		.ItemHeight(32)
+		.OnSelectionChanged_Lambda([this](const TSharedPtr<FTreeElement>& InSelectedItem, ESelectInfo::Type SelectionType)
+		{
+			if (!InSelectedItem.IsValid())
+			{
+				return;
+			}
+
+			TSharedPtr<SWidget> DetailsView;
+
+			FString TypePath;
+			FString InstanceDataAsText;
+
+			if (const FStateTreeTraceConditionEvent* ConditionEvent = InSelectedItem->Event.TryGet<FStateTreeTraceConditionEvent>())
+			{
+				TypePath = ConditionEvent->TypePath;
+				InstanceDataAsText = ConditionEvent->InstanceDataAsText;
+			}
+			else if (const FStateTreeTraceTaskEvent* TaskEvent = InSelectedItem->Event.TryGet<FStateTreeTraceTaskEvent>())
+			{
+				TypePath = TaskEvent->TypePath;
+				InstanceDataAsText = TaskEvent->InstanceDataAsText;
+			}
+
+			if (!TypePath.IsEmpty())
+			{
+				FDetailsViewArgs DetailsViewArgs;
+				DetailsViewArgs.NameAreaSettings = FDetailsViewArgs::HideNameArea;
+
+				UScriptStruct* ScriptStruct = FindObject<UScriptStruct>(nullptr, *TypePath, /*ExactClass*/false);
+				if (ScriptStruct == nullptr)
+				{
+					ScriptStruct = LoadObject<UScriptStruct>(nullptr, *TypePath);
+				}
+
+				if (ScriptStruct != nullptr)
+				{
+					SelectedNodeDataStruct.InitializeAs(ScriptStruct);
 
+					ScriptStruct->ImportText(*InstanceDataAsText, SelectedNodeDataStruct.GetMutableMemory(), /*OwnerObject*/nullptr, PPF_None, GLog, ScriptStruct->GetName());
+
+					SelectedNodeDataStruct.GetScriptStruct();
+					FPropertyEditorModule& PropertyEditorModule = FModuleManager::LoadModuleChecked<FPropertyEditorModule>("PropertyEditor");
+					const TSharedPtr<FStructOnScope> InstanceDataStruct = MakeShared<FStructOnScope>(SelectedNodeDataStruct.GetScriptStruct(), const_cast<uint8*>(SelectedNodeDataStruct.GetMemory()));
+
+					FStructureDetailsViewArgs StructureViewArgs;
+					StructureViewArgs.bShowObjects = true;
+					StructureViewArgs.bShowAssets = true;
+					StructureViewArgs.bShowClasses = true;
+					StructureViewArgs.bShowInterfaces = true;
+					const TSharedRef<IStructureDetailsView> StructDetailsView = PropertyEditorModule.CreateStructureDetailView(DetailsViewArgs, StructureViewArgs, InstanceDataStruct);
+
+					DetailsView = StructDetailsView->GetDetailsView()->AsShared();
+				}
+
+				// UObject
+				UClass* Class = FindObject<UClass>(nullptr, *TypePath, /*ExactClass*/false);
+				if (Class == nullptr)
+				{
+					Class = LoadObject<UClass>(nullptr, *TypePath);
+				}
+
+				if (Class != nullptr)
+				{
+					if (SelectedNodeDataObject.IsValid())
+					{
+						SelectedNodeDataObject->RemoveFromRoot();
+					}
+					FStateTreeTraceTextObjectFactory ObjectFactory;
+					if (ObjectFactory.CanCreateObjectsFromText(InstanceDataAsText))
+					{
+						ObjectFactory.ProcessBuffer(GetTransientPackage(), RF_Transactional, InstanceDataAsText);
+						SelectedNodeDataObject = ObjectFactory.NodeInstanceObject;
+						SelectedNodeDataObject->AddToRoot();
+
+						FPropertyEditorModule& PropertyEditorModule = FModuleManager::Get().GetModuleChecked<FPropertyEditorModule>("PropertyEditor");
+						const TSharedRef<IDetailsView> ObjectDetailsView = PropertyEditorModule.CreateDetailView(DetailsViewArgs);
+						ObjectDetailsView->SetObject(SelectedNodeDataObject.Get());
+						DetailsView = ObjectDetailsView->AsShared();
+					}
+				}
+			}
+
+			if (DetailsView)
+			{
+				PropertiesBorder->SetContent(DetailsView.ToSharedRef());
+			}
+			else
+			{
+				PropertiesBorder->ClearContent();
+			}
+		})
+		.AllowOverscroll(EAllowOverscroll::No);
+	
 	ChildSlot
 	[
 		SNew(SBorder)
@@ -184,16 +418,22 @@ void SStateTreeDebuggerView::Construct(const FArguments& InArgs, const UStateTre
 			.FillHeight(1.0f)
 			.Padding(0.0f, 0.0f, 0.0f, 4.0f)
 			[
-				SNew(SBox)
+				SNew(SSplitter)
+				.Orientation(Orient_Horizontal)
+				+ SSplitter::Slot()
+				[
+					TreeView.ToSharedRef()
+				]
+				+ SSplitter::Slot()
 				[
-					DebuggerTraces.ToSharedRef()
+					PropertiesBorder.ToSharedRef()
 				]
 			]
 		]
 	];
 }
 
-void SStateTreeDebuggerView::BindDebuggerToolbarCommands(const TSharedRef<FUICommandList> ToolkitCommands)
+void SStateTreeDebuggerView::BindDebuggerToolbarCommands(const TSharedRef<FUICommandList>& ToolkitCommands)
 {
 	const FStateTreeDebuggerCommands& Commands = FStateTreeDebuggerCommands::Get();
 	
@@ -263,9 +503,9 @@ void SStateTreeDebuggerView::ToggleBreakpoint() const
 void SStateTreeDebuggerView::OnTracesUpdated(
 	const TConstArrayView<const FStateTreeTraceEventVariantType> Events,
 	const TConstArrayView<UE::StateTreeDebugger::FFrameIndexSpan> Spans
-	) const
+	)
 {
-	DebuggerTraceListing->ClearMessages();
+	TreeElements.Reset();
 	
 	if (Events.IsEmpty())
 	{
@@ -273,59 +513,114 @@ void SStateTreeDebuggerView::OnTracesUpdated(
 	}
 
 	check(Spans.Num());
-
-	TArray<TSharedRef<FTokenizedMessage>> Messages;
-	Messages.Reserve(Events.Num());
-
 	check(StateTree.IsValid());
 
 	int32 SpanIdx = INDEX_NONE;
+
+	TArray<TSharedPtr<FTreeElement>, TInlineAllocator<8>> ParentStack;
+	EStateTreeUpdatePhase LastPhase = EStateTreeUpdatePhase::Unset;
+	EStateTreeUpdatePhase EventPhase = EStateTreeUpdatePhase::Unset;
+	
 	for (int32 EventIdx = 0; EventIdx < Events.Num(); EventIdx++)
 	{
 		const FStateTreeTraceEventVariantType& Event = Events[EventIdx];
-		
+		EventPhase = LastPhase;
+
 		if (SpanIdx == INDEX_NONE
 			|| (Spans.IsValidIndex(SpanIdx+1) && EventIdx >= Spans[SpanIdx+1].EventIdx))
 		{
 			SpanIdx = SpanIdx == INDEX_NONE ? 0 : SpanIdx+1;
-			Messages.Add(FTokenizedMessage::Create(EMessageSeverity::Info,
-				FText::FromString(FString::Printf(TEXT("---------------- Frame: %s ----------------"), *LexToString(Spans[SpanIdx].FrameIdx)))));
+
+			while (ParentStack.Num())
+			{
+				TSharedPtr<FTreeElement> RemovedElement = ParentStack.Pop();
+				if (RemovedElement->Children.IsEmpty())
+				{
+					if (ParentStack.Num())
+					{
+						ParentStack.Last()->Children.Remove(RemovedElement);
+					}
+					TreeElements.Remove(RemovedElement);
+				}
+			}
+
+			LastPhase = EStateTreeUpdatePhase::Unset;
+
+			// Create fake event to create a parent node with frame info
+			FStateTreeTraceLogEvent DummyEvent(EStateTreeUpdatePhase::Unset, FString::Printf(TEXT("Frame: %s"), *LexToString(Spans[SpanIdx].Frame.Index)));
+
+			ParentStack.Push(TreeElements.Add_GetRef(MakeShareable(new FTreeElement(
+				Spans[SpanIdx].Frame,
+				FStateTreeTraceEventVariantType(TInPlaceType<FStateTreeTraceLogEvent>(), DummyEvent)))));
 		}
-		
-		// Add log event messages directly
+
+		// Need to test each type explicitly with TVariant even if they are all FStateTreeTracePhaseEvent
 		if (const FStateTreeTraceLogEvent* LogEvent = Event.TryGet<FStateTreeTraceLogEvent>())
 		{
-			if (LogEvent->Message.Len())
-			{
-				Messages.Add(FTokenizedMessage::Create(EMessageSeverity::Info, FText::FromString(LogEvent->Message)));	
-			}			
-			continue;	
+			EventPhase = LogEvent->Phase;
+		}
+		else if (const FStateTreeTraceStateEvent* StateEvent = Event.TryGet<FStateTreeTraceStateEvent>())
+		{
+			EventPhase = StateEvent->Phase;
+		}
+		else if (const FStateTreeTraceTaskEvent* TaskEvent = Event.TryGet<FStateTreeTraceTaskEvent>())
+		{
+			EventPhase = TaskEvent->Phase;
+		}
+		else if (const FStateTreeTraceConditionEvent* ConditionEvent = Event.TryGet<FStateTreeTraceConditionEvent>())
+		{
+			EventPhase = ConditionEvent->Phase;
 		}
 
-		// Add messages for state events 
-		if (const FStateTreeTraceStateEvent* StateEvent = Event.TryGet<FStateTreeTraceStateEvent>())
+		// Create a hierarchy level for each update phase
+		if (EventPhase != LastPhase)
 		{
-			const FStateTreeStateHandle StateHandle(StateEvent->StateIdx);
-			if (const FCompactStateTreeState* CompactState = StateTree->GetStateFromHandle(StateHandle))
+			// Remove valid previous phase from parent stack 
+			if (ensure(ParentStack.Num()) && LastPhase != EStateTreeUpdatePhase::Unset)
 			{
-				Messages.Add(FTokenizedMessage::Create(EMessageSeverity::Info, FText::FromString(FString::Printf(TEXT("%s State '%s'"),
-								*StaticEnum<EStateTreeTraceNodeEventType>()->GetNameStringByValue((int64)StateEvent->EventType),
-								*CompactState->Name.ToString()))));
+				TSharedPtr<FTreeElement> RemovedElement = ParentStack.Pop();
+				if (RemovedElement->Children.IsEmpty())
+				{
+					if (ParentStack.Num())
+					{
+						ParentStack.Last()->Children.Remove(RemovedElement);	
+					}
+				}
 			}
+
+			LastPhase = EventPhase;
+
+			// Create fake log event to describe the phase
+			FStateTreeTraceLogEvent DummyEvent(EStateTreeUpdatePhase::Unset,
+				FString::Printf(TEXT("%s"), *StaticEnum<EStateTreeUpdatePhase>()->GetNameStringByValue((int64)EventPhase)));
+
+			// Add to the list of children of the frame element
+			check(ParentStack.Num());
+			ParentStack.Push(ParentStack.Last()->Children.Add_GetRef(MakeShareable(new FTreeElement(
+				Spans[SpanIdx].Frame,
+				FStateTreeTraceEventVariantType(TInPlaceType<FStateTreeTraceLogEvent>(), DummyEvent)))));
 		}
 
-		if (const FStateTreeTraceTaskEvent* TaskEvent = Event.TryGet<FStateTreeTraceTaskEvent>())
+		if (ensure(ParentStack.Num()))
+		{
+			ParentStack.Last()->Children.Add(MakeShareable(new FTreeElement(Spans[SpanIdx].Frame, Event)));
+		}
+	}
+
+	while (ParentStack.Num())
+	{
+		TSharedPtr<FTreeElement> RemovedElement = ParentStack.Pop();
+		if (RemovedElement->Children.IsEmpty())
 		{
-			FConstStructView NodeView = StateTree->GetNode(TaskEvent->TaskIdx);
-			const FStateTreeTaskBase* Task = NodeView.GetPtr<const FStateTreeTaskBase>();
-					
-			Messages.Add(FTokenizedMessage::Create(EMessageSeverity::Info, FText::FromString(FString::Printf(TEXT("  %s Task '%s'"),
-				*StaticEnum<EStateTreeTraceNodeEventType>()->GetNameStringByValue((int64)TaskEvent->EventType),
-				Task != nullptr ? *Task->Name.ToString() : *LexToString(TaskEvent->TaskIdx)))));
+			if (ParentStack.Num())
+			{
+				ParentStack.Last()->Children.Remove(RemovedElement);
+			}
+			TreeElements.Remove(RemovedElement);
 		}
 	}
 
-	DebuggerTraceListing->AddMessages(Messages, /*bMirrorToOutputLog*/false);
+	TreeView->RequestTreeRefresh();
 }
 
 void SStateTreeDebuggerView::OnBreakpointHit(const FStateTreeInstanceDebugId InstanceId, const FStateTreeStateHandle StateHandle, const TSharedPtr<FUICommandList> ActionList) const
@@ -355,6 +650,17 @@ void SStateTreeDebuggerView::OnBreakpointHit(const FStateTreeInstanceDebugId Ins
 	}
 }
 
+void SStateTreeDebuggerView::OnDebuggedInstanceSet()
+{
+	TreeElements.Reset();
+	if (TreeView)
+	{
+		TreeView->RequestTreeRefresh();	
+	}
+				
+	PropertiesBorder->ClearContent();
+}
+
 TSharedRef<SWidget> SStateTreeDebuggerView::OnGetDebuggerTracesMenu() const
 {
 	FMenuBuilder MenuBuilder(/*bInShouldCloseWindowAfterMenuSelection*/true, /*InCommandList*/nullptr);
diff --git a/Engine/Plugins/Runtime/StateTree/Source/StateTreeEditorModule/Private/SStateTreeDebuggerView.h b/Engine/Plugins/Runtime/StateTree/Source/StateTreeEditorModule/Private/SStateTreeDebuggerView.h
--- a/Engine/Plugins/Runtime/StateTree/Source/StateTreeEditorModule/Private/SStateTreeDebuggerView.h
+++ b/Engine/Plugins/Runtime/StateTree/Source/StateTreeEditorModule/Private/SStateTreeDebuggerView.h
@@ -4,8 +4,12 @@
 
 #if WITH_STATETREE_DEBUGGER
 
+#include "StateTreeTypes.h"
 #include "Debugger/StateTreeTraceTypes.h"
+#include "InstancedStruct.h"
+#include "Widgets/Views/STreeView.h"
 #include "Widgets/SCompoundWidget.h"
+#include "TraceServices/Model/Frames.h"
 
 namespace UE::StateTreeDebugger { struct FFrameIndexSpan; }
 struct FStateTreeDebugger;
@@ -13,9 +17,23 @@ struct FStateTreeInstanceDebugId;
 class FStateTreeEditor;
 class FStateTreeViewModel;
 class FUICommandList;
-class IMessageLogListing;
 class UStateTree;
 
+/** An item in the tree */
+class FTreeElement : public TSharedFromThis<FTreeElement>
+{
+public:
+	explicit FTreeElement(const TraceServices::FFrame& Frame, const FStateTreeTraceEventVariantType& Event)
+		: Frame(Frame), Event(Event)
+	{
+	}
+
+	TraceServices::FFrame Frame;
+	FStateTreeTraceEventVariantType Event;
+	TArray<TSharedPtr<FTreeElement>> Children;
+};
+
+
 class SStateTreeDebuggerView : public SCompoundWidget
 {
 public:
@@ -25,7 +43,7 @@ public:
 	SStateTreeDebuggerView();
 	~SStateTreeDebuggerView();
 
-	void Construct(const FArguments& InArgs, const UStateTree* InStateTree, TSharedRef<FStateTreeViewModel> InStateTreeViewModel, TSharedRef<FUICommandList> InCommandList);
+	void Construct(const FArguments& InArgs, const UStateTree* InStateTree, const TSharedRef<FStateTreeViewModel>& InStateTreeViewModel, const TSharedRef<FUICommandList>& InCommandList);
 
 private:
 	TSharedRef<SWidget> OnGetDebuggerInstancesMenu() const;
@@ -37,10 +55,11 @@ private:
 	void OnPIEResumed(bool bIsSimulating) const;
 	void OnPIESingleStepped(bool bIsSimulating) const;
 	
-	void OnTracesUpdated(const TConstArrayView<const FStateTreeTraceEventVariantType> Events, const TConstArrayView<UE::StateTreeDebugger::FFrameIndexSpan> Spans) const;
+	void OnTracesUpdated(const TConstArrayView<const FStateTreeTraceEventVariantType> Events, const TConstArrayView<UE::StateTreeDebugger::FFrameIndexSpan> Spans);
 	void OnBreakpointHit(const FStateTreeInstanceDebugId InstanceId, const FStateTreeStateHandle StateHandle, const TSharedPtr<FUICommandList> ActionList) const;
+	void OnDebuggedInstanceSet();
 
-	void BindDebuggerToolbarCommands(const TSharedRef<FUICommandList> ToolkitCommands);
+	void BindDebuggerToolbarCommands(const TSharedRef<FUICommandList>& ToolkitCommands);
 
 	/**
 	 * Stores raw log information so it can be processed
@@ -75,9 +94,13 @@ private:
 	TSharedPtr<FStateTreeDebugger> Debugger;
 	TSharedPtr<FStateTreeViewModel> StateTreeViewModel;
 	TWeakObjectPtr<const UStateTree> StateTree;
-	
-	TSharedPtr<SWidget> DebuggerTraces;
-	TSharedPtr<IMessageLogListing> DebuggerTraceListing;
+
+	TSharedPtr<STreeView<TSharedPtr<FTreeElement>>> TreeView;
+	TArray<TSharedPtr<FTreeElement>> TreeElements;
+
+	TSharedPtr<SBorder> PropertiesBorder;
+	FInstancedStruct SelectedNodeDataStruct;
+	TWeakObjectPtr<UObject> SelectedNodeDataObject;
 };
 
 #endif // WITH_STATETREE_DEBUGGER
\ No newline at end of file
diff --git a/Engine/Plugins/Runtime/StateTree/Source/StateTreeModule/Private/Debugger/StateTreeDebugger.cpp b/Engine/Plugins/Runtime/StateTree/Source/StateTreeModule/Private/Debugger/StateTreeDebugger.cpp
--- a/Engine/Plugins/Runtime/StateTree/Source/StateTreeModule/Private/Debugger/StateTreeDebugger.cpp
+++ b/Engine/Plugins/Runtime/StateTree/Source/StateTreeModule/Private/Debugger/StateTreeDebugger.cpp
@@ -184,6 +184,8 @@ void FStateTreeDebugger::SetDebuggedInstance(const FStateTreeDebuggerInstanceDes
 		bProcessingInitialEvents = true;
 		SyncToCurrentSessionDuration();
 		bProcessingInitialEvents = false;
+
+		OnDebuggedInstanceSet.ExecuteIfBound();
 	}
 }
 
@@ -405,7 +407,7 @@ bool FStateTreeDebugger::CanStepBack() const
 	return bPaused && CurrentFrameIndex != INDEX_NONE
 		&& FramesWithEvents.ContainsByPredicate([CurrentIndex = CurrentFrameIndex](const UE::StateTreeDebugger::FFrameIndexSpan& IteratedFrame)
 		{
-			return IteratedFrame.FrameIdx < CurrentIndex;
+			return IteratedFrame.Frame.Index < CurrentIndex;
 		});
 }
 
@@ -413,7 +415,7 @@ void FStateTreeDebugger::StepBack()
 {
 	const int32 PrevFrameIndex = FramesWithEvents.FindLastByPredicate([CurrentIndex = CurrentFrameIndex](const UE::StateTreeDebugger::FFrameIndexSpan& IteratedFrame)
 	{
-		return IteratedFrame.FrameIdx < CurrentIndex;
+		return IteratedFrame.Frame.Index < CurrentIndex;
 	});
 	
 	if (PrevFrameIndex != INDEX_NONE)
@@ -421,8 +423,8 @@ void FStateTreeDebugger::StepBack()
 		// Rebuild from beginning up to previous frame with events
 		Events.Reset();
 		LastTraceReadTime = UnsetTime;
-		ReadTrace(FramesWithEvents[PrevFrameIndex].FrameIdx);
-		check(CurrentFrameIndex == FramesWithEvents[PrevFrameIndex].FrameIdx);
+		ReadTrace(FramesWithEvents[PrevFrameIndex].Frame.Index);
+		check(CurrentFrameIndex == FramesWithEvents[PrevFrameIndex].Frame.Index);
 	}
 }
 
@@ -431,7 +433,7 @@ bool FStateTreeDebugger::CanStepForward() const
 	return bPaused && CurrentFrameIndex != INDEX_NONE
 		&& FramesWithEvents.ContainsByPredicate([CurrentIndex = CurrentFrameIndex](const UE::StateTreeDebugger::FFrameIndexSpan& IteratedFrame)
 		{
-			return IteratedFrame.FrameIdx > CurrentIndex;
+			return IteratedFrame.Frame.Index > CurrentIndex;
 		});
 }
 
@@ -439,12 +441,12 @@ void FStateTreeDebugger::StepForward()
 {
 	const UE::StateTreeDebugger::FFrameIndexSpan* NextFrame = FramesWithEvents.FindByPredicate([CurrentIndex = CurrentFrameIndex](const UE::StateTreeDebugger::FFrameIndexSpan& IteratedFrame)
 	{
-		return IteratedFrame.FrameIdx > CurrentIndex;
+		return IteratedFrame.Frame.Index > CurrentIndex;
 	});
 	
 	if (NextFrame != nullptr)
 	{
-		const uint64 FrameIndex = (*NextFrame).FrameIdx; 
+		const uint64 FrameIndex = (*NextFrame).Frame.Index; 
 		ReadTrace(FrameIndex);
 		check(CurrentFrameIndex == FrameIndex);
 	}
@@ -578,14 +580,14 @@ void FStateTreeDebugger::ReadTrace(
 
 	if (FramesWithEvents.ContainsByPredicate([CurrentIndex = CurrentFrameIndex](const UE::StateTreeDebugger::FFrameIndexSpan& IteratedFrame)
 		{
-			return IteratedFrame.FrameIdx == CurrentIndex;
+			return IteratedFrame.Frame.Index == CurrentIndex;
 		}))
 	{
 		CurrentFrameIndex = Frame.Index;	
 	}
 	else if (FramesWithEvents.Num())
 	{
-		CurrentFrameIndex = FramesWithEvents.Last().FrameIdx;
+		CurrentFrameIndex = FramesWithEvents.Last().Frame.Index;
 	}
 }
 
@@ -627,7 +629,7 @@ bool FStateTreeDebugger::ProcessEvent(const FStateTreeInstanceDebugId InstanceId
 
 				if (DebuggedInstance.IsValid() || InstanceStateTree == DebuggedAsset)
 				{
-					HitBreakpointStateIndex = StatesWithBreakpoint.Find(FStateTreeStateHandle(StateEvent->StateIdx));
+					HitBreakpointStateIndex = StatesWithBreakpoint.Find(FStateTreeStateHandle(StateEvent->Idx));
 					if (HitBreakpointStateIndex != INDEX_NONE)
 					{
 						HitBreakpointInstanceId = InstanceId;
@@ -639,9 +641,9 @@ bool FStateTreeDebugger::ProcessEvent(const FStateTreeInstanceDebugId InstanceId
 		// Store events for currently debugged entry 
 		if (DebuggedInstance.Id == InstanceId)
 		{
-			if (FramesWithEvents.IsEmpty() || FramesWithEvents.Last().FrameIdx < Frame.Index)
+			if (FramesWithEvents.IsEmpty() || FramesWithEvents.Last().Frame.Index < Frame.Index)
 			{
-				FramesWithEvents.Add(UE::StateTreeDebugger::FFrameIndexSpan(Frame.Index, Events.Num()));
+				FramesWithEvents.Add(UE::StateTreeDebugger::FFrameIndexSpan(Frame, Events.Num()));
 			}
 
 			Events.Emplace(Event);
diff --git a/Engine/Plugins/Runtime/StateTree/Source/StateTreeModule/Private/Debugger/StateTreeTrace.cpp b/Engine/Plugins/Runtime/StateTree/Source/StateTreeModule/Private/Debugger/StateTreeTrace.cpp
--- a/Engine/Plugins/Runtime/StateTree/Source/StateTreeModule/Private/Debugger/StateTreeTrace.cpp
+++ b/Engine/Plugins/Runtime/StateTree/Source/StateTreeModule/Private/Debugger/StateTreeTrace.cpp
@@ -4,6 +4,8 @@
 
 #include "Debugger/StateTreeTrace.h"
 #include "Debugger/StateTreeDebugger.h"
+#include "Exporters/Exporter.h"
+#include "Serialization/BufferArchive.h"
 #include "StateTree.h"
 #include "UObject/Package.h"
 #include "Trace/Trace.inl"
@@ -25,6 +27,7 @@ UE_TRACE_EVENT_BEGIN(StateTreeDebugger, LogEvent)
 	UE_TRACE_EVENT_FIELD(uint64, Cycle)
 	UE_TRACE_EVENT_FIELD(uint32, InstanceId)
 	UE_TRACE_EVENT_FIELD(uint32, InstanceSerial)
+	UE_TRACE_EVENT_FIELD(uint8, Phase)
 	UE_TRACE_EVENT_FIELD(UE::Trace::WideString, Message)
 UE_TRACE_EVENT_END()
 
@@ -32,7 +35,8 @@ UE_TRACE_EVENT_BEGIN(StateTreeDebugger, StateEvent)
 	UE_TRACE_EVENT_FIELD(uint64, Cycle)
 	UE_TRACE_EVENT_FIELD(uint32, InstanceId)
 	UE_TRACE_EVENT_FIELD(uint32, InstanceSerial)
-	UE_TRACE_EVENT_FIELD(uint16, StateIdx)
+	UE_TRACE_EVENT_FIELD(uint8, Phase)
+	UE_TRACE_EVENT_FIELD(uint16, Idx)
 	UE_TRACE_EVENT_FIELD(uint8, EventType)
 UE_TRACE_EVENT_END()
 
@@ -40,7 +44,20 @@ UE_TRACE_EVENT_BEGIN(StateTreeDebugger, TaskEvent)
 	UE_TRACE_EVENT_FIELD(uint64, Cycle)
 	UE_TRACE_EVENT_FIELD(uint32, InstanceId)
 	UE_TRACE_EVENT_FIELD(uint32, InstanceSerial)
-	UE_TRACE_EVENT_FIELD(uint16, TaskIdx)
+	UE_TRACE_EVENT_FIELD(uint8, Phase)
+	UE_TRACE_EVENT_FIELD(uint16, Idx)
+	UE_TRACE_EVENT_FIELD(uint8[], DataView)
+	UE_TRACE_EVENT_FIELD(uint8, EventType)
+	UE_TRACE_EVENT_FIELD(uint8, Status)
+UE_TRACE_EVENT_END()
+
+UE_TRACE_EVENT_BEGIN(StateTreeDebugger, ConditionEvent)
+	UE_TRACE_EVENT_FIELD(uint64, Cycle)
+	UE_TRACE_EVENT_FIELD(uint32, InstanceId)
+	UE_TRACE_EVENT_FIELD(uint32, InstanceSerial)
+	UE_TRACE_EVENT_FIELD(uint8, Phase)
+	UE_TRACE_EVENT_FIELD(uint16, Idx)
+	UE_TRACE_EVENT_FIELD(uint8[], DataView)
 	UE_TRACE_EVENT_FIELD(uint8, EventType)
 UE_TRACE_EVENT_END()
 
@@ -51,11 +68,44 @@ UE_TRACE_EVENT_BEGIN(StateTreeDebugger, ActiveStatesEvent)
 	UE_TRACE_EVENT_FIELD(uint16[], ActiveStates)
 UE_TRACE_EVENT_END()
 
+namespace UE::StateTreeTrace
+{
+	void SerializeDataViewToArchive(FBufferArchive Ar, const FStateTreeDataView DataView)
+	{
+		if (const UScriptStruct* ScriptStruct = Cast<const UScriptStruct>(DataView.GetStruct()))
+		{
+			TRACE_CPUPROFILER_EVENT_SCOPE(UE::StateTree::ExportStructAsText)
+			FString StructPath = ScriptStruct->GetPathName();
+			FString TextValue;
+
+			ScriptStruct->ExportText(TextValue, DataView.GetMemory(), DataView.GetMemory(), /*OwnerObject*/nullptr, PPF_None, /*ExportRootScope*/nullptr);
+
+			Ar << StructPath;		
+			Ar << TextValue;
+		}
+		else if (const UClass* Class = Cast<const UClass>(DataView.GetStruct()))
+		{
+			TRACE_CPUPROFILER_EVENT_SCOPE(UE::StateTree::ExportObjectAsText)
+			FString StructPath = Class->GetPathName();
+			FStringOutputDevice OutputDevice;
+			UObject* Object = DataView.GetMutablePtr<UObject>();
+
+			// Not using on scope FExportObjectInnerContext since it is very costly to build.
+			// Passing a null context will make the export use an already built thread local context.
+			UExporter::ExportToOutputDevice(nullptr, Object, /*Exporter*/nullptr, OutputDevice, TEXT("copy"), 0, PPF_ExportsNotFullyQualified | PPF_Copy | PPF_Delimited, false, Object->GetOuter());
+
+			Ar << StructPath;
+			Ar << OutputDevice;
+		}
+	}
+}
+
 void UE::StateTreeTrace::OutputInstanceLifetimeEvent(
 	const FStateTreeInstanceDebugId InstanceId,
 	const UStateTree* StateTree,
 	const TCHAR* InstanceName,
-	const EStateTreeTraceInstanceEventType EventType)
+	const EStateTreeTraceInstanceEventType EventType
+	)
 {
 	if (StateTree == nullptr)
 	{
@@ -75,7 +125,11 @@ void UE::StateTreeTrace::OutputInstanceLifetimeEvent(
 		<< InstanceEvent.EventType(static_cast<uint8>(EventType));
 }
 
-void UE::StateTreeTrace::OutputLogEventTrace(const FStateTreeInstanceDebugId InstanceId, const TCHAR* Fmt, ...)
+void UE::StateTreeTrace::OutputLogEventTrace(
+	const FStateTreeInstanceDebugId InstanceId,
+	const EStateTreeUpdatePhase Phase,
+	const TCHAR* Fmt, ...
+	)
 {
 	static TCHAR TraceStaticBuffer[8192];
 	GET_VARARGS(TraceStaticBuffer, UE_ARRAY_COUNT(TraceStaticBuffer), UE_ARRAY_COUNT(TraceStaticBuffer) - 1, Fmt, Fmt);
@@ -84,30 +138,74 @@ void UE::StateTreeTrace::OutputLogEventTrace(const FStateTreeInstanceDebugId Ins
 		<< LogEvent.Cycle(FPlatformTime::Cycles64())
 		<< LogEvent.InstanceId(InstanceId.Id)
 		<< LogEvent.InstanceSerial(InstanceId.SerialNumber)
+		<< LogEvent.Phase(static_cast<uint8>(Phase))
 		<< LogEvent.Message(TraceStaticBuffer);
 }
 
-void UE::StateTreeTrace::OutputStateEventTrace(const FStateTreeInstanceDebugId InstanceId, const uint16 StateIdx, const EStateTreeTraceNodeEventType EventType)
+void UE::StateTreeTrace::OutputStateEventTrace(
+	const FStateTreeInstanceDebugId InstanceId,
+	const EStateTreeUpdatePhase Phase,
+	const uint16 StateIdx,
+	const EStateTreeTraceNodeEventType EventType
+	)
 {
 	UE_TRACE_LOG(StateTreeDebugger, StateEvent, StateTreeDebugChannel)
 		<< StateEvent.Cycle(FPlatformTime::Cycles64())
 		<< StateEvent.InstanceId(InstanceId.Id)
 		<< StateEvent.InstanceSerial(InstanceId.SerialNumber)
-		<< StateEvent.StateIdx(StateIdx)
+		<< StateEvent.Phase(static_cast<uint8>(Phase))
+		<< StateEvent.Idx(StateIdx)
 		<< StateEvent.EventType(static_cast<uint8>(EventType));
 }
 
-void UE::StateTreeTrace::OutputTaskEventTrace(const FStateTreeInstanceDebugId InstanceId, const uint16 TaskIdx, const EStateTreeTraceNodeEventType EventType)
+void UE::StateTreeTrace::OutputTaskEventTrace(
+	const FStateTreeInstanceDebugId InstanceId,
+	const EStateTreeUpdatePhase Phase,
+	const uint16 TaskIdx,
+	const FStateTreeDataView DataView,
+	const EStateTreeTraceNodeEventType EventType,
+	const EStateTreeRunStatus Status
+	)
 {
+	FBufferArchive Archive;
+	SerializeDataViewToArchive(Archive, DataView);
+
 	UE_TRACE_LOG(StateTreeDebugger, TaskEvent, StateTreeDebugChannel)
 		<< TaskEvent.Cycle(FPlatformTime::Cycles64())
 		<< TaskEvent.InstanceId(InstanceId.Id)
 		<< TaskEvent.InstanceSerial(InstanceId.SerialNumber)
-		<< TaskEvent.TaskIdx(TaskIdx)
-		<< TaskEvent.EventType(static_cast<uint8>(EventType));
+		<< TaskEvent.Phase(static_cast<uint8>(Phase))
+		<< TaskEvent.Idx(TaskIdx)
+		<< TaskEvent.DataView(Archive.GetData(), Archive.Num())
+		<< TaskEvent.EventType(static_cast<uint8>(EventType))
+		<< TaskEvent.Status(static_cast<uint8>(Status));
 }
 
-void UE::StateTreeTrace::OutputActiveStatesEventTrace(const FStateTreeInstanceDebugId InstanceId, const FStateTreeActiveStates& ActiveStates)
+void UE::StateTreeTrace::OutputConditionEventTrace(
+	const FStateTreeInstanceDebugId InstanceId,
+	const EStateTreeUpdatePhase Phase,
+	const uint16 ConditionIdx,
+	const FStateTreeDataView DataView,
+	const EStateTreeTraceNodeEventType EventType
+	)
+{
+	FBufferArchive Archive;
+	SerializeDataViewToArchive(Archive, DataView);
+	
+	UE_TRACE_LOG(StateTreeDebugger, ConditionEvent, StateTreeDebugChannel)
+	<< ConditionEvent.Cycle(FPlatformTime::Cycles64())
+	<< ConditionEvent.InstanceId(InstanceId.Id)
+	<< ConditionEvent.InstanceSerial(InstanceId.SerialNumber)
+	<< ConditionEvent.Phase(static_cast<uint8>(Phase))
+	<< ConditionEvent.Idx(ConditionIdx)
+	<< ConditionEvent.DataView(Archive.GetData(), Archive.Num())
+	<< ConditionEvent.EventType(static_cast<uint8>(EventType));
+}
+
+void UE::StateTreeTrace::OutputActiveStatesEventTrace(
+	const FStateTreeInstanceDebugId InstanceId,
+	const FStateTreeActiveStates& ActiveStates
+	)
 {
 	TArray<uint16, TInlineAllocator<FStateTreeActiveStates::MaxStates>> StatesIndices;
 	for (int32 i = 0; i < ActiveStates.Num(); i++)
diff --git a/Engine/Plugins/Runtime/StateTree/Source/StateTreeModule/Private/Debugger/StateTreeTraceAnalyzer.cpp b/Engine/Plugins/Runtime/StateTree/Source/StateTreeModule/Private/Debugger/StateTreeTraceAnalyzer.cpp
--- a/Engine/Plugins/Runtime/StateTree/Source/StateTreeModule/Private/Debugger/StateTreeTraceAnalyzer.cpp
+++ b/Engine/Plugins/Runtime/StateTree/Source/StateTreeModule/Private/Debugger/StateTreeTraceAnalyzer.cpp
@@ -6,6 +6,7 @@
 #include "Debugger/StateTreeDebugger.h"
 #include "Debugger/StateTreeTraceProvider.h"
 #include "Debugger/StateTreeTraceTypes.h"
+#include "Serialization/MemoryReader.h"
 #include "TraceServices/Model/AnalysisSession.h"
 
 FStateTreeTraceAnalyzer::FStateTreeTraceAnalyzer(TraceServices::IAnalysisSession& InSession, FStateTreeTraceProvider& InProvider)
@@ -22,6 +23,7 @@ void FStateTreeTraceAnalyzer::OnAnalysisBegin(const FOnAnalysisContext& Context)
 	Builder.RouteEvent(RouteId_LogMessage, "StateTreeDebugger", "LogEvent");
 	Builder.RouteEvent(RouteId_State, "StateTreeDebugger", "StateEvent");
 	Builder.RouteEvent(RouteId_Task, "StateTreeDebugger", "TaskEvent");
+	Builder.RouteEvent(RouteId_Condition, "StateTreeDebugger", "ConditionEvent");
 	Builder.RouteEvent(RouteId_ActiveStates, "StateTreeDebugger", "ActiveStatesEvent");
 }
 
@@ -49,19 +51,19 @@ bool FStateTreeTraceAnalyzer::OnEvent(const uint16 RouteId, EStyle Style, const
 				FGCScopeGuard Guard;
 				WeakStateTree = FindObject<UStateTree>(Path);
 			}
-			
+
 			if (const UStateTree* StateTree = WeakStateTree.Get())
 			{
 				const uint32 CompiledDataHash = EventData.GetValue<uint32>("CompiledDataHash");
 				if (StateTree->LastCompiledEditorDataHash == CompiledDataHash)
 				{
-					const FStateTreeInstanceDebugId InstanceId(
-						EventData.GetValue<uint32>("InstanceId"),
-						EventData.GetValue<uint32>("InstanceSerial"));
 					FString InstanceName;
 					EventData.GetString("InstanceName", InstanceName);
-					const EStateTreeTraceInstanceEventType EventType = EventData.GetValue<EStateTreeTraceInstanceEventType>("EventType");
-					Provider.AppendInstance(StateTree, InstanceId, *InstanceName, EventType);
+
+					Provider.AppendInstance(StateTree,
+						FStateTreeInstanceDebugId(EventData.GetValue<uint32>("InstanceId"), EventData.GetValue<uint32>("InstanceSerial")),
+						*InstanceName,
+						EventData.GetValue<EStateTreeTraceInstanceEventType>("EventType"));
 				}
 				else
 				{
@@ -76,67 +78,71 @@ bool FStateTreeTraceAnalyzer::OnEvent(const uint16 RouteId, EStyle Style, const
 		}
 	case RouteId_LogMessage:
 		{
-			const uint64 Cycle = EventData.GetValue<uint64>("Cycle");
-			const FStateTreeInstanceDebugId InstanceId(
-				EventData.GetValue<uint32>("InstanceId"),
-				EventData.GetValue<uint32>("InstanceSerial"));
+			FString Message;
+        	EventData.GetString("Message", Message);
+			const FStateTreeTraceLogEvent Event(EventData.GetValue<EStateTreeUpdatePhase>("Phase"), Message);
 
-			FStateTreeTraceLogEvent Event;
-			EventData.GetString("Message", Event.Message);
-
-			Provider.AppendEvent(
-				InstanceId,
-				Context.EventTime.AsSeconds(Cycle),
+			Provider.AppendEvent(FStateTreeInstanceDebugId(EventData.GetValue<uint32>("InstanceId"), EventData.GetValue<uint32>("InstanceSerial")),
+				Context.EventTime.AsSeconds(EventData.GetValue<uint64>("Cycle")),
 				FStateTreeTraceEventVariantType(TInPlaceType<FStateTreeTraceLogEvent>(), Event));
 			break;
 		}
 	case RouteId_State:
 		{
-			const uint64 Cycle = EventData.GetValue<uint64>("Cycle");
-			const FStateTreeInstanceDebugId InstanceId(
-				EventData.GetValue<uint32>("InstanceId"),
-				EventData.GetValue<uint32>("InstanceSerial"));
-
-			FStateTreeTraceStateEvent Event;
-			Event.StateIdx = EventData.GetValue<uint16>("StateIdx");
-			Event.EventType = EventData.GetValue<EStateTreeTraceNodeEventType>("EventType");
-
-			Provider.AppendEvent(
-				InstanceId,
-				Context.EventTime.AsSeconds(Cycle),
+			const FStateTreeTraceStateEvent Event(
+				EventData.GetValue<EStateTreeUpdatePhase>("Phase"),
+				EventData.GetValue<int16>("Idx"),
+				EventData.GetValue<EStateTreeTraceNodeEventType>("EventType"));
+
+			Provider.AppendEvent(FStateTreeInstanceDebugId(EventData.GetValue<uint32>("InstanceId"), EventData.GetValue<uint32>("InstanceSerial")),
+				Context.EventTime.AsSeconds(EventData.GetValue<uint64>("Cycle")),
 				FStateTreeTraceEventVariantType(TInPlaceType<FStateTreeTraceStateEvent>(), Event));
 			break;
 		}
 	case RouteId_Task:
 		{
-			const uint64 Cycle = EventData.GetValue<uint64>("Cycle");
-			const FStateTreeInstanceDebugId InstanceId(
-				EventData.GetValue<uint32>("InstanceId"),
-				EventData.GetValue<uint32>("InstanceSerial"));
-
-			FStateTreeTraceTaskEvent Event;
-			Event.TaskIdx = EventData.GetValue<uint16>("TaskIdx");
-			Event.EventType = EventData.GetValue<EStateTreeTraceNodeEventType>("EventType");
-
-			Provider.AppendEvent(
-				InstanceId,
-				Context.EventTime.AsSeconds(Cycle),
+			FString TypePath, DataAsText;
+			FMemoryReaderView Archive(EventData.GetArrayView<uint8>("DataView"));
+			Archive << TypePath;
+			Archive << DataAsText;
+
+			const FStateTreeTraceTaskEvent Event(
+				EventData.GetValue<EStateTreeUpdatePhase>("Phase"),
+				EventData.GetValue<int16>("Idx"),
+				EventData.GetValue<EStateTreeTraceNodeEventType>("EventType"),
+				EventData.GetValue<EStateTreeRunStatus>("Status"),
+				TypePath, DataAsText);
+
+			Provider.AppendEvent(FStateTreeInstanceDebugId(EventData.GetValue<uint32>("InstanceId"), EventData.GetValue<uint32>("InstanceSerial")),
+				Context.EventTime.AsSeconds(EventData.GetValue<uint64>("Cycle")),
 				FStateTreeTraceEventVariantType(TInPlaceType<FStateTreeTraceTaskEvent>(), Event));
 			break;
 		}
+	case RouteId_Condition:
+		{
+			FString TypePath, DataAsText;
+			FMemoryReaderView Archive(EventData.GetArrayView<uint8>("DataView"));
+			Archive << TypePath;
+			Archive << DataAsText;
+
+			const FStateTreeTraceConditionEvent Event(
+				EventData.GetValue<EStateTreeUpdatePhase>("Phase"),
+				EventData.GetValue<int16>("Idx"),
+				EventData.GetValue<EStateTreeTraceNodeEventType>("EventType"),
+				TypePath, DataAsText);
+			
+			Provider.AppendEvent(FStateTreeInstanceDebugId(EventData.GetValue<uint32>("InstanceId"), EventData.GetValue<uint32>("InstanceSerial")),
+				Context.EventTime.AsSeconds(EventData.GetValue<uint64>("Cycle")),
+				FStateTreeTraceEventVariantType(TInPlaceType<FStateTreeTraceConditionEvent>(), Event));
+			break;
+		}
 	case RouteId_ActiveStates:
 		{
-			const uint64 Cycle = EventData.GetValue<uint64>("Cycle");
-			const FStateTreeInstanceDebugId InstanceId(
-				EventData.GetValue<uint32>("InstanceId"),
-				EventData.GetValue<uint32>("InstanceSerial"));
-
 			FStateTreeTraceActiveStatesEvent Event;
 			Event.ActiveStates = EventData.GetArrayView<uint16>("ActiveStates");
 
-			Provider.AppendEvent(
-				InstanceId,
-				Context.EventTime.AsSeconds(Cycle),
+			Provider.AppendEvent(FStateTreeInstanceDebugId(EventData.GetValue<uint32>("InstanceId"), EventData.GetValue<uint32>("InstanceSerial")),
+				Context.EventTime.AsSeconds(EventData.GetValue<uint64>("Cycle")),
 				FStateTreeTraceEventVariantType(TInPlaceType<FStateTreeTraceActiveStatesEvent>(), Event));
 			break;
 		}
@@ -147,4 +153,4 @@ bool FStateTreeTraceAnalyzer::OnEvent(const uint16 RouteId, EStyle Style, const
 	return true;
 }
 
-#endif // WITH_STATETREE_DEBUGGER
\ No newline at end of file
+#endif // WITH_STATETREE_DEBUGGER
diff --git a/Engine/Plugins/Runtime/StateTree/Source/StateTreeModule/Private/StateTreeExecutionContext.cpp b/Engine/Plugins/Runtime/StateTree/Source/StateTreeModule/Private/StateTreeExecutionContext.cpp
--- a/Engine/Plugins/Runtime/StateTree/Source/StateTreeModule/Private/StateTreeExecutionContext.cpp
+++ b/Engine/Plugins/Runtime/StateTree/Source/StateTreeModule/Private/StateTreeExecutionContext.cpp
@@ -14,6 +14,20 @@
 #define STATETREE_LOG(Verbosity, Format, ...) UE_VLOG_UELOG(GetOwner(), LogStateTree, Verbosity, TEXT("%s") Format, *GetInstanceDescription(), ##__VA_ARGS__)
 #define STATETREE_CLOG(Condition, Verbosity, Format, ...) UE_CVLOG_UELOG((Condition), GetOwner(), LogStateTree, Verbosity, TEXT("%s") Format, *GetInstanceDescription(), ##__VA_ARGS__)
 
+#if WITH_STATETREE_DEBUGGER
+	#define STATETREE_SCOPED_PHASE(Phase) TGuardValue<EStateTreeUpdatePhase> ScopedPhase(CurrentUpdatePhase, Phase)
+    #define STATETREE_TRACE_LOG_EVENT(Format, ...) TRACE_STATETREE_LOG_EVENT(GetInstanceDebugId(), CurrentUpdatePhase, Format, ##__VA_ARGS__)
+	#define STATETREE_TRACE_STATE_EVENT(Index, EventType) TRACE_STATETREE_STATE_EVENT(GetInstanceDebugId(), CurrentUpdatePhase, Index, EventType);
+	#define STATETREE_TRACE_TASK_EVENT(Index, DataView, EventType, Status) TRACE_STATETREE_TASK_EVENT(GetInstanceDebugId(), CurrentUpdatePhase, Index, DataView, EventType, Status);
+	#define STATETREE_TRACE_CONDITION_EVENT(Index, DataViews, EventType) TRACE_STATETREE_CONDITION_EVENT(GetInstanceDebugId(), CurrentUpdatePhase, Index, DataView, EventType);	
+#else
+	#define STATETREE_SCOPED_PHASE(Phase)
+	#define STATETREE_TRACE_LOG_EVENT(Format, ...)
+	#define STATETREE_TRACE_STATE_EVENT(Index, EventType)
+	#define STATETREE_TRACE_TASK_EVENT(Index, DataView, EventType, Status)
+	#define STATETREE_TRACE_CONDITION_EVENT(Index, DataViews, EventType)
+#endif // WITH_STATETREE_DEBUGGER
+
 namespace UE::StateTree
 {
 	constexpr int32 DebugIndentSize = 2;	// Debug printing indent for hierarchical data.
@@ -159,6 +173,7 @@ void FStateTreeExecutionContext::UpdateSubtreeStateParameters(const FCompactStat
 EStateTreeRunStatus FStateTreeExecutionContext::Start()
 {
 	CSV_SCOPED_TIMING_STAT_EXCLUSIVE(StateTree_Start);
+	STATETREE_SCOPED_PHASE(EStateTreeUpdatePhase::StartTree);
 
 	if (!IsValid())
 	{
@@ -252,6 +267,7 @@ EStateTreeRunStatus FStateTreeExecutionContext::Start()
 EStateTreeRunStatus FStateTreeExecutionContext::Stop()
 {
 	CSV_SCOPED_TIMING_STAT_EXCLUSIVE(StateTree_Stop);
+	STATETREE_SCOPED_PHASE(EStateTreeUpdatePhase::StopTree);
 
 	if (!IsValid())
 	{
@@ -347,6 +363,8 @@ EStateTreeRunStatus FStateTreeExecutionContext::Tick(const float DeltaTime)
 
 	if (Exec->LastTickStatus == EStateTreeRunStatus::Running)
 	{
+		STATETREE_SCOPED_PHASE(EStateTreeUpdatePhase::TickingTasks);
+
 		// Tick tasks on active states.
 		Exec->LastTickStatus = TickTasks(DeltaTime);
 		
@@ -362,6 +380,8 @@ EStateTreeRunStatus FStateTreeExecutionContext::Tick(const float DeltaTime)
 	static constexpr int32 MaxIterations = 5;
 	for (int32 Iter = 0; Iter < MaxIterations; Iter++)
 	{
+		STATETREE_SCOPED_PHASE(EStateTreeUpdatePhase::TriggerTransitions);
+
 		// Append events accumulated during the tick, so that transitions can immediately act on them.
 		// We'll consume the events only if they lead to state change below (EnterState is treated the same as Tick),
 		// or let them be processed next frame if no transition.
@@ -469,7 +489,7 @@ void FStateTreeExecutionContext::SendEvent(const FGameplayTag Tag, const FConstS
 	}
 
 	STATETREE_LOG(Verbose, TEXT("Send Event '%s'"), *Tag.ToString());
-	TRACE_STATETREE_LOG_EVENT(GetInstanceDebugId(), TEXT("Send Event '%s'"), *Tag.ToString());
+	STATETREE_TRACE_LOG_EVENT(TEXT("Send Event '%s'"), *Tag.ToString());
 
 	FStateTreeEventQueue& EventQueue = InstanceData.GetMutableEventQueue();
 	EventQueue.SendEvent(&Owner, Tag, Payload, Origin);
@@ -478,6 +498,7 @@ void FStateTreeExecutionContext::SendEvent(const FGameplayTag Tag, const FConstS
 void FStateTreeExecutionContext::RequestTransition(const FStateTreeTransitionRequest& Request)
 {
 	CSV_SCOPED_TIMING_STAT_EXCLUSIVE(StateTree_RequestTransition);
+	STATETREE_SCOPED_PHASE(EStateTreeUpdatePhase::RequestTransition);
 
 	if (!IsValid())
 	{
@@ -684,7 +705,7 @@ EStateTreeRunStatus FStateTreeExecutionContext::EnterState(const FStateTreeTrans
 	}
 	
 	STATETREE_LOG(Log, TEXT("Enter state '%s' (%d)"), *DebugGetStatePath(Transition.NextActiveStates), Exec.StateChangeCount);
-	TRACE_STATETREE_STATE_EVENT(GetInstanceDebugId(), Transition.NextActiveStates.Last().Index, EStateTreeTraceNodeEventType::OnEnter);
+	STATETREE_TRACE_STATE_EVENT(Transition.NextActiveStates.Last().Index, EStateTreeTraceNodeEventType::OnEnter);
 
 	for (int32 Index = 0; Index < Transition.NextActiveStates.Num() && Result != EStateTreeRunStatus::Failed; Index++)
 	{
@@ -736,13 +757,18 @@ EStateTreeRunStatus FStateTreeExecutionContext::EnterState(const FStateTreeTrans
 			if (bIsEnteringState && bShouldCallStateChange)
 			{
 				STATETREE_LOG(Verbose, TEXT("%*s  Task '%s'"), Index*UE::StateTree::DebugIndentSize, TEXT(""), *Task.Name.ToString());
-				TRACE_STATETREE_TASK_EVENT(GetInstanceDebugId(), TaskIndex, EStateTreeTraceNodeEventType::OnEnter);
-				QUICK_SCOPE_CYCLE_COUNTER(StateTree_Task_EnterState);
-				CSV_SCOPED_TIMING_STAT_EXCLUSIVE(StateTree_Task_EnterState);
-				const EStateTreeRunStatus Status = Task.EnterState(*this, CurrentTransition);
+
+				EStateTreeRunStatus Status = EStateTreeRunStatus::Unset;
+				{
+					QUICK_SCOPE_CYCLE_COUNTER(StateTree_Task_EnterState);
+					CSV_SCOPED_TIMING_STAT_EXCLUSIVE(StateTree_Task_EnterState);
+				
+					Status = Task.EnterState(*this, CurrentTransition);
+				}
 
 				if (Status != EStateTreeRunStatus::Running)
 				{
+					STATETREE_TRACE_TASK_EVENT(TaskIndex, DataViews[Task.DataViewIndex.Get()], EStateTreeTraceNodeEventType::OnTaskCompleted, Status);
 					// Store the first state that completed, will be used to decide where to trigger transitions.
 					if (!Exec.CompletedStateHandle.IsValid())
 					{
@@ -857,7 +883,7 @@ void FStateTreeExecutionContext::ExitState(const FStateTreeTransitionResult& Tra
 		CurrentTransition.ChangeType = ExitedStateChangeType[Index];
 
 		STATETREE_LOG(Log, TEXT("%*sState '%s' %s"), Index*UE::StateTree::DebugIndentSize, TEXT(""), *DebugGetStatePath(Transition.CurrentActiveStates, ExitedStateActiveIndex[Index]), *UEnum::GetDisplayValueAsText(CurrentTransition.ChangeType).ToString());
-		TRACE_STATETREE_STATE_EVENT(GetInstanceDebugId(), CurrentHandle.Index, EStateTreeTraceNodeEventType::OnExit);
+		STATETREE_TRACE_STATE_EVENT(CurrentHandle.Index, EStateTreeTraceNodeEventType::OnExit);
 
 		// Tasks
 		for (int32 TaskIndex = (State.TasksBegin + State.TasksNum) - 1; TaskIndex >= State.TasksBegin; TaskIndex--)
@@ -875,7 +901,6 @@ void FStateTreeExecutionContext::ExitState(const FStateTreeTransitionResult& Tra
 				if (bShouldCallStateChange)
 				{
 					STATETREE_LOG(Verbose, TEXT("%*s  Task '%s'"), Index*UE::StateTree::DebugIndentSize, TEXT(""), *Task.Name.ToString());
-					TRACE_STATETREE_TASK_EVENT(GetInstanceDebugId(), TaskIndex, EStateTreeTraceNodeEventType::OnExit);
 					{
 						QUICK_SCOPE_CYCLE_COUNTER(StateTree_Task_ExitState);
 						CSV_SCOPED_TIMING_STAT_EXCLUSIVE(StateTree_Task_ExitState);
@@ -910,7 +935,7 @@ void FStateTreeExecutionContext::StateCompleted()
 		FCurrentlyProcessedStateScope StateScope(*this, CurrentHandle);
 		
 		STATETREE_LOG(Verbose, TEXT("%*sState '%s'"), Index*UE::StateTree::DebugIndentSize, TEXT(""), *DebugGetStatePath(Exec.ActiveStates, Index));
-		TRACE_STATETREE_STATE_EVENT(GetInstanceDebugId(), CurrentHandle.Index, EStateTreeTraceNodeEventType::OnCompleted);
+		STATETREE_TRACE_STATE_EVENT(CurrentHandle.Index, EStateTreeTraceNodeEventType::OnStateCompleted);
 
 		// Notify Tasks
 		for (int32 TaskIndex = (State.TasksBegin + State.TasksNum) - 1; TaskIndex >= State.TasksBegin; TaskIndex--)
@@ -923,7 +948,6 @@ void FStateTreeExecutionContext::StateCompleted()
 				const FStateTreeTaskBase& Task = StateTree.Nodes[TaskIndex].Get<const FStateTreeTaskBase>();
 
 				STATETREE_LOG(Verbose, TEXT("%*s  Task '%s'"), Index*UE::StateTree::DebugIndentSize, TEXT(""), *Task.Name.ToString());
-				TRACE_STATETREE_TASK_EVENT(GetInstanceDebugId(), TaskIndex, EStateTreeTraceNodeEventType::OnCompleted);
 				Task.StateCompleted(*this, Exec.LastTickStatus, Exec.ActiveStates);
 			}
 		}
@@ -933,6 +957,7 @@ void FStateTreeExecutionContext::StateCompleted()
 EStateTreeRunStatus FStateTreeExecutionContext::TickEvaluatorsAndGlobalTasks(const float DeltaTime, bool bTickGlobalTasks)
 {
 	CSV_SCOPED_TIMING_STAT_EXCLUSIVE(StateTree_TickEvaluators);
+	STATETREE_SCOPED_PHASE(EStateTreeUpdatePhase::TickingGlobalTasks);
 
 	STATETREE_CLOG(StateTree.EvaluatorsNum > 0, VeryVerbose, TEXT("Ticking Evaluators"));
 
@@ -985,22 +1010,24 @@ EStateTreeRunStatus FStateTreeExecutionContext::TickEvaluatorsAndGlobalTasks(con
 				StateTree.PropertyBindings.CopyTo(DataViews, Task.BindingsBatch, DataViews[Task.DataViewIndex.Get()]);
 			}
 
+			EStateTreeRunStatus TaskResult = EStateTreeRunStatus::Unset;
 			{
 				QUICK_SCOPE_CYCLE_COUNTER(StateTree_Task_Tick);
 				CSV_SCOPED_TIMING_STAT_EXCLUSIVE(StateTree_Task_Tick);
 
-				const EStateTreeRunStatus TaskResult = Task.Tick(*this, DeltaTime);
+				TaskResult = Task.Tick(*this, DeltaTime);
+			}
 
-				// If a global task succeeds or fails, it will stop the whole tree.
-				if (TaskResult != EStateTreeRunStatus::Running)
-				{
-					Result = TaskResult;
-				}
-					
-				if (TaskResult == EStateTreeRunStatus::Failed)
-				{
-					bShouldTickTasks = false;
-				}
+			// If a global task succeeds or fails, it will stop the whole tree.
+			if (TaskResult != EStateTreeRunStatus::Running)
+			{
+				STATETREE_TRACE_TASK_EVENT(TaskIndex, DataViews[Task.DataViewIndex.Get()], EStateTreeTraceNodeEventType::OnTaskCompleted, TaskResult);
+				Result = TaskResult;
+			}
+				
+			if (TaskResult == EStateTreeRunStatus::Failed)
+			{
+				bShouldTickTasks = false;
 			}
 		}
 	}
@@ -1011,6 +1038,7 @@ EStateTreeRunStatus FStateTreeExecutionContext::TickEvaluatorsAndGlobalTasks(con
 bool FStateTreeExecutionContext::StartEvaluatorsAndGlobalTasks(FStateTreeIndex16& OutLastInitializedTaskIndex)
 {
 	CSV_SCOPED_TIMING_STAT_EXCLUSIVE(StateTree_StartEvaluators);
+	STATETREE_SCOPED_PHASE(EStateTreeUpdatePhase::StartGlobalTasks);
 
 	STATETREE_CLOG(StateTree.EvaluatorsNum > 0 || StateTree.GlobalTasksNum > 0, Verbose, TEXT("Start Evaluators & Global tasks"));
 
@@ -1070,6 +1098,7 @@ bool FStateTreeExecutionContext::StartEvaluatorsAndGlobalTasks(FStateTreeIndex16
 void FStateTreeExecutionContext::StopEvaluatorsAndGlobalTasks(const FStateTreeIndex16 LastInitializedTaskIndex)
 {
 	CSV_SCOPED_TIMING_STAT_EXCLUSIVE(StateTree_StopEvaluators);
+	STATETREE_SCOPED_PHASE(EStateTreeUpdatePhase::StopGlobalTasks);
 
 	STATETREE_CLOG(StateTree.EvaluatorsNum > 0, Verbose, TEXT("Stop Evaluators & Global Tasks"));
 
@@ -1199,28 +1228,31 @@ EStateTreeRunStatus FStateTreeExecutionContext::TickTasks(const float DeltaTime)
 			{
 				StateTree.PropertyBindings.CopyTo(DataViews, Task.BindingsBatch, DataViews[Task.DataViewIndex.Get()]);
 			}
-			
+
+			EStateTreeRunStatus TaskResult = EStateTreeRunStatus::Unset;
 			{
 				QUICK_SCOPE_CYCLE_COUNTER(StateTree_Task_Tick);
 				CSV_SCOPED_TIMING_STAT_EXCLUSIVE(StateTree_Task_Tick);
 
-				const EStateTreeRunStatus TaskResult = Task.Tick(*this, DeltaTime);
+				TaskResult = Task.Tick(*this, DeltaTime);
+			}
 
-				// TODO: Add more control over which states can control the failed/succeeded result.
-				if (TaskResult != EStateTreeRunStatus::Running)
-				{
-					// Store the first state that completed, will be used to decide where to trigger transitions.
-					if (!Exec.CompletedStateHandle.IsValid())
-					{
-						Exec.CompletedStateHandle = CurrentHandle;
-					}
-					Result = TaskResult;
-				}
-				
-				if (TaskResult == EStateTreeRunStatus::Failed)
+			// TODO: Add more control over which states can control the failed/succeeded result.
+			if (TaskResult != EStateTreeRunStatus::Running)
+			{
+				STATETREE_TRACE_TASK_EVENT(TaskIndex, DataViews[Task.DataViewIndex.Get()], EStateTreeTraceNodeEventType::OnTaskCompleted, TaskResult);
+
+				// Store the first state that completed, will be used to decide where to trigger transitions.
+				if (!Exec.CompletedStateHandle.IsValid())
 				{
-					bShouldTickTasks = false;
+					Exec.CompletedStateHandle = CurrentHandle;
 				}
+				Result = TaskResult;
+			}
+			
+			if (TaskResult == EStateTreeRunStatus::Failed)
+			{
+				bShouldTickTasks = false;
 			}
 		}
 		NumTotalTasks += State.TasksNum;
@@ -1252,19 +1284,20 @@ bool FStateTreeExecutionContext::TestAllConditions(const int32 ConditionsOffset,
 	for (int32 Index = 0; Index < ConditionsNum; Index++)
 	{
 		const FStateTreeConditionBase& Cond = StateTree.Nodes[ConditionsOffset + Index].Get<const FStateTreeConditionBase>();
+		FStateTreeDataView& DataView = DataViews[Cond.DataViewIndex.Get()]; 
 		if (Cond.bInstanceIsObject)
 		{
-			DataViews[Cond.DataViewIndex.Get()] = SharedInstanceData->GetMutableObject(Cond.InstanceIndex.Get());
+			DataView = SharedInstanceData->GetMutableObject(Cond.InstanceIndex.Get());
 		}
 		else
 		{
-			DataViews[Cond.DataViewIndex.Get()] = SharedInstanceData->GetMutableStruct(Cond.InstanceIndex.Get());
+			DataView = SharedInstanceData->GetMutableStruct(Cond.InstanceIndex.Get());
 		}
 
 		// Copy bound properties.
 		if (Cond.BindingsBatch.IsValid())
 		{
-			if (!StateTree.PropertyBindings.CopyTo(DataViews, Cond.BindingsBatch, DataViews[Cond.DataViewIndex.Get()]))
+			if (!StateTree.PropertyBindings.CopyTo(DataViews, Cond.BindingsBatch, DataView))
 			{
 				// If the source data cannot be accessed, the whole expression evaluates to false.
 				Values[0] = false;
@@ -1273,11 +1306,12 @@ bool FStateTreeExecutionContext::TestAllConditions(const int32 ConditionsOffset,
 		}
 
 		const bool bValue = Cond.TestCondition(*this);
+		STATETREE_TRACE_CONDITION_EVENT(ConditionsOffset + Index, DataView, bValue ? EStateTreeTraceNodeEventType::Passed : EStateTreeTraceNodeEventType::Failed);	
 
 		// Reset copied properties that might contain object references.
 		if (Cond.BindingsBatch.IsValid())
 		{
-			StateTree.PropertyBindings.ResetObjects(Cond.BindingsBatch, DataViews[Cond.DataViewIndex.Get()]);
+			StateTree.PropertyBindings.ResetObjects(Cond.BindingsBatch, DataView);
 		}
 
 		const int32 DeltaIndent = Cond.DeltaIndent;
@@ -1708,6 +1742,7 @@ bool FStateTreeExecutionContext::SelectState(const FStateTreeStateHandle NextSta
 bool FStateTreeExecutionContext::SelectStateInternal(const FStateTreeStateHandle NextState, FStateTreeActiveStates& OutNewActiveState, FStateTreeActiveStates& VisitedStates)
 {
 	CSV_SCOPED_TIMING_STAT_EXCLUSIVE(StateTree_SelectState);
+	STATETREE_SCOPED_PHASE(EStateTreeUpdatePhase::StateSelection);
 
 	const FStateTreeExecutionState& Exec = GetExecState();
 
diff --git a/Engine/Plugins/Runtime/StateTree/Source/StateTreeModule/Private/StateTreeModule.cpp b/Engine/Plugins/Runtime/StateTree/Source/StateTreeModule/Private/StateTreeModule.cpp
--- a/Engine/Plugins/Runtime/StateTree/Source/StateTreeModule/Private/StateTreeModule.cpp
+++ b/Engine/Plugins/Runtime/StateTree/Source/StateTreeModule/Private/StateTreeModule.cpp
@@ -6,6 +6,7 @@
 #if WITH_STATETREE_DEBUGGER
 #include "Debugger/StateTreeTraceModule.h"
 #include "Features/IModularFeatures.h"
+#include "HAL/IConsoleManager.h"
 #include "Misc/CoreDelegates.h"
 #include "ProfilingDebugging/TraceAuxiliary.h"
 #include "StateTreeSettings.h"
@@ -52,6 +53,21 @@ class FStateTreeModule : public IStateTreeModule
 	
 	FStateTreeTraceModule StateTreeTraceModule;
 	FDelegateHandle StoreServiceHandle;
+
+	FAutoConsoleCommand EnableDebugger = FAutoConsoleCommand(
+		TEXT("statetree.enabledebugger"),
+		TEXT("Sets `bUseDebugger` to true in StateTreeSettings and turns on traces if not already active."),
+		FConsoleCommandDelegate::CreateLambda(
+			[]()
+		{
+				UStateTreeSettings& Settings = UStateTreeSettings::Get();
+				if (Settings.bUseDebugger == false)
+				{
+					Settings.bUseDebugger = true;
+					FStateTreeModule& StateTreeModule = FModuleManager::GetModuleChecked<FStateTreeModule>("StateTree");
+					StateTreeModule.StartTraces();	
+				}
+		}));
 #endif // WITH_STATETREE_DEBUGGER
 };
 
diff --git a/Engine/Plugins/Runtime/StateTree/Source/StateTreeModule/Public/Debugger/StateTreeDebugger.h b/Engine/Plugins/Runtime/StateTree/Source/StateTreeModule/Public/Debugger/StateTreeDebugger.h
--- a/Engine/Plugins/Runtime/StateTree/Source/StateTreeModule/Public/Debugger/StateTreeDebugger.h
+++ b/Engine/Plugins/Runtime/StateTree/Source/StateTreeModule/Public/Debugger/StateTreeDebugger.h
@@ -26,14 +26,14 @@ namespace UE::StateTreeDebugger
 	struct FFrameIndexSpan
 	{
 		FFrameIndexSpan() = default;
-		FFrameIndexSpan(const uint64 FrameIdx, const uint32 EventIdx)
-			: FrameIdx(FrameIdx)
+		FFrameIndexSpan(const TraceServices::FFrame& Frame, const uint32 EventIdx)
+			: Frame(Frame)
 			, EventIdx(EventIdx)
 		{
 		}
 
 		/** Frame index in the analysis session */
-		uint64 FrameIdx = INDEX_NONE;
+		TraceServices::FFrame Frame;
 
 		/** Index of the first event for that Frame index */
 		int32 EventIdx = INDEX_NONE;
@@ -76,6 +76,7 @@ DECLARE_DELEGATE_TwoParams(FOnStateTreeDebuggerTracesUpdated, TConstArrayView<FS
 DECLARE_DELEGATE_TwoParams(FOnStateTreeDebuggerBreakpointHit, FStateTreeInstanceDebugId InstanceId, FStateTreeStateHandle StateHandle);
 DECLARE_DELEGATE_OneParam(FOnStateTreeDebuggerBreakpointsChanged, TConstArrayView<FStateTreeStateHandle> Breakpoints);
 DECLARE_DELEGATE_OneParam(FOnStateTreeDebuggerActiveStatesChanges, TConstArrayView<FStateTreeStateHandle> ActiveStates);
+DECLARE_DELEGATE(FOnStateTreeDebuggerDebuggedInstanceSet);
 
 struct STATETREEMODULE_API FStateTreeDebugger : FTickableGameObject
 {
@@ -114,6 +115,7 @@ struct STATETREEMODULE_API FStateTreeDebugger : FTickableGameObject
 	void StepForward();
 	
 	bool IsAnalysisSessionActive() const { return GetAnalysisSession() != nullptr; }
+	const TraceServices::IAnalysisSession* GetAnalysisSession() const;
 
 	TConstArrayView<FStateTreeDebuggerInstanceDesc> GetActiveInstances() const;
 	const FStateTreeDebuggerInstanceDesc& GetDebuggedInstance() const;
@@ -130,6 +132,7 @@ struct STATETREEMODULE_API FStateTreeDebugger : FTickableGameObject
 	
 	void ToggleBreakpoints(const TConstArrayView<FStateTreeStateHandle> SelectedStates);
 
+	FOnStateTreeDebuggerDebuggedInstanceSet OnDebuggedInstanceSet;
 	FOnStateTreeDebuggerTracesUpdated OnTracesUpdated;
 	FOnStateTreeDebuggerBreakpointHit OnBreakpointHit;
 	FOnStateTreeDebuggerBreakpointsChanged OnBreakpointsChanged;
@@ -158,7 +161,6 @@ private:
 
 	void SetActiveStates(const TConstArrayView<FStateTreeStateHandle> NewActiveStates);
 	
-	const TraceServices::IAnalysisSession* GetAnalysisSession() const;
 	UE::Trace::FStoreClient* GetStoreClient() const;
 	void GetSessionActiveInstances(TArray<FStateTreeDebuggerInstanceDesc>& OutInstances) const;
 
diff --git a/Engine/Plugins/Runtime/StateTree/Source/StateTreeModule/Public/Debugger/StateTreeTrace.h b/Engine/Plugins/Runtime/StateTree/Source/StateTreeModule/Public/Debugger/StateTreeTrace.h
--- a/Engine/Plugins/Runtime/StateTree/Source/StateTreeModule/Public/Debugger/StateTreeTrace.h
+++ b/Engine/Plugins/Runtime/StateTree/Source/StateTreeModule/Public/Debugger/StateTreeTrace.h
@@ -7,9 +7,12 @@
 #include "Trace/Trace.h"
 
 class UStateTree;
+struct FStateTreeDataView;
 struct FStateTreeActiveStates;
 struct FStateTreeInstanceDebugId;
+enum class EStateTreeRunStatus : uint8;
 enum class EStateTreeTraceNodeEventType : uint8;
+enum class EStateTreeUpdatePhase : uint8;
 enum class EStateTreeTraceInstanceEventType : uint8;
 
 UE_TRACE_CHANNEL_EXTERN(StateTreeDebugChannel, STATETREEMODULE_API)
@@ -17,11 +20,12 @@ UE_TRACE_CHANNEL_EXTERN(StateTreeDebugChannel, STATETREEMODULE_API)
 namespace UE::StateTreeTrace
 {
 	void OutputInstanceLifetimeEvent(const FStateTreeInstanceDebugId InstanceId, const UStateTree* StateTree, const TCHAR* InstanceName, const EStateTreeTraceInstanceEventType EventType);
-	void OutputLogEventTrace(const FStateTreeInstanceDebugId InstanceId, const TCHAR* Fmt, ...);
-	void OutputStateEventTrace(const FStateTreeInstanceDebugId InstanceId, const uint16 StateIdx, const EStateTreeTraceNodeEventType EventType);
-	void OutputTaskEventTrace(const FStateTreeInstanceDebugId InstanceId, const uint16 TaskIdx, const EStateTreeTraceNodeEventType EventType);
+	void OutputLogEventTrace(const FStateTreeInstanceDebugId InstanceId, const EStateTreeUpdatePhase Phase, const TCHAR* Fmt, ...);
+	void OutputStateEventTrace(const FStateTreeInstanceDebugId InstanceId, const EStateTreeUpdatePhase Phase, const uint16 StateIdx, const EStateTreeTraceNodeEventType EventType);
+	void OutputTaskEventTrace(const FStateTreeInstanceDebugId InstanceId, const EStateTreeUpdatePhase Phase, const uint16 TaskIdx, FStateTreeDataView DataView, const EStateTreeTraceNodeEventType EventType, const EStateTreeRunStatus Status);
+	void OutputConditionEventTrace(const FStateTreeInstanceDebugId InstanceId, const EStateTreeUpdatePhase Phase, const uint16 ConditionIdx, FStateTreeDataView DataView, const EStateTreeTraceNodeEventType EventType);
 	void OutputActiveStatesEventTrace(const FStateTreeInstanceDebugId InstanceId, const FStateTreeActiveStates& ActiveStates);
-};
+}
 
 #define TRACE_STATETREE_INSTANCE_EVENT(InstanceID, StateTree, InstanceName, EventType) \
 	if (UE_TRACE_CHANNELEXPR_IS_ENABLED(StateTreeDebugChannel)) \
@@ -29,22 +33,28 @@ namespace UE::StateTreeTrace
 		UE::StateTreeTrace::OutputInstanceLifetimeEvent(InstanceID, StateTree, InstanceName, EventType); \
 	}
 
-#define TRACE_STATETREE_LOG_EVENT(InstanceId, Format, ...) \
+#define TRACE_STATETREE_LOG_EVENT(InstanceId, Phase, Format, ...) \
 	if (UE_TRACE_CHANNELEXPR_IS_ENABLED(StateTreeDebugChannel)) \
 	{ \
-		UE::StateTreeTrace::OutputLogEventTrace(InstanceId, Format, ##__VA_ARGS__); \
+		UE::StateTreeTrace::OutputLogEventTrace(InstanceId, Phase, Format, ##__VA_ARGS__); \
 	}
 
-#define TRACE_STATETREE_STATE_EVENT(InstanceId, StateIdx, EventType) \
+#define TRACE_STATETREE_STATE_EVENT(InstanceId, Phase, StateIdx, EventType) \
 	if (UE_TRACE_CHANNELEXPR_IS_ENABLED(StateTreeDebugChannel)) \
 	{ \
-		UE::StateTreeTrace::OutputStateEventTrace(InstanceId, StateIdx, EventType); \
+		UE::StateTreeTrace::OutputStateEventTrace(InstanceId, Phase, StateIdx, EventType); \
 	}
 
-#define TRACE_STATETREE_TASK_EVENT(InstanceId, TaskIdx, EventType) \
+#define TRACE_STATETREE_TASK_EVENT(InstanceId, Phase, TaskIdx, DataView, EventType, Status) \
 	if (UE_TRACE_CHANNELEXPR_IS_ENABLED(StateTreeDebugChannel)) \
 	{ \
-		UE::StateTreeTrace::OutputTaskEventTrace(InstanceId, TaskIdx, EventType); \
+		UE::StateTreeTrace::OutputTaskEventTrace(InstanceId, Phase, TaskIdx, DataView, EventType, Status); \
+	}
+
+#define TRACE_STATETREE_CONDITION_EVENT(InstanceId, Phase, ConditionIdx, DataView, EventType) \
+	if (UE_TRACE_CHANNELEXPR_IS_ENABLED(StateTreeDebugChannel)) \
+	{ \
+		UE::StateTreeTrace::OutputConditionEventTrace(InstanceId, Phase, ConditionIdx, DataView, EventType); \
 	}
 
 #define TRACE_STATETREE_ACTIVE_STATES_EVENT(InstanceId, ActivateStates) \
@@ -56,9 +66,10 @@ namespace UE::StateTreeTrace
 #else //STATETREE_DEBUG_TRACE_ENABLED
 
 #define TRACE_STATETREE_INSTANCE_EVENT(InstanceID, StateTree, InstanceName, EventType)
-#define TRACE_STATETREE_LOG_EVENT(InstanceId, Format, ...)
-#define TRACE_STATETREE_STATE_EVENT(InstanceId, StateIdx, EventType)
-#define TRACE_STATETREE_TASK_EVENT(InstanceId, TaskIdx, EventType)
+#define TRACE_STATETREE_LOG_EVENT(InstanceId, Phase, Format, ...)
+#define TRACE_STATETREE_STATE_EVENT(InstanceId, Phase, StateIdx, EventType)
+#define TRACE_STATETREE_TASK_EVENT(InstanceId, Phase, TaskIdx, DataView, EventType, Status)
+#define TRACE_STATETREE_CONDITION_EVENT(InstanceId, Phase, ConditionIdx, DataView, EventType)
 #define TRACE_STATETREE_ACTIVE_STATES_EVENT(InstanceId, ActivateStates)
 
 #endif // STATETREE_DEBUG_TRACE_ENABLED
diff --git a/Engine/Plugins/Runtime/StateTree/Source/StateTreeModule/Public/Debugger/StateTreeTraceAnalyzer.h b/Engine/Plugins/Runtime/StateTree/Source/StateTreeModule/Public/Debugger/StateTreeTraceAnalyzer.h
--- a/Engine/Plugins/Runtime/StateTree/Source/StateTreeModule/Public/Debugger/StateTreeTraceAnalyzer.h
+++ b/Engine/Plugins/Runtime/StateTree/Source/StateTreeModule/Public/Debugger/StateTreeTraceAnalyzer.h
@@ -25,6 +25,7 @@ private:
 		RouteId_LogMessage,
 		RouteId_State,
 		RouteId_Task,
+		RouteId_Condition,
 		RouteId_ActiveStates
 	};
 
diff --git a/Engine/Plugins/Runtime/StateTree/Source/StateTreeModule/Public/Debugger/StateTreeTraceTypes.h b/Engine/Plugins/Runtime/StateTree/Source/StateTreeModule/Public/Debugger/StateTreeTraceTypes.h
--- a/Engine/Plugins/Runtime/StateTree/Source/StateTreeModule/Public/Debugger/StateTreeTraceTypes.h
+++ b/Engine/Plugins/Runtime/StateTree/Source/StateTreeModule/Public/Debugger/StateTreeTraceTypes.h
@@ -6,6 +6,8 @@
 #include "StateTreeTraceTypes.generated.h"
 
 struct FStateTreeStateHandle;
+enum class EStateTreeRunStatus : uint8;
+enum class EStateTreeUpdatePhase : uint8;
 
 UENUM()
 enum class EStateTreeTraceInstanceEventType : uint8
@@ -18,28 +20,84 @@ UENUM()
 enum class EStateTreeTraceNodeEventType : uint8
 {
 	Unset,
-	OnEnter,
-	OnExit,
-	OnCompleted
+	OnEnter,			// for State + Task events
+	OnExit,				// for State + Task events
+	OnStateCompleted,	// for State events
+	OnTaskCompleted,	// for Task events
+	Passed,				// for Condition events
+	Failed				// for Condition events
 };
 
 #if WITH_STATETREE_DEBUGGER
 
-struct FStateTreeTraceLogEvent
+struct FStateTreeTracePhaseEvent
 {
+	explicit FStateTreeTracePhaseEvent(const EStateTreeUpdatePhase Phase)
+		: Phase(Phase)
+	{
+	}
+
+	EStateTreeUpdatePhase Phase;
+};
+
+struct FStateTreeTraceLogEvent : FStateTreeTracePhaseEvent
+{
+	explicit FStateTreeTraceLogEvent(const EStateTreeUpdatePhase Phase, const FString& Message)
+		: FStateTreeTracePhaseEvent(Phase)
+		, Message(Message)
+	{
+	}
+
 	FString Message;
 };
 
-struct FStateTreeTraceStateEvent
+struct FStateTreeTraceNodeEvent : FStateTreeTracePhaseEvent
 {
-	int16 StateIdx = INDEX_NONE;
+	explicit FStateTreeTraceNodeEvent(const EStateTreeUpdatePhase Phase, const int16 Idx, const EStateTreeTraceNodeEventType EventType)
+		: FStateTreeTracePhaseEvent(Phase)
+		, Idx(Idx)
+		, EventType(EventType)
+	{
+	}
+
+	int16 Idx = INDEX_NONE;
 	EStateTreeTraceNodeEventType EventType = EStateTreeTraceNodeEventType::Unset;
 };
 
-struct FStateTreeTraceTaskEvent
+struct FStateTreeTraceStateEvent : FStateTreeTraceNodeEvent
 {
-	int16 TaskIdx = INDEX_NONE;
-	EStateTreeTraceNodeEventType EventType = EStateTreeTraceNodeEventType::Unset;
+	explicit FStateTreeTraceStateEvent(const EStateTreeUpdatePhase Phase, const int16 Idx, const EStateTreeTraceNodeEventType EventType)
+		: FStateTreeTraceNodeEvent(Phase, Idx, EventType)
+	{
+	}
+};
+
+struct FStateTreeTraceTaskEvent : FStateTreeTraceNodeEvent
+{
+	explicit FStateTreeTraceTaskEvent(const EStateTreeUpdatePhase Phase, const int16 Idx,	const EStateTreeTraceNodeEventType EventType, const EStateTreeRunStatus Status,	const FString& TypePath, const FString& InstanceDataAsText)
+		: FStateTreeTraceNodeEvent(Phase, Idx, EventType)
+		, TypePath(TypePath)
+		, InstanceDataAsText(InstanceDataAsText)
+		, Status(Status)
+	{
+	}
+
+	FString TypePath;
+	FString InstanceDataAsText;
+	EStateTreeRunStatus Status;	
+};
+
+struct FStateTreeTraceConditionEvent : FStateTreeTraceNodeEvent
+{
+	explicit FStateTreeTraceConditionEvent(const EStateTreeUpdatePhase Phase, const int16 Idx, const EStateTreeTraceNodeEventType EventType, const FString& TypePath, const FString& InstanceDataAsText)
+		: FStateTreeTraceNodeEvent(Phase, Idx, EventType)
+		, TypePath(TypePath)
+		, InstanceDataAsText(InstanceDataAsText)
+	{
+	}
+
+	FString TypePath;
+	FString InstanceDataAsText;
 };
 
 struct FStateTreeTraceActiveStatesEvent
@@ -48,6 +106,6 @@ struct FStateTreeTraceActiveStatesEvent
 };
 
 /** Type aliases for statetree trace events */
-using FStateTreeTraceEventVariantType = TVariant<FStateTreeTraceLogEvent, FStateTreeTraceStateEvent, FStateTreeTraceTaskEvent, FStateTreeTraceActiveStatesEvent>;
+using FStateTreeTraceEventVariantType = TVariant<FStateTreeTraceLogEvent, FStateTreeTraceNodeEvent, FStateTreeTraceStateEvent, FStateTreeTraceTaskEvent, FStateTreeTraceConditionEvent, FStateTreeTraceActiveStatesEvent>;
 
 #endif // WITH_STATETREE_DEBUGGER
diff --git a/Engine/Plugins/Runtime/StateTree/Source/StateTreeModule/Public/StateTreeExecutionContext.h b/Engine/Plugins/Runtime/StateTree/Source/StateTreeModule/Public/StateTreeExecutionContext.h
--- a/Engine/Plugins/Runtime/StateTree/Source/StateTreeModule/Public/StateTreeExecutionContext.h
+++ b/Engine/Plugins/Runtime/StateTree/Source/StateTreeModule/Public/StateTreeExecutionContext.h
@@ -6,6 +6,12 @@
 #include "StateTreeNodeBase.h"
 #include "Experimental/ConcurrentLinearAllocator.h"
 
+#if WITH_STATETREE_DEBUGGER
+#include "Debugger/StateTreeTraceTypes.h"
+#endif // WITH_STATETREE_DEBUGGER
+
+#include "StateTreeExecutionContext.generated.h"
+
 struct FGameplayTag;
 struct FInstancedPropertyBag;
 
@@ -16,6 +22,27 @@ struct FStateTreeEvent;
 struct FStateTreeTransitionRequest;
 struct FStateTreeInstanceDebugId;
 
+/**
+ * Enumeration for the different update phases.
+ * This is used as context information when tracing debug events.   
+ */
+UENUM()
+enum class EStateTreeUpdatePhase : uint8
+{
+	Unset,
+	TickingGlobalTasks,
+	StateSelection,
+	TickingTasks,
+	StateTransition,
+	TriggerTransitions,
+	RequestTransition,
+	StartGlobalTasks,
+	StopGlobalTasks,
+	StartTree,
+	StopTree
+};
+
+
 /**
  * StateTree Execution Context is a helper that is used to update and access StateTree instance data.
  *
@@ -440,6 +467,11 @@ protected:
 	/** Current state we're processing, or invalid if not applicable. */
 	FStateTreeStateHandle CurrentlyProcessedState;
 
+#if WITH_STATETREE_DEBUGGER
+	/** Current phase in the update. Used for debugging events */
+	EStateTreeUpdatePhase CurrentUpdatePhase = EStateTreeUpdatePhase::Unset;
+#endif // WITH_STATETREE_DEBUGGER
+
 	/** True if transitions are allowed to be requested directly instead of buffering. */
 	bool bAllowDirectTransitions = false;
 
