StateTree: Fix a potential crash in FStateTreePropertyBindings::ResolvePath().

diff --git a/Engine/Plugins/Runtime/StateTree/Source/StateTreeModule/Private/StateTreePropertyBindings.cpp b/Engine/Plugins/Runtime/StateTree/Source/StateTreeModule/Private/StateTreePropertyBindings.cpp
--- a/Engine/Plugins/Runtime/StateTree/Source/StateTreeModule/Private/StateTreePropertyBindings.cpp
+++ b/Engine/Plugins/Runtime/StateTree/Source/StateTreeModule/Private/StateTreePropertyBindings.cpp
@@ -269,8 +269,8 @@ bool FStateTreePropertyBindings::ResolvePath(const UStruct* Struct, const FState
 		for (int32 Index = 1; Index < TempIndirections.Num(); Index++)
 		{
 			const int32 IndirectionIndex = PropertyIndirections.Num();
+			PrevIndirection->NextIndex = FStateTreeIndex16(IndirectionIndex); // Set PrevIndirection before array add, as it can invalidate the pointer.
 			FStateTreePropertyIndirection& NewIndirection = PropertyIndirections.Add_GetRef(TempIndirections[Index]);
-			PrevIndirection->NextIndex = FStateTreeIndex16(IndirectionIndex);
 			PrevIndirection = &NewIndirection;
 		}
 	}
