Fixed mismatched copy constructor/assignment operator in FStateTreeEditorPropertyPath.

diff --git a/Engine/Plugins/Runtime/StateTree/Source/StateTreeModule/Public/StateTreePropertyBindings.h b/Engine/Plugins/Runtime/StateTree/Source/StateTreeModule/Public/StateTreePropertyBindings.h
--- a/Engine/Plugins/Runtime/StateTree/Source/StateTreeModule/Public/StateTreePropertyBindings.h
+++ b/Engine/Plugins/Runtime/StateTree/Source/StateTreeModule/Public/StateTreePropertyBindings.h
@@ -414,7 +414,6 @@ struct UE_DEPRECATED(5.3, "Use FStateTreePropertyPath instead.") STATETREEMODULE
 
 	FStateTreeEditorPropertyPath() = default;
 	FStateTreeEditorPropertyPath(const FGuid& InStructID, const TCHAR* PropertyName) : StructID(InStructID) { Path.Add(PropertyName); }
-	FStateTreeEditorPropertyPath(const FStateTreeEditorPropertyPath& InPath) : StructID(InPath.StructID), Path(InPath.Path) {}
 	
 	/**
 	 * Returns the property path as a one string. Highlight allows to decorate a specific segment.
