[StateTree] moved command list used by the tree view to the StateTreeEditor. Stored there, so that other windows (e.g. debugger) can add commands to it, even if the tree view is not spawned yet.

diff --git a/Engine/Plugins/Runtime/StateTree/Source/StateTreeEditorModule/Private/SStateTreeView.cpp b/Engine/Plugins/Runtime/StateTree/Source/StateTreeEditorModule/Private/SStateTreeView.cpp
--- a/Engine/Plugins/Runtime/StateTree/Source/StateTreeEditorModule/Private/SStateTreeView.cpp
+++ b/Engine/Plugins/Runtime/StateTree/Source/StateTreeEditorModule/Private/SStateTreeView.cpp
@@ -34,7 +34,7 @@ SStateTreeView::~SStateTreeView()
 	}
 }
 
-void SStateTreeView::Construct(const FArguments& InArgs, TSharedRef<FStateTreeViewModel> InStateTreeViewModel)
+void SStateTreeView::Construct(const FArguments& InArgs, TSharedRef<FStateTreeViewModel> InStateTreeViewModel, const TSharedRef<FUICommandList>& InCommandList)
 {
 	StateTreeViewModel = InStateTreeViewModel;
 
@@ -130,7 +130,7 @@ void SStateTreeView::Construct(const FArguments& InArgs, TSharedRef<FStateTreeVi
 
 	UpdateTree(true);
 
-	CommandList = MakeShared<FUICommandList>();
+	CommandList = InCommandList;
 	BindCommands();
 }
 
diff --git a/Engine/Plugins/Runtime/StateTree/Source/StateTreeEditorModule/Private/SStateTreeView.h b/Engine/Plugins/Runtime/StateTree/Source/StateTreeEditorModule/Private/SStateTreeView.h
--- a/Engine/Plugins/Runtime/StateTree/Source/StateTreeEditorModule/Private/SStateTreeView.h
+++ b/Engine/Plugins/Runtime/StateTree/Source/StateTreeEditorModule/Private/SStateTreeView.h
@@ -47,7 +47,7 @@ public:
 	SStateTreeView();
 	virtual ~SStateTreeView() override;
 
-	void Construct(const FArguments& InArgs, TSharedRef<FStateTreeViewModel> StateTreeViewModel);
+	void Construct(const FArguments& InArgs, TSharedRef<FStateTreeViewModel> StateTreeViewModel, const TSharedRef<FUICommandList>& InCommandList);
 
 	void SavePersistentExpandedStates();
 
diff --git a/Engine/Plugins/Runtime/StateTree/Source/StateTreeEditorModule/Private/StateTreeEditor.cpp b/Engine/Plugins/Runtime/StateTree/Source/StateTreeEditorModule/Private/StateTreeEditor.cpp
--- a/Engine/Plugins/Runtime/StateTree/Source/StateTreeEditorModule/Private/StateTreeEditor.cpp
+++ b/Engine/Plugins/Runtime/StateTree/Source/StateTreeEditorModule/Private/StateTreeEditor.cpp
@@ -296,7 +296,7 @@ TSharedRef<SDockTab> FStateTreeEditor::SpawnTab_StateTreeView(const FSpawnTabArg
 		.Label(NSLOCTEXT("StateTreeEditor", "StateTreeViewTab", "StateTree"))
 		.TabColorScale(GetTabColorScale())
 		[
-			SAssignNew(StateTreeView, SStateTreeView, StateTreeViewModel.ToSharedRef())
+			SAssignNew(StateTreeView, SStateTreeView, StateTreeViewModel.ToSharedRef(), TreeViewCommandList)
 		];
 }
 
@@ -387,7 +387,7 @@ TSharedRef<SDockTab> FStateTreeEditor::SpawnTab_Debugger(const FSpawnTabArgs& Ar
 		.Label(LOCTEXT("DebuggerTitle", "Debugger"))
 		.TabColorScale(GetTabColorScale())
 		[
-			SAssignNew(DebuggerView, SStateTreeDebuggerView, StateTree, StateTreeViewModel.ToSharedRef(), ToolkitCommands)
+			SAssignNew(DebuggerView, SStateTreeDebuggerView, StateTree, StateTreeViewModel.ToSharedRef(), TreeViewCommandList)
 		];
 }
 #endif // WITH_STATETREE_DEBUGGER
diff --git a/Engine/Plugins/Runtime/StateTree/Source/StateTreeEditorModule/Private/StateTreeEditor.h b/Engine/Plugins/Runtime/StateTree/Source/StateTreeEditorModule/Private/StateTreeEditor.h
--- a/Engine/Plugins/Runtime/StateTree/Source/StateTreeEditorModule/Private/StateTreeEditor.h
+++ b/Engine/Plugins/Runtime/StateTree/Source/StateTreeEditorModule/Private/StateTreeEditor.h
@@ -26,6 +26,7 @@ namespace UE::StateTree::Editor
 class FStateTreeEditor : public IStateTreeEditor, public FSelfRegisteringEditorUndoClient, public FGCObject
 {
 public:
+	FStateTreeEditor() : TreeViewCommandList(new FUICommandList()) {}
 
 	virtual void RegisterTabSpawners(const TSharedRef<class FTabManager>& TabManager) override;
 	virtual void UnregisterTabSpawners(const TSharedRef<class FTabManager>& TabManager) override;
@@ -89,11 +90,14 @@ private:
 
 	FText GetStatisticsText() const;
 
-	/* State Tree being edited */
+	/** State Tree being edited */
 	UStateTree* StateTree = nullptr;
 
 	uint32 EditorDataHash = 0;
 	bool bLastCompileSucceeded = true;
+
+	/** The command list used by the tree view. Stored here, so that other windows (e.g. debugger) can add commands to it, even if the tree view is not spawned yet. */
+	TSharedRef<FUICommandList> TreeViewCommandList;
 	
 	/** Selection Property View */
 	TSharedPtr<class IDetailsView> SelectionDetailsView;
