[StateTree] fixed circular pointer dependency preventing proper cleanup when closing the StateTree asset editor.

diff --git a/Engine/Plugins/Runtime/StateTree/Source/StateTreeEditorModule/Private/SStateTreeViewRow.cpp b/Engine/Plugins/Runtime/StateTree/Source/StateTreeEditorModule/Private/SStateTreeViewRow.cpp
--- a/Engine/Plugins/Runtime/StateTree/Source/StateTreeEditorModule/Private/SStateTreeViewRow.cpp
+++ b/Engine/Plugins/Runtime/StateTree/Source/StateTreeEditorModule/Private/SStateTreeViewRow.cpp
@@ -33,17 +33,21 @@ void SStateTreeViewRow::Construct(const FArguments& InArgs, const TSharedRef<STa
 	static const FLinearColor TasksBackground = FLinearColor(FColor(17, 117, 131));
 	static const FLinearColor LinkBackground = FLinearColor(FColor(84, 84, 84));
 	static const FLinearColor IconTint = FLinearColor(1, 1, 1, 0.5f);
-	static const FLinearColor TypeIconTint = FLinearColor(0, 0, 0, 0.5f);
 
 	this->ChildSlot
 	.HAlign(HAlign_Fill)
 	[
 		SNew(SBox)
-	    .MinDesiredWidth_Lambda([ViewBox]()
+		.MinDesiredWidth_Lambda([WeakOwnerViewBox = ViewBox.ToWeakPtr()]()
 			{
-    			// Make the row at least as wide as the view.
-	    		// The -1 is needed or we'll see a scrollbar. 
-				return ViewBox->GetTickSpaceGeometry().GetLocalSize().X - 1;
+				// Captured as weak ptr so we don't prevent our parent widget from being destroyed (circular pointer reference).
+				if (const TSharedPtr<SScrollBox> OwnerViewBox = WeakOwnerViewBox.Pin())
+				{
+					// Make the row at least as wide as the view.
+					// The -1 is needed or we'll see a scrollbar.
+					return OwnerViewBox->GetTickSpaceGeometry().GetLocalSize().X - 1;
+				}
+				return 0.f;
 			})
 		[
 			SNew(SHorizontalBox)
@@ -777,7 +781,7 @@ FText SStateTreeViewRow::GetTransitionsIcon(const UStateTreeState& State, const
 
 EVisibility SStateTreeViewRow::GetTransitionsVisibility(const UStateTreeState& State, const EStateTreeTransitionTrigger Trigger) const
 {
-	// Handle completed, succeeeded and failed transitions.
+	// Handle completed, succeeded and failed transitions.
 	if (EnumHasAnyFlags(Trigger, EStateTreeTransitionTrigger::OnStateCompleted))
 	{
 		const bool bIsLeafState = (State.Children.Num() == 0);
