[StateTree] debugger update
- moved debugger related files to a dedicated directory
- trace provider now allows to read all timelines associated to a given state tree asset instance
- added multiple slate widgets to manage tracks and timelines (directly inspired from RewindDebugger to make it easier to integrate the tools together eventually).
- added synchronized lists of instance names and timelines
- converted EStateTreeUpdatePhase to enum flags used to build frame details hierarchically

diff --git a/Engine/Plugins/Runtime/StateTree/Source/StateTreeEditorModule/Private/Debugger/SStateTreeDebuggerEventTimelineView.cpp b/Engine/Plugins/Runtime/StateTree/Source/StateTreeEditorModule/Private/Debugger/SStateTreeDebuggerEventTimelineView.cpp
new file mode 100644
--- /dev/null
+++ b/Engine/Plugins/Runtime/StateTree/Source/StateTreeEditorModule/Private/Debugger/SStateTreeDebuggerEventTimelineView.cpp
@@ -0,0 +1,134 @@
+// Copyright Epic Games, Inc. All Rights Reserved.
+
+#include "SStateTreeDebuggerEventTimelineView.h"
+#include "Framework/Application/SlateApplication.h"
+#include "Rendering/DrawElements.h"
+#include "SSimpleTimeSlider.h"
+
+#define LOCTEXT_NAMESPACE "SStateTreeDebuggerEventTimelineView"
+
+
+void SStateTreeDebuggerEventTimelineView::Construct(const FArguments& InArgs)
+{
+	ViewRange = InArgs._ViewRange;
+	DesiredSize = InArgs._DesiredSize;
+	EventData = InArgs._EventData;
+
+	EventBrush = FAppStyle::GetBrush("Sequencer.KeyDiamond");
+	EventBorderBrush = FAppStyle::GetBrush("Sequencer.KeyDiamondBorder");
+}
+
+int32 SStateTreeDebuggerEventTimelineView::OnPaint(const FPaintArgs& Args, const FGeometry& AllottedGeometry, const FSlateRect& MyCullingRect, FSlateWindowElementList& OutDrawElements, const int32 LayerId, const FWidgetStyle& InWidgetStyle, const bool bParentEnabled) const
+{
+	const int32 NewLayer = PaintEvents(AllottedGeometry, MyCullingRect, OutDrawElements, LayerId, InWidgetStyle, bParentEnabled);
+
+	return FMath::Max(NewLayer, SCompoundWidget::OnPaint(Args, AllottedGeometry, MyCullingRect, OutDrawElements, NewLayer, InWidgetStyle, ShouldBeEnabled(bParentEnabled)));
+}
+
+FVector2D SStateTreeDebuggerEventTimelineView::ComputeDesiredSize(float) const
+{
+	return DesiredSize.Get();
+}
+
+int32 SStateTreeDebuggerEventTimelineView::PaintEvents(const FGeometry& AllottedGeometry, const FSlateRect& MyCullingRect, FSlateWindowElementList& OutDrawElements, int32 LayerId, const FWidgetStyle& InWidgetStyle, bool bParentEnabled) const
+{
+	const TRange<double> DebugTimeRange = ViewRange.Get();
+
+	const SSimpleTimeSlider::FScrubRangeToScreen RangeToScreen(DebugTimeRange, AllottedGeometry.GetLocalSize());
+
+	if (const TSharedPtr<FTimelineEventData> TimelineEventData = EventData.Get())
+	{
+		const FSlateBrush* Brush = FAppStyle::GetBrush("Sequencer.SectionArea.Background");
+
+		const int NumWindows = TimelineEventData->Windows.Num();
+		if (NumWindows > 0)
+		{
+			for (int i = 0; i < NumWindows; i++)
+			{
+				const FTimelineEventData::EventWindow& Window = TimelineEventData->Windows[i];
+
+				FVector2D EventSize = EventBrush->GetImageSize();
+
+				const float XStart = RangeToScreen.InputToLocalX(Window.TimeStart);
+				const float XEnd = RangeToScreen.InputToLocalX(Window.TimeEnd);
+				const float Y = (AllottedGeometry.Size.Y - EventSize.Y) / 2;
+
+				// window bar
+				FSlateDrawElement::MakeBox
+				(OutDrawElements,
+					LayerId++,
+					AllottedGeometry.ToPaintGeometry(FVector2D(XEnd - XStart, EventSize.Y - 2), FSlateLayoutTransform(FVector2D(XStart, Y + 1))),
+					Brush,
+					ESlateDrawEffect::None,
+					Window.Color);
+
+				if (Window.Description.IsEmpty() == false)
+				{
+					FSlateDrawElement::MakeText
+					(
+						OutDrawElements,
+						LayerId,
+						AllottedGeometry.ToPaintGeometry(EventSize, FSlateLayoutTransform(FVector2D(XStart + 8, Y))),
+						Window.Description,
+						FCoreStyle::GetDefaultFontStyle("Regular", 8),
+						ESlateDrawEffect::None,
+						FLinearColor::Black);
+				}
+			}
+		}
+
+		const int NumPoints = TimelineEventData->Points.Num();
+
+		double PrevPointTime = 0;
+		int OverlappingPointCount = 0;
+
+		if (NumPoints > 0)
+		{
+			for (int i = 0; i < NumPoints; i++)
+			{
+				const FTimelineEventData::EventPoint& Point = TimelineEventData->Points[i];
+
+				FVector2D EventSize = EventBrush->GetImageSize();
+
+				float X = RangeToScreen.InputToLocalX(Point.Time);
+				X = X - EventSize.X / 2;
+				float Y = (AllottedGeometry.Size.Y - EventSize.Y) / 2;
+				if (Point.Time == PrevPointTime)
+				{
+					OverlappingPointCount++;
+
+					static constexpr int OverlapOffsetAmount = 2;
+					Y += OverlapOffsetAmount * OverlappingPointCount;
+				}
+				else
+				{
+					OverlappingPointCount = 0;
+				}
+
+				PrevPointTime = Point.Time;
+
+				FSlateDrawElement::MakeBox
+				(OutDrawElements,
+					LayerId,
+					AllottedGeometry.ToPaintGeometry(EventSize, FSlateLayoutTransform(FVector2f(X, Y))),
+					EventBrush,
+					ESlateDrawEffect::None,
+					Point.Color);
+
+				FSlateDrawElement::MakeBox
+				(OutDrawElements,
+					LayerId,
+					AllottedGeometry.ToPaintGeometry(EventSize, FSlateLayoutTransform(FVector2f(X, Y))),
+					EventBorderBrush,
+					ESlateDrawEffect::None,
+					FLinearColor::Black);
+			}
+		}
+
+		LayerId++;
+	}
+
+	return LayerId;
+}
+
+#undef LOCTEXT_NAMESPACE
diff --git a/Engine/Plugins/Runtime/StateTree/Source/StateTreeEditorModule/Private/Debugger/SStateTreeDebuggerEventTimelineView.h b/Engine/Plugins/Runtime/StateTree/Source/StateTreeEditorModule/Private/Debugger/SStateTreeDebuggerEventTimelineView.h
new file mode 100644
--- /dev/null
+++ b/Engine/Plugins/Runtime/StateTree/Source/StateTreeEditorModule/Private/Debugger/SStateTreeDebuggerEventTimelineView.h
@@ -0,0 +1,82 @@
+// Copyright Epic Games, Inc. All Rights Reserved.
+
+#pragma once
+
+#include "Widgets/SCompoundWidget.h"
+#include "Widgets/DeclarativeSyntaxSupport.h"
+
+class FPaintArgs;
+class FSlateWindowElementList;
+class FTimeSliderController;
+
+
+/**
+ * Displays a sequence of points and windows for a statetree debugger instance timeline.
+ * This is a slightly modified version of the SEventTimelineView used by RewindDebugger
+ * so it should make it easier to eventually merge the tools.
+ */
+class SStateTreeDebuggerEventTimelineView : public SCompoundWidget
+{
+public:
+	struct FTimelineEventData
+	{
+	public:
+		struct EventPoint
+		{
+			double Time;
+			FText Type;
+			FText Description;
+			FLinearColor Color;
+		};
+		
+		struct EventWindow
+		{
+			double TimeStart;
+			double TimeEnd;
+			FText Type;
+			FText Description;
+			FLinearColor Color;
+		};
+
+		TArray<EventPoint> Points;
+		TArray<EventWindow> Windows;
+	};
+	
+	SLATE_BEGIN_ARGS(SStateTreeDebuggerEventTimelineView)
+		: _ViewRange(TRange<double>(0,10))
+		, _DesiredSize(FVector2D(100.f,20.f))
+	{}
+    	/** View time range */
+    	SLATE_ATTRIBUTE(TRange<double>, ViewRange);
+
+		/** Data for events to render */
+		SLATE_ATTRIBUTE(TSharedPtr<FTimelineEventData>, EventData);
+
+		/** Desired widget size */
+		SLATE_ATTRIBUTE(FVector2D, DesiredSize);
+	
+	SLATE_END_ARGS()
+
+
+	/**
+	 * Construct the widget
+	 * 
+	 * @param InArgs A declaration from which to construct the widget
+	 */
+	void Construct( const FArguments& InArgs );
+
+protected:
+	// SWidget interface
+	virtual FVector2D ComputeDesiredSize(float) const override;
+	virtual int32 OnPaint( const FPaintArgs& Args, const FGeometry& AllottedGeometry, const FSlateRect& MyCullingRect, FSlateWindowElementList& OutDrawElements, int32 LayerId, const FWidgetStyle& InWidgetStyle, bool bParentEnabled ) const override;
+	
+	int32 PaintEvents(const FGeometry& AllottedGeometry, const FSlateRect& MyCullingRect, FSlateWindowElementList& OutDrawElements, int32 LayerId, const FWidgetStyle& InWidgetStyle, bool bParentEnabled ) const;
+	
+	TAttribute<TRange<double>> ViewRange;
+	TAttribute<TSharedPtr<FTimelineEventData>> EventData;
+	TAttribute<FVector2D> DesiredSize;
+	
+	const FSlateBrush* EventBrush = nullptr;
+	const FSlateBrush* EventBorderBrush = nullptr;
+};
+
diff --git a/Engine/Plugins/Runtime/StateTree/Source/StateTreeEditorModule/Private/Debugger/SStateTreeDebuggerInstanceTree.cpp b/Engine/Plugins/Runtime/StateTree/Source/StateTreeEditorModule/Private/Debugger/SStateTreeDebuggerInstanceTree.cpp
new file mode 100644
--- /dev/null
+++ b/Engine/Plugins/Runtime/StateTree/Source/StateTreeEditorModule/Private/Debugger/SStateTreeDebuggerInstanceTree.cpp
@@ -0,0 +1,135 @@
+// Copyright Epic Games, Inc. All Rights Reserved.
+
+#if WITH_STATETREE_DEBUGGER
+
+#include "Debugger/SStateTreeDebuggerInstanceTree.h"
+#include "RewindDebuggerTrack.h"
+#include "Widgets/Images/SLayeredImage.h"
+
+//----------------------------------------------------------------------//
+// SStateTreeInstanceTree
+//----------------------------------------------------------------------//
+SStateTreeDebuggerInstanceTree::SStateTreeDebuggerInstanceTree()
+{
+}
+
+SStateTreeDebuggerInstanceTree::~SStateTreeDebuggerInstanceTree()
+{
+}
+
+TSharedRef<ITableRow> SStateTreeDebuggerInstanceTree::GenerateTreeRow(TSharedPtr<RewindDebugger::FRewindDebuggerTrack> Item, const TSharedRef<STableViewBase>& OwnerTable) const
+{
+	const FSlateIcon ObjectIcon = Item->GetIcon();
+
+	const TSharedRef<SLayeredImage> LayeredIcons = SNew(SLayeredImage)
+				.DesiredSizeOverride(FVector2D(16, 16))
+				.Image(ObjectIcon.GetIcon());
+
+	if (ObjectIcon.GetOverlayIcon())
+	{
+		LayeredIcons->AddLayer(ObjectIcon.GetOverlayIcon());
+	}
+	
+	return SNew(STableRow<TSharedPtr<RewindDebugger::FRewindDebuggerTrack>>, OwnerTable)
+		[
+			SNew(SHorizontalBox)
+			+SHorizontalBox::Slot().AutoWidth().Padding(2)
+			[
+				LayeredIcons
+			]
+			+SHorizontalBox::Slot()
+			.FillWidth(1.0f)
+			.Padding(1.0f, 0.0f)
+			.VAlign(VAlign_Center)
+			[
+				SNew(STextBlock)
+				.Text(Item->GetDisplayName())
+				.Font(FAppStyle::GetFontStyle(TEXT("PropertyWindow.NormalFont")))
+			]
+		];
+}
+
+void SStateTreeDebuggerInstanceTree::TreeExpansionChanged(TSharedPtr<RewindDebugger::FRewindDebuggerTrack> Item, const bool bShouldBeExpanded) const
+{
+	Item->SetIsExpanded(bShouldBeExpanded);
+	OnExpansionChanged.ExecuteIfBound();
+}
+
+void SStateTreeDebuggerInstanceTree::SetSelection(const TSharedPtr<RewindDebugger::FRewindDebuggerTrack>& SelectedItem) const
+{
+	TreeView->SetSelection(SelectedItem);
+}
+
+void SStateTreeDebuggerInstanceTree::ScrollTo(const double ScrollOffset) const
+{
+	TreeView->SetScrollOffset(ScrollOffset);
+}
+
+void SStateTreeDebuggerInstanceTree::Construct(const FArguments& InArgs)
+{
+	InstanceTracks = InArgs._InstanceTracks;
+
+	TreeView = SNew(STreeView<TSharedPtr<RewindDebugger::FRewindDebuggerTrack>>)
+				.ItemHeight(20.0f)
+				.TreeItemsSource(InstanceTracks)
+				.OnGenerateRow(this, &SStateTreeDebuggerInstanceTree::GenerateTreeRow)
+				.OnGetChildren_Lambda([](const TSharedPtr<RewindDebugger::FRewindDebuggerTrack>& Item, TArray<TSharedPtr<RewindDebugger::FRewindDebuggerTrack>>& OutChildren)
+					{
+						Item->IterateSubTracks([&OutChildren](const TSharedPtr<RewindDebugger::FRewindDebuggerTrack>& Track)
+						{
+							if (Track->IsVisible())
+							{
+								OutChildren.Add(Track);
+							}
+						});
+					})
+				.OnExpansionChanged(this, &SStateTreeDebuggerInstanceTree::TreeExpansionChanged)
+				.SelectionMode(ESelectionMode::Single)
+				.OnSelectionChanged(InArgs._OnSelectionChanged)
+				.OnMouseButtonDoubleClick(InArgs._OnMouseButtonDoubleClick)
+				.ExternalScrollbar(InArgs._ExternalScrollBar)
+				.AllowOverscroll(EAllowOverscroll::No)
+				.OnTreeViewScrolled(InArgs._OnScrolled)
+				.ScrollbarDragFocusCause(EFocusCause::SetDirectly)
+				.OnContextMenuOpening(InArgs._OnContextMenuOpening);
+
+	ChildSlot
+	[
+		TreeView.ToSharedRef()
+	];
+
+	OnExpansionChanged = InArgs._OnExpansionChanged;
+}
+
+namespace UE::StateTreeDebugger
+{
+static void RestoreExpansion(TSharedPtr<RewindDebugger::FRewindDebuggerTrack> Track, TSharedPtr<STreeView<TSharedPtr<RewindDebugger::FRewindDebuggerTrack>>>& TreeView)
+{
+	TreeView->SetItemExpansion(Track, Track->GetIsExpanded());
+	Track->IterateSubTracks([&TreeView](TSharedPtr<RewindDebugger::FRewindDebuggerTrack> SubTrack)
+	{
+		RestoreExpansion(SubTrack, TreeView);
+	});
+}
+} // UE::StateTreeDebugger
+
+void SStateTreeDebuggerInstanceTree::RestoreExpansion()
+{
+	for (const auto& Track : *InstanceTracks)
+	{
+		UE::StateTreeDebugger::RestoreExpansion(Track, TreeView);
+	}
+}
+
+void SStateTreeDebuggerInstanceTree::Refresh()
+{
+	TreeView->RebuildList();
+
+	if (InstanceTracks)
+	{
+		// make sure any newly added TreeView nodes are created expanded
+		RestoreExpansion();
+	}
+}
+
+#endif // WITH_STATETREE_DEBUGGER
\ No newline at end of file
diff --git a/Engine/Plugins/Runtime/StateTree/Source/StateTreeEditorModule/Private/Debugger/SStateTreeDebuggerInstanceTree.h b/Engine/Plugins/Runtime/StateTree/Source/StateTreeEditorModule/Private/Debugger/SStateTreeDebuggerInstanceTree.h
new file mode 100644
--- /dev/null
+++ b/Engine/Plugins/Runtime/StateTree/Source/StateTreeEditorModule/Private/Debugger/SStateTreeDebuggerInstanceTree.h
@@ -0,0 +1,55 @@
+// Copyright Epic Games, Inc. All Rights Reserved.
+
+#pragma once
+
+#if WITH_STATETREE_DEBUGGER
+
+#include "Widgets/Views/STreeView.h"
+
+namespace RewindDebugger
+{
+	class FRewindDebuggerTrack;
+}
+
+struct FStateTreeDebugger;
+
+class SStateTreeDebuggerInstanceTree : public SCompoundWidget
+{
+	using FOnSelectionChanged = STreeView<TSharedPtr<RewindDebugger::FRewindDebuggerTrack>>::FOnSelectionChanged;
+	using FOnMouseButtonDoubleClick = STreeView<TSharedPtr<RewindDebugger::FRewindDebuggerTrack>>::FOnMouseButtonDoubleClick;
+
+public:
+	SLATE_BEGIN_ARGS(SStateTreeDebuggerInstanceTree) { }
+	SLATE_ARGUMENT(TArray<TSharedPtr<RewindDebugger::FRewindDebuggerTrack>>*, InstanceTracks);
+	SLATE_ARGUMENT(TSharedPtr< SScrollBar >, ExternalScrollBar);
+	SLATE_EVENT(FOnSelectionChanged, OnSelectionChanged)
+	SLATE_EVENT(FOnMouseButtonDoubleClick, OnMouseButtonDoubleClick)
+	SLATE_EVENT(FOnContextMenuOpening, OnContextMenuOpening)
+	SLATE_EVENT(FSimpleDelegate, OnExpansionChanged)
+	SLATE_EVENT(FOnTableViewScrolled, OnScrolled)
+	SLATE_END_ARGS()
+
+	SStateTreeDebuggerInstanceTree();
+	virtual ~SStateTreeDebuggerInstanceTree();
+
+	void Construct(const FArguments& InArgs);
+
+	void Refresh();
+	void RestoreExpansion();
+	
+	void SetSelection(const TSharedPtr<RewindDebugger::FRewindDebuggerTrack>& SelectedItem) const;
+	void ScrollTo(double ScrollOffset) const;
+
+private:
+	TSharedRef<ITableRow> GenerateTreeRow(TSharedPtr<RewindDebugger::FRewindDebuggerTrack> Item, const TSharedRef<STableViewBase>& OwnerTable) const;
+	void TreeExpansionChanged(TSharedPtr<RewindDebugger::FRewindDebuggerTrack> Item, bool bShouldBeExpanded) const;
+	
+	TArray<TSharedPtr<RewindDebugger::FRewindDebuggerTrack>>* InstanceTracks = nullptr;
+
+	TSharedPtr<STreeView<TSharedPtr<RewindDebugger::FRewindDebuggerTrack>>> TreeView;
+
+	FSimpleDelegate OnExpansionChanged;	
+};
+
+
+#endif // WITH_STATETREE_DEBUGGER
\ No newline at end of file
diff --git a/Engine/Plugins/Runtime/StateTree/Source/StateTreeEditorModule/Private/Debugger/SStateTreeDebuggerTimelines.cpp b/Engine/Plugins/Runtime/StateTree/Source/StateTreeEditorModule/Private/Debugger/SStateTreeDebuggerTimelines.cpp
new file mode 100644
--- /dev/null
+++ b/Engine/Plugins/Runtime/StateTree/Source/StateTreeEditorModule/Private/Debugger/SStateTreeDebuggerTimelines.cpp
@@ -0,0 +1,205 @@
+// Copyright Epic Games, Inc. All Rights Reserved.
+
+#if WITH_STATETREE_DEBUGGER
+
+#include "Debugger/SStateTreeDebuggerTimelines.h"
+#include "Widgets/Layout/SSpacer.h"
+#include "RewindDebuggerTrack.h"
+
+//#define LOCTEXT_NAMESPACE "StateTreeEditor"
+
+
+//----------------------------------------------------------------------//
+// SRewindDebuggerTimelineTableRow
+//----------------------------------------------------------------------//
+class SRewindDebuggerTimelineTableRow : public STableRow<TSharedPtr<RewindDebugger::FRewindDebuggerTrack>>
+{
+	public:
+	
+	SLATE_BEGIN_ARGS(SRewindDebuggerTimelineTableRow) {}
+	SLATE_END_ARGS()
+	
+	void Construct(const FArguments& InArgs, const TSharedRef<STableViewBase>& InOwnerTableView)
+	{
+		// @todo: why not passing InArgs ???
+		STableRow::Construct(STableRow::FArguments(), InOwnerTableView);
+
+		SetExpanderArrowVisibility(EVisibility::Collapsed);
+	}
+
+	virtual int32 GetIndentLevel() const override
+	{
+		// don't indent timeline tracks
+		return 0;
+	}
+};
+
+
+//----------------------------------------------------------------------//
+// SStateTreeInstanceTimelines
+//----------------------------------------------------------------------//
+SStateTreeDebuggerTimelines::SStateTreeDebuggerTimelines() 
+	: SSimpleTimeSlider()
+{ 
+}
+
+SStateTreeDebuggerTimelines::~SStateTreeDebuggerTimelines() 
+{
+}
+
+void SStateTreeDebuggerTimelines::SetSelection(const TSharedPtr<RewindDebugger::FRewindDebuggerTrack>& SelectedItem)
+{
+	InstanceTreeView->SetSelection(SelectedItem);
+}
+
+void SStateTreeDebuggerTimelines::ScrollTo(const double ScrollOffset)
+{
+	InstanceTreeView->SetScrollOffset(ScrollOffset);
+}
+
+TSharedRef<ITableRow> SStateTreeDebuggerTimelines::GenerateTreeRow(TSharedPtr<RewindDebugger::FRewindDebuggerTrack> InItem, const TSharedRef<STableViewBase>& OwnerTable)
+{
+	const TSharedPtr<SRewindDebuggerTimelineTableRow> Track = SNew(SRewindDebuggerTimelineTableRow, OwnerTable);
+	TSharedPtr<SWidget> TimelineView = InItem->GetTimelineView();
+
+	if (!TimelineView.IsValid())
+	{
+		TimelineView = SNew(SSpacer).Size(FVector2D(100, 20));
+	}
+
+	Track->SetContent(TimelineView.ToSharedRef());
+
+	return Track.ToSharedRef();
+}
+
+namespace UE::StateTreeDebugger
+{
+void TimelineViewGetChildren(TSharedPtr<RewindDebugger::FRewindDebuggerTrack> InItem, TArray<TSharedPtr<RewindDebugger::FRewindDebuggerTrack>>& OutChildren)
+{
+	InItem->IterateSubTracks([&OutChildren](TSharedPtr<RewindDebugger::FRewindDebuggerTrack> Track)
+	{
+		if (Track->IsVisible())
+		{
+			OutChildren.Add(Track);
+		}
+	});
+}
+} // UE::StateTreeDebugger
+
+void SStateTreeDebuggerTimelines::TreeExpansionChanged(TSharedPtr<RewindDebugger::FRewindDebuggerTrack> InItem, bool bShouldBeExpanded)
+{
+	InItem->SetIsExpanded(bShouldBeExpanded);
+	OnExpansionChanged.ExecuteIfBound();
+}
+
+void SStateTreeDebuggerTimelines::Construct(const FArguments& InArgs)
+{
+	SSimpleTimeSlider::Construct(SSimpleTimeSlider::FArguments()
+		.ViewRange(InArgs._ViewRange)
+		.OnViewRangeChanged(InArgs._OnViewRangeChanged)
+		.ClampRange(InArgs._ClampRange)
+		.ScrubPosition(InArgs._ScrubPosition)
+		.OnScrubPositionChanged(InArgs._OnScrubPositionChanged));
+
+	InstanceTracks = InArgs._DebugComponents;
+
+	OnExpansionChanged = InArgs._OnExpansionChanged;
+
+	InstanceTreeView = SNew(STreeView<TSharedPtr<RewindDebugger::FRewindDebuggerTrack>>)
+									.ItemHeight(20.0f)
+									.OnSelectionChanged(InArgs._OnSelectionChanged)
+									.TreeItemsSource(InstanceTracks)
+									.OnGenerateRow(this, &SStateTreeDebuggerTimelines::GenerateTreeRow)
+									.OnGetChildren_Static(&UE::StateTreeDebugger::TimelineViewGetChildren)
+									.OnExpansionChanged(this, &SStateTreeDebuggerTimelines::TreeExpansionChanged)
+									.SelectionMode(ESelectionMode::Single)
+									.AllowOverscroll(EAllowOverscroll::No)
+									.ExternalScrollbar(InArgs._ExternalScrollbar)
+									.OnTreeViewScrolled(InArgs._OnScrolled);
+
+	ChildSlot
+	[
+		InstanceTreeView.ToSharedRef()
+	];
+}
+
+int32 SStateTreeDebuggerTimelines::OnPaint( const FPaintArgs& Args, const FGeometry& AllottedGeometry, const FSlateRect& MyCullingRect, FSlateWindowElementList& OutDrawElements, int32 LayerId, const FWidgetStyle& InWidgetStyle, bool bParentEnabled ) const
+{
+	const TRange<double> LocalViewRange = ViewRange.Get();
+	const FScrubRangeToScreen RangeToScreen( LocalViewRange, AllottedGeometry.GetLocalSize() );
+
+	// Paint Child Widgets (Tracks) 
+	LayerId = SCompoundWidget::OnPaint(Args, AllottedGeometry, MyCullingRect, OutDrawElements, LayerId, InWidgetStyle, ShouldBeEnabled(bParentEnabled));
+
+	// Paint Major Time Marks
+	FDrawTickArgs TickArgs;
+	TickArgs.AllottedGeometry = AllottedGeometry;
+	TickArgs.bOnlyDrawMajorTicks = true;
+	TickArgs.TickColor = FLinearColor(0.2f,0.2f,0.2f,0.2f);
+	TickArgs.ClippingRect = MyCullingRect;
+	TickArgs.DrawEffects = ESlateDrawEffect::None;
+	TickArgs.StartLayer = ++LayerId;
+	TickArgs.TickOffset = 0;
+	TickArgs.MajorTickHeight = AllottedGeometry.Size.Y;
+	DrawTicks(OutDrawElements, RangeToScreen, TickArgs);
+    	
+	// Paint Cursor
+	const float XPos = RangeToScreen.InputToLocalX( ScrubPosition.Get() );
+
+	const TArray<FVector2D> Points {{XPos, 0}, { XPos, AllottedGeometry.Size.Y}};
+
+	FSlateDrawElement::MakeLines(
+				OutDrawElements,
+						++LayerId,
+						AllottedGeometry.ToPaintGeometry(),
+						Points,
+						ESlateDrawEffect::None,
+						FLinearColor::White,
+						false
+						);
+
+	return LayerId;
+}
+
+FVector2D SStateTreeDebuggerTimelines::ComputeDesiredSize(float X) const
+{
+	FVector2D DefaultSize = SSimpleTimeSlider::ComputeDesiredSize(X);
+	if (InstanceTracks && InstanceTracks->Num() > 0)
+	{
+		return FVector2D(InstanceTracks->Num() * DefaultSize[0], InstanceTracks->Num() * DefaultSize[1]);
+	}
+	return DefaultSize;
+}
+
+namespace UE::StateTreeDebugger
+{
+static void RestoreTimelineExpansion(const TSharedPtr<RewindDebugger::FRewindDebuggerTrack>& Track, TSharedPtr<STreeView<TSharedPtr<RewindDebugger::FRewindDebuggerTrack>>>& TreeView)
+{
+	TreeView->SetItemExpansion(Track, Track->GetIsExpanded());
+	Track->IterateSubTracks([&TreeView](const TSharedPtr<RewindDebugger::FRewindDebuggerTrack>& SubTrack)
+		{
+			RestoreTimelineExpansion(SubTrack, TreeView);
+		});
+}
+} // UE::StateTreeDebugger
+
+void SStateTreeDebuggerTimelines::RestoreExpansion()
+{
+	for (const auto& Track : *InstanceTracks)
+	{
+		UE::StateTreeDebugger::RestoreTimelineExpansion(Track, InstanceTreeView);
+	}
+}
+
+void SStateTreeDebuggerTimelines::Refresh()
+{
+	InstanceTreeView->RebuildList();
+
+	if (InstanceTracks)
+	{
+		// make sure any newly added TreeView nodes are created expanded
+		RestoreExpansion();
+	}
+}
+
+#endif // WITH_STATETREE_DEBUGGER
\ No newline at end of file
diff --git a/Engine/Plugins/Runtime/StateTree/Source/StateTreeEditorModule/Private/Debugger/SStateTreeDebuggerTimelines.h b/Engine/Plugins/Runtime/StateTree/Source/StateTreeEditorModule/Private/Debugger/SStateTreeDebuggerTimelines.h
new file mode 100644
--- /dev/null
+++ b/Engine/Plugins/Runtime/StateTree/Source/StateTreeEditorModule/Private/Debugger/SStateTreeDebuggerTimelines.h
@@ -0,0 +1,87 @@
+// Copyright Epic Games, Inc. All Rights Reserved.
+
+#pragma once
+
+#if WITH_STATETREE_DEBUGGER
+
+#include "SSimpleTimeSlider.h"
+#include "Widgets/Views/STreeView.h"
+
+namespace RewindDebugger
+{
+	class FRewindDebuggerTrack;
+}
+
+struct FStateTreeDebugger;
+
+class SStateTreeDebuggerTimelines : public SSimpleTimeSlider
+{
+	using FOnSelectionChanged = STreeView<TSharedPtr<RewindDebugger::FRewindDebuggerTrack>>::FOnSelectionChanged;
+public:
+	SLATE_BEGIN_ARGS(SStateTreeDebuggerTimelines) {}
+	 
+		SLATE_ARGUMENT( TArray< TSharedPtr< RewindDebugger::FRewindDebuggerTrack > >*, DebugComponents );
+		SLATE_ARGUMENT( TSharedPtr<SScrollBar>, ExternalScrollbar)
+		SLATE_EVENT( FOnSelectionChanged, OnSelectionChanged )
+		SLATE_EVENT( FOnTableViewScrolled, OnScrolled)
+		SLATE_EVENT( FSimpleDelegate, OnExpansionChanged )
+
+	
+		/** The scrub position */
+    	SLATE_ATTRIBUTE(double, ScrubPosition);
+    
+    	/** View time range */
+    	SLATE_ATTRIBUTE(TRange<double>, ViewRange);
+    
+    	/** Clamp time range */
+    	SLATE_ATTRIBUTE(TRange<double>, ClampRange);
+    
+		/** Called when the scrub position changes */
+		SLATE_EVENT(FOnScrubPositionChanged, OnScrubPositionChanged);
+	
+		/** Called right before the scrubber begins to move */
+		SLATE_EVENT(FSimpleDelegate, OnBeginScrubberMovement);
+	
+		/** Called right after the scrubber handle is released by the user */
+		SLATE_EVENT(FSimpleDelegate, OnEndScrubberMovement);
+	
+    	/** Called when the view range changes */
+    	SLATE_EVENT(FOnRangeChanged, OnViewRangeChanged);
+    	
+	SLATE_END_ARGS()
+
+	/**
+	* Default constructor.
+	*/
+	SStateTreeDebuggerTimelines();
+	virtual ~SStateTreeDebuggerTimelines();
+
+	/**
+	* Constructs the application.
+	*
+	* @param InArgs The Slate argument list.
+	*/
+	void Construct(const FArguments& InArgs);
+
+	void Refresh();
+	void RestoreExpansion();
+
+	void SetSelection(const TSharedPtr<RewindDebugger::FRewindDebuggerTrack>& SelectedItem);
+	void ScrollTo(double ScrollOffset);
+	
+	virtual int32 OnPaint( const FPaintArgs& Args, const FGeometry& AllottedGeometry, const FSlateRect& MyCullingRect, FSlateWindowElementList& OutDrawElements, int32 LayerId, const FWidgetStyle& InWidgetStyle, bool bParentEnabled ) const override;
+
+protected:
+	virtual FVector2D ComputeDesiredSize(float) const override;
+
+private:
+	TSharedRef<ITableRow> GenerateTreeRow(TSharedPtr<RewindDebugger::FRewindDebuggerTrack> InItem, const TSharedRef<STableViewBase>& OwnerTable);
+	void TreeExpansionChanged(TSharedPtr<RewindDebugger::FRewindDebuggerTrack> InItem, bool bShouldBeExpanded);
+	
+	TArray<TSharedPtr<RewindDebugger::FRewindDebuggerTrack>>* InstanceTracks = nullptr;
+	TSharedPtr<STreeView<TSharedPtr<RewindDebugger::FRewindDebuggerTrack>>> InstanceTreeView;
+	
+	FSimpleDelegate OnExpansionChanged;
+};
+
+#endif // WITH_STATETREE_DEBUGGER
\ No newline at end of file
diff --git a/Engine/Plugins/Runtime/StateTree/Source/StateTreeEditorModule/Private/Debugger/SStateTreeDebuggerView.cpp b/Engine/Plugins/Runtime/StateTree/Source/StateTreeEditorModule/Private/Debugger/SStateTreeDebuggerView.cpp
new file mode 100644
--- /dev/null
+++ b/Engine/Plugins/Runtime/StateTree/Source/StateTreeEditorModule/Private/Debugger/SStateTreeDebuggerView.cpp
@@ -0,0 +1,947 @@
+// Copyright Epic Games, Inc. All Rights Reserved.
+
+#if WITH_STATETREE_DEBUGGER
+
+#include "SStateTreeDebuggerView.h"
+#include "Debugger/StateTreeDebugger.h"
+#include "DetailsViewArgs.h"
+#include "Editor.h"
+#include "Factories.h"
+#include "Framework/MultiBox/MultiBoxBuilder.h"
+#include "IStructureDetailsView.h"
+#include "Kismet2/DebuggerCommands.h"
+#include "Modules/ModuleManager.h"
+#include "PropertyEditorModule.h"
+#include "SStateTreeDebuggerInstanceTree.h"
+#include "SStateTreeDebuggerTimelines.h"
+#include "StateTree.h"
+#include "StateTreeDebuggerCommands.h"
+#include "StateTreeDebuggerTrack.h"
+#include "StateTreeExecutionContext.h"
+#include "StateTreeState.h"
+#include "StateTreeTaskBase.h"
+#include "StateTreeViewModel.h"
+#include "Widgets/Input/SComboButton.h"
+#include "Widgets/Layout/SBorder.h"
+#include "Widgets/Layout/SScrollBox.h"
+#include "Widgets/Layout/SSpacer.h"
+
+#define LOCTEXT_NAMESPACE "StateTreeEditor"
+
+
+namespace UE::StateTreeDebugger
+{
+//----------------------------------------------------------------------//
+// FTraceTextObjectFactory
+//----------------------------------------------------------------------//
+struct FTraceTextObjectFactory : FCustomizableTextObjectFactory
+{
+	UObject* NodeInstanceObject = nullptr;
+	FTraceTextObjectFactory() : FCustomizableTextObjectFactory(GWarn) {}
+
+protected:
+	virtual bool CanCreateClass(UClass* ObjectClass, bool& bOmitSubObjs) const override
+	{
+		return true;
+	}
+
+	virtual void ProcessConstructedObject(UObject* CreatedObject) override
+	{
+		NodeInstanceObject = CreatedObject;
+	}
+};
+
+
+//----------------------------------------------------------------------//
+// FEventTreeElement
+//----------------------------------------------------------------------//
+/** An item in the trace event tree */
+struct FEventTreeElement : TSharedFromThis<FEventTreeElement>
+{
+	explicit FEventTreeElement(const TraceServices::FFrame& Frame, const FStateTreeTraceEventVariantType& Event)
+		: Frame(Frame), Event(Event)
+	{
+	}
+
+	TraceServices::FFrame Frame;
+	FStateTreeTraceEventVariantType Event;
+	TArray<TSharedPtr<FEventTreeElement>> Children;
+};
+} // UE::StateTreeDebugger
+
+
+//----------------------------------------------------------------------//
+// SStateTreeDebuggerTableRow
+//----------------------------------------------------------------------//
+class SStateTreeDebuggerTableRow : public SMultiColumnTableRow<TSharedPtr<UE::StateTreeDebugger::FEventTreeElement>>
+{
+public:
+	void Construct(const FArguments& InArgs,
+		const TSharedPtr<STableViewBase>& InOwnerTableView,
+		const TSharedPtr<UE::StateTreeDebugger::FEventTreeElement>& InElement,
+		const TSharedRef<FStateTreeViewModel>& InStateTreeViewModel)
+	{
+		Item = InElement;
+		StateTreeViewModel = InStateTreeViewModel.ToSharedPtr();
+		SMultiColumnTableRow::Construct(InArgs, InOwnerTableView.ToSharedRef());
+	}
+
+protected:
+	// SMultiColumnTableRow overrides
+	virtual TSharedRef<SWidget> GenerateWidgetForColumn(const FName& ColumnName) override
+	{
+		const TSharedPtr<SHorizontalBox> Contents = SNew(SHorizontalBox);
+
+		Contents->AddSlot()
+			.VAlign(VAlign_Fill)
+			.HAlign(HAlign_Left)
+			.AutoWidth()
+			[
+				SNew(SExpanderArrow, SharedThis(this))
+				.ShouldDrawWires(true)
+				.IndentAmount(32)
+				.BaseIndentLevel(0)
+			];
+		
+		if (ColumnName == FName("Desc"))
+		{
+			Contents->AddSlot()
+				.Padding(5, 0)
+				.FillWidth(1.f)
+				[
+					SNew(STextBlock)
+					.Font(FCoreStyle::GetDefaultFontStyle("Mono", 9))
+					.Text_Lambda([&Event=Item->Event, this]()
+					{
+						return GetTextForEvent(Event);
+					})
+				];
+		}
+
+		return Contents.ToSharedRef();
+	}
+
+	FText GetTextForEvent(const FStateTreeTraceEventVariantType& Event) const
+	{
+		const UStateTree* StateTree = StateTreeViewModel->GetStateTree();
+
+		// Use log event messages directly
+		if (const FStateTreeTraceLogEvent* LogEvent = Event.TryGet<FStateTreeTraceLogEvent>())
+		{
+			if (LogEvent->Message.Len())
+			{
+				return FText::FromString(*LogEvent->Message);
+			}
+		}
+		// Process state events (index has a different meaning)
+		else if (const FStateTreeTraceStateEvent* StateEvent = Event.TryGet<FStateTreeTraceStateEvent>())
+		{
+			const FStateTreeStateHandle StateHandle(StateEvent->Idx);
+			if (const FCompactStateTreeState* CompactState = StateTree->GetStateFromHandle(StateHandle))
+			{
+				return FText::FromString
+				(FString::Printf
+					(TEXT("%s State '%s'"),
+						*StaticEnum<EStateTreeTraceNodeEventType>()->GetNameStringByValue((int64)StateEvent->EventType),
+						*CompactState->Name.ToString()));
+			}
+		}
+		// Process Tasks events
+		else if (const FStateTreeTraceTaskEvent* TaskEvent = Event.TryGet<FStateTreeTraceTaskEvent>())
+		{
+			const FConstStructView NodeView = StateTree->GetNode(TaskEvent->Idx);
+			const FStateTreeNodeBase* Node = NodeView.GetPtr<const FStateTreeNodeBase>();
+
+			return FText::FromString
+			(FString::Printf
+				(TEXT("%s:%s %s '%s'"),
+					*StaticEnum<EStateTreeTraceNodeEventType>()->GetNameStringByValue((int64)TaskEvent->EventType),
+					*StaticEnum<EStateTreeRunStatus>()->GetNameStringByValue((int64)TaskEvent->Status),
+					*NodeView.GetScriptStruct()->GetName(),
+					Node != nullptr ? *Node->Name.ToString() : *LexToString(TaskEvent->Idx)));
+		}
+		// Process Conditions events
+		else if (const FStateTreeTraceConditionEvent* ConditionEvent = Event.TryGet<FStateTreeTraceConditionEvent>())
+		{
+			const FConstStructView NodeView = StateTree->GetNode(ConditionEvent->Idx);
+			const FStateTreeNodeBase* Node = NodeView.GetPtr<const FStateTreeNodeBase>();
+
+			return FText::FromString
+			(FString::Printf
+				(TEXT("%s %s '%s'"),
+					*StaticEnum<EStateTreeTraceNodeEventType>()->GetNameStringByValue((int64)ConditionEvent->EventType),
+					*NodeView.GetScriptStruct()->GetName(),
+					Node != nullptr ? *Node->Name.ToString() : *LexToString(ConditionEvent->Idx)));
+		}
+		// Process ActiveStates events
+		else if (const FStateTreeTraceActiveStatesEvent* ActiveStatesEvent = Event.TryGet<FStateTreeTraceActiveStatesEvent>())
+		{
+			FString StatePath;
+			for (int32 i = 0; i < ActiveStatesEvent->ActiveStates.Num(); i++)
+			{
+				const FCompactStateTreeState& State = StateTree->GetStates()[ActiveStatesEvent->ActiveStates[i].Index];
+				StatePath.Appendf(TEXT("%s%s"), i == 0 ? TEXT("") : TEXT("."), *State.Name.ToString());
+			}
+
+			return FText::FromString(FString::Printf(TEXT("New active states: '%s'"), *StatePath));
+		}
+
+		return FText();
+	}
+
+	TSharedPtr<FStateTreeViewModel> StateTreeViewModel;
+	TSharedPtr<UE::StateTreeDebugger::FEventTreeElement> Item;
+};
+
+
+//----------------------------------------------------------------------//
+// SStateTreeDebuggerView
+//----------------------------------------------------------------------//
+SStateTreeDebuggerView::SStateTreeDebuggerView()
+{
+	FEditorDelegates::BeginPIE.AddRaw(this, &SStateTreeDebuggerView::OnPIEStarted);
+	FEditorDelegates::EndPIE.AddRaw(this, &SStateTreeDebuggerView::OnPIEStopped);
+	FEditorDelegates::PausePIE.AddRaw(this, &SStateTreeDebuggerView::OnPIEPaused);
+	FEditorDelegates::ResumePIE.AddRaw(this, &SStateTreeDebuggerView::OnPIEResumed);
+	FEditorDelegates::SingleStepPIE.AddRaw(this, &SStateTreeDebuggerView::OnPIESingleStepped);
+}
+
+SStateTreeDebuggerView::~SStateTreeDebuggerView()
+{
+	if (SelectedNodeDataObject.IsValid())
+	{
+		SelectedNodeDataObject->RemoveFromRoot();
+	}
+
+	if (Debugger)
+	{
+		Debugger->OnScrubStateChanged.Unbind();
+		Debugger->OnBreakpointHit.Unbind();
+		Debugger->OnNewInstance.Unbind();
+		Debugger->OnSelectedInstanceCleared.Unbind();
+	}
+
+	FEditorDelegates::BeginPIE.RemoveAll(this);
+	FEditorDelegates::EndPIE.RemoveAll(this);
+	FEditorDelegates::PausePIE.RemoveAll(this);
+	FEditorDelegates::ResumePIE.RemoveAll(this);
+	FEditorDelegates::SingleStepPIE.RemoveAll(this);
+}
+
+void SStateTreeDebuggerView::OnPIEStarted(const bool bIsSimulating) const
+{
+	if (!Debugger->IsAnalysisSessionActive())
+	{
+		Debugger->StartLastLiveSessionAnalysis();
+	}
+}
+
+void SStateTreeDebuggerView::OnPIEStopped(const bool bIsSimulating) const
+{
+	Debugger->Unpause();
+}
+
+void SStateTreeDebuggerView::OnPIEPaused(const bool bIsSimulating) const
+{
+	Debugger->Pause();
+}
+
+void SStateTreeDebuggerView::OnPIEResumed(const bool bIsSimulating) const
+{
+	Debugger->Unpause();
+}
+
+void SStateTreeDebuggerView::OnPIESingleStepped(bool bSimulating) const
+{
+	Debugger->SyncToCurrentSessionDuration();
+}
+
+void SStateTreeDebuggerView::Construct(const FArguments& InArgs, const UStateTree* InStateTree, const TSharedRef<FStateTreeViewModel>& InStateTreeViewModel, const TSharedRef<FUICommandList>& InCommandList)
+{
+	StateTreeViewModel = InStateTreeViewModel;
+	StateTree = InStateTree;
+
+	Debugger = InStateTreeViewModel->GetDebugger();
+
+	// Bind callbacks to the debugger delegates
+	Debugger->OnNewInstance.BindSP(this, &SStateTreeDebuggerView::OnNewInstance);
+	Debugger->OnScrubStateChanged.BindSP(this, &SStateTreeDebuggerView::OnDebuggerScrubStateChanged);
+	Debugger->OnBreakpointHit.BindSP(this, &SStateTreeDebuggerView::OnBreakpointHit, InCommandList);
+	Debugger->OnSelectedInstanceCleared.BindSP(this, &SStateTreeDebuggerView::OnSelectedInstanceCleared);
+
+	// Bind our scrub time attribute to follow the value computed by the debugger
+	ScrubTimeAttribute = TAttribute<double>(InStateTreeViewModel->GetDebugger(), &FStateTreeDebugger::GetScrubTime);
+
+	// Put debugger in proper simulation state when view is constructed after PIE/SIE was started
+	if (FPlayWorldCommandCallbacks::HasPlayWorldAndPaused())
+	{
+		Debugger->Pause();
+	}
+
+	// Add & Bind commands
+	BindDebuggerToolbarCommands(InCommandList);
+
+	// Register the play world commands
+	InCommandList->Append(FPlayWorldCommands::GlobalPlayWorldActions.ToSharedRef());
+
+	InCommandList->MapAction(
+		FStateTreeDebuggerCommands::Get().ToggleBreakpoint,
+		FExecuteAction::CreateSP(this, &SStateTreeDebuggerView::ToggleBreakpoint),
+		FCanExecuteAction::CreateSP(this, &SStateTreeDebuggerView::CanToggleBreakpoint)
+	);
+	
+	// Toolbar
+	FSlimHorizontalToolBarBuilder ToolbarBuilder(InCommandList, FMultiBoxCustomization::None, /*InExtender*/ nullptr, /*InForceSmallIcons*/ true);
+	ToolbarBuilder.BeginSection(TEXT("Debugging"));
+	{
+		const FPlayWorldCommands& PlayWorldCommand = FPlayWorldCommands::Get();
+		ToolbarBuilder.AddToolBarButton(PlayWorldCommand.RepeatLastPlay);
+		ToolbarBuilder.AddToolBarButton(PlayWorldCommand.PausePlaySession,
+					NAME_None, TAttribute<FText>(), TAttribute<FText>(), FSlateIcon(FAppStyle::GetAppStyleSetName(), "PlayWorld.PausePlaySession.Small"));
+		ToolbarBuilder.AddToolBarButton(PlayWorldCommand.ResumePlaySession,
+					NAME_None, TAttribute<FText>(), TAttribute<FText>(), FSlateIcon(FAppStyle::GetAppStyleSetName(), "PlayWorld.ResumePlaySession.Small"));	
+		ToolbarBuilder.AddToolBarButton(PlayWorldCommand.StopPlaySession,
+					NAME_None, TAttribute<FText>(), TAttribute<FText>(), FSlateIcon(FAppStyle::GetAppStyleSetName(), "PlayWorld.StopPlaySession.Small"));
+		ToolbarBuilder.AddSeparator();
+		ToolbarBuilder.AddToolBarButton(FStateTreeDebuggerCommands::Get().PreviousFrameWithStateChange);
+		ToolbarBuilder.AddToolBarButton(FStateTreeDebuggerCommands::Get().PreviousFrameWithEvents);
+		ToolbarBuilder.AddToolBarButton(FStateTreeDebuggerCommands::Get().NextFrameWithEvents);
+		ToolbarBuilder.AddToolBarButton(FStateTreeDebuggerCommands::Get().NextFrameWithStateChange);
+	}
+	ToolbarBuilder.EndSection();
+
+	// Auto-select session if there is only one available
+	TArray<FStateTreeDebugger::FTraceDescriptor> TraceDescriptors;
+	Debugger->GetLiveTraces(TraceDescriptors);
+	
+	if (TraceDescriptors.Num() == 1)
+	{
+		Debugger->StartSessionAnalysis(TraceDescriptors.Last());
+	}
+
+	// Trace selection combo
+	const TSharedRef<SWidget> TraceSelectionBox = SNew(SComboButton)
+		.OnGetMenuContent(this, &SStateTreeDebuggerView::OnGetDebuggerTracesMenu)
+		.ButtonContent()
+		[
+			SNew(STextBlock)
+			.ToolTipText(LOCTEXT("SelectTraceSession", "Pick trace session to debug"))
+			.Text_Lambda([Debugger = Debugger]()
+			{
+				return Debugger.IsValid() ? Debugger->GetSelectedTraceDescription() : FText::GetEmpty();
+			})
+		];
+
+	// Instances TreeView
+	InstancesTreeView =	SNew(SStateTreeDebuggerInstanceTree)
+		.OnExpansionChanged_Lambda([this]() { InstanceTimelinesTreeView->RestoreExpansion(); })
+		.OnScrolled_Lambda([this](double ScrollOffset)
+		{
+			InstanceTimelinesTreeView->ScrollTo(ScrollOffset);
+		})
+		.InstanceTracks(&InstanceTracks)
+		.OnSelectionChanged_Lambda([this](TSharedPtr<RewindDebugger::FRewindDebuggerTrack> SelectedItem, ESelectInfo::Type SelectInfo)
+			{
+				InstanceTimelinesTreeView->SetSelection(SelectedItem);
+
+				const FStateTreeDebuggerTrack* StateTreeTrack = static_cast<FStateTreeDebuggerTrack*>(SelectedItem.Get());
+				Debugger->SelectInstance(StateTreeTrack != nullptr ? StateTreeTrack->GetInstanceId() : FStateTreeInstanceDebugId::Invalid);
+			});
+	
+	// Timelines TreeView
+	InstanceTimelinesTreeView = SNew(SStateTreeDebuggerTimelines)
+	   .OnExpansionChanged_Lambda([this]() { InstancesTreeView->RestoreExpansion(); })
+	   .OnScrolled_Lambda([this](double ScrollOffset){ InstancesTreeView->ScrollTo(ScrollOffset); })
+	   .DebugComponents(&InstanceTracks)
+	   .ViewRange_Lambda([this](){ return ViewRange; })
+	   .ClampRange_Lambda([this](){ return TRange<double>(0.0f, Debugger->GetRecordingDuration()); })
+	   .OnViewRangeChanged_Lambda([this](TRange<double> NewRange){ ViewRange = NewRange; })
+	   .ScrubPosition(ScrubTimeAttribute)
+	   .OnScrubPositionChanged_Lambda([this](double NewScrubTime, bool bIsScrubbing){ OnTimeLineScrubPositionChanged( NewScrubTime, bIsScrubbing ); });
+
+	// EventsTreeView
+	EventsTreeView = SNew(STreeView<TSharedPtr<UE::StateTreeDebugger::FEventTreeElement>>)
+		.HeaderRow(SNew(SHeaderRow)
+			+SHeaderRow::Column("Desc")
+			.DefaultLabel(LOCTEXT("FrameDetailsColumnHeader", "Frame Details")))
+			.OnGenerateRow_Lambda([this](const TSharedPtr<UE::StateTreeDebugger::FEventTreeElement>& InElement, const TSharedRef<STableViewBase>& InOwnerTableView)
+			{
+				return SNew(SStateTreeDebuggerTableRow, InOwnerTableView, InElement, StateTreeViewModel.ToSharedRef());
+			})
+			.OnGetChildren_Lambda([](const TSharedPtr<const UE::StateTreeDebugger::FEventTreeElement>& InParent, TArray<TSharedPtr<UE::StateTreeDebugger::FEventTreeElement>>& OutChildren)
+			{
+				if (const UE::StateTreeDebugger::FEventTreeElement* Parent = InParent.Get())
+				{
+					OutChildren.Append(Parent->Children);
+				}
+			})
+		.TreeItemsSource(&EventsTreeElements)
+		.ItemHeight(32)
+		.OnSelectionChanged_Lambda([this](const TSharedPtr<UE::StateTreeDebugger::FEventTreeElement>& InSelectedItem, ESelectInfo::Type SelectionType)
+		{
+			if (!InSelectedItem.IsValid())
+			{
+				return;
+			}
+
+			TSharedPtr<SWidget> DetailsView;
+
+			FString TypePath;
+			FString InstanceDataAsText;
+
+			if (const FStateTreeTraceConditionEvent* ConditionEvent = InSelectedItem->Event.TryGet<FStateTreeTraceConditionEvent>())
+			{
+				TypePath = ConditionEvent->TypePath;
+				InstanceDataAsText = ConditionEvent->InstanceDataAsText;
+			}
+			else if (const FStateTreeTraceTaskEvent* TaskEvent = InSelectedItem->Event.TryGet<FStateTreeTraceTaskEvent>())
+			{
+				TypePath = TaskEvent->TypePath;
+				InstanceDataAsText = TaskEvent->InstanceDataAsText;
+			}
+
+			if (!TypePath.IsEmpty())
+			{
+				FDetailsViewArgs DetailsViewArgs;
+				DetailsViewArgs.NameAreaSettings = FDetailsViewArgs::HideNameArea;
+
+				UScriptStruct* ScriptStruct = FindObject<UScriptStruct>(nullptr, *TypePath, /*ExactClass*/false);
+				if (ScriptStruct == nullptr)
+				{
+					ScriptStruct = LoadObject<UScriptStruct>(nullptr, *TypePath);
+				}
+
+				if (ScriptStruct != nullptr)
+				{
+					SelectedNodeDataStruct.InitializeAs(ScriptStruct);
+
+					ScriptStruct->ImportText(*InstanceDataAsText, SelectedNodeDataStruct.GetMutableMemory(), /*OwnerObject*/nullptr, PPF_None, GLog, ScriptStruct->GetName());
+
+					SelectedNodeDataStruct.GetScriptStruct();
+					FPropertyEditorModule& PropertyEditorModule = FModuleManager::LoadModuleChecked<FPropertyEditorModule>("PropertyEditor");
+					const TSharedPtr<FStructOnScope> InstanceDataStruct = MakeShared<FStructOnScope>(SelectedNodeDataStruct.GetScriptStruct(), const_cast<uint8*>(SelectedNodeDataStruct.GetMemory()));
+
+					FStructureDetailsViewArgs StructureViewArgs;
+					StructureViewArgs.bShowObjects = true;
+					StructureViewArgs.bShowAssets = true;
+					StructureViewArgs.bShowClasses = true;
+					StructureViewArgs.bShowInterfaces = true;
+					const TSharedRef<IStructureDetailsView> StructDetailsView = PropertyEditorModule.CreateStructureDetailView(DetailsViewArgs, StructureViewArgs, InstanceDataStruct);
+
+					DetailsView = StructDetailsView->GetDetailsView()->AsShared();
+				}
+
+				// UObject
+				UClass* Class = FindObject<UClass>(nullptr, *TypePath, /*ExactClass*/false);
+				if (Class == nullptr)
+				{
+					Class = LoadObject<UClass>(nullptr, *TypePath);
+				}
+
+				if (Class != nullptr)
+				{
+					if (SelectedNodeDataObject.IsValid())
+					{
+						SelectedNodeDataObject->RemoveFromRoot();
+					}
+					UE::StateTreeDebugger::FTraceTextObjectFactory ObjectFactory;
+					if (ObjectFactory.CanCreateObjectsFromText(InstanceDataAsText))
+					{
+						ObjectFactory.ProcessBuffer(GetTransientPackage(), RF_Transactional, InstanceDataAsText);
+						SelectedNodeDataObject = ObjectFactory.NodeInstanceObject;
+						SelectedNodeDataObject->AddToRoot();
+
+						FPropertyEditorModule& PropertyEditorModule = FModuleManager::Get().GetModuleChecked<FPropertyEditorModule>("PropertyEditor");
+						const TSharedRef<IDetailsView> ObjectDetailsView = PropertyEditorModule.CreateDetailView(DetailsViewArgs);
+						ObjectDetailsView->SetObject(SelectedNodeDataObject.Get());
+						DetailsView = ObjectDetailsView->AsShared();
+					}
+				}
+			}
+
+			if (DetailsView)
+			{
+				PropertiesBorder->SetContent(DetailsView.ToSharedRef());
+			}
+			else
+			{
+				PropertiesBorder->ClearContent();
+			}
+		})
+		.AllowOverscroll(EAllowOverscroll::No);
+
+	ChildSlot
+	[
+		SNew(SBorder)
+		.Padding(4.0f)
+		.BorderImage(FAppStyle::GetBrush("ToolPanel.GroupBorder"))
+		[
+			SNew(SVerticalBox)
+			+ SVerticalBox::Slot()
+			.AutoHeight()
+			.Padding(0.0f, 0.0f, 0.0f, 4.0f)
+			[
+				ToolbarBuilder.MakeWidget()
+			]
+			+ SVerticalBox::Slot().AutoHeight()
+			[
+				SNew(SBorder)
+				[
+					SAssignNew(HeaderSplitter, SSplitter)
+					.Orientation(Orient_Horizontal)
+					+ SSplitter::Slot()
+					.Value(0.2f)
+					.MinSize(350)
+					.Resizable(false)
+					[
+						TraceSelectionBox
+					]
+					+ SSplitter::Slot()
+					.Resizable(false)
+					[
+						SNew(SSimpleTimeSlider)
+						.DesiredSize({100, 24})
+						.ClampRangeHighlightSize(0.15f)
+						.ClampRangeHighlightColor(FLinearColor::Red.CopyWithNewOpacity(0.5f))
+						.ScrubPosition(ScrubTimeAttribute)
+						.ViewRange_Lambda([this]() { return ViewRange; })
+						.OnViewRangeChanged_Lambda([this](TRange<double> NewRange){ ViewRange = NewRange; })
+						.ClampRange_Lambda([this](){ return TRange<double>(0.0f, Debugger->GetRecordingDuration()); })
+						.OnScrubPositionChanged_Lambda([this](double NewScrubTime, bool bIsScrubbing){ OnTimeLineScrubPositionChanged(NewScrubTime, bIsScrubbing); })
+					]
+				]
+			]
+			+ SVerticalBox::Slot()
+			[
+				SNew(SSplitter)
+				.Orientation(Orient_Vertical)
+				+ SSplitter::Slot()
+				.Value(0.2f)
+				[
+					SAssignNew(TreeViewsSplitter, SSplitter)
+					.Orientation(Orient_Horizontal)
+					+ SSplitter::Slot()
+					.Value(0.2f)
+					.MinSize(350)
+					.OnSlotResized_Lambda([this](float Size)
+						{
+							// Sync both header and content
+							TreeViewsSplitter->SlotAt(0).SetSizeValue(Size);
+							HeaderSplitter->SlotAt(0).SetSizeValue(Size);
+						})
+					[
+						SNew(SScrollBox)
+						.Orientation(Orient_Horizontal)
+						+ SScrollBox::Slot()
+						.FillSize(1.0f)
+						[
+							InstancesTreeView.ToSharedRef()
+						]
+					]
+					+ SSplitter::Slot()
+					.OnSlotResized_Lambda([this](float Size)
+						{
+							TreeViewsSplitter->SlotAt(1).SetSizeValue(Size);
+							HeaderSplitter->SlotAt(1).SetSizeValue(Size);
+						})
+					[
+						SNew(SScrollBox)
+						.Orientation(Orient_Vertical)
+						+ SScrollBox::Slot()
+						[
+							InstanceTimelinesTreeView.ToSharedRef()
+						]
+					]
+				]
+				+ SSplitter::Slot()
+				[
+					SNew(SSplitter)
+					.Orientation(Orient_Horizontal)
+					+ SSplitter::Slot()
+					.MinSize(400)
+					[
+						SNew(SScrollBox)
+						.Orientation(Orient_Horizontal)
+						+ SScrollBox::Slot()
+						.FillSize(1.0f)
+						[
+							EventsTreeView.ToSharedRef()
+						]
+					]
+					+ SSplitter::Slot()
+					.MinSize(400)
+					[
+						SNew(SScrollBox)
+						.Orientation(Orient_Horizontal)
+						+ SScrollBox::Slot()
+						.FillSize(1.0f)
+						[
+							SAssignNew(PropertiesBorder, SBorder)
+						]
+					]
+				]
+			]
+		]
+	];
+}
+
+void SStateTreeDebuggerView::Tick(const FGeometry& AllottedGeometry, const double InCurrentTime, const float InDeltaTime)
+{
+	QUICK_SCOPE_CYCLE_COUNTER(STAT_StateTreeDebuggerView_TickView);
+
+	//Stick to most recent data
+	if (Debugger->IsPaused() == false)
+	{
+		Debugger->SetScrubTime(Debugger->GetRecordingDuration());
+	}
+	
+	RefreshTracks();
+}
+
+void SStateTreeDebuggerView::RefreshTracks()
+{
+	bool bChanged = false;
+	for (const TSharedPtr<RewindDebugger::FRewindDebuggerTrack>& DebugTrack : InstanceTracks)
+	{
+		bChanged = DebugTrack->Update() || bChanged;
+	}
+	
+	if (bChanged)
+	{
+		InstancesTreeView->Refresh();
+		InstanceTimelinesTreeView->Refresh();
+		TrackCursor();
+	}
+}
+
+void SStateTreeDebuggerView::BindDebuggerToolbarCommands(const TSharedRef<FUICommandList>& ToolkitCommands)
+{
+	const FStateTreeDebuggerCommands& Commands = FStateTreeDebuggerCommands::Get();
+	
+	ToolkitCommands->MapAction(
+		Commands.PreviousFrameWithStateChange,
+		FExecuteAction::CreateSP(this, &SStateTreeDebuggerView::StepBackToPreviousStateChange),
+		FCanExecuteAction::CreateSP(this, &SStateTreeDebuggerView::CanStepBackToPreviousStateChange));
+
+	ToolkitCommands->MapAction(
+		Commands.PreviousFrameWithEvents,
+		FExecuteAction::CreateSP(this, &SStateTreeDebuggerView::StepBackToPreviousStateWithEvents),
+		FCanExecuteAction::CreateSP(this, &SStateTreeDebuggerView::CanStepBackToPreviousStateWithEvents));
+
+	ToolkitCommands->MapAction(
+		Commands.NextFrameWithEvents,
+		FExecuteAction::CreateSP(this, &SStateTreeDebuggerView::StepForwardToNextStateWithEvents),
+		FCanExecuteAction::CreateSP(this, &SStateTreeDebuggerView::CanStepForwardToNextStateWithEvents));
+
+	ToolkitCommands->MapAction(
+		Commands.NextFrameWithStateChange,
+		FExecuteAction::CreateSP(this, &SStateTreeDebuggerView::StepForwardToNextStateChange),
+		FCanExecuteAction::CreateSP(this, &SStateTreeDebuggerView::CanStepForwardToNextStateChange));
+}
+
+bool SStateTreeDebuggerView::CanStepBackToPreviousStateWithEvents() const
+{
+	return Debugger->CanStepBackToPreviousStateWithEvents();
+}
+
+void SStateTreeDebuggerView::StepBackToPreviousStateWithEvents()
+{
+	Debugger->StepBackToPreviousStateWithEvents();
+}
+
+bool SStateTreeDebuggerView::CanStepForwardToNextStateWithEvents() const
+{
+	return Debugger->CanStepForwardToNextStateWithEvents();
+}
+
+void SStateTreeDebuggerView::StepForwardToNextStateWithEvents()
+{
+	Debugger->StepForwardToNextStateWithEvents();
+}
+
+bool SStateTreeDebuggerView::CanStepBackToPreviousStateChange() const
+{
+	return Debugger->CanStepBackToPreviousStateChange();
+}
+
+void SStateTreeDebuggerView::StepBackToPreviousStateChange()
+{
+	Debugger->StepBackToPreviousStateChange();
+}
+
+bool SStateTreeDebuggerView::CanStepForwardToNextStateChange() const
+{
+	return Debugger->CanStepForwardToNextStateChange();
+}
+
+void SStateTreeDebuggerView::StepForwardToNextStateChange()
+{
+	Debugger->StepForwardToNextStateChange();
+}
+
+bool SStateTreeDebuggerView::CanToggleBreakpoint() const
+{
+	return (Debugger.IsValid() && StateTreeViewModel.IsValid() && StateTreeViewModel->HasSelection());
+}
+
+void SStateTreeDebuggerView::ToggleBreakpoint() const
+{
+	check(StateTreeViewModel);
+	check(Debugger);
+	check(StateTree.IsValid());
+
+	TArray<UStateTreeState*> States;
+	StateTreeViewModel->GetSelectedStates(States);
+
+	TArray<FStateTreeStateHandle> StateHandles;
+	StateHandles.Reserve(States.Num());
+	for (const UStateTreeState* SelectedState : States)
+	{
+		if (SelectedState->Type != EStateTreeStateType::State || SelectedState->Parent == nullptr)
+		{
+			continue;
+		}
+		
+		FStateTreeStateHandle Handle = StateTree->GetStateHandleFromId(SelectedState->ID);
+		if (Handle.IsValid())
+		{
+			StateHandles.Add(Handle);
+		}
+	}
+	Debugger->ToggleBreakpoints(StateHandles);
+}
+
+void SStateTreeDebuggerView::OnTimeLineScrubPositionChanged(double Time, bool bIsScrubbing)
+{
+	Debugger->Pause();
+	Debugger->SetScrubTime(Time);
+}
+
+void SStateTreeDebuggerView::OnDebuggerScrubStateChanged(const UE::StateTreeDebugger::FScrubState& ScrubState)
+{
+	// Rebuild frame details from the events of that frame
+	EventsTreeElements.Reset();
+	EventsTreeView->RequestTreeRefresh();
+
+	const UE::StateTreeDebugger::FInstanceEventCollection& EventCollection = ScrubState.GetEventCollection();
+	const TConstArrayView<const FStateTreeTraceEventVariantType> Events = EventCollection.Events;
+
+	if (Events.IsEmpty() || ScrubState.IsInBounds() == false)
+	{
+		return;
+	}
+
+	const TConstArrayView<UE::StateTreeDebugger::FFrameSpan> Spans = EventCollection.FrameSpans;
+	check(Spans.Num());
+	check(StateTree.IsValid());
+
+	TArray<TSharedPtr<UE::StateTreeDebugger::FEventTreeElement>, TInlineAllocator<8>> ParentStack;
+	EStateTreeUpdatePhase LastPhase = EStateTreeUpdatePhase::Unset;
+	EStateTreeUpdatePhase EventPhase = EStateTreeUpdatePhase::Unset;
+
+	const int32 SpanIdx = ScrubState.FrameSpanIndex;
+	const int32 FirstEventIdx = Spans[SpanIdx].EventIdx;
+	const int32 MaxEventIdx = Spans.IsValidIndex(SpanIdx+1) ? Spans[SpanIdx+1].EventIdx : Events.Num();
+
+	for (int32 EventIdx = FirstEventIdx; EventIdx < MaxEventIdx; EventIdx++)
+	{
+		const FStateTreeTraceEventVariantType& Event = Events[EventIdx];
+		EventPhase = LastPhase;
+
+		// Need to test each type explicitly with TVariant even if they are all FStateTreeTracePhaseEvent
+		if (const FStateTreeTraceLogEvent* LogEvent = Event.TryGet<FStateTreeTraceLogEvent>())
+		{
+			EventPhase = LogEvent->Phase;
+		}
+		else if (const FStateTreeTraceStateEvent* StateEvent = Event.TryGet<FStateTreeTraceStateEvent>())
+		{
+			EventPhase = StateEvent->Phase;
+		}
+		else if (const FStateTreeTraceTaskEvent* TaskEvent = Event.TryGet<FStateTreeTraceTaskEvent>())
+		{
+			EventPhase = TaskEvent->Phase;
+		}
+		else if (const FStateTreeTraceConditionEvent* ConditionEvent = Event.TryGet<FStateTreeTraceConditionEvent>())
+		{
+			EventPhase = ConditionEvent->Phase;
+		}
+
+		// Create a hierarchy level for each update phase
+		if (EventPhase != LastPhase)
+		{
+			const UEnum* PhaseEnum = StaticEnum<EStateTreeUpdatePhase>();
+			EStateTreeUpdatePhase PhasesDiff = EventPhase ^ LastPhase;
+			const int32 NumEnum = PhaseEnum->NumEnums();
+			check(NumEnum > 0);
+
+			// Pop phases first from last enum (every bit in the previous phase that differs must be popped)
+			if (EnumHasAnyFlags(LastPhase, PhasesDiff))
+			{
+				for (int i = NumEnum - 1; i >= 0; --i)
+				{
+					const EStateTreeUpdatePhase Phase = static_cast<EStateTreeUpdatePhase>(PhaseEnum->GetValueByIndex(i));
+					if (EnumHasAnyFlags(LastPhase & PhasesDiff, Phase))
+					{
+						check(ParentStack.Num());
+						TSharedPtr<UE::StateTreeDebugger::FEventTreeElement> RemovedElement = ParentStack.Pop();
+						if (RemovedElement->Children.IsEmpty())
+						{
+							if (ParentStack.Num())
+							{
+								ParentStack.Last()->Children.Remove(RemovedElement);
+							}
+						}
+
+						EnumRemoveFlags(PhasesDiff, Phase);
+						if (EnumHasAnyFlags(LastPhase, PhasesDiff) == false)
+						{
+							break;
+						}
+					}
+				}
+			}
+
+			// Push required phases from first enum
+			if (EnumHasAnyFlags(EventPhase, PhasesDiff))
+			{
+				for (int i = 0; i < NumEnum; ++i)
+				{
+					const EStateTreeUpdatePhase Phase = static_cast<EStateTreeUpdatePhase>(PhaseEnum->GetValueByIndex(i));
+					if (EnumHasAnyFlags(EventPhase & PhasesDiff, Phase))
+					{
+						// Create fake log event to describe the phase
+						FStateTreeTraceLogEvent DummyEvent(EStateTreeUpdatePhase::Unset,
+							FString::Printf(TEXT("%s"), *StaticEnum<EStateTreeUpdatePhase>()->GetValueOrBitfieldAsString((int64)Phase)));
+
+						// Create Tree element to hold the event
+						TSharedPtr<UE::StateTreeDebugger::FEventTreeElement> NewElement = MakeShareable(new UE::StateTreeDebugger::FEventTreeElement(
+							Spans[SpanIdx].Frame,
+							FStateTreeTraceEventVariantType(TInPlaceType<FStateTreeTraceLogEvent>(), DummyEvent)));
+
+						// Push tree element to the proper stack level
+						TArray<TSharedPtr<UE::StateTreeDebugger::FEventTreeElement>>& Elements = ParentStack.IsEmpty() ? EventsTreeElements : ParentStack.Last()->Children;
+						ParentStack.Push(Elements.Add_GetRef(NewElement));
+						
+						EnumRemoveFlags(PhasesDiff, Phase);
+						if (EnumHasAnyFlags(EventPhase, PhasesDiff) == false)
+						{
+							break;
+						}
+					}
+				}
+			}
+
+			LastPhase = EventPhase;
+		}
+
+		TArray<TSharedPtr<UE::StateTreeDebugger::FEventTreeElement>>& Elements = ParentStack.IsEmpty() ? EventsTreeElements : ParentStack.Last()->Children;
+		Elements.Add(MakeShareable(new UE::StateTreeDebugger::FEventTreeElement(Spans[SpanIdx].Frame, Event)));
+	}
+}
+
+void SStateTreeDebuggerView::OnBreakpointHit(const FStateTreeInstanceDebugId InstanceId, const FStateTreeStateHandle StateHandle, const TSharedRef<FUICommandList> ActionList) const
+{
+	if (FPlayWorldCommands::Get().PausePlaySession.IsValid())
+	{
+		if (ActionList->CanExecuteAction(FPlayWorldCommands::Get().PausePlaySession.ToSharedRef()))
+		{
+			ActionList->ExecuteAction(FPlayWorldCommands::Get().PausePlaySession.ToSharedRef());
+		}
+	}
+}
+
+void SStateTreeDebuggerView::OnNewInstance(FStateTreeInstanceDebugId InstanceId)
+{
+	const TSharedPtr<RewindDebugger::FRewindDebuggerTrack>* ExistingTrack = InstanceTracks.FindByPredicate([InstanceId](const TSharedPtr<RewindDebugger::FRewindDebuggerTrack>& Track)
+	{
+		return static_cast<FStateTreeDebuggerTrack*>(Track.Get())->GetInstanceId() == InstanceId;
+	});
+
+	if (ExistingTrack == nullptr)
+	{
+		InstanceTracks.Add(MakeShared<FStateTreeDebuggerTrack>(Debugger, InstanceId, Debugger->GetInstanceDescription(InstanceId), ViewRange));
+	}
+
+	InstancesTreeView->Refresh();
+	InstanceTimelinesTreeView->Refresh();
+}
+
+void SStateTreeDebuggerView::OnSelectedInstanceCleared()
+{
+	EventsTreeElements.Reset();
+	if (EventsTreeView)
+	{
+		EventsTreeView->RequestTreeRefresh();	
+	}
+
+	PropertiesBorder->ClearContent();
+}
+
+TSharedRef<SWidget> SStateTreeDebuggerView::OnGetDebuggerTracesMenu() const
+{
+	FMenuBuilder MenuBuilder(/*bInShouldCloseWindowAfterMenuSelection*/true, /*InCommandList*/nullptr);
+
+	if (Debugger.IsValid())
+	{
+		TArray<FStateTreeDebugger::FTraceDescriptor> TraceDescriptors;
+		Debugger->GetLiveTraces(TraceDescriptors);
+
+		for (const FStateTreeDebugger::FTraceDescriptor& TraceDescriptor : TraceDescriptors)
+		{
+			const FText Desc = Debugger->DescribeTrace(TraceDescriptor);
+
+			FUIAction ItemAction(FExecuteAction::CreateLambda([WeakDebugger = Debugger, TraceDescriptor]()
+			{
+				if (WeakDebugger)
+				{
+					WeakDebugger->StartSessionAnalysis(TraceDescriptor);
+				}
+			}));
+			MenuBuilder.AddMenuEntry(Desc, TAttribute<FText>(), FSlateIcon(), ItemAction);
+		}
+
+		// Failsafe when no match
+		if (TraceDescriptors.Num() == 0)
+		{
+			const FText Desc = LOCTEXT("NoLiveSessions","Can't find live trace sessions");
+			FUIAction ItemAction(FExecuteAction::CreateLambda([WeakDebugger = Debugger]()
+			{
+				if (WeakDebugger)
+				{
+					WeakDebugger->StartSessionAnalysis(FStateTreeDebugger::FTraceDescriptor());
+				}
+			}));
+			MenuBuilder.AddMenuEntry(Desc, TAttribute<FText>(), FSlateIcon(), ItemAction);
+		}
+	}
+
+	return MenuBuilder.MakeWidget();
+}
+
+void SStateTreeDebuggerView::TrackCursor()
+{
+	const double ScrubTime = ScrubTimeAttribute.Get();	
+	TRange<double> CurrentViewRange = ViewRange;
+	const double ViewRangeDuration = CurrentViewRange.GetUpperBoundValue() - CurrentViewRange.GetLowerBoundValue();
+
+	static constexpr double LeadingMarginFraction = 0.05;
+	static constexpr double TrailingMarginFraction = 0.01;
+
+	if (ScrubTime > CurrentViewRange.GetUpperBoundValue() - ViewRangeDuration * LeadingMarginFraction)
+	{
+		CurrentViewRange.SetUpperBound(ScrubTime + ViewRangeDuration * LeadingMarginFraction);
+		CurrentViewRange.SetLowerBound(CurrentViewRange.GetUpperBoundValue() - ViewRangeDuration);
+	}
+
+	if (ScrubTime < CurrentViewRange.GetLowerBoundValue() - ViewRangeDuration * TrailingMarginFraction)
+	{
+		CurrentViewRange.SetLowerBound(ScrubTime);
+		CurrentViewRange.SetUpperBound(CurrentViewRange.GetLowerBoundValue() + ViewRangeDuration);
+	}
+
+	ViewRange = CurrentViewRange;
+}
+
+#undef LOCTEXT_NAMESPACE
+
+#endif // WITH_STATETREE_DEBUGGER
\ No newline at end of file
diff --git a/Engine/Plugins/Runtime/StateTree/Source/StateTreeEditorModule/Private/Debugger/SStateTreeDebuggerView.h b/Engine/Plugins/Runtime/StateTree/Source/StateTreeEditorModule/Private/Debugger/SStateTreeDebuggerView.h
new file mode 100644
--- /dev/null
+++ b/Engine/Plugins/Runtime/StateTree/Source/StateTreeEditorModule/Private/Debugger/SStateTreeDebuggerView.h
@@ -0,0 +1,138 @@
+// Copyright Epic Games, Inc. All Rights Reserved.
+
+#pragma once
+#include "Debugger/StateTreeDebugger.h"
+
+#if WITH_STATETREE_DEBUGGER
+
+#include "Debugger/StateTreeTraceTypes.h"
+#include "InstancedStruct.h"
+#include "Widgets/Views/STreeView.h"
+#include "Widgets/SCompoundWidget.h"
+
+namespace RewindDebugger
+{
+	class FRewindDebuggerTrack;
+}
+namespace UE::StateTreeDebugger
+{
+	struct FFrameSpan;
+	struct FEventTreeElement;
+}
+struct FStateTreeDebugger;
+struct FStateTreeInstanceDebugId;
+class FStateTreeEditor;
+class FStateTreeViewModel;
+class FUICommandList;
+class UStateTree;
+class SStateTreeDebuggerTimelines;
+class SStateTreeDebuggerInstanceTree;
+
+
+/**
+ * Widget holding the timelines for all statetree instances matching a given asset
+ * in addition to some frame details panels.
+ */
+class SStateTreeDebuggerView : public SCompoundWidget
+{
+public:
+	SLATE_BEGIN_ARGS(SStateTreeDebuggerView) {}
+	SLATE_END_ARGS()
+
+	SStateTreeDebuggerView();
+	~SStateTreeDebuggerView();
+
+	void Construct(const FArguments& InArgs, const UStateTree* InStateTree, const TSharedRef<FStateTreeViewModel>& InStateTreeViewModel, const TSharedRef<FUICommandList>& InCommandList);
+
+private:
+	virtual void Tick(const FGeometry& AllottedGeometry, const double InCurrentTime, const float InDeltaTime) override;
+
+	void RefreshTracks();
+
+	void TrackCursor();
+
+	TSharedRef<SWidget> OnGetDebuggerTracesMenu() const;
+
+	void OnPIEStarted(bool bIsSimulating) const;
+	void OnPIEStopped(bool bIsSimulating) const;
+	void OnPIEPaused(bool bIsSimulating) const;
+	void OnPIEResumed(bool bIsSimulating) const;
+	void OnPIESingleStepped(bool bIsSimulating) const;
+	
+	void OnBreakpointHit(const FStateTreeInstanceDebugId InstanceId, const FStateTreeStateHandle StateHandle, const TSharedRef<FUICommandList> ActionList) const;
+	void OnNewInstance(FStateTreeInstanceDebugId InstanceId);
+	void OnSelectedInstanceCleared();
+
+	void BindDebuggerToolbarCommands(const TSharedRef<FUICommandList>& ToolkitCommands);
+
+	bool CanStepBackToPreviousStateWithEvents() const;
+	void StepBackToPreviousStateWithEvents();
+
+	bool CanStepForwardToNextStateWithEvents() const;
+	void StepForwardToNextStateWithEvents();
+
+	bool CanStepBackToPreviousStateChange() const;
+	void StepBackToPreviousStateChange();
+
+	bool CanStepForwardToNextStateChange() const;
+	void StepForwardToNextStateChange();
+	
+	bool CanToggleBreakpoint() const;
+	void ToggleBreakpoint() const;
+
+	TSharedPtr<FStateTreeDebugger> Debugger;
+	TSharedPtr<FStateTreeViewModel> StateTreeViewModel;
+	TWeakObjectPtr<const UStateTree> StateTree;
+
+	/** Tracks for all instances producing trace events for the associated state tree asset. */
+	TArray<TSharedPtr<RewindDebugger::FRewindDebuggerTrack>> InstanceTracks;
+
+	/**
+	 * Tree view displaying the instance names and synced with InstanceTimelinesTreeView.
+	 * Note that this is currently used as a list view but kept the tree view to be close
+	 * to the rewind debugger track implementation.
+	 */
+	TSharedPtr<SStateTreeDebuggerInstanceTree> InstancesTreeView;
+
+	/**
+	 * Tree view displaying the instance timelines and synced with InstancesTreeView.
+	 * Note that this is currently used as a list view but kept the tree view to be close
+	 * to the rewind debugger track implementation.
+	 */
+	TSharedPtr<SStateTreeDebuggerTimelines> InstanceTimelinesTreeView;
+
+	/** Splitter between instances selector and simple time slider. Used with TreeViewsSplitter to keep header and content in sync */
+	TSharedPtr<SSplitter> HeaderSplitter;
+
+	/** Splitter between instances names and their timelines. Used with HeaderSplitter to keep header and content in sync */
+	TSharedPtr<SSplitter> TreeViewsSplitter;
+
+	/** All trace events received for a given instance. */
+	TArray<TSharedPtr<UE::StateTreeDebugger::FEventTreeElement>> EventsTreeElements;
+	
+	/** Tree view displaying the frame events of the instance associated to the selected track. */
+	TSharedPtr<STreeView<TSharedPtr<UE::StateTreeDebugger::FEventTreeElement>>> EventsTreeView;
+
+	/** Widget holding the event details when available (i.e. property details of the associated struct or object). */
+	TSharedPtr<SBorder> PropertiesBorder;
+
+	/** Attribute provided by the debugger scrub position to control cursor and timelines positions. */
+	TAttribute<double> ScrubTimeAttribute;
+
+	/** Callback from timeline widgets to update the debugger scrub state. */
+	void OnTimeLineScrubPositionChanged(double Time, bool bIsScrubbing);
+
+	/** Callback used to reflect debugger scrub state in the UI. */
+	void OnDebuggerScrubStateChanged(const UE::StateTreeDebugger::FScrubState& ScrubState);
+
+	/** Range controlled by the timeline widgets and used to adjust cursor position and track content. */
+	TRange<double> ViewRange = TRange<double>(0, 10);
+
+	/** Struct created from the event data when statetree node was holding a struct. */
+	FInstancedStruct SelectedNodeDataStruct;
+
+	/** Object created from the event data when statetree node was holding an object. */
+	TWeakObjectPtr<UObject> SelectedNodeDataObject;
+};
+
+#endif // WITH_STATETREE_DEBUGGER
\ No newline at end of file
diff --git a/Engine/Plugins/Runtime/StateTree/Source/StateTreeEditorModule/Private/Debugger/StateTreeDebuggerCommands.cpp b/Engine/Plugins/Runtime/StateTree/Source/StateTreeEditorModule/Private/Debugger/StateTreeDebuggerCommands.cpp
new file mode 100644
--- /dev/null
+++ b/Engine/Plugins/Runtime/StateTree/Source/StateTreeEditorModule/Private/Debugger/StateTreeDebuggerCommands.cpp
@@ -0,0 +1,35 @@
+// Copyright Epic Games, Inc. All Rights Reserved.
+
+#if WITH_STATETREE_DEBUGGER
+
+#include "StateTreeDebuggerCommands.h"
+
+#include "StateTreeEditorStyle.h"
+#include "Styling/AppStyle.h"
+
+#define LOCTEXT_NAMESPACE "StateTreeDebugger"
+
+FStateTreeDebuggerCommands::FStateTreeDebuggerCommands()
+	: TCommands
+	(
+		"StateTreeEditor.Debugger",								// Context name for fast lookup
+		LOCTEXT("StateTreeDebugger", "StateTree Debugger"),		// Localized context name for displaying
+		NAME_None,												// Parent context name
+		FStateTreeEditorStyle::Get().GetStyleSetName()
+	)
+{
+}
+
+void FStateTreeDebuggerCommands::RegisterCommands()
+{
+	UI_COMMAND(ToggleBreakpoint, "Toggle Breakpoint", "Adds or removes a breakpoint on each selected node", EUserInterfaceActionType::Button, FInputChord(EKeys::F9))
+
+	UI_COMMAND(PreviousFrameWithStateChange, "Previous State Change", "Jump to previous state changed frame", EUserInterfaceActionType::Button, FInputChord(EModifierKey::Control, EKeys::Left));
+	UI_COMMAND(PreviousFrameWithEvents, "Previous Frame", "Step one frame back", EUserInterfaceActionType::Button, FInputChord(EKeys::Left));
+	UI_COMMAND(NextFrameWithEvents, "Next Frame", "Step one frame forward", EUserInterfaceActionType::Button, FInputChord(EKeys::Right));
+	UI_COMMAND(NextFrameWithStateChange, "Next State Change", "Jump to next state changed frame", EUserInterfaceActionType::Button, FInputChord(EModifierKey::Control, EKeys::Right));
+}
+
+#undef LOCTEXT_NAMESPACE
+
+#endif // WITH_STATETREE_DEBUGGER
\ No newline at end of file
diff --git a/Engine/Plugins/Runtime/StateTree/Source/StateTreeEditorModule/Private/Debugger/StateTreeDebuggerCommands.h b/Engine/Plugins/Runtime/StateTree/Source/StateTreeEditorModule/Private/Debugger/StateTreeDebuggerCommands.h
new file mode 100644
--- /dev/null
+++ b/Engine/Plugins/Runtime/StateTree/Source/StateTreeEditorModule/Private/Debugger/StateTreeDebuggerCommands.h
@@ -0,0 +1,27 @@
+// Copyright Epic Games, Inc. All Rights Reserved.
+
+#pragma once
+
+#if WITH_STATETREE_DEBUGGER
+
+#include "Framework/Commands/Commands.h"
+
+/**
+ * StateTree Debugger command set.
+ */
+class FStateTreeDebuggerCommands : public TCommands<FStateTreeDebuggerCommands>
+{
+public:
+	FStateTreeDebuggerCommands();
+
+	// TCommands<> overrides
+	virtual void RegisterCommands() override;
+
+	TSharedPtr<FUICommandInfo> ToggleBreakpoint;
+	TSharedPtr<FUICommandInfo> PreviousFrameWithStateChange;
+	TSharedPtr<FUICommandInfo> PreviousFrameWithEvents;
+	TSharedPtr<FUICommandInfo> NextFrameWithEvents;
+	TSharedPtr<FUICommandInfo> NextFrameWithStateChange;
+};
+
+#endif // WITH_STATETREE_DEBUGGER
\ No newline at end of file
diff --git a/Engine/Plugins/Runtime/StateTree/Source/StateTreeEditorModule/Private/Debugger/StateTreeDebuggerTrack.cpp b/Engine/Plugins/Runtime/StateTree/Source/StateTreeEditorModule/Private/Debugger/StateTreeDebuggerTrack.cpp
new file mode 100644
--- /dev/null
+++ b/Engine/Plugins/Runtime/StateTree/Source/StateTreeEditorModule/Private/Debugger/StateTreeDebuggerTrack.cpp
@@ -0,0 +1,115 @@
+// Copyright Epic Games, Inc. All Rights Reserved.
+
+#if WITH_STATETREE_DEBUGGER
+
+#include "Debugger/StateTreeDebuggerTrack.h"
+#include "Debugger/StateTreeDebugger.h"
+#include "SStateTreeDebuggerEventTimelineView.h"
+#include "Styling/AppStyle.h"
+
+//----------------------------------------------------------------------//
+// FStateTreeTrack
+//----------------------------------------------------------------------//
+FStateTreeDebuggerTrack::FStateTreeDebuggerTrack(const TSharedPtr<FStateTreeDebugger>& InDebugger, const FStateTreeInstanceDebugId InInstanceId, const FText InName, const TRange<double>& InViewRange)
+	: TrackName(InName)
+	, StateTreeDebugger(InDebugger)
+	, InstanceId(InInstanceId)
+	, ViewRange(InViewRange)
+{
+	EventData = MakeShared<SStateTreeDebuggerEventTimelineView::FTimelineEventData>();
+}
+
+TSharedPtr<SWidget> FStateTreeDebuggerTrack::GetDetailsViewInternal() 
+{
+	return SNullWidget::NullWidget;
+}
+
+bool FStateTreeDebuggerTrack::UpdateInternal()
+{
+	const int32 PrevNumPoints = EventData->Points.Num();
+	const int32 PrevNumWindows = EventData->Windows.Num();
+
+	EventData->Points.SetNum(0, false);
+	EventData->Windows.SetNum(0);
+	
+	const FStateTreeDebugger* Debugger = StateTreeDebugger.Get();
+	check(Debugger);
+	const UStateTree* StateTree = Debugger->GetAsset();
+	const UE::StateTreeDebugger::FInstanceEventCollection& EventCollection = Debugger->GetEventCollection(InstanceId);
+	const double RecordingDuration = Debugger->GetRecordingDuration();
+	
+	if (StateTree != nullptr && EventCollection.IsValid())
+	{
+		auto MakeRandomColor = [](const uint32 InSeed)->FLinearColor
+		{
+			const FRandomStream Stream(InSeed);
+			const uint8 Hue = (uint8)(Stream.FRand() * 255.0f);
+			constexpr uint8 SatVal = 196;
+			return FLinearColor::MakeFromHSV8(Hue, SatVal, SatVal);
+		};
+		
+		const TConstArrayView<UE::StateTreeDebugger::FFrameSpan> Spans = EventCollection.FrameSpans;
+		const TConstArrayView<FStateTreeTraceEventVariantType> Events = EventCollection.Events;
+		const uint32 NumStateChanges = EventCollection.ActiveStatesChanges.Num();
+		
+		for (uint32 StatechangeIndex = 0; StatechangeIndex < NumStateChanges; ++StatechangeIndex)
+		{
+			const uint32 SpanIndex = EventCollection.ActiveStatesChanges[StatechangeIndex].Key;
+			const uint32 EventIndex = EventCollection.ActiveStatesChanges[StatechangeIndex].Value;
+			const FStateTreeTraceActiveStatesEvent& Event = Events[EventIndex].Get<FStateTreeTraceActiveStatesEvent>();
+				
+			FString StatePath;
+			for (int32 StateIndex = 0; StateIndex < Event.ActiveStates.Num(); StateIndex++)
+			{
+				const FCompactStateTreeState& State = StateTree->GetStates()[Event.ActiveStates[StateIndex].Index];
+				StatePath.Appendf(TEXT("%s%s"), StateIndex == 0 ? TEXT("") : TEXT("."), *State.Name.ToString());
+			}
+			
+			SStateTreeDebuggerEventTimelineView::FTimelineEventData::EventWindow& Window = EventData->Windows.AddDefaulted_GetRef();
+			Window.Color = MakeRandomColor(GetTypeHash(StatePath));
+			Window.Description = FText::FromString(StatePath);
+			Window.TimeStart = EventCollection.FrameSpans[SpanIndex].Frame.StartTime;
+
+			// When there is another state change after the current one in the list we use its start time to close the window.
+			if (StatechangeIndex < NumStateChanges-1)
+			{
+				const uint32 NextSpanIndex = EventCollection.ActiveStatesChanges[StatechangeIndex+1].Key;
+				Window.TimeEnd = EventCollection.FrameSpans[NextSpanIndex].Frame.StartTime;
+			}
+			else
+			{
+				Window.TimeEnd = Debugger->IsActiveInstance(RecordingDuration, InstanceId) ? RecordingDuration : EventCollection.FrameSpans[SpanIndex].Frame.EndTime;
+			}
+		}
+
+		for (int32 SpanIndex = 0; SpanIndex < Spans.Num(); SpanIndex++)
+		{
+			const UE::StateTreeDebugger::FFrameSpan& Span = Spans[SpanIndex];
+
+			const int32 StartIndex = Span.EventIdx;
+			const int32 MaxIndex = (SpanIndex + 1 < Spans.Num()) ? Spans[SpanIndex+1].EventIdx : Events.Num();
+			for (int EventIndex = StartIndex; EventIndex < MaxIndex; ++EventIndex)
+			{
+				if (Events[EventIndex].IsType<FStateTreeTraceLogEvent>())
+				{
+					SStateTreeDebuggerEventTimelineView::FTimelineEventData::EventPoint Point;
+					Point.Time = Span.Frame.StartTime;
+					Point.Color = FColorList::Salmon;
+					EventData->Points.Add(Point);					
+				}
+			}
+		}
+	}
+
+	const bool bChanged = (PrevNumPoints != EventData->Points.Num() || PrevNumWindows != EventData->Windows.Num());
+	return bChanged;
+}
+
+TSharedPtr<SWidget> FStateTreeDebuggerTrack::GetTimelineViewInternal()
+{
+	return SNew(SStateTreeDebuggerEventTimelineView)
+		.ViewRange_Lambda([this](){ return ViewRange; })
+		.EventData_Lambda([this](){ return EventData; });
+}
+
+#endif // WITH_STATETREE_DEBUGGER
\ No newline at end of file
diff --git a/Engine/Plugins/Runtime/StateTree/Source/StateTreeEditorModule/Private/Debugger/StateTreeDebuggerTrack.h b/Engine/Plugins/Runtime/StateTree/Source/StateTreeEditorModule/Private/Debugger/StateTreeDebuggerTrack.h
new file mode 100644
--- /dev/null
+++ b/Engine/Plugins/Runtime/StateTree/Source/StateTreeEditorModule/Private/Debugger/StateTreeDebuggerTrack.h
@@ -0,0 +1,40 @@
+// Copyright Epic Games, Inc. All Rights Reserved.
+
+#pragma once
+
+#if WITH_STATETREE_DEBUGGER
+
+#include "RewindDebuggerTrack.h"
+#include "SStateTreeDebuggerEventTimelineView.h"
+#include "StateTreeTypes.h"
+
+struct FStateTreeDebugger;
+
+/**
+ * Struct use to store timeline events for a single StateTree instance
+ * This is currently not tied to RewindDebugger but we use this base
+ * type to facilitate a future integration.
+ */
+struct FStateTreeDebuggerTrack : RewindDebugger::FRewindDebuggerTrack
+{
+	FStateTreeDebuggerTrack(const TSharedPtr<FStateTreeDebugger>& InDebugger, const FStateTreeInstanceDebugId InInstanceId, const FText InName, const TRange<double>& InViewRange);
+	
+	FStateTreeInstanceDebugId GetInstanceId() const { return InstanceId; }
+private:
+	virtual bool UpdateInternal() override;
+	virtual TSharedPtr<SWidget> GetDetailsViewInternal() override;
+	virtual TSharedPtr<SWidget> GetTimelineViewInternal() override;
+	virtual FSlateIcon GetIconInternal() override { return Icon; }
+	virtual FText GetDisplayNameInternal() const override { return TrackName; }
+
+	TSharedPtr<SStateTreeDebuggerEventTimelineView::FTimelineEventData> EventData;
+	
+	FSlateIcon Icon;
+	FText TrackName;
+
+	TSharedPtr<FStateTreeDebugger> StateTreeDebugger;
+	FStateTreeInstanceDebugId InstanceId;
+	const TRange<double>& ViewRange;
+};
+
+#endif // WITH_STATETREE_DEBUGGER
\ No newline at end of file
diff --git a/Engine/Plugins/Runtime/StateTree/Source/StateTreeEditorModule/Private/SStateTreeDebuggerView.cpp b/Engine/Plugins/Runtime/StateTree/Source/StateTreeEditorModule/Private/SStateTreeDebuggerView.cpp
deleted file mode 100644
diff --git a/Engine/Plugins/Runtime/StateTree/Source/StateTreeEditorModule/Private/SStateTreeDebuggerView.h b/Engine/Plugins/Runtime/StateTree/Source/StateTreeEditorModule/Private/SStateTreeDebuggerView.h
deleted file mode 100644
diff --git a/Engine/Plugins/Runtime/StateTree/Source/StateTreeEditorModule/Private/SStateTreeView.cpp b/Engine/Plugins/Runtime/StateTree/Source/StateTreeEditorModule/Private/SStateTreeView.cpp
--- a/Engine/Plugins/Runtime/StateTree/Source/StateTreeEditorModule/Private/SStateTreeView.cpp
+++ b/Engine/Plugins/Runtime/StateTree/Source/StateTreeEditorModule/Private/SStateTreeView.cpp
@@ -5,7 +5,7 @@
 #include "Widgets/Layout/SScrollBox.h"
 #include "SPositiveActionButton.h"
 #include "SStateTreeViewRow.h"
-#include "StateTreeDebuggerCommands.h"
+#include "Debugger/StateTreeDebuggerCommands.h"
 #include "StateTreeViewModel.h"
 #include "StateTreeState.h"
 #include "StateTreeEditorCommands.h"
diff --git a/Engine/Plugins/Runtime/StateTree/Source/StateTreeEditorModule/Private/StateTreeDebuggerCommands.cpp b/Engine/Plugins/Runtime/StateTree/Source/StateTreeEditorModule/Private/StateTreeDebuggerCommands.cpp
deleted file mode 100644
diff --git a/Engine/Plugins/Runtime/StateTree/Source/StateTreeEditorModule/Private/StateTreeDebuggerCommands.h b/Engine/Plugins/Runtime/StateTree/Source/StateTreeEditorModule/Private/StateTreeDebuggerCommands.h
deleted file mode 100644
diff --git a/Engine/Plugins/Runtime/StateTree/Source/StateTreeEditorModule/Private/StateTreeEditor.cpp b/Engine/Plugins/Runtime/StateTree/Source/StateTreeEditorModule/Private/StateTreeEditor.cpp
--- a/Engine/Plugins/Runtime/StateTree/Source/StateTreeEditorModule/Private/StateTreeEditor.cpp
+++ b/Engine/Plugins/Runtime/StateTree/Source/StateTreeEditorModule/Private/StateTreeEditor.cpp
@@ -27,7 +27,7 @@
 #include "ToolMenuEntry.h"
 #include "Widgets/Input/SMultiLineEditableTextBox.h"
 #include "FileHelpers.h"
-#include "SStateTreeDebuggerView.h"
+#include "Debugger/SStateTreeDebuggerView.h"
 #include "StateTreeSettings.h"
 #include "ToolMenuSection.h"
 
diff --git a/Engine/Plugins/Runtime/StateTree/Source/StateTreeEditorModule/Private/StateTreeEditorModule.cpp b/Engine/Plugins/Runtime/StateTree/Source/StateTreeEditorModule/Private/StateTreeEditorModule.cpp
--- a/Engine/Plugins/Runtime/StateTree/Source/StateTreeEditorModule/Private/StateTreeEditorModule.cpp
+++ b/Engine/Plugins/Runtime/StateTree/Source/StateTreeEditorModule/Private/StateTreeEditorModule.cpp
@@ -16,7 +16,7 @@
 #include "StateTree.h"
 #include "StateTreeCompiler.h"
 #include "StateTreeCompilerLog.h"
-#include "StateTreeDebuggerCommands.h"
+#include "Debugger/StateTreeDebuggerCommands.h"
 #include "StateTreeDelegates.h"
 #include "StateTreeEditor.h"
 #include "StateTreeEditorCommands.h"
diff --git a/Engine/Plugins/Runtime/StateTree/Source/StateTreeEditorModule/Private/StateTreeEditorStyle.cpp b/Engine/Plugins/Runtime/StateTree/Source/StateTreeEditorModule/Private/StateTreeEditorStyle.cpp
--- a/Engine/Plugins/Runtime/StateTree/Source/StateTreeEditorModule/Private/StateTreeEditorStyle.cpp
+++ b/Engine/Plugins/Runtime/StateTree/Source/StateTreeEditorModule/Private/StateTreeEditorStyle.cpp
@@ -178,8 +178,10 @@ FStateTreeEditorStyle::FStateTreeEditorStyle()
 	{
 		// From generic engine editor
 		UE::StateTree::Editor::FContentRootScope Scope(this, EngineEditorSlateContentDir);
-		Set("StateTreeEditor.Debugger.Back", new IMAGE_BRUSH_SVG("Starship/Blueprints/diff_prev", CoreStyleConstants::Icon20x20));
-		Set("StateTreeEditor.Debugger.Forward", new IMAGE_BRUSH_SVG("Starship/Blueprints/diff_next", CoreStyleConstants::Icon20x20));
+		Set("StateTreeEditor.Debugger.PreviousFrameWithStateChange", new IMAGE_BRUSH("Sequencer/Transport_Bar/Go_To_Front_24x", CoreStyleConstants::Icon16x16));
+		Set("StateTreeEditor.Debugger.PreviousFrameWithEvents", new IMAGE_BRUSH("Sequencer/Transport_Bar/Step_Backwards_24x", CoreStyleConstants::Icon16x16));
+		Set("StateTreeEditor.Debugger.NextFrameWithEvents", new IMAGE_BRUSH("Sequencer/Transport_Bar/Step_Forward_24x", CoreStyleConstants::Icon16x16));
+		Set("StateTreeEditor.Debugger.NextFrameWithStateChange", new IMAGE_BRUSH("Sequencer/Transport_Bar/Go_To_End_24x", CoreStyleConstants::Icon16x16));
 	}
 
 	{
diff --git a/Engine/Plugins/Runtime/StateTree/Source/StateTreeEditorModule/StateTreeEditorModule.Build.cs b/Engine/Plugins/Runtime/StateTree/Source/StateTreeEditorModule/StateTreeEditorModule.Build.cs
--- a/Engine/Plugins/Runtime/StateTree/Source/StateTreeEditorModule/StateTreeEditorModule.Build.cs
+++ b/Engine/Plugins/Runtime/StateTree/Source/StateTreeEditorModule/StateTreeEditorModule.Build.cs
@@ -46,7 +46,8 @@ namespace UnrealBuildTool.Rules
 				"ToolMenus",
 				"ToolWidgets",
 				"ApplicationCore",
-				"DeveloperSettings"
+				"DeveloperSettings",
+				"RewindDebuggerInterface"
 			}
 			);
 
diff --git a/Engine/Plugins/Runtime/StateTree/Source/StateTreeModule/Private/Debugger/StateTreeDebugger.cpp b/Engine/Plugins/Runtime/StateTree/Source/StateTreeModule/Private/Debugger/StateTreeDebugger.cpp
--- a/Engine/Plugins/Runtime/StateTree/Source/StateTreeModule/Private/Debugger/StateTreeDebugger.cpp
+++ b/Engine/Plugins/Runtime/StateTree/Source/StateTreeModule/Private/Debugger/StateTreeDebugger.cpp
@@ -21,6 +21,8 @@
 #include "TraceServices/Model/Diagnostics.h"
 #include "GenericPlatform/GenericPlatformMisc.h"
 
+#define LOCTEXT_NAMESPACE "StateTreeDebugger"
+
 //----------------------------------------------------------------//
 // UE::StateTreeDebugger
 //----------------------------------------------------------------//
@@ -70,27 +72,12 @@ namespace UE::StateTreeDebugger
 } // UE::StateTreeDebugger
 
 
-//----------------------------------------------------------------//
-// FStateTreeDebuggerInstanceDesc
-//----------------------------------------------------------------//
-FStateTreeDebuggerInstanceDesc::FStateTreeDebuggerInstanceDesc(const UStateTree* InStateTree, const FStateTreeInstanceDebugId InId, const FString& InName)
-	: StateTree(InStateTree)
-	, Name(InName)
-	, Id(InId)	
-{ 
-}
-
-bool FStateTreeDebuggerInstanceDesc::IsValid() const
-{
-	return StateTree.IsValid() && Name.Len() && Id.IsValid();
-}
-
-
 //----------------------------------------------------------------//
 // FStateTreeDebugger
 //----------------------------------------------------------------//
 FStateTreeDebugger::FStateTreeDebugger()
 	: StateTreeModule(FModuleManager::GetModuleChecked<IStateTreeModule>("StateTreeModule"))
+	, ScrubState(EventCollections)
 {
 }
 
@@ -101,21 +88,10 @@ FStateTreeDebugger::~FStateTreeDebugger()
 
 void FStateTreeDebugger::Tick(float DeltaTime)
 {
-	GetSessionActiveInstances(ActiveInstances);
+	TRACE_CPUPROFILER_EVENT_SCOPE(FStateTreeDebugger::Tick);
 
-	if (!bPaused)
-	{
-		TRACE_CPUPROFILER_EVENT_SCOPE(FStateTreeDebugger::Tick);
-		SyncToCurrentSessionDuration();
-	}
-
-	if (DebuggedInstance.IsValid())
-	{
-		if (!ActiveInstances.Contains(DebuggedInstance))
-		{
-			SetDebuggedInstance(FStateTreeDebuggerInstanceDesc());
-		}
-	}
+	UpdateInstances();
+	SyncToCurrentSessionDuration();
 }
 
 bool FStateTreeDebugger::IsTickable() const
@@ -131,7 +107,7 @@ void FStateTreeDebugger::StopAnalysis()
 		AnalysisSession.Reset();	
 	}
 
-	ResetAnalysisData();
+	EventCollections.Reset();
 }
 
 void FStateTreeDebugger::Pause()
@@ -152,49 +128,64 @@ void FStateTreeDebugger::SyncToCurrentSessionDuration()
 {
 	if (const TraceServices::IAnalysisSession* Session = GetAnalysisSession())
 	{
-		double SessionDuration = 0.;
 		{
 			TraceServices::FAnalysisSessionReadScope SessionReadScope(*Session);
-			SessionDuration = Session->GetDurationSeconds();
+			RecordingDuration = Session->GetDurationSeconds();
 		}
-
-		ReadTrace(SessionDuration);
+		ReadTrace(RecordingDuration);
 	}
 }
 
-const FStateTreeDebuggerInstanceDesc& FStateTreeDebugger::GetDebuggedInstance() const
+FText FStateTreeDebugger::GetInstanceDescription(FStateTreeInstanceDebugId InstanceId) const
 {
-	return DebuggedInstance;
-}
+	using namespace UE::StateTreeDebugger;
+	const FInstanceDescriptor* FoundInstance = InstanceDescs.FindByPredicate([InstanceId](const FInstanceDescriptor& InstanceDesc)
+		{
+			return InstanceDesc.Id == InstanceId;
+		});
 
-FString FStateTreeDebugger::GetDebuggedInstanceDescription() const
-{
-	return DescribeInstance(DebuggedInstance);
+	return (FoundInstance != nullptr) ? DescribeInstance(*FoundInstance) : LOCTEXT("InstanceNotFound","Instance not found");
 }
 
-void FStateTreeDebugger::SetDebuggedInstance(const FStateTreeDebuggerInstanceDesc& SelectedInstance)
+void FStateTreeDebugger::SelectInstance(const FStateTreeInstanceDebugId InstanceId)
 {
-	if (DebuggedInstance != SelectedInstance)
+	if (SelectedInstanceId != InstanceId)
 	{
-		DebuggedInstance = SelectedInstance;
+		HitBreakpointStateIndex = INDEX_NONE;
+		HitBreakpointInstanceId.Reset();
 		
-		// Reparse traces
-		ResetAnalysisData();
+		SelectedInstanceId = InstanceId;
 
-		bProcessingInitialEvents = true;
-		SyncToCurrentSessionDuration();
-		bProcessingInitialEvents = false;
+		// Notify so listener can cleanup anything related to previous instance
+		OnSelectedInstanceCleared.ExecuteIfBound();
 
-		OnDebuggedInstanceSet.ExecuteIfBound();
+		// Update event collection index for newly debugged instance
+		ScrubState.EventCollectionIndex = InstanceId.IsValid() ? EventCollections.IndexOfByPredicate([InstanceId = InstanceId](const UE::StateTreeDebugger::FInstanceEventCollection& Entry)
+			{
+				return Entry.InstanceId == InstanceId;
+			})
+			: INDEX_NONE;
+
+		// Set scrub time to refresh internal tracking for the new instance
+		// This will notify scrub position changed and active states
+		SetScrubTime(ScrubState.ScrubTime);
 	}
 }
 
-TConstArrayView<FStateTreeDebuggerInstanceDesc> FStateTreeDebugger::GetActiveInstances() const
+void FStateTreeDebugger::GetSessionInstances(TArray<UE::StateTreeDebugger::FInstanceDescriptor>& OutInstances) const
 {
-	return ActiveInstances;
+	if (const TraceServices::IAnalysisSession* Session = GetAnalysisSession())
+	{
+		TraceServices::FAnalysisSessionReadScope SessionReadScope(*Session);
+
+		if (const IStateTreeTraceProvider* Provider = Session->ReadProvider<IStateTreeTraceProvider>(FStateTreeTraceProvider::ProviderName))
+		{
+			Provider->GetInstances(OutInstances);
+		}
+	}
 }
 
-void FStateTreeDebugger::GetSessionActiveInstances(TArray<FStateTreeDebuggerInstanceDesc>& OutInstances) const
+void FStateTreeDebugger::UpdateInstances()
 {
 	if (const TraceServices::IAnalysisSession* Session = GetAnalysisSession())
 	{
@@ -202,7 +193,7 @@ void FStateTreeDebugger::GetSessionActiveInstances(TArray<FStateTreeDebuggerInst
 
 		if (const IStateTreeTraceProvider* Provider = Session->ReadProvider<IStateTreeTraceProvider>(FStateTreeTraceProvider::ProviderName))
 		{
-			Provider->GetActiveInstances(OutInstances);
+			Provider->GetInstances(InstanceDescs);
 		}
 	}
 }
@@ -235,7 +226,7 @@ void FStateTreeDebugger::StartSessionAnalysis(const FTraceDescriptor& TraceDescr
 	StopAnalysis();
 
 	const uint32 TraceId = TraceDescriptor.TraceId;
-	
+
 	// Make sure it is still live
 	const UE::Trace::FStoreClient::FSessionInfo* SessionInfo = StoreClient->GetSessionInfoByTraceId(TraceId);
 	if (SessionInfo != nullptr)
@@ -265,9 +256,7 @@ void FStateTreeDebugger::StartSessionAnalysis(const FTraceDescriptor& TraceDescr
 			checkf(!AnalysisSession.IsValid(), TEXT("Must make sure that current session was properly stopped before starting a new one otherwise it can cause threading issues"));
 			AnalysisSession = TraceAnalysisService->StartAnalysis(TraceId, *TraceName, MoveTemp(TraceData));
 			
-			bProcessingInitialEvents = true;
 			SyncToCurrentSessionDuration();
-			bProcessingInitialEvents = false;
 		}
 
 		ActiveSessionTraceDescriptor = AnalysisSession.IsValid() ? TraceDescriptor : FTraceDescriptor();
@@ -360,40 +349,60 @@ void FStateTreeDebugger::UpdateMetadata(FTraceDescriptor& TraceDescriptor) const
 	TraceDescriptor.SessionInfo = Analyzer.SessionInfo;
 }
 
-FString FStateTreeDebugger::GetDebuggedTraceDescription() const
+FText FStateTreeDebugger::GetSelectedTraceDescription() const
 {
 	if (ActiveSessionTraceDescriptor.IsValid())
 	{
 		return DescribeTrace(ActiveSessionTraceDescriptor);
 	}
 
-	return TEXT("No trace selected");
+	return LOCTEXT("NoSelectedTraceDescriptor", "No trace selected");
+}
+
+void FStateTreeDebugger::SetScrubTime(const double ScrubTime)
+{
+	ScrubState.SetScrubTime(ScrubTime);
+	OnScrubStateChanged.Execute(ScrubState);
+
+	RefreshActiveStates();
 }
 
-FString FStateTreeDebugger::DescribeTrace(const FTraceDescriptor& TraceDescriptor)
+bool FStateTreeDebugger::IsActiveInstance(const double Time, const FStateTreeInstanceDebugId InstanceId) const
+{
+	for (const UE::StateTreeDebugger::FInstanceDescriptor& Desc : InstanceDescs)
+	{
+		if (Desc.Id == InstanceId && Desc.Lifetime.Contains(Time))
+		{
+			return true;
+		}
+	}
+	return false;
+}
+
+FText FStateTreeDebugger::DescribeTrace(const FTraceDescriptor& TraceDescriptor)
 {
 	if (TraceDescriptor.IsValid())
 	{
 		const TraceServices::FSessionInfo& SessionInfo = TraceDescriptor.SessionInfo;
 
-		return FString::Printf(TEXT("%s-%s-%s-%s-%s"),
+		return FText::FromString(FString::Printf(TEXT("%s-%s-%s-%s-%s"),
 			*LexToString(TraceDescriptor.TraceId),
 			*SessionInfo.Platform,
 			*SessionInfo.AppName,
 			LexToString(SessionInfo.ConfigurationType),
-			LexToString(SessionInfo.TargetType));
+			LexToString(SessionInfo.TargetType)));
 	}
 
-	return TEXT("Invalid");
+	return LOCTEXT("InvalidTraceDescriptor", "Invalid");
 }
 
-FString FStateTreeDebugger::DescribeInstance(const FStateTreeDebuggerInstanceDesc& InstanceDesc)
+FText FStateTreeDebugger::DescribeInstance(const UE::StateTreeDebugger::FInstanceDescriptor& InstanceDesc)
 {
-	if (FStateTreeDebuggerInstanceDesc() == InstanceDesc)
+	if (InstanceDesc.IsValid() == false)
 	{
-		return TEXT("No instance selected");
+		return LOCTEXT("NoSelectedInstanceDescriptor", "No instance selected");
 	}
-	return LexToString(InstanceDesc);
+	return FText::FromString(LexToString(InstanceDesc));
 }
 
 void FStateTreeDebugger::SetActiveStates(const TConstArrayView<FStateTreeStateHandle> NewActiveStates)
@@ -402,59 +411,75 @@ void FStateTreeDebugger::SetActiveStates(const TConstArrayView<FStateTreeStateHa
 	OnActiveStatesChanged.ExecuteIfBound(ActiveStates);
 }
 
-bool FStateTreeDebugger::CanStepBack() const
+void FStateTreeDebugger::RefreshActiveStates()
 {
-	return bPaused && CurrentFrameIndex != INDEX_NONE
-		&& FramesWithEvents.ContainsByPredicate([CurrentIndex = CurrentFrameIndex](const UE::StateTreeDebugger::FFrameIndexSpan& IteratedFrame)
-		{
-			return IteratedFrame.Frame.Index < CurrentIndex;
-		});
+	TArray<FStateTreeStateHandle> NewActiveStates;
+
+	if (ScrubState.IsPointingToValidActiveStates())
+	{
+		const UE::StateTreeDebugger::FInstanceEventCollection& EventCollection = EventCollections[ScrubState.EventCollectionIndex];
+		const int32 EventIndex = EventCollection.ActiveStatesChanges[ScrubState.ActiveStatesIndex].Value;
+		NewActiveStates = EventCollection.Events[EventIndex].Get<FStateTreeTraceActiveStatesEvent>().ActiveStates;
+	}
+
+	SetActiveStates(NewActiveStates);
 }
 
-void FStateTreeDebugger::StepBack()
+bool FStateTreeDebugger::CanStepBackToPreviousStateWithEvents() const
 {
-	const int32 PrevFrameIndex = FramesWithEvents.FindLastByPredicate([CurrentIndex = CurrentFrameIndex](const UE::StateTreeDebugger::FFrameIndexSpan& IteratedFrame)
-	{
-		return IteratedFrame.Frame.Index < CurrentIndex;
-	});
+	return bPaused ? ScrubState.HasPreviousFrame() : false;
+}
+
+void FStateTreeDebugger::StepBackToPreviousStateWithEvents()
+{
+	ScrubState.GotoPreviousFrame();
+	OnScrubStateChanged.Execute(ScrubState);
 	
-	if (PrevFrameIndex != INDEX_NONE)
-	{
-		// Rebuild from beginning up to previous frame with events
-		Events.Reset();
-		LastTraceReadTime = UnsetTime;
-		ReadTrace(FramesWithEvents[PrevFrameIndex].Frame.Index);
-		check(CurrentFrameIndex == FramesWithEvents[PrevFrameIndex].Frame.Index);
-	}
+	RefreshActiveStates();
 }
 
-bool FStateTreeDebugger::CanStepForward() const
+bool FStateTreeDebugger::CanStepForwardToNextStateWithEvents() const
 {
-	return bPaused && CurrentFrameIndex != INDEX_NONE
-		&& FramesWithEvents.ContainsByPredicate([CurrentIndex = CurrentFrameIndex](const UE::StateTreeDebugger::FFrameIndexSpan& IteratedFrame)
-		{
-			return IteratedFrame.Frame.Index > CurrentIndex;
-		});
+	return bPaused ? ScrubState.HasNextFrame() : false;
 }
 
-void FStateTreeDebugger::StepForward()
+void FStateTreeDebugger::StepForwardToNextStateWithEvents()
 {
-	const UE::StateTreeDebugger::FFrameIndexSpan* NextFrame = FramesWithEvents.FindByPredicate([CurrentIndex = CurrentFrameIndex](const UE::StateTreeDebugger::FFrameIndexSpan& IteratedFrame)
-	{
-		return IteratedFrame.Frame.Index > CurrentIndex;
-	});
+	ScrubState.GotoNextFrame();
+	OnScrubStateChanged.Execute(ScrubState);
+
+	RefreshActiveStates();
+}
+
+bool FStateTreeDebugger::CanStepBackToPreviousStateChange() const
+{
+	return bPaused ? ScrubState.HasPreviousActiveStates() : false;
+}
+
+void FStateTreeDebugger::StepBackToPreviousStateChange()
+{
+	ScrubState.GotoPreviousActiveStates();
+	OnScrubStateChanged.Execute(ScrubState);
 	
-	if (NextFrame != nullptr)
-	{
-		const uint64 FrameIndex = (*NextFrame).Frame.Index; 
-		ReadTrace(FrameIndex);
-		check(CurrentFrameIndex == FrameIndex);
-	}
+	RefreshActiveStates();
+}
+
+bool FStateTreeDebugger::CanStepForwardToNextStateChange() const
+{
+	return bPaused ? ScrubState.HasNextActiveStates() : false;
+}
+
+void FStateTreeDebugger::StepForwardToNextStateChange()
+{
+	ScrubState.GotoNextActiveStates();
+	OnScrubStateChanged.Execute(ScrubState);
+
+	RefreshActiveStates();
 }
 
 void FStateTreeDebugger::ToggleBreakpoints(const TConstArrayView<FStateTreeStateHandle> SelectedStates)
 {
-	TArray<FStateTreeStateHandle> CurrentBreakpoints = StatesWithBreakpoint;	
+	TArray<FStateTreeStateHandle> CurrentBreakpoints = StatesWithBreakpoint;
 	TArray<FStateTreeStateHandle> NewBreakpoints(SelectedStates);
 
 	// remove from the selected states any state that already has a breakpoint
@@ -490,19 +515,6 @@ UE::Trace::FStoreClient* FStateTreeDebugger::GetStoreClient() const
 	return StateTreeModule.GetStoreClient();
 }
 
-void FStateTreeDebugger::ResetAnalysisData()
-{
-	Events.Reset();
-	HitBreakpointStateIndex = INDEX_NONE;
-	HitBreakpointInstanceId.Reset();
-	LastTraceReadTime = UnsetTime;
-	CurrentScrubTime = UnsetTime;
-	CurrentFrameIndex = INDEX_NONE;
-	FramesWithEvents.Reset();
-
-	SetActiveStates({});
-}
-
 void FStateTreeDebugger::ReadTrace(const uint64 FrameIndex)
 {
 	if (const TraceServices::IAnalysisSession* Session = GetAnalysisSession())
@@ -513,7 +525,6 @@ void FStateTreeDebugger::ReadTrace(const uint64 FrameIndex)
 
 		if (const TraceServices::FFrame* TargetFrame = FrameProvider.GetFrame(TraceFrameType_Game, FrameIndex))
 		{
-			CurrentScrubTime = TargetFrame->EndTime;
 			ReadTrace(*Session, FrameProvider, *TargetFrame);
 		}
 	}
@@ -533,8 +544,18 @@ void FStateTreeDebugger::ReadTrace(const double ScrubTime)
 		TraceServices::FFrame TargetFrame;
 		if (FrameProvider.GetFrameFromTime(TraceFrameType_Game, ScrubTime, TargetFrame))
 		{
-			CurrentScrubTime = ScrubTime;
-			ReadTrace(*Session, FrameProvider, TargetFrame);						
+			// Process only completed frames 
+			if (TargetFrame.EndTime == std::numeric_limits<double>::infinity())
+			{
+				if (const TraceServices::FFrame* PreviousCompleteFrame = FrameProvider.GetFrame(TraceFrameType_Game, TargetFrame.Index - 1))
+				{
+					ReadTrace(*Session, FrameProvider, *PreviousCompleteFrame);
+				}
+			}
+			else
+			{
+				ReadTrace(*Session, FrameProvider, TargetFrame);
+			}
 		}
 	}
 
@@ -544,16 +565,26 @@ void FStateTreeDebugger::ReadTrace(const double ScrubTime)
 
 void FStateTreeDebugger::SendNotifications()
 {
-	if (bNewEvents)
+	if (NewInstances.Num() > 0)
 	{
-		OnTracesUpdated.ExecuteIfBound(Events, FramesWithEvents);
-		bNewEvents = false;
+		for (const FStateTreeInstanceDebugId NewInstanceId : NewInstances)
+		{
+			OnNewInstance.ExecuteIfBound(NewInstanceId);
+		}
+		NewInstances.Reset();
 	}
 
 	if (HitBreakpointStateIndex != INDEX_NONE)
 	{
 		check(HitBreakpointInstanceId.IsValid());
 		check(StatesWithBreakpoint.IsValidIndex(HitBreakpointStateIndex));
+
+		// Make sure the instance is selected in case the breakpoint was set for any instances 
+		if (SelectedInstanceId != HitBreakpointInstanceId)
+		{
+			SelectInstance(HitBreakpointInstanceId);
+		}
+
 		OnBreakpointHit.ExecuteIfBound(HitBreakpointInstanceId, StatesWithBreakpoint[HitBreakpointStateIndex]);
 
 		HitBreakpointStateIndex = INDEX_NONE;
@@ -574,95 +605,128 @@ void FStateTreeDebugger::ReadTrace(
 		if (const IStateTreeTraceProvider* Provider = Session.ReadProvider<IStateTreeTraceProvider>(FStateTreeTraceProvider::ProviderName))
 		{
 			AddEvents(LastTraceReadTime, Frame.EndTime, FrameProvider, *Provider);
-			LastTraceReadTime = CurrentScrubTime;
+			LastTraceReadTime = Frame.EndTime;
 		}
 	}
-
-	if (FramesWithEvents.ContainsByPredicate([CurrentIndex = CurrentFrameIndex](const UE::StateTreeDebugger::FFrameIndexSpan& IteratedFrame)
-		{
-			return IteratedFrame.Frame.Index == CurrentIndex;
-		}))
-	{
-		CurrentFrameIndex = Frame.Index;	
-	}
-	else if (FramesWithEvents.Num())
-	{
-		CurrentFrameIndex = FramesWithEvents.Last().Frame.Index;
-	}
 }
 
 bool FStateTreeDebugger::ProcessEvent(const FStateTreeInstanceDebugId InstanceId, const TraceServices::FFrame& Frame, const FStateTreeTraceEventVariantType& Event)
 {
-	if (DebuggedInstance.Id == InstanceId && Event.IsType<FStateTreeTraceActiveStatesEvent>())
-	{
-		// Active states
-		SetActiveStates(Event.Get<FStateTreeTraceActiveStatesEvent>().ActiveStates);
-	}
-	else
+	using namespace UE::StateTreeDebugger;
+
+	// Breakpoints 
+	if (bPaused == false // ignored when scrubbing a paused session
+		&& StateTreeAsset != nullptr // asset is required to properly match state handles
+		&& HitBreakpointStateIndex == INDEX_NONE // stop on first hit breakpoint
+		&& StatesWithBreakpoint.Num()
+		&& (SelectedInstanceId == InstanceId || SelectedInstanceId.IsInvalid())) // allow breakpoints on any instances if not specified
 	{
-		// Breakpoints 
-		if (bPaused == false // ignored when scrubbing a paused session
-			&& bProcessingInitialEvents == false // only processed for events added after the initial parsing
-			&& DebuggedAsset != nullptr // asset is required to properly match state handles
-			&& HitBreakpointStateIndex == INDEX_NONE // stop on first hit breakpoint
-			&& StatesWithBreakpoint.Num()
-			&& (DebuggedInstance.Id == InstanceId || !DebuggedInstance.IsValid())) // allow breakpoints on any instances if not specified
+		const FStateTreeTraceStateEvent* StateEvent = Event.TryGet<FStateTreeTraceStateEvent>();
+		if (StateEvent != nullptr && StateEvent->EventType == EStateTreeTraceNodeEventType::OnEnter)
 		{
-			const FStateTreeTraceStateEvent* StateEvent = Event.TryGet<FStateTreeTraceStateEvent>();
-			if (StateEvent != nullptr && StateEvent->EventType == EStateTreeTraceNodeEventType::OnEnter)
+			const UStateTree* InstanceStateTree = nullptr;
+			if (SelectedInstanceId.IsInvalid())
 			{
-				const UStateTree* InstanceStateTree = nullptr;
-				if (!DebuggedInstance.IsValid())
-				{
-					// No specific instance selected yet, find matching descriptor from id to extract the associated StateTree asset
-					const FStateTreeDebuggerInstanceDesc* FoundInstance = ActiveInstances.FindByPredicate(
-						[InstanceId](const FStateTreeDebuggerInstanceDesc& InstanceDesc)
-						{
-							return InstanceDesc.Id == InstanceId;
-						});
-
-					if (FoundInstance != nullptr)
+				// No specific instance selected yet, find matching descriptor from id to extract the associated StateTree asset
+				const FInstanceDescriptor* FoundInstance = InstanceDescs.FindByPredicate(
+					[InstanceId](const FInstanceDescriptor& InstanceDesc)
 					{
-						InstanceStateTree = FoundInstance->StateTree.Get();
-					}
+						return InstanceDesc.Id == InstanceId;
+					});
+
+				if (FoundInstance != nullptr)
+				{
+					InstanceStateTree = FoundInstance->StateTree.Get();
 				}
+			}
 
-				if (DebuggedInstance.IsValid() || InstanceStateTree == DebuggedAsset)
+			if (SelectedInstanceId.IsValid() || InstanceStateTree == StateTreeAsset)
+			{
+				HitBreakpointStateIndex = StatesWithBreakpoint.Find(FStateTreeStateHandle(StateEvent->Idx));
+				if (HitBreakpointStateIndex != INDEX_NONE)
 				{
-					HitBreakpointStateIndex = StatesWithBreakpoint.Find(FStateTreeStateHandle(StateEvent->Idx));
-					if (HitBreakpointStateIndex != INDEX_NONE)
-					{
-						HitBreakpointInstanceId = InstanceId;
-					}
+					HitBreakpointInstanceId = InstanceId;
 				}
 			}
 		}
+	}
 
-		// Store events for currently debugged entry 
-		if (DebuggedInstance.Id == InstanceId)
+	FInstanceEventCollection* ExistingCollection = EventCollections.FindByPredicate([InstanceId](const FInstanceEventCollection& Entry)
 		{
-			if (FramesWithEvents.IsEmpty() || FramesWithEvents.Last().Frame.Index < Frame.Index)
-			{
-				FramesWithEvents.Add(UE::StateTreeDebugger::FFrameIndexSpan(Frame, Events.Num()));
-			}
+			return Entry.InstanceId == InstanceId;
+		});
+
+	// Create missing EventCollection if necessary
+	if (ExistingCollection == nullptr)
+	{
+		// Push deferred notification for new instance Id
+		NewInstances.Push(InstanceId);
 
-			Events.Emplace(Event);
-			bNewEvents = true;
+		// Update the active event collection index when it's newly created for the currently debugged instance.
+		// Otherwise (i.e. EventCollection already exists) it is updated when switching instance (i.e. SelectInstance)
+		if (SelectedInstanceId == InstanceId && ScrubState.EventCollectionIndex == INDEX_NONE)
+		{
+			ScrubState.EventCollectionIndex = EventCollections.Num();
 		}
+
+		ExistingCollection = &EventCollections.Emplace_GetRef(InstanceId);
 	}
-	
+
+	check(ExistingCollection);
+	TArray<FStateTreeTraceEventVariantType>& Events = ExistingCollection->Events;
+
+	// Add new frame span if none added yet or new frame
+	if (ExistingCollection->FrameSpans.IsEmpty() || ExistingCollection->FrameSpans.Last().Frame.Index < Frame.Index)
+	{
+		ExistingCollection->FrameSpans.Add(UE::StateTreeDebugger::FFrameSpan(Frame, Events.Num()));
+	}
+
+	// Add activate states change info
+	if (Event.IsType<FStateTreeTraceActiveStatesEvent>())
+	{
+		checkf(ExistingCollection->FrameSpans.Num() > 0, TEXT("Expecting to always be in a frame span at this point."));
+		const uint32 FrameSpanIndex = ExistingCollection->FrameSpans.Num()-1;
+
+		// Add new entry for the first event or if the last event is for a different frame
+		if (ExistingCollection->ActiveStatesChanges.IsEmpty()
+			|| ExistingCollection->ActiveStatesChanges.Last().Key != FrameSpanIndex)
+		{
+			ExistingCollection->ActiveStatesChanges.Push({FrameSpanIndex, Events.Num()});
+		}
+		else
+		{
+			// Multiple events for change of active states in the same frame, keep the last one until we implement scrubbing within a frame
+			ExistingCollection->ActiveStatesChanges.Last().Value = Events.Num();
+		}
+	}
+
+	// Store event in the collection
+	Events.Emplace(Event);
+
 	return /*bKeepProcessing*/true;
 }
 
+const UE::StateTreeDebugger::FInstanceEventCollection& FStateTreeDebugger::GetEventCollection(FStateTreeInstanceDebugId InstanceId) const\
+{
+	using namespace UE::StateTreeDebugger;
+	const FInstanceEventCollection* ExistingCollection = EventCollections.FindByPredicate([InstanceId](const FInstanceEventCollection& Entry)
+	{
+		return Entry.InstanceId == InstanceId;
+	});
+
+	return ExistingCollection != nullptr ? *ExistingCollection : FInstanceEventCollection::Invalid;
+}
+
 void FStateTreeDebugger::AddEvents(float StartTime, float EndTime, const TraceServices::IFrameProvider& FrameProvider, const IStateTreeTraceProvider& StateTreeTraceProvider)
 {
-	StateTreeTraceProvider.ReadTimelines(DebuggedInstance.Id,
+	check(StateTreeAsset.IsValid());
+	StateTreeTraceProvider.ReadTimelines(*StateTreeAsset,
 		[this, StartTime, EndTime, &FrameProvider](const FStateTreeInstanceDebugId InstanceId, const IStateTreeTraceProvider::FEventsTimeline& TimelineData)
 		{
 			// Keep track of the frames containing events. Starting with an invalid frame.
 			TraceServices::FFrame Frame;
 			Frame.Index = INDEX_NONE;
-			
+
 			TimelineData.EnumerateEvents(StartTime,	EndTime,
 				[this, InstanceId, &FrameProvider, &Frame](const double EventStartTime, const double EventEndTime, uint32 InDepth, const FStateTreeTraceEventVariantType& Event)
 				{
@@ -670,13 +734,15 @@ void FStateTreeDebugger::AddEvents(float StartTime, float EndTime, const TraceSe
 					if (Frame.Index == INDEX_NONE ||
 						(EventEndTime < Frame.StartTime || Frame.EndTime < EventStartTime))
 					{
-						FrameProvider.GetFrameFromTime(TraceFrameType_Game, EventStartTime, Frame);	
+						FrameProvider.GetFrameFromTime(TraceFrameType_Game, EventStartTime, Frame);
 					}
-					
+
 					const bool bKeepProcessing = ProcessEvent(InstanceId, Frame, Event);
 					return bKeepProcessing ? TraceServices::EEventEnumerate::Continue : TraceServices::EEventEnumerate::Stop;
 				});
 		});
 }
 
-#endif // WITH_STATETREE_DEBUGGER
\ No newline at end of file
+#undef LOCTEXT_NAMESPACE
+
+#endif // WITH_STATETREE_DEBUGGER
diff --git a/Engine/Plugins/Runtime/StateTree/Source/StateTreeModule/Private/Debugger/StateTreeDebuggerTypes.cpp b/Engine/Plugins/Runtime/StateTree/Source/StateTreeModule/Private/Debugger/StateTreeDebuggerTypes.cpp
new file mode 100644
--- /dev/null
+++ b/Engine/Plugins/Runtime/StateTree/Source/StateTreeModule/Private/Debugger/StateTreeDebuggerTypes.cpp
@@ -0,0 +1,159 @@
+// Copyright Epic Games, Inc. All Rights Reserved.
+
+#if WITH_STATETREE_DEBUGGER
+
+#include "Debugger/StateTreeDebuggerTypes.h"
+
+namespace UE::StateTreeDebugger
+{
+
+//----------------------------------------------------------------//
+// FInstanceDescriptor
+//----------------------------------------------------------------//
+FInstanceDescriptor::FInstanceDescriptor(const UStateTree* InStateTree, const FStateTreeInstanceDebugId InId, const FString& InName, const TRange<double> InLifetime)
+	: Lifetime(InLifetime)
+	, StateTree(InStateTree)
+	, Name(InName)
+	, Id(InId)
+{ 
+}
+
+bool FInstanceDescriptor::IsValid() const
+{
+	return StateTree.IsValid() && Name.Len() && Id.IsValid();
+}
+
+
+//----------------------------------------------------------------//
+// FInstanceEventCollection
+//----------------------------------------------------------------//
+const FInstanceEventCollection FInstanceEventCollection::Invalid;
+
+
+//----------------------------------------------------------------//
+// FScrubState
+//----------------------------------------------------------------//
+void FScrubState::SetScrubTime(const double NewScrubTime)
+{
+	if (EventCollectionIndex != INDEX_NONE)
+	{
+		const TArray<FFrameSpan>& Spans = EventCollections[EventCollectionIndex].FrameSpans;
+		const uint32 NewFrameSpanIndex = Spans.IndexOfByPredicate([Time = NewScrubTime](const FFrameSpan& Span)
+		{
+			return Span.Frame.StartTime <= Time && Time <= Span.Frame.EndTime; 
+		});
+
+		if (NewFrameSpanIndex != INDEX_NONE)
+		{
+			SetFrameSpanIndex(NewFrameSpanIndex);
+		}
+		else if (Spans.Num() > 0 && NewScrubTime > Spans.Last().Frame.EndTime)
+		{
+			// Clamp to last span so it keep the scrub valid when sticking to most recent data.
+			SetFrameSpanIndex(EventCollections[EventCollectionIndex].FrameSpans.Num() - 1);
+		}
+		else
+		{
+			TraceFrameIndex = INDEX_NONE;
+			FrameSpanIndex = INDEX_NONE;
+			ActiveStatesIndex = INDEX_NONE;
+		}
+	}
+
+	// This will set back to the exact value provided since SetFrameSpanIndex will snap it to the start time of the matching frame.
+	// It will be consistent with the case where EventCollectionIndex is not set.
+	ScrubTime = NewScrubTime;
+}
+
+void FScrubState::SetFrameSpanIndex(const uint32 NewFrameSpanIndex)
+{
+	FrameSpanIndex = NewFrameSpanIndex;
+	checkf(EventCollections.IsValidIndex(EventCollectionIndex), TEXT("Internal method expecting validity checks before getting called."));
+	const FInstanceEventCollection& EventCollection = EventCollections[EventCollectionIndex];
+
+	checkf(EventCollections[EventCollectionIndex].FrameSpans.IsValidIndex(FrameSpanIndex), TEXT("Internal method expecting validity checks before getting called."));
+	ScrubTime = EventCollection.FrameSpans[FrameSpanIndex].Frame.StartTime;
+	TraceFrameIndex = EventCollection.FrameSpans[FrameSpanIndex].Frame.Index;
+	UpdateActiveStatesIndex();
+}
+
+void FScrubState::SetActiveStatesIndex(const uint32 NewActiveStatesIndex)
+{
+	ActiveStatesIndex = NewActiveStatesIndex;
+
+	checkf(EventCollections.IsValidIndex(EventCollectionIndex), TEXT("Internal method expecting validity checks before getting called."));
+	const FInstanceEventCollection& EventCollection = EventCollections[EventCollectionIndex];
+
+	checkf(EventCollection.ActiveStatesChanges.IsValidIndex(ActiveStatesIndex), TEXT("Internal method expecting validity checks before getting called."));
+	FrameSpanIndex = EventCollection.ActiveStatesChanges[ActiveStatesIndex].Key;
+	ScrubTime = EventCollection.FrameSpans[FrameSpanIndex].Frame.StartTime;
+	TraceFrameIndex = EventCollection.FrameSpans[FrameSpanIndex].Frame.Index;
+}
+
+bool FScrubState::HasPreviousFrame() const
+{
+	return IsInBounds() ? EventCollections[EventCollectionIndex].FrameSpans.IsValidIndex(FrameSpanIndex-1) : false;
+}
+
+double FScrubState::GotoPreviousFrame()
+{
+	SetFrameSpanIndex(FrameSpanIndex - 1);
+	return ScrubTime;
+}
+
+bool FScrubState::HasNextFrame() const
+{
+	return IsInBounds() ? EventCollections[EventCollectionIndex].FrameSpans.IsValidIndex(FrameSpanIndex+1) : false;
+}
+
+double FScrubState::GotoNextFrame()
+{
+	SetFrameSpanIndex(FrameSpanIndex + 1);
+	return ScrubTime;
+}
+
+bool FScrubState::HasPreviousActiveStates() const
+{
+	return (EventCollectionIndex != INDEX_NONE && ActiveStatesIndex != INDEX_NONE)
+		? EventCollections[EventCollectionIndex].ActiveStatesChanges.IsValidIndex(ActiveStatesIndex-1)
+		: false;
+}
+
+double FScrubState::GotoPreviousActiveStates()
+{
+	SetActiveStatesIndex(ActiveStatesIndex - 1);
+	return ScrubTime;
+}
+
+bool FScrubState::HasNextActiveStates() const
+{
+	return (EventCollectionIndex != INDEX_NONE && ActiveStatesIndex != INDEX_NONE)
+		? EventCollections[EventCollectionIndex].ActiveStatesChanges.IsValidIndex(ActiveStatesIndex+1)
+		: false;
+}
+
+double FScrubState::GotoNextActiveStates()
+{
+	SetActiveStatesIndex(ActiveStatesIndex + 1);
+	return ScrubTime;
+}
+
+const FInstanceEventCollection& FScrubState::GetEventCollection() const
+{
+	return EventCollectionIndex != INDEX_NONE ? EventCollections[EventCollectionIndex] : FInstanceEventCollection::Invalid;
+}
+
+void FScrubState::UpdateActiveStatesIndex()
+{
+	check(EventCollectionIndex != INDEX_NONE);
+	const FInstanceEventCollection& EventCollection = EventCollections[EventCollectionIndex];
+
+	// Need to find the index of a frame span that contains an active states changed event; either the current one has it otherwise look backward to find the last one
+	ActiveStatesIndex = EventCollection.ActiveStatesChanges.FindLastByPredicate([CurrentFrameSpanIndex = FrameSpanIndex](const TTuple<uint32, uint32> SpanAndEventIndices)
+	{
+		return SpanAndEventIndices.Key <= CurrentFrameSpanIndex;
+	});
+}
+
+} // UE::StateTreeDebugger
+#endif // WITH_STATETREE_DEBUGGER
diff --git a/Engine/Plugins/Runtime/StateTree/Source/StateTreeModule/Private/Debugger/StateTreeTraceAnalyzer.cpp b/Engine/Plugins/Runtime/StateTree/Source/StateTreeModule/Private/Debugger/StateTreeTraceAnalyzer.cpp
--- a/Engine/Plugins/Runtime/StateTree/Source/StateTreeModule/Private/Debugger/StateTreeTraceAnalyzer.cpp
+++ b/Engine/Plugins/Runtime/StateTree/Source/StateTreeModule/Private/Debugger/StateTreeTraceAnalyzer.cpp
@@ -60,9 +60,10 @@ bool FStateTreeTraceAnalyzer::OnEvent(const uint16 RouteId, EStyle Style, const
 					FString InstanceName;
 					EventData.GetString("InstanceName", InstanceName);
 
-					Provider.AppendInstance(StateTree,
+					Provider.AppendInstanceEvent(StateTree,
 						FStateTreeInstanceDebugId(EventData.GetValue<uint32>("InstanceId"), EventData.GetValue<uint32>("InstanceSerial")),
 						*InstanceName,
+						Context.EventTime.AsSeconds(EventData.GetValue<uint64>("Cycle")),
 						EventData.GetValue<EStateTreeTraceInstanceEventType>("EventType"));
 				}
 				else
diff --git a/Engine/Plugins/Runtime/StateTree/Source/StateTreeModule/Private/Debugger/StateTreeTraceProvider.cpp b/Engine/Plugins/Runtime/StateTree/Source/StateTreeModule/Private/Debugger/StateTreeTraceProvider.cpp
--- a/Engine/Plugins/Runtime/StateTree/Source/StateTreeModule/Private/Debugger/StateTreeTraceProvider.cpp
+++ b/Engine/Plugins/Runtime/StateTree/Source/StateTreeModule/Private/Debugger/StateTreeTraceProvider.cpp
@@ -5,7 +5,6 @@
 #include "Debugger/StateTreeTraceProvider.h"
 #include "StateTreeTypes.h"
 #include "Debugger/StateTreeDebugger.h"
-#include "Algo/RemoveIf.h"
 
 FName FStateTreeTraceProvider::ProviderName("StateTreeDebuggerProvider");
 
@@ -46,48 +45,69 @@ bool FStateTreeTraceProvider::ReadTimelines(const FStateTreeInstanceDebugId Inst
 	return false;
 }
 
-void FStateTreeTraceProvider::AppendEvent(const FStateTreeInstanceDebugId InInstanceId, const double InTime, const FStateTreeTraceEventVariantType& Event)
+bool FStateTreeTraceProvider::ReadTimelines(const UStateTree& StateTree, TFunctionRef<void(const FStateTreeInstanceDebugId ProcessedInstanceId, const FEventsTimeline&)> Callback) const
 {
-	Session.WriteAccessCheck();
+	Session.ReadAccessCheck();
 
-	TSharedPtr<TraceServices::TPointTimeline<FStateTreeTraceEventVariantType>> Timeline;
-	const uint32* IndexPtr = InstanceIdToDebuggerEntryTimelines.Find(InInstanceId);
-	if (IndexPtr != nullptr)
+	for (auto It = InstanceIdToDebuggerEntryTimelines.CreateConstIterator(); It; ++It)
 	{
-		Timeline = EventsTimelines[*IndexPtr];
+		check(EventsTimelines.IsValidIndex(It.Value()));
+		check(Descriptors.Num() == EventsTimelines.Num());
+
+		if (Descriptors[It.Value()].StateTree == &StateTree)
+		{
+			Callback(Descriptors[It.Value()].Id, *EventsTimelines[It.Value()]);
+		}
 	}
-	else
+
+	return EventsTimelines.Num() > 0;
+}
+
+void FStateTreeTraceProvider::AppendEvent(const FStateTreeInstanceDebugId InInstanceId, const double InTime, const FStateTreeTraceEventVariantType& InEvent)
+{
+	Session.WriteAccessCheck();
+
+	const uint32* IndexPtr = InstanceIdToDebuggerEntryTimelines.Find(InInstanceId);
+	if (ensureMsgf(IndexPtr != nullptr, TEXT("Timeline can't be found: probably caused by a missing instance event 'Started'.")))
 	{
-		InstanceIdToDebuggerEntryTimelines.Add(InInstanceId, EventsTimelines.Num());
-		Timeline = EventsTimelines.Emplace_GetRef(MakeShared<TraceServices::TPointTimeline<FStateTreeTraceEventVariantType>>(Session.GetLinearAllocator()));
+		EventsTimelines[*IndexPtr]->AppendEvent(InTime, InEvent);
 	}
-
-	Timeline->AppendEvent(InTime, Event);
+	
 	Session.UpdateDurationSeconds(InTime);
 }
 
-void FStateTreeTraceProvider::AppendInstance(
+void FStateTreeTraceProvider::AppendInstanceEvent(
 	const UStateTree* InStateTree,
 	const FStateTreeInstanceDebugId InInstanceId,
 	const TCHAR* InInstanceName,
-	const EStateTreeTraceInstanceEventType EventType)
+	const double InTime,
+	const EStateTreeTraceInstanceEventType InEventType)
 {
-	if (EventType == EStateTreeTraceInstanceEventType::Started)
+	if (InEventType == EStateTreeTraceInstanceEventType::Started)
 	{
-		ActiveInstances.Emplace(InStateTree, InInstanceId, InInstanceName);	
+		Descriptors.Emplace(InStateTree, InInstanceId, InInstanceName, TRange<double>(InTime, std::numeric_limits<double>::max()));
+
+		check(InstanceIdToDebuggerEntryTimelines.Find(InInstanceId) == nullptr);
+		InstanceIdToDebuggerEntryTimelines.Add(InInstanceId, EventsTimelines.Num());
+		
+		EventsTimelines.Emplace(MakeShared<TraceServices::TPointTimeline<FStateTreeTraceEventVariantType>>(Session.GetLinearAllocator()));
 	}
-	else if (EventType == EStateTreeTraceInstanceEventType::Stopped)
+	else if (InEventType == EStateTreeTraceInstanceEventType::Stopped)
 	{
-		ActiveInstances.SetNum(Algo::RemoveIf(ActiveInstances, [InInstanceId](const FStateTreeDebuggerInstanceDesc& Instance)
-			{
-				return Instance.Id == InInstanceId;
-			}));
+		const uint32* Index = InstanceIdToDebuggerEntryTimelines.Find(InInstanceId);
+		if (ensureMsgf(Index != nullptr, TEXT("Unable to find matching instance. Looks like we never received event 'Started' for this instance.")))
+		{
+			check(Descriptors.IsValidIndex(*Index));
+			Descriptors[*Index].Lifetime.SetUpperBound(InTime);
+		}
 	}
+
+	Session.UpdateDurationSeconds(InTime);
 }
 
-void FStateTreeTraceProvider::GetActiveInstances(TArray<FStateTreeDebuggerInstanceDesc>& OutInstances) const
+void FStateTreeTraceProvider::GetInstances(TArray<UE::StateTreeDebugger::FInstanceDescriptor>& OutInstances) const
 {
-	OutInstances = ActiveInstances;
+	OutInstances = Descriptors;
 }
 
 #undef LOCTEXT_NAMESPACE
diff --git a/Engine/Plugins/Runtime/StateTree/Source/StateTreeModule/Private/StateTreeExecutionContext.cpp b/Engine/Plugins/Runtime/StateTree/Source/StateTreeModule/Private/StateTreeExecutionContext.cpp
--- a/Engine/Plugins/Runtime/StateTree/Source/StateTreeModule/Private/StateTreeExecutionContext.cpp
+++ b/Engine/Plugins/Runtime/StateTree/Source/StateTreeModule/Private/StateTreeExecutionContext.cpp
@@ -15,11 +15,11 @@
 #define STATETREE_CLOG(Condition, Verbosity, Format, ...) UE_CVLOG_UELOG((Condition), GetOwner(), LogStateTree, Verbosity, TEXT("%s") Format, *GetInstanceDescription(), ##__VA_ARGS__)
 
 #if WITH_STATETREE_DEBUGGER
-	#define STATETREE_SCOPED_PHASE(Phase) TGuardValue<EStateTreeUpdatePhase> ScopedPhase(CurrentUpdatePhase, Phase)
-    #define STATETREE_TRACE_LOG_EVENT(Format, ...) TRACE_STATETREE_LOG_EVENT(GetInstanceDebugId(), CurrentUpdatePhase, Format, ##__VA_ARGS__)
-	#define STATETREE_TRACE_STATE_EVENT(Index, EventType) TRACE_STATETREE_STATE_EVENT(GetInstanceDebugId(), CurrentUpdatePhase, Index, EventType);
-	#define STATETREE_TRACE_TASK_EVENT(Index, DataView, EventType, Status) TRACE_STATETREE_TASK_EVENT(GetInstanceDebugId(), CurrentUpdatePhase, Index, DataView, EventType, Status);
-	#define STATETREE_TRACE_CONDITION_EVENT(Index, DataViews, EventType) TRACE_STATETREE_CONDITION_EVENT(GetInstanceDebugId(), CurrentUpdatePhase, Index, DataView, EventType);	
+	#define STATETREE_SCOPED_PHASE(Phase) TGuardValue<EStateTreeUpdatePhase> PREPROCESSOR_JOIN(ScopedPhaseMask, __LINE__)(CurrentUpdatePhaseMask, (CurrentUpdatePhaseMask | Phase))
+    #define STATETREE_TRACE_LOG_EVENT(Format, ...) TRACE_STATETREE_LOG_EVENT(GetInstanceDebugId(), CurrentUpdatePhaseMask, Format, ##__VA_ARGS__)
+	#define STATETREE_TRACE_STATE_EVENT(Index, EventType) TRACE_STATETREE_STATE_EVENT(GetInstanceDebugId(), CurrentUpdatePhaseMask, Index, EventType);
+	#define STATETREE_TRACE_TASK_EVENT(Index, DataView, EventType, Status) TRACE_STATETREE_TASK_EVENT(GetInstanceDebugId(), CurrentUpdatePhaseMask, Index, DataView, EventType, Status);
+	#define STATETREE_TRACE_CONDITION_EVENT(Index, DataViews, EventType) TRACE_STATETREE_CONDITION_EVENT(GetInstanceDebugId(), CurrentUpdatePhaseMask, Index, DataView, EventType);	
 #else
 	#define STATETREE_SCOPED_PHASE(Phase)
 	#define STATETREE_TRACE_LOG_EVENT(Format, ...)
@@ -320,6 +320,7 @@ EStateTreeRunStatus FStateTreeExecutionContext::Stop()
 EStateTreeRunStatus FStateTreeExecutionContext::Tick(const float DeltaTime)
 {
 	CSV_SCOPED_TIMING_STAT_EXCLUSIVE(StateTree_Tick);
+	STATETREE_SCOPED_PHASE(EStateTreeUpdatePhase::TickStateTree);
 
 	if (!IsValid())
 	{
@@ -1221,14 +1222,17 @@ EStateTreeRunStatus FStateTreeExecutionContext::TickTasks(const float DeltaTime)
 			{
 				continue;
 			}
+
+			const FStateTreeDataView TaskDataView = DataViews[Task.DataViewIndex.Get()];
 			
 			// Copy bound properties.
 			// Only copy properties when the task is actually ticked, and copy properties at tick is requested.
 			if (Task.BindingsBatch.IsValid() && Task.bShouldCopyBoundPropertiesOnTick)
 			{
-				StateTree.PropertyBindings.CopyTo(DataViews, Task.BindingsBatch, DataViews[Task.DataViewIndex.Get()]);
+				StateTree.PropertyBindings.CopyTo(DataViews, Task.BindingsBatch, TaskDataView);
 			}
 
+			STATETREE_TRACE_TASK_EVENT(TaskIndex, TaskDataView, EStateTreeTraceNodeEventType::OnTickingTask, EStateTreeRunStatus::Running);
 			EStateTreeRunStatus TaskResult = EStateTreeRunStatus::Unset;
 			{
 				QUICK_SCOPE_CYCLE_COUNTER(StateTree_Task_Tick);
@@ -1237,11 +1241,13 @@ EStateTreeRunStatus FStateTreeExecutionContext::TickTasks(const float DeltaTime)
 				TaskResult = Task.Tick(*this, DeltaTime);
 			}
 
+			STATETREE_TRACE_TASK_EVENT(TaskIndex, TaskDataView,
+				TaskResult != EStateTreeRunStatus::Running ? EStateTreeTraceNodeEventType::OnTaskCompleted : EStateTreeTraceNodeEventType::OnTaskTicked,
+				TaskResult);
+			
 			// TODO: Add more control over which states can control the failed/succeeded result.
 			if (TaskResult != EStateTreeRunStatus::Running)
 			{
-				STATETREE_TRACE_TASK_EVENT(TaskIndex, DataViews[Task.DataViewIndex.Get()], EStateTreeTraceNodeEventType::OnTaskCompleted, TaskResult);
-
 				// Store the first state that completed, will be used to decide where to trigger transitions.
 				if (!Exec.CompletedStateHandle.IsValid())
 				{
@@ -1430,7 +1436,11 @@ bool FStateTreeExecutionContext::TriggerTransitions()
 	FAllowDirectTransitionsScope AllowDirectTransitionsScope(*this); // Set flag for the scope of this function to allow direct transitions without buffering.
 	FStateTreeExecutionState& Exec = GetExecState();
 
-	STATETREE_CLOG(!EventsToProcess.IsEmpty(), Verbose, TEXT("Trigger transitions with events [%s]"), *DebugGetEventsAsString());
+	if (EventsToProcess.Num() > 0)
+	{
+		STATETREE_LOG(Verbose, TEXT("Trigger transitions with events [%s]"), *DebugGetEventsAsString());
+		STATETREE_TRACE_LOG_EVENT(TEXT("Trigger transitions with events [%s]"), *DebugGetEventsAsString());
+	}
 
 	NextTransition.Reset();
 
diff --git a/Engine/Plugins/Runtime/StateTree/Source/StateTreeModule/Public/Debugger/IStateTreeTraceProvider.h b/Engine/Plugins/Runtime/StateTree/Source/StateTreeModule/Public/Debugger/IStateTreeTraceProvider.h
--- a/Engine/Plugins/Runtime/StateTree/Source/StateTreeModule/Public/Debugger/IStateTreeTraceProvider.h
+++ b/Engine/Plugins/Runtime/StateTree/Source/StateTreeModule/Public/Debugger/IStateTreeTraceProvider.h
@@ -8,15 +8,20 @@
 #include "TraceServices/Model/AnalysisSession.h"
 #include "TraceServices/Containers/Timelines.h"
 
-struct FStateTreeDebuggerInstanceDesc;
+namespace UE::StateTreeDebugger { struct FInstanceDescriptor; }
 struct FStateTreeInstanceDebugId;
+class UStateTree;
 
 class IStateTreeTraceProvider : public TraceServices::IProvider
 {
 public:
 	typedef TraceServices::ITimeline<FStateTreeTraceEventVariantType> FEventsTimeline;
 
-	virtual void GetActiveInstances(TArray<FStateTreeDebuggerInstanceDesc>& OutInstances) const = 0;
+	/**
+	 * Return all instances with events in the traces.
+	 * @param OutInstances List of all instances with events in the traces regardless if they are active or not.
+	 */
+	virtual void GetInstances(TArray<UE::StateTreeDebugger::FInstanceDescriptor>& OutInstances) const = 0;
 
 	/**
 	 * Execute given function receiving an event timeline for a given instance or all timelines if instance not specified.  
@@ -25,6 +30,14 @@ public:
 	 * @return True if the specified instance was found for a given Id or at least one timeline was found when no specific Id is provided. 
 	 */
 	virtual bool ReadTimelines(const FStateTreeInstanceDebugId InstanceId, TFunctionRef<void(const FStateTreeInstanceDebugId ProcessedInstanceId, const FEventsTimeline&)> Callback) const = 0;
+
+	/**
+	 * Execute given function receiving an event timeline for all timelines matching a given asset.  
+	 * @param StateTree Asset that must be used by the instance to be processed
+	 * @param Callback Function called for timeline(s) matching the provided asset 
+	 * @return True if the specified instance was found for a given Id or at least one timeline was found when no specific Id is provided. 
+	 */
+	virtual bool ReadTimelines(const UStateTree& StateTree, TFunctionRef<void(const FStateTreeInstanceDebugId ProcessedInstanceId, const FEventsTimeline&)> Callback) const = 0;
 };
 
 #endif // WITH_STATETREE_DEBUGGER
diff --git a/Engine/Plugins/Runtime/StateTree/Source/StateTreeModule/Public/Debugger/StateTreeDebugger.h b/Engine/Plugins/Runtime/StateTree/Source/StateTreeModule/Public/Debugger/StateTreeDebugger.h
--- a/Engine/Plugins/Runtime/StateTree/Source/StateTreeModule/Public/Debugger/StateTreeDebugger.h
+++ b/Engine/Plugins/Runtime/StateTree/Source/StateTreeModule/Public/Debugger/StateTreeDebugger.h
@@ -5,11 +5,12 @@
 #if WITH_STATETREE_DEBUGGER
 
 #include "IStateTreeTraceProvider.h"
-#include "StateTreeTypes.h"
 #include "StateTree.h"
+#include "StateTreeDebuggerTypes.h"
+#include "StateTreeTypes.h"
 #include "Tickable.h"
-#include "TraceServices/Model/Diagnostics.h"
 #include "TraceServices/Model/AnalysisSession.h"
+#include "TraceServices/Model/Diagnostics.h"
 #include "TraceServices/Model/Frames.h"
 
 namespace UE::Trace
@@ -21,61 +22,11 @@ class IStateTreeModule;
 class UStateTreeState;
 class UStateTree;
 
-namespace UE::StateTreeDebugger
-{
-	struct FFrameIndexSpan
-	{
-		FFrameIndexSpan() = default;
-		FFrameIndexSpan(const TraceServices::FFrame& Frame, const uint32 EventIdx)
-			: Frame(Frame)
-			, EventIdx(EventIdx)
-		{
-		}
-
-		/** Frame index in the analysis session */
-		TraceServices::FFrame Frame;
-
-		/** Index of the first event for that Frame index */
-		int32 EventIdx = INDEX_NONE;
-	};
-}
-
-struct STATETREEMODULE_API FStateTreeDebuggerInstanceDesc
-{
-	FStateTreeDebuggerInstanceDesc() = default;
-	FStateTreeDebuggerInstanceDesc(const UStateTree* InStateTree, const FStateTreeInstanceDebugId InId, const FString& InName);
-
-	bool IsValid() const;
-
-	bool operator==(const FStateTreeDebuggerInstanceDesc& Other) const
-	{
-		return StateTree == Other.StateTree
-			&& Id == Other.Id;
-	}
-
-	bool operator!=(const FStateTreeDebuggerInstanceDesc& Other) const
-	{
-		return !(*this == Other);
-	}
-
-	friend FString LexToString(const FStateTreeDebuggerInstanceDesc& InstanceDesc)
-	{
-		return FString::Printf(TEXT("%s | %s | %s"),
-			*GetNameSafe(InstanceDesc.StateTree.Get()),
-			*LexToString(InstanceDesc.Id),
-			*InstanceDesc.Name);
-	}
-
-	TWeakObjectPtr<const UStateTree> StateTree = nullptr;
-	FString Name;
-	FStateTreeInstanceDebugId Id = FStateTreeInstanceDebugId::Invalid;
-};
-
-
-DECLARE_DELEGATE_TwoParams(FOnStateTreeDebuggerTracesUpdated, TConstArrayView<FStateTreeTraceEventVariantType>, TConstArrayView<UE::StateTreeDebugger::FFrameIndexSpan>);
+DECLARE_DELEGATE_OneParam(FOnStateTreeDebuggerScrubStateChanged, const UE::StateTreeDebugger::FScrubState& ScrubState);
 DECLARE_DELEGATE_TwoParams(FOnStateTreeDebuggerBreakpointHit, FStateTreeInstanceDebugId InstanceId, FStateTreeStateHandle StateHandle);
 DECLARE_DELEGATE_OneParam(FOnStateTreeDebuggerBreakpointsChanged, TConstArrayView<FStateTreeStateHandle> Breakpoints);
 DECLARE_DELEGATE_OneParam(FOnStateTreeDebuggerActiveStatesChanges, TConstArrayView<FStateTreeStateHandle> ActiveStates);
+DECLARE_DELEGATE_OneParam(FOnStateTreeDebuggerNewInstance, FStateTreeInstanceDebugId InstanceId);
 DECLARE_DELEGATE(FOnStateTreeDebuggerDebuggedInstanceSet);
 
 struct STATETREEMODULE_API FStateTreeDebugger : FTickableGameObject
@@ -98,42 +49,66 @@ struct STATETREEMODULE_API FStateTreeDebugger : FTickableGameObject
 	FStateTreeDebugger();
 	virtual ~FStateTreeDebugger() override;
 
-	void SetAsset(const UStateTree* StateTreeAsset) { DebuggedAsset = StateTreeAsset; }
+	const UStateTree* GetAsset() const { return StateTreeAsset.Get(); }
+	void SetAsset(const UStateTree* InStateTreeAsset) { StateTreeAsset = InStateTreeAsset; }
 	
-	/** Stops reading traces every frame to preserve current state */
+	/** Stops reading traces every frame to preserve current state. */
 	void Pause();
 
-	/** Resumes reading traces every frame */
+	/** Indicates that the debugger was explicitly paused and is no longer fetching new events from the analysis session. */
+	bool IsPaused() const { return bPaused; }
+
+	/** Resumes reading traces every frame. */
 	void Unpause();
 
 	/** Forces a single refresh to latest state. Useful when simulation is paused. */
 	void SyncToCurrentSessionDuration();
 	
-	bool CanStepBack() const;
-	bool CanStepForward() const;
-	void StepBack();
-	void StepForward();
-	
+	bool CanStepBackToPreviousStateWithEvents() const;
+	void StepBackToPreviousStateWithEvents();
+
+	bool CanStepForwardToNextStateWithEvents() const;
+	void StepForwardToNextStateWithEvents();
+
+	bool CanStepBackToPreviousStateChange() const;
+	void StepBackToPreviousStateChange();
+
+	bool CanStepForwardToNextStateChange() const;
+	void StepForwardToNextStateChange();
+
 	bool IsAnalysisSessionActive() const { return GetAnalysisSession() != nullptr; }
 	const TraceServices::IAnalysisSession* GetAnalysisSession() const;
 
-	TConstArrayView<FStateTreeDebuggerInstanceDesc> GetActiveInstances() const;
-	const FStateTreeDebuggerInstanceDesc& GetDebuggedInstance() const;
-	void SetDebuggedInstance(const FStateTreeDebuggerInstanceDesc& SelectedInstance);
-	FString GetDebuggedInstanceDescription() const;
-	static FString DescribeTrace(const FTraceDescriptor& TraceDescriptor);
-	static FString DescribeInstance(const FStateTreeDebuggerInstanceDesc& StateTreeInstanceDesc);
+	bool IsActiveInstance(double Time, FStateTreeInstanceDebugId InstanceId) const;
+	FText GetInstanceDescription(FStateTreeInstanceDebugId InstanceId) const;
+	void SelectInstance(const FStateTreeInstanceDebugId InstanceId);
+
+	static FText DescribeTrace(const FTraceDescriptor& TraceDescriptor);
+	static FText DescribeInstance(const UE::StateTreeDebugger::FInstanceDescriptor& StateTreeInstanceDesc);
+
+	/**
+	 * Finds and returns the event collection associated to a given instance Id.
+	 * An invalid empty collection is returned if not found (IsValid needs to be called).
+	 * @param InstanceId Id of the instance for which the event collection is returned. 
+	 * @return Event collection associated to the provided Id or an invalid one if not found.
+	 */
+	const UE::StateTreeDebugger::FInstanceEventCollection& GetEventCollection(FStateTreeInstanceDebugId InstanceId) const;
+	
+	double GetRecordingDuration() const { return RecordingDuration; }
+	double GetScrubTime() const	{ return ScrubState.ScrubTime; }
+	void SetScrubTime(double ScrubTime);
 
 	void GetLiveTraces(TArray<FTraceDescriptor>& OutTraceDescriptors) const;
 	void StartLastLiveSessionAnalysis();
 	void StartSessionAnalysis(const FTraceDescriptor& TraceDescriptor);
-	FTraceDescriptor GetDebuggedTraceDescriptor() const { return ActiveSessionTraceDescriptor; }
-	FString GetDebuggedTraceDescription() const;
+	FTraceDescriptor GetSelectedTraceDescriptor() const { return ActiveSessionTraceDescriptor; }
+	FText GetSelectedTraceDescription() const;
 	
 	void ToggleBreakpoints(const TConstArrayView<FStateTreeStateHandle> SelectedStates);
 
-	FOnStateTreeDebuggerDebuggedInstanceSet OnDebuggedInstanceSet;
-	FOnStateTreeDebuggerTracesUpdated OnTracesUpdated;
+	FOnStateTreeDebuggerNewInstance OnNewInstance;
+	FOnStateTreeDebuggerDebuggedInstanceSet OnSelectedInstanceCleared;
+	FOnStateTreeDebuggerScrubStateChanged OnScrubStateChanged;
 	FOnStateTreeDebuggerBreakpointHit OnBreakpointHit;
 	FOnStateTreeDebuggerBreakpointsChanged OnBreakpointsChanged;
 	FOnStateTreeDebuggerActiveStatesChanges OnActiveStatesChanged;
@@ -147,7 +122,6 @@ protected:
 	
 private:
 	void StopAnalysis();
-	void ResetAnalysisData();
 
 	void ReadTrace(double ScrubTime);
 	void ReadTrace(uint64 FrameIndex);
@@ -160,43 +134,69 @@ private:
 	void SendNotifications();
 
 	void SetActiveStates(const TConstArrayView<FStateTreeStateHandle> NewActiveStates);
-	
+
+	/**
+	 * Recompute index of the span that contains the active states change event and update the active states.
+	 * This method handles unselected instances in which case it will reset the active states and set the span index to INDEX_NONE
+	 * */
+	void RefreshActiveStates();
+
 	UE::Trace::FStoreClient* GetStoreClient() const;
-	void GetSessionActiveInstances(TArray<FStateTreeDebuggerInstanceDesc>& OutInstances) const;
+	void GetSessionInstances(TArray<UE::StateTreeDebugger::FInstanceDescriptor>& OutInstances) const;
+	void UpdateInstances();
 
 	bool ProcessEvent(const FStateTreeInstanceDebugId InstanceId, const TraceServices::FFrame& Frame, const FStateTreeTraceEventVariantType& Event);
 	void AddEvents(float StartTime, float EndTime, const TraceServices::IFrameProvider& FrameProvider, const IStateTreeTraceProvider& StateTreeTraceProvider);
 	void UpdateMetadata(FTraceDescriptor& TraceDescriptor) const;
 
-	TWeakObjectPtr<const UStateTree> DebuggedAsset;
-	
+	/** Module used to access the store client and analysis sessions .*/
 	IStateTreeModule& StateTreeModule;
+
+	/** The StateTree asset associated to this debugger. All instances will be using this asset. */ 
+	TWeakObjectPtr<const UStateTree> StateTreeAsset;
 	
 	/** The trace analysis session. */
 	TSharedPtr<const TraceServices::IAnalysisSession> AnalysisSession;
-	
+
+	/** Descriptor of the currently selected session */
 	FTraceDescriptor ActiveSessionTraceDescriptor;
 
-	TArray<FStateTreeDebuggerInstanceDesc> ActiveInstances;
-	FStateTreeDebuggerInstanceDesc DebuggedInstance;
-	
+	/** Descriptors for all instances of the StateTree asset that have traces in the analysis session and still active. */
+	TArray<UE::StateTreeDebugger::FInstanceDescriptor> InstanceDescs;
+
+	/** Processed events for each instance. */
+	TArray<UE::StateTreeDebugger::FInstanceEventCollection> EventCollections;
+
+	/** Specific instance selected for more details */
+	FStateTreeInstanceDebugId SelectedInstanceId;
+
+	/** Handles of states on which a breakpoint has been set. This is per asset and not specific to an instance. */
 	TArray<FStateTreeStateHandle> StatesWithBreakpoint;
+
+	/** List of currently active states in the selected instance */
 	TArray<FStateTreeStateHandle> ActiveStates;
-	TArray<FStateTreeTraceEventVariantType> Events;
-	
-	FStateTreeInstanceDebugId HitBreakpointInstanceId = FStateTreeInstanceDebugId::Invalid;
-	int32 HitBreakpointStateIndex = INDEX_NONE;
 
 	inline static constexpr double UnsetTime = -1;
+
+	/** Recording duration of the analysis session. This is not related to the world simulation time. */
 	double RecordingDuration = UnsetTime;
-	double CurrentScrubTime = UnsetTime;
+
+	/** Last time in the recording that we use to fetch events and we will use for the next read. */
 	double LastTraceReadTime = UnsetTime;
 
-	uint64 CurrentFrameIndex = INDEX_NONE;
-	TArray<UE::StateTreeDebugger::FFrameIndexSpan> FramesWithEvents;
+	/** Combined information regarding current scrub time (e.g. frame index, event collection index, etc.) */ 
+	UE::StateTreeDebugger::FScrubState ScrubState;
+
+	/** Indicates the instance for which a breakpoint has been hit */
+	FStateTreeInstanceDebugId HitBreakpointInstanceId = FStateTreeInstanceDebugId::Invalid;
+
+	/** Indicates the state for which a breakpoint has been hit */
+	int32 HitBreakpointStateIndex = INDEX_NONE;
+
+	/** List of new instances discovered by processing event in the analysis session. */
+	TArray<FStateTreeInstanceDebugId> NewInstances;
 
-	bool bProcessingInitialEvents = false;
-	bool bNewEvents = false;
+	/** Indicates that the debugger was explicitly paused and will wait before fetching new events from the analysis session provider. */
 	bool bPaused = false;
 };
 
diff --git a/Engine/Plugins/Runtime/StateTree/Source/StateTreeModule/Public/Debugger/StateTreeDebuggerTypes.h b/Engine/Plugins/Runtime/StateTree/Source/StateTreeModule/Public/Debugger/StateTreeDebuggerTypes.h
new file mode 100644
--- /dev/null
+++ b/Engine/Plugins/Runtime/StateTree/Source/StateTreeModule/Public/Debugger/StateTreeDebuggerTypes.h
@@ -0,0 +1,207 @@
+// Copyright Epic Games, Inc. All Rights Reserved.
+
+#pragma once
+#include "StateTreeTraceTypes.h"
+
+#if WITH_STATETREE_DEBUGGER
+
+#include "StateTreeTypes.h"
+#include "StateTree.h"
+#include "TraceServices/Model/Frames.h"
+
+class UStateTree;
+
+namespace UE::StateTreeDebugger
+{
+/**
+ * Struct indicating the index of the first event for a given trace recording frame.
+ */
+struct FFrameSpan
+{
+	FFrameSpan() = default;
+	FFrameSpan(const TraceServices::FFrame& Frame, const uint32 EventIdx)
+		: Frame(Frame)
+		, EventIdx(EventIdx)
+	{
+	}
+
+	/** Frame index in the analysis session */
+	TraceServices::FFrame Frame;
+
+	/** Index of the first event for that Frame index */
+	int32 EventIdx = INDEX_NONE;
+};
+
+
+/**
+ * Struct describing a state tree instance for a given StateTree asset
+ */
+struct STATETREEMODULE_API FInstanceDescriptor
+{
+	FInstanceDescriptor() = default;
+	FInstanceDescriptor(const UStateTree* InStateTree, const FStateTreeInstanceDebugId InId, const FString& InName, const TRange<double> InLifetime);
+
+	bool IsValid() const;
+
+	bool operator==(const FInstanceDescriptor& Other) const
+	{
+		return StateTree == Other.StateTree && Id == Other.Id;
+	}
+
+	bool operator!=(const FInstanceDescriptor& Other) const
+	{
+		return !(*this == Other);
+	}
+
+	friend FString LexToString(const FInstanceDescriptor& InstanceDesc)
+	{
+		return FString::Printf(TEXT("%s | %s | %s"),
+			*GetNameSafe(InstanceDesc.StateTree.Get()),
+			*LexToString(InstanceDesc.Id),
+			*InstanceDesc.Name);
+	}
+
+	friend uint32 GetTypeHash(const FInstanceDescriptor& Desc)
+	{
+		return GetTypeHash(Desc.Id);
+	}
+
+	TRange<double> Lifetime = TRange<double>(0);
+	TWeakObjectPtr<const UStateTree> StateTree = nullptr;
+	FString Name;
+	FStateTreeInstanceDebugId Id = FStateTreeInstanceDebugId::Invalid;
+};
+
+
+/**
+ * Struct holding organized events associated to a given state tree instance.
+ */
+struct STATETREEMODULE_API FInstanceEventCollection
+{
+	FInstanceEventCollection() = default;
+	explicit FInstanceEventCollection(const FStateTreeInstanceDebugId& InstanceId)
+		: InstanceId(InstanceId)
+	{
+	}
+
+	friend bool operator==(const FInstanceEventCollection& Lhs, const FInstanceEventCollection& RHS)
+	{
+		return Lhs.InstanceId == RHS.InstanceId;
+	}
+
+	friend bool operator!=(const FInstanceEventCollection& Lhs, const FInstanceEventCollection& RHS)
+	{
+		return !(Lhs == RHS);
+	}
+
+	bool IsValid() const { return InstanceId.IsValid(); }
+	bool IsInvalid() const { return !IsValid(); }
+
+	/** Id of the instance associated to the stored events. */
+	FStateTreeInstanceDebugId InstanceId;
+
+	/** All events received for this instance. */
+	TArray<FStateTreeTraceEventVariantType> Events;
+
+	/** Spans for frames with events. Each span contains the frame information and the index of the first event for that frame. */
+	TArray<FFrameSpan> FrameSpans;
+
+	/** Indices of span and event for frames with a change of activate states. */
+	TArray<TTuple<uint32, uint32>> ActiveStatesChanges;
+
+	/**
+	 * Returns the event collection associated to the currently selected instance.
+	 * An invalid empty collection is returned if there is no selected instance. (IsValid needs to be called).
+	 * @return Event collection associated to the selected instance or an invalid one if not found.
+	 */
+	static const FInstanceEventCollection Invalid;
+};
+
+struct STATETREEMODULE_API FScrubState
+{
+	explicit FScrubState(const TArray<FInstanceEventCollection>& EventCollections)
+		: EventCollections(EventCollections)
+	{
+	}
+
+private:
+	const TArray<FInstanceEventCollection>& EventCollections;
+	
+public:
+	double ScrubTime = 0;
+	uint32 EventCollectionIndex = INDEX_NONE;
+	uint64 TraceFrameIndex = INDEX_NONE;
+	uint32 FrameSpanIndex = INDEX_NONE;
+	uint32 ActiveStatesIndex = INDEX_NONE;
+	
+
+	void SetScrubTime(double NewScrubTime);
+
+	/**
+	 * Indicates if the current scrub state points to a valid frame.
+	 * @return True if the frame index is set
+	 */
+	bool IsInBounds() const { return FrameSpanIndex != INDEX_NONE; }
+
+	/**
+	 * Indicates if the current scrub state points to an active states entry in the event collection.
+	 * @return True if the collection and active states indices are set
+	 */
+	bool IsPointingToValidActiveStates() const { return EventCollectionIndex != INDEX_NONE && ActiveStatesIndex != INDEX_NONE; }
+
+	/** Indicates if there is a frame before with events. */
+	bool HasPreviousFrame() const;
+	
+	/**
+	 * Set scrubbing info using the previous frame with events.
+	 * HasPreviousFrame must be used to validate that this method can be called otherwise some checks might fail.
+	 * @return Adjusted scrub time
+	 */
+	double GotoPreviousFrame();
+	
+	/** Indicates if there is a frame after with events. */
+	bool HasNextFrame() const;
+	
+	/**
+	 * Set scrubbing info using the next frame with events.
+	 * HasPreviousFrame must be used to validate that this method can be called otherwise some checks might fail.
+	 * @return Adjusted scrub time
+	 */
+	double GotoNextFrame();
+	
+	/** Indicates if there is a frame before where the StateTree has a different list of active states. */
+	bool HasPreviousActiveStates() const;
+	
+	/**
+	 * Set scrubbing info using the previous frame where the StateTree has a different list of active states.
+	 * HasPreviousActiveStates must be used to validate that this method can be called otherwise some checks might fail.
+	 * @return Adjusted scrub time
+	 */
+	double GotoPreviousActiveStates();
+	
+	/** Indicates if there is a frame after where the StateTree has a different list of active states. */
+	bool HasNextActiveStates() const;
+	
+	/**
+	 * Set scrubbing info using the next frame where the StateTree has a different list of active states.
+	 * HasNextActiveStates must be used to validate that this method can be called otherwise some checks might fail.
+	 * @return Adjusted scrub time
+	 */
+	double GotoNextActiveStates();
+
+	/**
+	 * Returns the event collection associated to the selected instance.
+	 * An invalid empty collection is returned if there is no selected instance (IsValid needs to be called).
+	 * @return Event collection associated to the selected instance or an invalid one if not found.
+	 */
+	const FInstanceEventCollection& GetEventCollection() const;
+
+private:
+	void SetFrameSpanIndex(uint32 NewFrameSpanIndex);
+	void SetActiveStatesIndex(uint32 NewActiveStatesIndex);
+	void UpdateActiveStatesIndex();
+};
+
+} // UE::StateTreeDebugger
+
+#endif // WITH_STATETREE_DEBUGGER
diff --git a/Engine/Plugins/Runtime/StateTree/Source/StateTreeModule/Public/Debugger/StateTreeTrace.h b/Engine/Plugins/Runtime/StateTree/Source/StateTreeModule/Public/Debugger/StateTreeTrace.h
--- a/Engine/Plugins/Runtime/StateTree/Source/StateTreeModule/Public/Debugger/StateTreeTrace.h
+++ b/Engine/Plugins/Runtime/StateTree/Source/StateTreeModule/Public/Debugger/StateTreeTrace.h
@@ -12,7 +12,7 @@ struct FStateTreeActiveStates;
 struct FStateTreeInstanceDebugId;
 enum class EStateTreeRunStatus : uint8;
 enum class EStateTreeTraceNodeEventType : uint8;
-enum class EStateTreeUpdatePhase : uint8;
+enum class EStateTreeUpdatePhase : uint16;
 enum class EStateTreeTraceInstanceEventType : uint8;
 
 UE_TRACE_CHANNEL_EXTERN(StateTreeDebugChannel, STATETREEMODULE_API)
diff --git a/Engine/Plugins/Runtime/StateTree/Source/StateTreeModule/Public/Debugger/StateTreeTraceProvider.h b/Engine/Plugins/Runtime/StateTree/Source/StateTreeModule/Public/Debugger/StateTreeTraceProvider.h
--- a/Engine/Plugins/Runtime/StateTree/Source/StateTreeModule/Public/Debugger/StateTreeTraceProvider.h
+++ b/Engine/Plugins/Runtime/StateTree/Source/StateTreeModule/Public/Debugger/StateTreeTraceProvider.h
@@ -11,6 +11,7 @@
 namespace TraceServices { class IAnalysisSession; }
 enum class EStateTreeTraceInstanceEventType : uint8;
 class UStateTree;
+namespace UE::StateTreeDebugger { struct FInstanceDescriptor; }
 
 class FStateTreeTraceProvider : public IStateTreeTraceProvider
 {
@@ -19,22 +20,25 @@ public:
 
 	explicit FStateTreeTraceProvider(TraceServices::IAnalysisSession& InSession);
 	
-	void AppendEvent(FStateTreeInstanceDebugId InInstanceId, double InTime, const FStateTreeTraceEventVariantType& Event);
-	void AppendInstance(
+	void AppendEvent(FStateTreeInstanceDebugId InInstanceId, double InTime, const FStateTreeTraceEventVariantType& InEvent);
+	void AppendInstanceEvent(
 		const UStateTree* InStateTree,
 		const FStateTreeInstanceDebugId InInstanceId,
 		const TCHAR* InInstanceName,
-		const EStateTreeTraceInstanceEventType EventType);
+		double Time,
+		EStateTreeTraceInstanceEventType EventType);
 
 protected:
 	/** IStateTreeDebuggerProvider interface */
-	virtual void GetActiveInstances(TArray<FStateTreeDebuggerInstanceDesc>& OutInstances) const override;
+	virtual void GetInstances(TArray<UE::StateTreeDebugger::FInstanceDescriptor>& OutInstances) const override;
 	virtual bool ReadTimelines(const FStateTreeInstanceDebugId InstanceId, TFunctionRef<void(const FStateTreeInstanceDebugId ProcessedInstanceId, const FEventsTimeline&)> Callback) const override;
+	virtual bool ReadTimelines(const UStateTree& StateTree, TFunctionRef<void(const FStateTreeInstanceDebugId ProcessedInstanceId, const FEventsTimeline&)> Callback) const override;
 
 private:
 	TraceServices::IAnalysisSession& Session;
-	TArray<FStateTreeDebuggerInstanceDesc> ActiveInstances;
+
 	TMap<FStateTreeInstanceDebugId, uint32> InstanceIdToDebuggerEntryTimelines;
+	TArray<UE::StateTreeDebugger::FInstanceDescriptor> Descriptors;
 	TArray<TSharedRef<TraceServices::TPointTimeline<FStateTreeTraceEventVariantType>>> EventsTimelines;
 };
 #endif // WITH_STATETREE_DEBUGGER
\ No newline at end of file
diff --git a/Engine/Plugins/Runtime/StateTree/Source/StateTreeModule/Public/Debugger/StateTreeTraceTypes.h b/Engine/Plugins/Runtime/StateTree/Source/StateTreeModule/Public/Debugger/StateTreeTraceTypes.h
--- a/Engine/Plugins/Runtime/StateTree/Source/StateTreeModule/Public/Debugger/StateTreeTraceTypes.h
+++ b/Engine/Plugins/Runtime/StateTree/Source/StateTreeModule/Public/Debugger/StateTreeTraceTypes.h
@@ -7,7 +7,7 @@
 
 struct FStateTreeStateHandle;
 enum class EStateTreeRunStatus : uint8;
-enum class EStateTreeUpdatePhase : uint8;
+enum class EStateTreeUpdatePhase : uint16;
 
 UENUM()
 enum class EStateTreeTraceInstanceEventType : uint8
@@ -23,7 +23,9 @@ enum class EStateTreeTraceNodeEventType : uint8
 	OnEnter,			// for State + Task events
 	OnExit,				// for State + Task events
 	OnStateCompleted,	// for State events
+	OnTickingTask,		// for Task events
 	OnTaskCompleted,	// for Task events
+	OnTaskTicked,		// for Task events
 	Passed,				// for Condition events
 	Failed				// for Condition events
 };
diff --git a/Engine/Plugins/Runtime/StateTree/Source/StateTreeModule/Public/StateTreeExecutionContext.h b/Engine/Plugins/Runtime/StateTree/Source/StateTreeModule/Public/StateTreeExecutionContext.h
--- a/Engine/Plugins/Runtime/StateTree/Source/StateTreeModule/Public/StateTreeExecutionContext.h
+++ b/Engine/Plugins/Runtime/StateTree/Source/StateTreeModule/Public/StateTreeExecutionContext.h
@@ -6,10 +6,6 @@
 #include "StateTreeNodeBase.h"
 #include "Experimental/ConcurrentLinearAllocator.h"
 
-#if WITH_STATETREE_DEBUGGER
-#include "Debugger/StateTreeTraceTypes.h"
-#endif // WITH_STATETREE_DEBUGGER
-
 #include "StateTreeExecutionContext.generated.h"
 
 struct FGameplayTag;
@@ -24,23 +20,26 @@ struct FStateTreeInstanceDebugId;
 
 /**
  * Enumeration for the different update phases.
- * This is used as context information when tracing debug events.   
+ * This is used as context information when tracing debug events.
+ * The values are ordered based on their hierarchy starting with the "higher" parents.
  */
 UENUM()
-enum class EStateTreeUpdatePhase : uint8
+enum class EStateTreeUpdatePhase : uint16
 {
-	Unset,
-	TickingGlobalTasks,
-	StateSelection,
-	TickingTasks,
-	StateTransition,
-	TriggerTransitions,
-	RequestTransition,
-	StartGlobalTasks,
-	StopGlobalTasks,
-	StartTree,
-	StopTree
+	Unset				= 0,
+	StartTree			= 1ull << 1,
+	StopTree			= 1ull << 2,
+	StartGlobalTasks	= 1ull << 3,
+	StopGlobalTasks		= 1ull << 4,
+	TickStateTree		= 1ull << 5,
+	TickingGlobalTasks	= 1ull << 6,
+	TickingTasks		= 1ull << 7,
+	TriggerTransitions	= 1ull << 8,
+	StateSelection		= 1ull << 9,
+	StateTransition		= 1ull << 10,
+	RequestTransition	= 1ull << 11,
 };
+ENUM_CLASS_FLAGS(EStateTreeUpdatePhase);
 
 
 /**
@@ -469,7 +468,7 @@ protected:
 
 #if WITH_STATETREE_DEBUGGER
 	/** Current phase in the update. Used for debugging events */
-	EStateTreeUpdatePhase CurrentUpdatePhase = EStateTreeUpdatePhase::Unset;
+	EStateTreeUpdatePhase CurrentUpdatePhaseMask = EStateTreeUpdatePhase::Unset;
 #endif // WITH_STATETREE_DEBUGGER
 
 	/** True if transitions are allowed to be requested directly instead of buffering. */
