State Tree: Add specific state and status when the tree is stopped.

diff --git a/Engine/Plugins/Runtime/StateTree/Source/StateTreeModule/Private/StateTreeExecutionContext.cpp b/Engine/Plugins/Runtime/StateTree/Source/StateTreeModule/Private/StateTreeExecutionContext.cpp
--- a/Engine/Plugins/Runtime/StateTree/Source/StateTreeModule/Private/StateTreeExecutionContext.cpp
+++ b/Engine/Plugins/Runtime/StateTree/Source/StateTreeModule/Private/StateTreeExecutionContext.cpp
@@ -294,14 +294,14 @@ EStateTreeRunStatus FStateTreeExecutionContext::Stop()
 	{
 		// Transition to Succeeded state.
 		FStateTreeTransitionResult Transition;
-		Transition.TargetState = FStateTreeStateHandle::Succeeded;
+		Transition.TargetState = FStateTreeStateHandle::Stopped;
 		Transition.CurrentActiveStates = Exec.ActiveStates;
 		Transition.CurrentRunStatus = Exec.LastTickStatus;
-		Transition.NextActiveStates = FStateTreeActiveStates(FStateTreeStateHandle::Succeeded);
+		Transition.NextActiveStates = FStateTreeActiveStates(FStateTreeStateHandle::Stopped);
 
 		ExitState(Transition);
 
-		Result = EStateTreeRunStatus::Succeeded;
+		Result = EStateTreeRunStatus::Stopped;
 	}
 
 	// Stop evaluators and global tasks.
diff --git a/Engine/Plugins/Runtime/StateTree/Source/StateTreeModule/Private/StateTreeTypes.cpp b/Engine/Plugins/Runtime/StateTree/Source/StateTreeModule/Private/StateTreeTypes.cpp
--- a/Engine/Plugins/Runtime/StateTree/Source/StateTreeModule/Private/StateTreeTypes.cpp
+++ b/Engine/Plugins/Runtime/StateTree/Source/StateTreeModule/Private/StateTreeTypes.cpp
@@ -10,6 +10,7 @@ DEFINE_LOG_CATEGORY(LogStateTree);
 const FStateTreeStateHandle FStateTreeStateHandle::Invalid = FStateTreeStateHandle();
 const FStateTreeStateHandle FStateTreeStateHandle::Succeeded = FStateTreeStateHandle(FStateTreeStateHandle::SucceededIndex);
 const FStateTreeStateHandle FStateTreeStateHandle::Failed = FStateTreeStateHandle(FStateTreeStateHandle::FailedIndex);
+const FStateTreeStateHandle FStateTreeStateHandle::Stopped = FStateTreeStateHandle(FStateTreeStateHandle::StoppedIndex);
 const FStateTreeStateHandle FStateTreeStateHandle::Root = FStateTreeStateHandle(0);
 
 const FStateTreeIndex16 FStateTreeIndex16::Invalid = FStateTreeIndex16();
diff --git a/Engine/Plugins/Runtime/StateTree/Source/StateTreeModule/Public/StateTreeTypes.h b/Engine/Plugins/Runtime/StateTree/Source/StateTreeModule/Public/StateTreeTypes.h
--- a/Engine/Plugins/Runtime/StateTree/Source/StateTreeModule/Public/StateTreeTypes.h
+++ b/Engine/Plugins/Runtime/StateTree/Source/StateTreeModule/Public/StateTreeTypes.h
@@ -32,7 +32,10 @@ enum class EStateTreeRunStatus : uint8
 	
 	/** Tree execution has stopped on success. */
 	Succeeded,
-	
+
+	/** The State Tree was requested to stop without a particular success or failure state. */
+	Stopped,
+
 	/** Status not set. */
 	Unset,
 };
@@ -178,13 +181,15 @@ struct STATETREEMODULE_API FStateTreeStateHandle
 {
 	GENERATED_BODY()
 
-	static constexpr uint16 InvalidIndex = uint16(-1);		// Index value indicating invalid item.
-	static constexpr uint16 SucceededIndex = uint16(-2);	// Index value indicating a Succeeded item.
-	static constexpr uint16 FailedIndex = uint16(-3);		// Index value indicating a Failed item.
+	static constexpr uint16 InvalidIndex = uint16(-1);		// Index value indicating invalid state.
+	static constexpr uint16 SucceededIndex = uint16(-2);	// Index value indicating a Succeeded state.
+	static constexpr uint16 FailedIndex = uint16(-3);		// Index value indicating a Failed state.
+	static constexpr uint16 StoppedIndex = uint16(-4);		// Index value indicating a Stopped state.
 	
 	static const FStateTreeStateHandle Invalid;
 	static const FStateTreeStateHandle Succeeded;
 	static const FStateTreeStateHandle Failed;
+	static const FStateTreeStateHandle Stopped;
 	static const FStateTreeStateHandle Root;
 
 	FStateTreeStateHandle() = default;
@@ -192,7 +197,7 @@ struct STATETREEMODULE_API FStateTreeStateHandle
 
 	bool IsValid() const { return Index != InvalidIndex; }
 	void Invalidate() { Index = InvalidIndex; }
-	bool IsCompletionState() const { return Index == SucceededIndex || Index == FailedIndex; }
+	bool IsCompletionState() const { return Index == SucceededIndex || Index == FailedIndex || Index == StoppedIndex; }
 	EStateTreeRunStatus ToCompletionStatus() const
 	{
 		if (Index == SucceededIndex)
@@ -203,6 +208,10 @@ struct STATETREEMODULE_API FStateTreeStateHandle
 		{
 			return EStateTreeRunStatus::Failed;
 		}
+		else if (Index == StoppedIndex)
+		{
+			return EStateTreeRunStatus::Stopped;
+		}
 		return EStateTreeRunStatus::Unset;
 	}
 
@@ -213,9 +222,10 @@ struct STATETREEMODULE_API FStateTreeStateHandle
 	{
 		switch (Index)
 		{
-		case InvalidIndex:		return TEXT("Invalid Item");
-		case SucceededIndex:	return TEXT("Succeeded Item");
-		case FailedIndex: 		return TEXT("Failed Item");
+		case InvalidIndex:		return TEXT("Invalid");
+		case SucceededIndex:	return TEXT("Succeeded");
+		case FailedIndex: 		return TEXT("Failed");
+		case StoppedIndex: 		return TEXT("Stopped");
 		default: 				return FString::Printf(TEXT("%d"), Index);
 		}
 	}
diff --git a/Engine/Plugins/Runtime/StateTree/Source/StateTreeTestSuite/Private/StateTreeTest.cpp b/Engine/Plugins/Runtime/StateTree/Source/StateTreeTestSuite/Private/StateTreeTest.cpp
--- a/Engine/Plugins/Runtime/StateTree/Source/StateTreeTestSuite/Private/StateTreeTest.cpp
+++ b/Engine/Plugins/Runtime/StateTree/Source/StateTreeTestSuite/Private/StateTreeTest.cpp
@@ -290,6 +290,7 @@ struct FStateTreeTest_FailEnterState : FAITestBase
 		AITEST_TRUE("StateTree Task1 should exit state", Exec.Expect(Task1.GetName(), ExitStateStr));
 		AITEST_TRUE("StateTree Task2 should exit state", Exec.Expect(Task2.GetName(), ExitStateStr));
 		AITEST_FALSE("StateTree Task3 should not exit state", Exec.Expect(Task3.GetName(), ExitStateStr));
+		AITEST_TRUE("StateTree status should be stopped", Status == EStateTreeRunStatus::Stopped);
 		Exec.LogClear();
 
 		return true;
