[StateTree] type conversion fixes. Modified the enum type casts to consider underlying type to catch type conversion issue at compile time if underlying type changes.
#rnx

diff --git a/Engine/Plugins/Runtime/StateTree/Source/StateTreeModule/Private/Debugger/StateTreeDebugger.cpp b/Engine/Plugins/Runtime/StateTree/Source/StateTreeModule/Private/Debugger/StateTreeDebugger.cpp
--- a/Engine/Plugins/Runtime/StateTree/Source/StateTreeModule/Private/Debugger/StateTreeDebugger.cpp
+++ b/Engine/Plugins/Runtime/StateTree/Source/StateTreeModule/Private/Debugger/StateTreeDebugger.cpp
@@ -717,7 +717,7 @@ const UE::StateTreeDebugger::FInstanceEventCollection& FStateTreeDebugger::GetEv
 	return ExistingCollection != nullptr ? *ExistingCollection : FInstanceEventCollection::Invalid;
 }
 
-void FStateTreeDebugger::AddEvents(float StartTime, float EndTime, const TraceServices::IFrameProvider& FrameProvider, const IStateTreeTraceProvider& StateTreeTraceProvider)
+void FStateTreeDebugger::AddEvents(const double StartTime, const double EndTime, const TraceServices::IFrameProvider& FrameProvider, const IStateTreeTraceProvider& StateTreeTraceProvider)
 {
 	check(StateTreeAsset.IsValid());
 	StateTreeTraceProvider.ReadTimelines(*StateTreeAsset,
diff --git a/Engine/Plugins/Runtime/StateTree/Source/StateTreeModule/Private/Debugger/StateTreeTrace.cpp b/Engine/Plugins/Runtime/StateTree/Source/StateTreeModule/Private/Debugger/StateTreeTrace.cpp
--- a/Engine/Plugins/Runtime/StateTree/Source/StateTreeModule/Private/Debugger/StateTreeTrace.cpp
+++ b/Engine/Plugins/Runtime/StateTree/Source/StateTreeModule/Private/Debugger/StateTreeTrace.cpp
@@ -27,7 +27,7 @@ UE_TRACE_EVENT_BEGIN(StateTreeDebugger, LogEvent)
 	UE_TRACE_EVENT_FIELD(uint64, Cycle)
 	UE_TRACE_EVENT_FIELD(uint32, InstanceId)
 	UE_TRACE_EVENT_FIELD(uint32, InstanceSerial)
-	UE_TRACE_EVENT_FIELD(uint8, Phase)
+	UE_TRACE_EVENT_FIELD(uint16, Phase)
 	UE_TRACE_EVENT_FIELD(UE::Trace::WideString, Message)
 UE_TRACE_EVENT_END()
 
@@ -35,7 +35,7 @@ UE_TRACE_EVENT_BEGIN(StateTreeDebugger, StateEvent)
 	UE_TRACE_EVENT_FIELD(uint64, Cycle)
 	UE_TRACE_EVENT_FIELD(uint32, InstanceId)
 	UE_TRACE_EVENT_FIELD(uint32, InstanceSerial)
-	UE_TRACE_EVENT_FIELD(uint8, Phase)
+	UE_TRACE_EVENT_FIELD(uint16, Phase)
 	UE_TRACE_EVENT_FIELD(uint16, Idx)
 	UE_TRACE_EVENT_FIELD(uint8, EventType)
 UE_TRACE_EVENT_END()
@@ -44,7 +44,7 @@ UE_TRACE_EVENT_BEGIN(StateTreeDebugger, TaskEvent)
 	UE_TRACE_EVENT_FIELD(uint64, Cycle)
 	UE_TRACE_EVENT_FIELD(uint32, InstanceId)
 	UE_TRACE_EVENT_FIELD(uint32, InstanceSerial)
-	UE_TRACE_EVENT_FIELD(uint8, Phase)
+	UE_TRACE_EVENT_FIELD(uint16, Phase)
 	UE_TRACE_EVENT_FIELD(uint16, Idx)
 	UE_TRACE_EVENT_FIELD(uint8[], DataView)
 	UE_TRACE_EVENT_FIELD(uint8, EventType)
@@ -55,7 +55,7 @@ UE_TRACE_EVENT_BEGIN(StateTreeDebugger, ConditionEvent)
 	UE_TRACE_EVENT_FIELD(uint64, Cycle)
 	UE_TRACE_EVENT_FIELD(uint32, InstanceId)
 	UE_TRACE_EVENT_FIELD(uint32, InstanceSerial)
-	UE_TRACE_EVENT_FIELD(uint8, Phase)
+	UE_TRACE_EVENT_FIELD(uint16, Phase)
 	UE_TRACE_EVENT_FIELD(uint16, Idx)
 	UE_TRACE_EVENT_FIELD(uint8[], DataView)
 	UE_TRACE_EVENT_FIELD(uint8, EventType)
@@ -122,7 +122,7 @@ void UE::StateTreeTrace::OutputInstanceLifetimeEvent(
 		<< InstanceEvent.InstanceId(InstanceId.Id)
 		<< InstanceEvent.InstanceSerial(InstanceId.SerialNumber)
 		<< InstanceEvent.InstanceName(InstanceName)
-		<< InstanceEvent.EventType(static_cast<uint8>(EventType));
+		<< InstanceEvent.EventType(static_cast<std::underlying_type_t<EStateTreeTraceNodeEventType>>(EventType));
 }
 
 void UE::StateTreeTrace::OutputLogEventTrace(
@@ -138,14 +138,14 @@ void UE::StateTreeTrace::OutputLogEventTrace(
 		<< LogEvent.Cycle(FPlatformTime::Cycles64())
 		<< LogEvent.InstanceId(InstanceId.Id)
 		<< LogEvent.InstanceSerial(InstanceId.SerialNumber)
-		<< LogEvent.Phase(static_cast<uint8>(Phase))
+		<< LogEvent.Phase(static_cast<std::underlying_type_t<EStateTreeUpdatePhase>>(Phase))
 		<< LogEvent.Message(TraceStaticBuffer);
 }
 
 void UE::StateTreeTrace::OutputStateEventTrace(
 	const FStateTreeInstanceDebugId InstanceId,
 	const EStateTreeUpdatePhase Phase,
-	const uint16 StateIdx,
+	const FStateTreeStateHandle StateHandle,
 	const EStateTreeTraceNodeEventType EventType
 	)
 {
@@ -153,15 +153,15 @@ void UE::StateTreeTrace::OutputStateEventTrace(
 		<< StateEvent.Cycle(FPlatformTime::Cycles64())
 		<< StateEvent.InstanceId(InstanceId.Id)
 		<< StateEvent.InstanceSerial(InstanceId.SerialNumber)
-		<< StateEvent.Phase(static_cast<uint8>(Phase))
-		<< StateEvent.Idx(StateIdx)
-		<< StateEvent.EventType(static_cast<uint8>(EventType));
+		<< StateEvent.Phase(static_cast<std::underlying_type_t<EStateTreeUpdatePhase>>(Phase))
+		<< StateEvent.Idx(StateHandle.Index)
+		<< StateEvent.EventType(static_cast<std::underlying_type_t<EStateTreeTraceNodeEventType>>(EventType));
 }
 
 void UE::StateTreeTrace::OutputTaskEventTrace(
 	const FStateTreeInstanceDebugId InstanceId,
 	const EStateTreeUpdatePhase Phase,
-	const uint16 TaskIdx,
+	const FStateTreeIndex16 TaskIdx,
 	const FStateTreeDataView DataView,
 	const EStateTreeTraceNodeEventType EventType,
 	const EStateTreeRunStatus Status
@@ -174,17 +174,17 @@ void UE::StateTreeTrace::OutputTaskEventTrace(
 		<< TaskEvent.Cycle(FPlatformTime::Cycles64())
 		<< TaskEvent.InstanceId(InstanceId.Id)
 		<< TaskEvent.InstanceSerial(InstanceId.SerialNumber)
-		<< TaskEvent.Phase(static_cast<uint8>(Phase))
-		<< TaskEvent.Idx(TaskIdx)
+		<< TaskEvent.Phase(static_cast<std::underlying_type_t<EStateTreeUpdatePhase>>(Phase))
+		<< TaskEvent.Idx(TaskIdx.Get())
 		<< TaskEvent.DataView(Archive.GetData(), Archive.Num())
-		<< TaskEvent.EventType(static_cast<uint8>(EventType))
-		<< TaskEvent.Status(static_cast<uint8>(Status));
+		<< TaskEvent.EventType(static_cast<std::underlying_type_t<EStateTreeTraceNodeEventType>>(EventType))
+		<< TaskEvent.Status(static_cast<std::underlying_type_t<EStateTreeRunStatus>>(Status));
 }
 
 void UE::StateTreeTrace::OutputConditionEventTrace(
 	const FStateTreeInstanceDebugId InstanceId,
 	const EStateTreeUpdatePhase Phase,
-	const uint16 ConditionIdx,
+	const FStateTreeIndex16 ConditionIdx,
 	const FStateTreeDataView DataView,
 	const EStateTreeTraceNodeEventType EventType
 	)
@@ -196,10 +196,10 @@ void UE::StateTreeTrace::OutputConditionEventTrace(
 	<< ConditionEvent.Cycle(FPlatformTime::Cycles64())
 	<< ConditionEvent.InstanceId(InstanceId.Id)
 	<< ConditionEvent.InstanceSerial(InstanceId.SerialNumber)
-	<< ConditionEvent.Phase(static_cast<uint8>(Phase))
-	<< ConditionEvent.Idx(ConditionIdx)
+	<< ConditionEvent.Phase(static_cast<std::underlying_type_t<EStateTreeUpdatePhase>>(Phase))
+	<< ConditionEvent.Idx(ConditionIdx.Get())
 	<< ConditionEvent.DataView(Archive.GetData(), Archive.Num())
-	<< ConditionEvent.EventType(static_cast<uint8>(EventType));
+	<< ConditionEvent.EventType(static_cast<std::underlying_type_t<EStateTreeTraceNodeEventType>>(EventType));
 }
 
 void UE::StateTreeTrace::OutputActiveStatesEventTrace(
diff --git a/Engine/Plugins/Runtime/StateTree/Source/StateTreeModule/Private/StateTreeExecutionContext.cpp b/Engine/Plugins/Runtime/StateTree/Source/StateTreeModule/Private/StateTreeExecutionContext.cpp
--- a/Engine/Plugins/Runtime/StateTree/Source/StateTreeModule/Private/StateTreeExecutionContext.cpp
+++ b/Engine/Plugins/Runtime/StateTree/Source/StateTreeModule/Private/StateTreeExecutionContext.cpp
@@ -17,9 +17,9 @@
 #if WITH_STATETREE_DEBUGGER
 	#define STATETREE_SCOPED_PHASE(Phase) TGuardValue<EStateTreeUpdatePhase> PREPROCESSOR_JOIN(ScopedPhaseMask, __LINE__)(CurrentUpdatePhaseMask, (CurrentUpdatePhaseMask | Phase))
     #define STATETREE_TRACE_LOG_EVENT(Format, ...) TRACE_STATETREE_LOG_EVENT(GetInstanceDebugId(), CurrentUpdatePhaseMask, Format, ##__VA_ARGS__)
-	#define STATETREE_TRACE_STATE_EVENT(Index, EventType) TRACE_STATETREE_STATE_EVENT(GetInstanceDebugId(), CurrentUpdatePhaseMask, Index, EventType);
-	#define STATETREE_TRACE_TASK_EVENT(Index, DataView, EventType, Status) TRACE_STATETREE_TASK_EVENT(GetInstanceDebugId(), CurrentUpdatePhaseMask, Index, DataView, EventType, Status);
-	#define STATETREE_TRACE_CONDITION_EVENT(Index, DataViews, EventType) TRACE_STATETREE_CONDITION_EVENT(GetInstanceDebugId(), CurrentUpdatePhaseMask, Index, DataView, EventType);	
+	#define STATETREE_TRACE_STATE_EVENT(Index, EventType) TRACE_STATETREE_STATE_EVENT(GetInstanceDebugId(), CurrentUpdatePhaseMask, FStateTreeStateHandle(Index), EventType);
+	#define STATETREE_TRACE_TASK_EVENT(Index, DataView, EventType, Status) TRACE_STATETREE_TASK_EVENT(GetInstanceDebugId(), CurrentUpdatePhaseMask, FStateTreeIndex16(Index), DataView, EventType, Status);
+	#define STATETREE_TRACE_CONDITION_EVENT(Index, DataViews, EventType) TRACE_STATETREE_CONDITION_EVENT(GetInstanceDebugId(), CurrentUpdatePhaseMask, FStateTreeIndex16(Index), DataView, EventType);	
 #else
 	#define STATETREE_SCOPED_PHASE(Phase)
 	#define STATETREE_TRACE_LOG_EVENT(Format, ...)
diff --git a/Engine/Plugins/Runtime/StateTree/Source/StateTreeModule/Private/StateTreePropertyBindings.cpp b/Engine/Plugins/Runtime/StateTree/Source/StateTreeModule/Private/StateTreePropertyBindings.cpp
--- a/Engine/Plugins/Runtime/StateTree/Source/StateTreeModule/Private/StateTreePropertyBindings.cpp
+++ b/Engine/Plugins/Runtime/StateTree/Source/StateTreeModule/Private/StateTreePropertyBindings.cpp
@@ -178,7 +178,9 @@ bool FStateTreePropertyBindings::ResolvePath(const UStruct* Struct, const FState
 	{
 		FStateTreePropertyIndirection& Indirection = TempIndirections.AddDefaulted_GetRef();
 
-		Indirection.Offset = PathIndirection.GetPropertyOffset();
+		check(PathIndirection.GetPropertyOffset() >= MIN_uint16 && PathIndirection.GetPropertyOffset() <= MAX_uint16);
+
+		Indirection.Offset = static_cast<uint16>(PathIndirection.GetPropertyOffset());
 		Indirection.Type = PathIndirection.GetAccessType();
 
 		if (Indirection.Type == EStateTreePropertyAccessType::IndexArray)
diff --git a/Engine/Plugins/Runtime/StateTree/Source/StateTreeModule/Public/Debugger/StateTreeDebugger.h b/Engine/Plugins/Runtime/StateTree/Source/StateTreeModule/Public/Debugger/StateTreeDebugger.h
--- a/Engine/Plugins/Runtime/StateTree/Source/StateTreeModule/Public/Debugger/StateTreeDebugger.h
+++ b/Engine/Plugins/Runtime/StateTree/Source/StateTreeModule/Public/Debugger/StateTreeDebugger.h
@@ -146,7 +146,7 @@ private:
 	void UpdateInstances();
 
 	bool ProcessEvent(const FStateTreeInstanceDebugId InstanceId, const TraceServices::FFrame& Frame, const FStateTreeTraceEventVariantType& Event);
-	void AddEvents(float StartTime, float EndTime, const TraceServices::IFrameProvider& FrameProvider, const IStateTreeTraceProvider& StateTreeTraceProvider);
+	void AddEvents(double StartTime, double EndTime, const TraceServices::IFrameProvider& FrameProvider, const IStateTreeTraceProvider& StateTreeTraceProvider);
 	void UpdateMetadata(FTraceDescriptor& TraceDescriptor) const;
 
 	/** Module used to access the store client and analysis sessions .*/
diff --git a/Engine/Plugins/Runtime/StateTree/Source/StateTreeModule/Public/Debugger/StateTreeTrace.h b/Engine/Plugins/Runtime/StateTree/Source/StateTreeModule/Public/Debugger/StateTreeTrace.h
--- a/Engine/Plugins/Runtime/StateTree/Source/StateTreeModule/Public/Debugger/StateTreeTrace.h
+++ b/Engine/Plugins/Runtime/StateTree/Source/StateTreeModule/Public/Debugger/StateTreeTrace.h
@@ -10,6 +10,8 @@ class UStateTree;
 struct FStateTreeDataView;
 struct FStateTreeActiveStates;
 struct FStateTreeInstanceDebugId;
+struct FStateTreeIndex16;
+struct FStateTreeStateHandle;
 enum class EStateTreeRunStatus : uint8;
 enum class EStateTreeTraceNodeEventType : uint8;
 enum class EStateTreeUpdatePhase : uint16;
@@ -21,9 +23,9 @@ namespace UE::StateTreeTrace
 {
 	void OutputInstanceLifetimeEvent(const FStateTreeInstanceDebugId InstanceId, const UStateTree* StateTree, const TCHAR* InstanceName, const EStateTreeTraceInstanceEventType EventType);
 	void OutputLogEventTrace(const FStateTreeInstanceDebugId InstanceId, const EStateTreeUpdatePhase Phase, const TCHAR* Fmt, ...);
-	void OutputStateEventTrace(const FStateTreeInstanceDebugId InstanceId, const EStateTreeUpdatePhase Phase, const uint16 StateIdx, const EStateTreeTraceNodeEventType EventType);
-	void OutputTaskEventTrace(const FStateTreeInstanceDebugId InstanceId, const EStateTreeUpdatePhase Phase, const uint16 TaskIdx, FStateTreeDataView DataView, const EStateTreeTraceNodeEventType EventType, const EStateTreeRunStatus Status);
-	void OutputConditionEventTrace(const FStateTreeInstanceDebugId InstanceId, const EStateTreeUpdatePhase Phase, const uint16 ConditionIdx, FStateTreeDataView DataView, const EStateTreeTraceNodeEventType EventType);
+	void OutputStateEventTrace(const FStateTreeInstanceDebugId InstanceId, const EStateTreeUpdatePhase Phase, const FStateTreeStateHandle StateHandle, const EStateTreeTraceNodeEventType EventType);
+	void OutputTaskEventTrace(const FStateTreeInstanceDebugId InstanceId, const EStateTreeUpdatePhase Phase, const FStateTreeIndex16 TaskIdx, FStateTreeDataView DataView, const EStateTreeTraceNodeEventType EventType, const EStateTreeRunStatus Status);
+	void OutputConditionEventTrace(const FStateTreeInstanceDebugId InstanceId, const EStateTreeUpdatePhase Phase, const FStateTreeIndex16 ConditionIdx, FStateTreeDataView DataView, const EStateTreeTraceNodeEventType EventType);
 	void OutputActiveStatesEventTrace(const FStateTreeInstanceDebugId InstanceId, const FStateTreeActiveStates& ActiveStates);
 }
 
diff --git a/Engine/Plugins/Runtime/StateTree/Source/StateTreeModule/Public/StateTreeTypes.h b/Engine/Plugins/Runtime/StateTree/Source/StateTreeModule/Public/StateTreeTypes.h
--- a/Engine/Plugins/Runtime/StateTree/Source/StateTreeModule/Public/StateTreeTypes.h
+++ b/Engine/Plugins/Runtime/StateTree/Source/StateTreeModule/Public/StateTreeTypes.h
@@ -192,8 +192,19 @@ struct STATETREEMODULE_API FStateTreeStateHandle
 	static const FStateTreeStateHandle Stopped;
 	static const FStateTreeStateHandle Root;
 
+	/** @return true if the given index can be represented by the type. */
+	static bool IsValidIndex(const int32 Index)
+	{
+		return Index >= 0 && Index < (int32)MAX_uint16;
+	}
+	
 	FStateTreeStateHandle() = default;
-	explicit FStateTreeStateHandle(uint16 InIndex) : Index(InIndex) {}
+	explicit FStateTreeStateHandle(const uint16 InIndex) : Index(InIndex) {}
+	explicit FStateTreeStateHandle(const int32 InIndex) : Index()
+	{
+		check(InIndex == INDEX_NONE || IsValidIndex(InIndex));
+		Index = InIndex == INDEX_NONE ? InvalidIndex : static_cast<uint16>(InIndex);	
+	}
 
 	bool IsValid() const { return Index != InvalidIndex; }
 	void Invalidate() { Index = InvalidIndex; }
@@ -560,7 +571,7 @@ struct STATETREEMODULE_API FStateTreeRandomTimeDuration
 	{
 		const int32 MinVal = FMath::Max(0, static_cast<int32>(Duration) - static_cast<int32>(RandomVariance));
 		const int32 MaxVal = static_cast<int32>(Duration) + static_cast<int32>(RandomVariance);
-		return FMath::RandRange(MinVal, MaxVal) / Scale;
+		return static_cast<decltype(Scale)>(FMath::RandRange(MinVal, MaxVal)) / Scale;
 	}
 	
 protected:
diff --git a/Engine/Plugins/Runtime/StateTree/Source/StateTreeModule/StateTreeModule.Build.cs b/Engine/Plugins/Runtime/StateTree/Source/StateTreeModule/StateTreeModule.Build.cs
--- a/Engine/Plugins/Runtime/StateTree/Source/StateTreeModule/StateTreeModule.Build.cs
+++ b/Engine/Plugins/Runtime/StateTree/Source/StateTreeModule/StateTreeModule.Build.cs
@@ -40,6 +40,8 @@ namespace UnrealBuildTool.Rules
 				);
 			}
 
+			UnsafeTypeCastWarningLevel = WarningLevel.Error;
+
 			if (Target.Platform == UnrealTargetPlatform.Win64 && 
 				(Target.Configuration != UnrealTargetConfiguration.Shipping || Target.bBuildEditor))
 			{
