[StateTreeDebugger] fixed some Prev/Next frame button actions when scrub position ends up between frames with events
#rnx

diff --git a/Engine/Plugins/Runtime/StateTree/Source/StateTreeModule/Private/Debugger/StateTreeDebuggerTypes.cpp b/Engine/Plugins/Runtime/StateTree/Source/StateTreeModule/Private/Debugger/StateTreeDebuggerTypes.cpp
--- a/Engine/Plugins/Runtime/StateTree/Source/StateTreeModule/Private/Debugger/StateTreeDebuggerTypes.cpp
+++ b/Engine/Plugins/Runtime/StateTree/Source/StateTreeModule/Private/Debugger/StateTreeDebuggerTypes.cpp
@@ -15,7 +15,7 @@ FInstanceDescriptor::FInstanceDescriptor(const UStateTree* InStateTree, const FS
 	, StateTree(InStateTree)
 	, Name(InName)
 	, Id(InId)
-{ 
+{
 }
 
 bool FInstanceDescriptor::IsValid() const
@@ -39,24 +39,42 @@ void FScrubState::SetScrubTime(const double NewScrubTime)
 	{
 		const TArray<FFrameSpan>& Spans = EventCollections[EventCollectionIndex].FrameSpans;
 		const uint32 NewFrameSpanIndex = Spans.IndexOfByPredicate([Time = NewScrubTime](const FFrameSpan& Span)
-		{
-			return Span.Frame.StartTime <= Time && Time <= Span.Frame.EndTime; 
-		});
+			{
+				return Span.Frame.StartTime <= Time && Time <= Span.Frame.EndTime;
+			});
 
 		if (NewFrameSpanIndex != INDEX_NONE)
 		{
+			PreviousFrameSpanIndex = NewFrameSpanIndex > 0 ? NewFrameSpanIndex - 1 : NewFrameSpanIndex;
 			SetFrameSpanIndex(NewFrameSpanIndex);
 		}
 		else if (Spans.Num() > 0 && NewScrubTime > Spans.Last().Frame.EndTime)
 		{
 			// Clamp to last span so it keep the scrub valid when sticking to most recent data.
-			SetFrameSpanIndex(EventCollections[EventCollectionIndex].FrameSpans.Num() - 1);
+			const uint32 LastSpanIndex = EventCollections[EventCollectionIndex].FrameSpans.Num() - 1;
+			SetFrameSpanIndex(LastSpanIndex);
+			PreviousFrameSpanIndex = LastSpanIndex;
 		}
 		else
 		{
 			TraceFrameIndex = INDEX_NONE;
 			FrameSpanIndex = INDEX_NONE;
-			ActiveStatesIndex = INDEX_NONE;
+
+			// Find the last span before the given time (to allow Prev/Next actions)
+			PreviousFrameSpanIndex = Spans.FindLastByPredicate([Time = NewScrubTime](const FFrameSpan& Span)
+				{
+					return Span.Frame.EndTime < Time;
+				});
+
+			// If previous span found let's update the active states index (to allow Prev/Next actions)
+			if (PreviousFrameSpanIndex != INDEX_NONE)
+			{
+				UpdateActiveStatesIndex(PreviousFrameSpanIndex);
+			}
+			else
+			{
+				ActiveStatesIndex = INDEX_NONE;
+			}
 		}
 	}
 
@@ -74,7 +92,7 @@ void FScrubState::SetFrameSpanIndex(const uint32 NewFrameSpanIndex)
 	checkf(EventCollections[EventCollectionIndex].FrameSpans.IsValidIndex(FrameSpanIndex), TEXT("Internal method expecting validity checks before getting called."));
 	ScrubTime = EventCollection.FrameSpans[FrameSpanIndex].Frame.StartTime;
 	TraceFrameIndex = EventCollection.FrameSpans[FrameSpanIndex].Frame.Index;
-	UpdateActiveStatesIndex();
+	UpdateActiveStatesIndex(NewFrameSpanIndex);
 }
 
 void FScrubState::SetActiveStatesIndex(const uint32 NewActiveStatesIndex)
@@ -92,30 +110,41 @@ void FScrubState::SetActiveStatesIndex(const uint32 NewActiveStatesIndex)
 
 bool FScrubState::HasPreviousFrame() const
 {
-	return IsInBounds() ? EventCollections[EventCollectionIndex].FrameSpans.IsValidIndex(FrameSpanIndex-1) : false;
+	if (EventCollectionIndex != INDEX_NONE)
+	{
+		return IsInBounds() ? EventCollections[EventCollectionIndex].FrameSpans.IsValidIndex(FrameSpanIndex - 1) : PreviousFrameSpanIndex != INDEX_NONE;
+	}
+	return false;
 }
 
 double FScrubState::GotoPreviousFrame()
 {
-	SetFrameSpanIndex(FrameSpanIndex - 1);
+	SetFrameSpanIndex(IsInBounds() ? (FrameSpanIndex - 1) : PreviousFrameSpanIndex);
 	return ScrubTime;
 }
 
 bool FScrubState::HasNextFrame() const
 {
-	return IsInBounds() ? EventCollections[EventCollectionIndex].FrameSpans.IsValidIndex(FrameSpanIndex+1) : false;
+	if (EventCollectionIndex != INDEX_NONE)
+	{
+		const TConstArrayView<FFrameSpan> FrameSpans = EventCollections[EventCollectionIndex].FrameSpans;
+		return IsInBounds()
+				   ? FrameSpans.IsValidIndex(FrameSpanIndex + 1)
+				   : (PreviousFrameSpanIndex != INDEX_NONE && FrameSpans.IsValidIndex(PreviousFrameSpanIndex + 1));
+	}
+	return false;
 }
 
 double FScrubState::GotoNextFrame()
 {
-	SetFrameSpanIndex(FrameSpanIndex + 1);
+	SetFrameSpanIndex(IsInBounds() ? (FrameSpanIndex + 1) : PreviousFrameSpanIndex + 1);
 	return ScrubTime;
 }
 
 bool FScrubState::HasPreviousActiveStates() const
 {
 	return (EventCollectionIndex != INDEX_NONE && ActiveStatesIndex != INDEX_NONE)
-		? EventCollections[EventCollectionIndex].ActiveStatesChanges.IsValidIndex(ActiveStatesIndex-1)
+		? EventCollections[EventCollectionIndex].ActiveStatesChanges.IsValidIndex(ActiveStatesIndex - 1)
 		: false;
 }
 
@@ -128,7 +157,7 @@ double FScrubState::GotoPreviousActiveStates()
 bool FScrubState::HasNextActiveStates() const
 {
 	return (EventCollectionIndex != INDEX_NONE && ActiveStatesIndex != INDEX_NONE)
-		? EventCollections[EventCollectionIndex].ActiveStatesChanges.IsValidIndex(ActiveStatesIndex+1)
+		? EventCollections[EventCollectionIndex].ActiveStatesChanges.IsValidIndex(ActiveStatesIndex + 1)
 		: false;
 }
 
@@ -143,17 +172,17 @@ const FInstanceEventCollection& FScrubState::GetEventCollection() const
 	return EventCollectionIndex != INDEX_NONE ? EventCollections[EventCollectionIndex] : FInstanceEventCollection::Invalid;
 }
 
-void FScrubState::UpdateActiveStatesIndex()
+void FScrubState::UpdateActiveStatesIndex(const uint32 SpanIndex)
 {
 	check(EventCollectionIndex != INDEX_NONE);
 	const FInstanceEventCollection& EventCollection = EventCollections[EventCollectionIndex];
 
 	// Need to find the index of a frame span that contains an active states changed event; either the current one has it otherwise look backward to find the last one
-	ActiveStatesIndex = EventCollection.ActiveStatesChanges.FindLastByPredicate([CurrentFrameSpanIndex = FrameSpanIndex](const TTuple<uint32, uint32> SpanAndEventIndices)
-	{
-		return SpanAndEventIndices.Key <= CurrentFrameSpanIndex;
-	});
+	ActiveStatesIndex = EventCollection.ActiveStatesChanges.FindLastByPredicate(
+		[SpanIndex](const TTuple<uint32, uint32> SpanAndEventIndices)
+		{
+			return SpanAndEventIndices.Key <= SpanIndex;
+		});
 }
-
 } // UE::StateTreeDebugger
 #endif // WITH_STATETREE_DEBUGGER
diff --git a/Engine/Plugins/Runtime/StateTree/Source/StateTreeModule/Public/Debugger/StateTreeDebuggerTypes.h b/Engine/Plugins/Runtime/StateTree/Source/StateTreeModule/Public/Debugger/StateTreeDebuggerTypes.h
--- a/Engine/Plugins/Runtime/StateTree/Source/StateTreeModule/Public/Debugger/StateTreeDebuggerTypes.h
+++ b/Engine/Plugins/Runtime/StateTree/Source/StateTreeModule/Public/Debugger/StateTreeDebuggerTypes.h
@@ -132,6 +132,7 @@ public:
 	uint32 EventCollectionIndex = INDEX_NONE;
 	uint64 TraceFrameIndex = INDEX_NONE;
 	uint32 FrameSpanIndex = INDEX_NONE;
+	uint32 PreviousFrameSpanIndex = INDEX_NONE;
 	uint32 ActiveStatesIndex = INDEX_NONE;
 	
 
@@ -199,7 +200,7 @@ public:
 private:
 	void SetFrameSpanIndex(uint32 NewFrameSpanIndex);
 	void SetActiveStatesIndex(uint32 NewActiveStatesIndex);
-	void UpdateActiveStatesIndex();
+	void UpdateActiveStatesIndex(uint32 SpanIndex);
 };
 
 } // UE::StateTreeDebugger
