Downgrade more errors :(
#rnx

diff --git a/Engine/Plugins/Runtime/StateTree/Source/StateTreeModule/Private/StateTree.cpp b/Engine/Plugins/Runtime/StateTree/Source/StateTreeModule/Private/StateTree.cpp
--- a/Engine/Plugins/Runtime/StateTree/Source/StateTreeModule/Private/StateTree.cpp
+++ b/Engine/Plugins/Runtime/StateTree/Source/StateTreeModule/Private/StateTree.cpp
@@ -265,7 +265,7 @@ void UStateTree::PostLoad()
 	
 	if (!Link())
 	{
-		UE_LOG(LogStateTree, Error, TEXT("%s failed to link. Asset will not be usable at runtime."), *GetFullName());	
+		UE_LOG(LogStateTree, Log, TEXT("%s failed to link. Asset will not be usable at runtime."), *GetFullName());	
 	}
 }
 
@@ -297,7 +297,7 @@ void UStateTree::Serialize(FStructuredArchiveRecord Record)
 	{
 		if (!Link())
 		{
-			UE_LOG(LogStateTree, Error, TEXT("%s failed to link. Asset will not be usable at runtime."), *GetName());	
+			UE_LOG(LogStateTree, Log, TEXT("%s failed to link. Asset will not be usable at runtime."), *GetName());	
 		}
 	}
 }
