StateTree: Fix object instance data not being handled correctly when copying a state

diff --git a/Engine/Plugins/Runtime/StateTree/Source/StateTreeEditorModule/Private/Customizations/StateTreeEditorNodeDetails.cpp b/Engine/Plugins/Runtime/StateTree/Source/StateTreeEditorModule/Private/Customizations/StateTreeEditorNodeDetails.cpp
--- a/Engine/Plugins/Runtime/StateTree/Source/StateTreeEditorModule/Private/Customizations/StateTreeEditorNodeDetails.cpp
+++ b/Engine/Plugins/Runtime/StateTree/Source/StateTreeEditorModule/Private/Customizations/StateTreeEditorNodeDetails.cpp
@@ -1824,68 +1824,76 @@ void FStateTreeEditorNodeDetails::OnStructPicked(const UScriptStruct* InStruct)
 	check(StructProperty);
 	check(StateTree);
 
+	TArray<UObject*> OuterObjects;
 	TArray<void*> RawNodeData;
+	StructProperty->GetOuterObjects(OuterObjects);
 	StructProperty->AccessRawData(RawNodeData);
 
 	GEditor->BeginTransaction(LOCTEXT("SelectNode", "Select Node"));
 
 	StructProperty->NotifyPreChange();
 
-	for (void* Data : RawNodeData)
+	if (OuterObjects.Num() == RawNodeData.Num())
 	{
-		if (FStateTreeEditorNode* Node = static_cast<FStateTreeEditorNode*>(Data))
+		for (int32 Index = 0; Index < RawNodeData.Num(); Index++)
 		{
-			Node->Reset();
-			
-			if (InStruct)
+			if (UObject* Outer = OuterObjects[Index])
 			{
-				// Generate new ID.
-				Node->ID = FGuid::NewGuid();
-
-				// Initialize node
-				Node->Node.InitializeAs(InStruct);
-				
-				// Generate new name and instantiate instance data.
-				if (InStruct->IsChildOf(FStateTreeTaskBase::StaticStruct()))
-				{
-					FStateTreeTaskBase& Task = Node->Node.GetMutable<FStateTreeTaskBase>();
-					Task.Name = FName(InStruct->GetDisplayNameText().ToString());
-
-					if (const UScriptStruct* InstanceType = Cast<const UScriptStruct>(Task.GetInstanceDataType()))
+				if (FStateTreeEditorNode* Node = static_cast<FStateTreeEditorNode*>(RawNodeData[Index]))
 					{
-						Node->Instance.InitializeAs(InstanceType);
-					}
-					else if (const UClass* InstanceClass = Cast<const UClass>(Task.GetInstanceDataType()))
+					Node->Reset();
+					
+					if (InStruct)
 					{
-						Node->InstanceObject = NewObject<UObject>(EditorData, InstanceClass);
-					}
-				}
-				else if (InStruct->IsChildOf(FStateTreeEvaluatorBase::StaticStruct()))
-				{
-					FStateTreeEvaluatorBase& Eval = Node->Node.GetMutable<FStateTreeEvaluatorBase>();
-					Eval.Name = FName(InStruct->GetDisplayNameText().ToString());
+						// Generate new ID.
+						Node->ID = FGuid::NewGuid();
+
+						// Initialize node
+						Node->Node.InitializeAs(InStruct);
+						
+						// Generate new name and instantiate instance data.
+						if (InStruct->IsChildOf(FStateTreeTaskBase::StaticStruct()))
+						{
+							FStateTreeTaskBase& Task = Node->Node.GetMutable<FStateTreeTaskBase>();
+							Task.Name = FName(InStruct->GetDisplayNameText().ToString());
 
-					if (const UScriptStruct* InstanceType = Cast<const UScriptStruct>(Eval.GetInstanceDataType()))
-					{
-						Node->Instance.InitializeAs(InstanceType);
-					}
-					else if (const UClass* InstanceClass = Cast<const UClass>(Eval.GetInstanceDataType()))
-					{
-						Node->InstanceObject = NewObject<UObject>(EditorData, InstanceClass);
-					}
-				}
-				else if (InStruct->IsChildOf(FStateTreeConditionBase::StaticStruct()))
-				{
-					FStateTreeConditionBase& Cond = Node->Node.GetMutable<FStateTreeConditionBase>();
-					Cond.Name = FName(InStruct->GetDisplayNameText().ToString());
+							if (const UScriptStruct* InstanceType = Cast<const UScriptStruct>(Task.GetInstanceDataType()))
+							{
+								Node->Instance.InitializeAs(InstanceType);
+							}
+							else if (const UClass* InstanceClass = Cast<const UClass>(Task.GetInstanceDataType()))
+							{
+								Node->InstanceObject = NewObject<UObject>(Outer, InstanceClass);
+							}
+						}
+						else if (InStruct->IsChildOf(FStateTreeEvaluatorBase::StaticStruct()))
+						{
+							FStateTreeEvaluatorBase& Eval = Node->Node.GetMutable<FStateTreeEvaluatorBase>();
+							Eval.Name = FName(InStruct->GetDisplayNameText().ToString());
 
-					if (const UScriptStruct* InstanceType = Cast<const UScriptStruct>(Cond.GetInstanceDataType()))
-					{
-						Node->Instance.InitializeAs(InstanceType);
-					}
-					else if (const UClass* InstanceClass = Cast<const UClass>(Cond.GetInstanceDataType()))
-					{
-						Node->InstanceObject = NewObject<UObject>(EditorData, InstanceClass);
+							if (const UScriptStruct* InstanceType = Cast<const UScriptStruct>(Eval.GetInstanceDataType()))
+							{
+								Node->Instance.InitializeAs(InstanceType);
+							}
+							else if (const UClass* InstanceClass = Cast<const UClass>(Eval.GetInstanceDataType()))
+							{
+								Node->InstanceObject = NewObject<UObject>(Outer, InstanceClass);
+							}
+						}
+						else if (InStruct->IsChildOf(FStateTreeConditionBase::StaticStruct()))
+						{
+							FStateTreeConditionBase& Cond = Node->Node.GetMutable<FStateTreeConditionBase>();
+							Cond.Name = FName(InStruct->GetDisplayNameText().ToString());
+
+							if (const UScriptStruct* InstanceType = Cast<const UScriptStruct>(Cond.GetInstanceDataType()))
+							{
+								Node->Instance.InitializeAs(InstanceType);
+							}
+							else if (const UClass* InstanceClass = Cast<const UClass>(Cond.GetInstanceDataType()))
+							{
+								Node->InstanceObject = NewObject<UObject>(Outer, InstanceClass);
+							}
+						}
 					}
 				}
 			}
@@ -1910,51 +1918,59 @@ void FStateTreeEditorNodeDetails::OnClassPicked(const UClass* InClass) const
 	check(StructProperty);
 	check(StateTree);
 
+	TArray<UObject*> OuterObjects;
 	TArray<void*> RawNodeData;
+	StructProperty->GetOuterObjects(OuterObjects);
 	StructProperty->AccessRawData(RawNodeData);
 
 	GEditor->BeginTransaction(LOCTEXT("SelectBlueprintNode", "Select Blueprint Node"));
 
 	StructProperty->NotifyPreChange();
 
-	for (void* Data : RawNodeData)
+	if (OuterObjects.Num() == RawNodeData.Num())
 	{
-		if (FStateTreeEditorNode* Node = static_cast<FStateTreeEditorNode*>(Data))
+		for (int32 Index = 0; Index < RawNodeData.Num(); Index++)
 		{
-			Node->Reset();
-
-			if (InClass && InClass->IsChildOf(UStateTreeTaskBlueprintBase::StaticClass()))
+			if (UObject* Outer = OuterObjects[Index])
 			{
-				Node->Node.InitializeAs(FStateTreeBlueprintTaskWrapper::StaticStruct());
-				FStateTreeBlueprintTaskWrapper& Task = Node->Node.GetMutable<FStateTreeBlueprintTaskWrapper>();
-				Task.TaskClass = const_cast<UClass*>(InClass);
-				Task.Name = FName(InClass->GetDisplayNameText().ToString());
-				
-				Node->InstanceObject = NewObject<UObject>(EditorData, InClass);
-
-				Node->ID = FGuid::NewGuid();
-			}
-			else if (InClass && InClass->IsChildOf(UStateTreeEvaluatorBlueprintBase::StaticClass()))
-			{
-				Node->Node.InitializeAs(FStateTreeBlueprintEvaluatorWrapper::StaticStruct());
-				FStateTreeBlueprintEvaluatorWrapper& Eval = Node->Node.GetMutable<FStateTreeBlueprintEvaluatorWrapper>();
-				Eval.EvaluatorClass = const_cast<UClass*>(InClass);
-				Eval.Name = FName(InClass->GetDisplayNameText().ToString());
-				
-				Node->InstanceObject = NewObject<UObject>(EditorData, InClass);
+				if (FStateTreeEditorNode* Node = static_cast<FStateTreeEditorNode*>(RawNodeData[Index]))
+				{
+					Node->Reset();
 
-				Node->ID = FGuid::NewGuid();
-			}
-			else if (InClass && InClass->IsChildOf(UStateTreeConditionBlueprintBase::StaticClass()))
-			{
-				Node->Node.InitializeAs(FStateTreeBlueprintConditionWrapper::StaticStruct());
-				FStateTreeBlueprintConditionWrapper& Cond = Node->Node.GetMutable<FStateTreeBlueprintConditionWrapper>();
-				Cond.ConditionClass = const_cast<UClass*>(InClass);
-				Cond.Name = FName(InClass->GetDisplayNameText().ToString());
+					if (InClass && InClass->IsChildOf(UStateTreeTaskBlueprintBase::StaticClass()))
+					{
+						Node->Node.InitializeAs(FStateTreeBlueprintTaskWrapper::StaticStruct());
+						FStateTreeBlueprintTaskWrapper& Task = Node->Node.GetMutable<FStateTreeBlueprintTaskWrapper>();
+						Task.TaskClass = const_cast<UClass*>(InClass);
+						Task.Name = FName(InClass->GetDisplayNameText().ToString());
+						
+						Node->InstanceObject = NewObject<UObject>(Outer, InClass);
+
+						Node->ID = FGuid::NewGuid();
+					}
+					else if (InClass && InClass->IsChildOf(UStateTreeEvaluatorBlueprintBase::StaticClass()))
+					{
+						Node->Node.InitializeAs(FStateTreeBlueprintEvaluatorWrapper::StaticStruct());
+						FStateTreeBlueprintEvaluatorWrapper& Eval = Node->Node.GetMutable<FStateTreeBlueprintEvaluatorWrapper>();
+						Eval.EvaluatorClass = const_cast<UClass*>(InClass);
+						Eval.Name = FName(InClass->GetDisplayNameText().ToString());
+						
+						Node->InstanceObject = NewObject<UObject>(Outer, InClass);
+
+						Node->ID = FGuid::NewGuid();
+					}
+					else if (InClass && InClass->IsChildOf(UStateTreeConditionBlueprintBase::StaticClass()))
+					{
+						Node->Node.InitializeAs(FStateTreeBlueprintConditionWrapper::StaticStruct());
+						FStateTreeBlueprintConditionWrapper& Cond = Node->Node.GetMutable<FStateTreeBlueprintConditionWrapper>();
+						Cond.ConditionClass = const_cast<UClass*>(InClass);
+						Cond.Name = FName(InClass->GetDisplayNameText().ToString());
 
-				Node->InstanceObject = NewObject<UObject>(EditorData, InClass);
+						Node->InstanceObject = NewObject<UObject>(Outer, InClass);
 
-				Node->ID = FGuid::NewGuid();
+						Node->ID = FGuid::NewGuid();
+					}
+				}
 			}
 		}
 	}
diff --git a/Engine/Plugins/Runtime/StateTree/Source/StateTreeEditorModule/Private/StateTreeEditorData.cpp b/Engine/Plugins/Runtime/StateTree/Source/StateTreeEditorModule/Private/StateTreeEditorData.cpp
--- a/Engine/Plugins/Runtime/StateTree/Source/StateTreeEditorModule/Private/StateTreeEditorData.cpp
+++ b/Engine/Plugins/Runtime/StateTree/Source/StateTreeEditorModule/Private/StateTreeEditorData.cpp
@@ -7,6 +7,7 @@
 #include "StateTreeEvaluatorBase.h"
 #include "StateTreeTaskBase.h"
 #include "Algo/LevenshteinDistance.h"
+#include "StateTreeEditorModule.h"
 
 #include UE_INLINE_GENERATED_CPP_BY_NAME(StateTreeEditorData)
 
@@ -22,6 +23,7 @@ void UStateTreeEditorData::PostLoad()
 {
 	Super::PostLoad();
 	ReparentStates();
+	FixObjectNodes();
 }
 
 void UStateTreeEditorData::PostEditChangeChainProperty(FPropertyChangedChainEvent& PropertyChangedEvent)
@@ -330,6 +332,7 @@ void UStateTreeEditorData::ReparentStates()
 		UObject* ExpectedOuter = ParentState ? Cast<UObject>(ParentState) : Cast<UObject>(TreeData);
 		if (State.GetOuter() != ExpectedOuter)
 		{
+			UE_LOG(LogStateTreeEditor, Log, TEXT("%s: Fixing outer on state %s."), *TreeData->GetFullName(), *GetNameSafe(&State));
 			State.Rename(nullptr, ExpectedOuter, REN_DontCreateRedirectors | REN_DoNotDirty | REN_ForceNoResetLoaders);
 		}
 		
@@ -339,6 +342,80 @@ void UStateTreeEditorData::ReparentStates()
 	});
 }
 
+
+void UStateTreeEditorData::FixObjectInstance(TSet<UObject*>& SeenObjects, UObject& Outer, FStateTreeEditorNode& Node)
+{
+	if (Node.InstanceObject)
+	{
+		// Found a duplicate reference to an object, make unique copy.
+		if (SeenObjects.Contains(Node.InstanceObject))
+		{
+			UE_LOG(LogStateTreeEditor, Log, TEXT("%s: Making duplicate node instance %s unique."), *GetFullName(), *GetNameSafe(Node.InstanceObject));
+			Node.InstanceObject = DuplicateObject(Node.InstanceObject, &Outer);
+		}
+		else
+		{
+			// Make sure the instance object is property outered.
+			if (Node.InstanceObject->GetOuter() != &Outer)
+			{
+				UE_LOG(LogStateTreeEditor, Log, TEXT("%s: Fixing outer on node instance %s."), *GetFullName(), *GetNameSafe(Node.InstanceObject));
+				Node.InstanceObject->Rename(nullptr, &Outer, REN_DontCreateRedirectors | REN_DoNotDirty | REN_ForceNoResetLoaders);
+			}
+		}
+		SeenObjects.Add(Node.InstanceObject);
+	}
+};
+
+void UStateTreeEditorData::FixObjectNodes()
+{
+	// Older version of State Trees had all instances outered to the editor data. This causes issues with State copy/paste.
+	// Instance data does not get duplicated but the copied state will reference the object on the source state instead.
+	//
+	// Ensure that all node objects are parented to their states, and make duplicated instances unique.
+
+	TSet<UObject*> SeenObjects;
+	
+	VisitHierarchy([&SeenObjects, TreeData = this](UStateTreeState& State, UStateTreeState* ParentState) mutable 
+	{
+
+		// Enter conditions
+		for (FStateTreeEditorNode& Node : State.EnterConditions)
+		{
+			TreeData->FixObjectInstance(SeenObjects, State, Node);
+		}
+		
+		// Tasks
+		for (FStateTreeEditorNode& Node : State.Tasks)
+		{
+			TreeData->FixObjectInstance(SeenObjects, State, Node);
+		}
+
+		TreeData->FixObjectInstance(SeenObjects, State, State.SingleTask);
+
+
+		// Transitions
+		for (FStateTreeTransition& Transition : State.Transitions)
+		{
+			for (FStateTreeEditorNode& Node : Transition.Conditions)
+			{
+				TreeData->FixObjectInstance(SeenObjects, State, Node);
+			}
+		}
+		
+		return EStateTreeVisitor::Continue;
+	});
+
+	for (FStateTreeEditorNode& Node : Evaluators)
+	{
+		FixObjectInstance(SeenObjects, *this, Node);
+	}
+
+	for (FStateTreeEditorNode& Node : GlobalTasks)
+	{
+		FixObjectInstance(SeenObjects, *this, Node);
+	}
+}
+
 EStateTreeVisitor UStateTreeEditorData::VisitStateNodes(const UStateTreeState& State, TFunctionRef<EStateTreeVisitor(const UStateTreeState* State, const FGuid& ID, const FName& Name, const EStateTreeNodeType NodeType, const UScriptStruct* NodeStruct, const UStruct* InstanceStruct)> InFunc) const
 {
 	bool bContinue = true;
diff --git a/Engine/Plugins/Runtime/StateTree/Source/StateTreeEditorModule/Public/StateTreeEditorData.h b/Engine/Plugins/Runtime/StateTree/Source/StateTreeEditorModule/Public/StateTreeEditorData.h
--- a/Engine/Plugins/Runtime/StateTree/Source/StateTreeEditorModule/Public/StateTreeEditorData.h
+++ b/Engine/Plugins/Runtime/StateTree/Source/StateTreeEditorModule/Public/StateTreeEditorData.h
@@ -197,6 +197,11 @@ PRAGMA_ENABLE_DEPRECATION_WARNINGS
 
 	// ~StateTree Builder API
 
+private:
+	void FixObjectInstance(TSet<UObject*>& SeenObjects, UObject& Outer, FStateTreeEditorNode& Node);
+	void FixObjectNodes();
+
+public:
 	/** Schema describing which inputs, evaluators, and tasks a StateTree can contain */	
 	UPROPERTY(EditDefaultsOnly, Category = Common, Instanced)
 	TObjectPtr<UStateTreeSchema> Schema = nullptr;
