[StateTree debugger]
- added traces for transitions
- support for recursive state selection
- fixed breakpoints not stopping execution since timelines were added
- added meta data to types used for event visualization
- forced expansion for all treeview items in frame details
#rnx

diff --git a/Engine/Plugins/Runtime/StateTree/Source/StateTreeEditorModule/Private/Debugger/SStateTreeDebuggerInstanceTree.cpp b/Engine/Plugins/Runtime/StateTree/Source/StateTreeEditorModule/Private/Debugger/SStateTreeDebuggerInstanceTree.cpp
--- a/Engine/Plugins/Runtime/StateTree/Source/StateTreeEditorModule/Private/Debugger/SStateTreeDebuggerInstanceTree.cpp
+++ b/Engine/Plugins/Runtime/StateTree/Source/StateTreeEditorModule/Private/Debugger/SStateTreeDebuggerInstanceTree.cpp
@@ -55,6 +55,15 @@ void SStateTreeDebuggerInstanceTree::TreeExpansionChanged(TSharedPtr<RewindDebug
 	OnExpansionChanged.ExecuteIfBound();
 }
 
+TSharedPtr<RewindDebugger::FRewindDebuggerTrack> SStateTreeDebuggerInstanceTree::GetSelection() const
+{
+	TArray<TSharedPtr<RewindDebugger::FRewindDebuggerTrack>> SelectedItems;
+	const int32 NumSelected = TreeView->GetSelectedItems(SelectedItems);
+
+	// TreeView uses 'SelectionMode(ESelectionMode::Single)' so number of selected items is 0 or 1
+	return NumSelected > 0 ? SelectedItems.Top() : nullptr;	
+}
+
 void SStateTreeDebuggerInstanceTree::SetSelection(const TSharedPtr<RewindDebugger::FRewindDebuggerTrack>& SelectedItem) const
 {
 	TreeView->SetSelection(SelectedItem);
diff --git a/Engine/Plugins/Runtime/StateTree/Source/StateTreeEditorModule/Private/Debugger/SStateTreeDebuggerInstanceTree.h b/Engine/Plugins/Runtime/StateTree/Source/StateTreeEditorModule/Private/Debugger/SStateTreeDebuggerInstanceTree.h
--- a/Engine/Plugins/Runtime/StateTree/Source/StateTreeEditorModule/Private/Debugger/SStateTreeDebuggerInstanceTree.h
+++ b/Engine/Plugins/Runtime/StateTree/Source/StateTreeEditorModule/Private/Debugger/SStateTreeDebuggerInstanceTree.h
@@ -36,7 +36,8 @@ public:
 
 	void Refresh();
 	void RestoreExpansion();
-	
+
+	TSharedPtr<RewindDebugger::FRewindDebuggerTrack> GetSelection() const;
 	void SetSelection(const TSharedPtr<RewindDebugger::FRewindDebuggerTrack>& SelectedItem) const;
 	void ScrollTo(double ScrollOffset) const;
 
diff --git a/Engine/Plugins/Runtime/StateTree/Source/StateTreeEditorModule/Private/Debugger/SStateTreeDebuggerView.cpp b/Engine/Plugins/Runtime/StateTree/Source/StateTreeEditorModule/Private/Debugger/SStateTreeDebuggerView.cpp
--- a/Engine/Plugins/Runtime/StateTree/Source/StateTreeEditorModule/Private/Debugger/SStateTreeDebuggerView.cpp
+++ b/Engine/Plugins/Runtime/StateTree/Source/StateTreeEditorModule/Private/Debugger/SStateTreeDebuggerView.cpp
@@ -19,7 +19,6 @@
 #include "StateTreeDebuggerTrack.h"
 #include "StateTreeExecutionContext.h"
 #include "StateTreeState.h"
-#include "StateTreeTaskBase.h"
 #include "StateTreeViewModel.h"
 #include "Widgets/Input/SComboButton.h"
 #include "Widgets/Layout/SBorder.h"
@@ -113,7 +112,15 @@ protected:
 					.Font(FCoreStyle::GetDefaultFontStyle("Mono", 9))
 					.Text_Lambda([&Event=Item->Event, this]()
 					{
-						return GetTextForEvent(Event);
+						FString EventDescription;
+						if (const UStateTree* StateTree = StateTreeViewModel->GetStateTree())
+						{
+							Visit([&EventDescription, StateTree](auto& TypedEvent)
+								{
+									EventDescription = TypedEvent.ToString(*StateTree);
+								}, Event);
+						}
+						return FText::FromString(EventDescription);
 					})
 				];
 		}
@@ -121,74 +128,6 @@ protected:
 		return Contents.ToSharedRef();
 	}
 
-	FText GetTextForEvent(const FStateTreeTraceEventVariantType& Event) const
-	{
-		const UStateTree* StateTree = StateTreeViewModel->GetStateTree();
-
-		// Use log event messages directly
-		if (const FStateTreeTraceLogEvent* LogEvent = Event.TryGet<FStateTreeTraceLogEvent>())
-		{
-			if (LogEvent->Message.Len())
-			{
-				return FText::FromString(*LogEvent->Message);
-			}
-		}
-		// Process state events (index has a different meaning)
-		else if (const FStateTreeTraceStateEvent* StateEvent = Event.TryGet<FStateTreeTraceStateEvent>())
-		{
-			const FStateTreeStateHandle StateHandle(StateEvent->Idx);
-			if (const FCompactStateTreeState* CompactState = StateTree->GetStateFromHandle(StateHandle))
-			{
-				return FText::FromString
-				(FString::Printf
-					(TEXT("%s State '%s'"),
-						*StaticEnum<EStateTreeTraceNodeEventType>()->GetNameStringByValue((int64)StateEvent->EventType),
-						*CompactState->Name.ToString()));
-			}
-		}
-		// Process Tasks events
-		else if (const FStateTreeTraceTaskEvent* TaskEvent = Event.TryGet<FStateTreeTraceTaskEvent>())
-		{
-			const FConstStructView NodeView = StateTree->GetNode(TaskEvent->Idx);
-			const FStateTreeNodeBase* Node = NodeView.GetPtr<const FStateTreeNodeBase>();
-
-			return FText::FromString
-			(FString::Printf
-				(TEXT("%s:%s %s '%s'"),
-					*StaticEnum<EStateTreeTraceNodeEventType>()->GetNameStringByValue((int64)TaskEvent->EventType),
-					*StaticEnum<EStateTreeRunStatus>()->GetNameStringByValue((int64)TaskEvent->Status),
-					*NodeView.GetScriptStruct()->GetName(),
-					Node != nullptr ? *Node->Name.ToString() : *LexToString(TaskEvent->Idx)));
-		}
-		// Process Conditions events
-		else if (const FStateTreeTraceConditionEvent* ConditionEvent = Event.TryGet<FStateTreeTraceConditionEvent>())
-		{
-			const FConstStructView NodeView = StateTree->GetNode(ConditionEvent->Idx);
-			const FStateTreeNodeBase* Node = NodeView.GetPtr<const FStateTreeNodeBase>();
-
-			return FText::FromString
-			(FString::Printf
-				(TEXT("%s %s '%s'"),
-					*StaticEnum<EStateTreeTraceNodeEventType>()->GetNameStringByValue((int64)ConditionEvent->EventType),
-					*NodeView.GetScriptStruct()->GetName(),
-					Node != nullptr ? *Node->Name.ToString() : *LexToString(ConditionEvent->Idx)));
-		}
-		// Process ActiveStates events
-		else if (const FStateTreeTraceActiveStatesEvent* ActiveStatesEvent = Event.TryGet<FStateTreeTraceActiveStatesEvent>())
-		{
-			FString StatePath;
-			for (int32 i = 0; i < ActiveStatesEvent->ActiveStates.Num(); i++)
-			{
-				const FCompactStateTreeState& State = StateTree->GetStates()[ActiveStatesEvent->ActiveStates[i].Index];
-				StatePath.Appendf(TEXT("%s%s"), i == 0 ? TEXT("") : TEXT("."), *State.Name.ToString());
-			}
-
-			return FText::FromString(FString::Printf(TEXT("New active states: '%s'"), *StatePath));
-		}
-
-		return FText();
-	}
-
 	TSharedPtr<FStateTreeViewModel> StateTreeViewModel;
 	TSharedPtr<UE::StateTreeDebugger::FEventTreeElement> Item;
 };
@@ -389,17 +328,12 @@ void SStateTreeDebuggerView::Construct(const FArguments& InArgs, const UStateTre
 			FString TypePath;
 			FString InstanceDataAsText;
 
-			if (const FStateTreeTraceConditionEvent* ConditionEvent = InSelectedItem->Event.TryGet<FStateTreeTraceConditionEvent>())
-			{
-				TypePath = ConditionEvent->TypePath;
-				InstanceDataAsText = ConditionEvent->InstanceDataAsText;
-			}
-			else if (const FStateTreeTraceTaskEvent* TaskEvent = InSelectedItem->Event.TryGet<FStateTreeTraceTaskEvent>())
-			{
-				TypePath = TaskEvent->TypePath;
-				InstanceDataAsText = TaskEvent->InstanceDataAsText;
-			}
-
+			Visit([&TypePath, &InstanceDataAsText](auto& TypedEvent)
+				{
+					TypePath = TypedEvent.GetDataTypePath();
+					InstanceDataAsText = TypedEvent.GetDataAsText();
+				}, InSelectedItem->Event);
+				
 			if (!TypePath.IsEmpty())
 			{
 				FDetailsViewArgs DetailsViewArgs;
@@ -721,7 +655,6 @@ void SStateTreeDebuggerView::OnDebuggerScrubStateChanged(const UE::StateTreeDebu
 {
 	// Rebuild frame details from the events of that frame
 	EventsTreeElements.Reset();
-	EventsTreeView->RequestTreeRefresh();
 
 	const UE::StateTreeDebugger::FInstanceEventCollection& EventCollection = ScrubState.GetEventCollection();
 	const TConstArrayView<const FStateTreeTraceEventVariantType> Events = EventCollection.Events;
@@ -735,35 +668,63 @@ void SStateTreeDebuggerView::OnDebuggerScrubStateChanged(const UE::StateTreeDebu
 	check(Spans.Num());
 	check(StateTree.IsValid());
 
-	TArray<TSharedPtr<UE::StateTreeDebugger::FEventTreeElement>, TInlineAllocator<8>> ParentStack;
+	struct FParentInfo
+	{
+		TSharedPtr<UE::StateTreeDebugger::FEventTreeElement> Element;
+		FStateTreeStateHandle State;
+	};
+	TArray<FParentInfo, TInlineAllocator<8>> ParentStack;
 	EStateTreeUpdatePhase LastPhase = EStateTreeUpdatePhase::Unset;
 	EStateTreeUpdatePhase EventPhase = EStateTreeUpdatePhase::Unset;
 
 	const int32 SpanIdx = ScrubState.FrameSpanIndex;
 	const int32 FirstEventIdx = Spans[SpanIdx].EventIdx;
+	const TraceServices::FFrame Frame = Spans[SpanIdx].Frame;
 	const int32 MaxEventIdx = Spans.IsValidIndex(SpanIdx+1) ? Spans[SpanIdx+1].EventIdx : Events.Num();
+	
+	auto CreateParentFunc = [this, &ParentStack, Frame](const EStateTreeUpdatePhase NodePhase, const FStateTreeStateHandle StateHandle)
+		{
+			// Create log event to describe the phase
+			FString NodeDesc;
+			if (StateHandle.IsValid())
+			{
+				const FCompactStateTreeState* CompactState = StateTree->GetStateFromHandle(StateHandle);
+				NodeDesc = FString::Printf(TEXT(" '%s'"), CompactState != nullptr ? *CompactState->Name.ToString() : *StateHandle.Describe());
+			}
+
+			FStateTreeTraceLogEvent LogEvent(NodePhase,	FString::Printf(TEXT("%s%s"), *UEnum::GetDisplayValueAsText(NodePhase).ToString(), *NodeDesc));
+
+			// Create Tree element to hold the event
+			const TSharedPtr<UE::StateTreeDebugger::FEventTreeElement> NewElement = MakeShareable(
+				new UE::StateTreeDebugger::FEventTreeElement(Frame, FStateTreeTraceEventVariantType(TInPlaceType<FStateTreeTraceLogEvent>(), LogEvent)));
 
+			// Push tree element to the proper hierarchy level
+			TArray<TSharedPtr<UE::StateTreeDebugger::FEventTreeElement>>& TreeElements = (ParentStack.IsEmpty()) ? EventsTreeElements : ParentStack.Top().Element->Children;
+			return ParentStack.Push({TreeElements.Add_GetRef(NewElement), StateHandle});
+		};
+	
 	for (int32 EventIdx = FirstEventIdx; EventIdx < MaxEventIdx; EventIdx++)
 	{
 		const FStateTreeTraceEventVariantType& Event = Events[EventIdx];
-		EventPhase = LastPhase;
-
-		// Need to test each type explicitly with TVariant even if they are all FStateTreeTracePhaseEvent
-		if (const FStateTreeTraceLogEvent* LogEvent = Event.TryGet<FStateTreeTraceLogEvent>())
-		{
-			EventPhase = LogEvent->Phase;
-		}
-		else if (const FStateTreeTraceStateEvent* StateEvent = Event.TryGet<FStateTreeTraceStateEvent>())
-		{
-			EventPhase = StateEvent->Phase;
-		}
-		else if (const FStateTreeTraceTaskEvent* TaskEvent = Event.TryGet<FStateTreeTraceTaskEvent>())
+		Visit([&EventPhase](auto& TypedEvent)
 		{
-			EventPhase = TaskEvent->Phase;
-		}
-		else if (const FStateTreeTraceConditionEvent* ConditionEvent = Event.TryGet<FStateTreeTraceConditionEvent>())
+			EventPhase = TypedEvent.Phase;
+		}, Event);
+
+		// Keep track of required state push/pop to be processed after phases
+		FStateTreeStateHandle StateToPush;
+		bool bPopState = false;
+
+		if (const FStateTreeTraceStateEvent* StateEvent = Event.TryGet<FStateTreeTraceStateEvent>())
 		{
-			EventPhase = ConditionEvent->Phase;
+			if (StateEvent->EventType == EStateTreeTraceNodeEventType::PushStateSelection)
+			{
+				StateToPush = StateEvent->GetStateHandle();
+			}
+			else if (StateEvent->EventType == EStateTreeTraceNodeEventType::PopStateSelection)
+			{
+				bPopState = true;
+			}
 		}
 
 		// Create a hierarchy level for each update phase
@@ -783,15 +744,7 @@ void SStateTreeDebuggerView::OnDebuggerScrubStateChanged(const UE::StateTreeDebu
 					if (EnumHasAnyFlags(LastPhase & PhasesDiff, Phase))
 					{
 						check(ParentStack.Num());
-						TSharedPtr<UE::StateTreeDebugger::FEventTreeElement> RemovedElement = ParentStack.Pop();
-						if (RemovedElement->Children.IsEmpty())
-						{
-							if (ParentStack.Num())
-							{
-								ParentStack.Last()->Children.Remove(RemovedElement);
-							}
-						}
-
+						ParentStack.Pop();
 						EnumRemoveFlags(PhasesDiff, Phase);
 						if (EnumHasAnyFlags(LastPhase, PhasesDiff) == false)
 						{
@@ -804,23 +757,13 @@ void SStateTreeDebuggerView::OnDebuggerScrubStateChanged(const UE::StateTreeDebu
 			// Push required phases from first enum
 			if (EnumHasAnyFlags(EventPhase, PhasesDiff))
 			{
+				const FStateTreeStateHandle StateHandle = ParentStack.IsEmpty() ? FStateTreeStateHandle() : ParentStack.Top().State; 
 				for (int i = 0; i < NumEnum; ++i)
 				{
 					const EStateTreeUpdatePhase Phase = static_cast<EStateTreeUpdatePhase>(PhaseEnum->GetValueByIndex(i));
 					if (EnumHasAnyFlags(EventPhase & PhasesDiff, Phase))
 					{
-						// Create fake log event to describe the phase
-						FStateTreeTraceLogEvent DummyEvent(EStateTreeUpdatePhase::Unset,
-							FString::Printf(TEXT("%s"), *StaticEnum<EStateTreeUpdatePhase>()->GetValueOrBitfieldAsString((int64)Phase)));
-
-						// Create Tree element to hold the event
-						TSharedPtr<UE::StateTreeDebugger::FEventTreeElement> NewElement = MakeShareable(new UE::StateTreeDebugger::FEventTreeElement(
-							Spans[SpanIdx].Frame,
-							FStateTreeTraceEventVariantType(TInPlaceType<FStateTreeTraceLogEvent>(), DummyEvent)));
-
-						// Push tree element to the proper stack level
-						TArray<TSharedPtr<UE::StateTreeDebugger::FEventTreeElement>>& Elements = ParentStack.IsEmpty() ? EventsTreeElements : ParentStack.Last()->Children;
-						ParentStack.Push(Elements.Add_GetRef(NewElement));
+						CreateParentFunc(Phase, StateHandle);
 						
 						EnumRemoveFlags(PhasesDiff, Phase);
 						if (EnumHasAnyFlags(EventPhase, PhasesDiff) == false)
@@ -830,12 +773,39 @@ void SStateTreeDebuggerView::OnDebuggerScrubStateChanged(const UE::StateTreeDebu
 					}
 				}
 			}
+		}
 
-			LastPhase = EventPhase;
+		LastPhase = EventPhase;
+
+		if (bPopState)
+		{
+			ParentStack.Pop();
+			// We use pop state to maintain the state stack but there is no element to visualize				
+			continue;
 		}
 
-		TArray<TSharedPtr<UE::StateTreeDebugger::FEventTreeElement>>& Elements = ParentStack.IsEmpty() ? EventsTreeElements : ParentStack.Last()->Children;
-		Elements.Add(MakeShareable(new UE::StateTreeDebugger::FEventTreeElement(Spans[SpanIdx].Frame, Event)));
+		if (StateToPush.IsValid())
+		{
+			// We use pushed state to create a new parent only
+			CreateParentFunc(EStateTreeUpdatePhase::StateSelection, StateToPush);
+			continue;
+		}
+
+		TArray<TSharedPtr<UE::StateTreeDebugger::FEventTreeElement>>& Elements = ParentStack.IsEmpty() ? EventsTreeElements : ParentStack.Top().Element->Children;
+        Elements.Add(MakeShareable(new UE::StateTreeDebugger::FEventTreeElement(Spans[SpanIdx].Frame, Event)));
+	}
+
+	EventsTreeView->ClearExpandedItems();
+	ExpandAll(EventsTreeElements);
+	EventsTreeView->RequestTreeRefresh();
+}
+
+void SStateTreeDebuggerView::ExpandAll(const TArray<TSharedPtr<UE::StateTreeDebugger::FEventTreeElement>>& Items)
+{
+	for (const TSharedPtr<UE::StateTreeDebugger::FEventTreeElement>& Item : Items)
+	{
+		EventsTreeView->SetItemExpansion(Item, true);
+		ExpandAll(Item->Children);
 	}
 }
 
diff --git a/Engine/Plugins/Runtime/StateTree/Source/StateTreeEditorModule/Private/Debugger/SStateTreeDebuggerView.h b/Engine/Plugins/Runtime/StateTree/Source/StateTreeEditorModule/Private/Debugger/SStateTreeDebuggerView.h
--- a/Engine/Plugins/Runtime/StateTree/Source/StateTreeEditorModule/Private/Debugger/SStateTreeDebuggerView.h
+++ b/Engine/Plugins/Runtime/StateTree/Source/StateTreeEditorModule/Private/Debugger/SStateTreeDebuggerView.h
@@ -125,6 +125,9 @@ private:
 	/** Callback used to reflect debugger scrub state in the UI. */
 	void OnDebuggerScrubStateChanged(const UE::StateTreeDebugger::FScrubState& ScrubState);
 
+	/** Recursively sets tree items as expanded. */
+	void ExpandAll(const TArray<TSharedPtr<UE::StateTreeDebugger::FEventTreeElement>>& Items);
+
 	/** Range controlled by the timeline widgets and used to adjust cursor position and track content. */
 	TRange<double> ViewRange = TRange<double>(0, 10);
 
diff --git a/Engine/Plugins/Runtime/StateTree/Source/StateTreeModule/Private/Debugger/StateTreeDebugger.cpp b/Engine/Plugins/Runtime/StateTree/Source/StateTreeModule/Private/Debugger/StateTreeDebugger.cpp
--- a/Engine/Plugins/Runtime/StateTree/Source/StateTreeModule/Private/Debugger/StateTreeDebugger.cpp
+++ b/Engine/Plugins/Runtime/StateTree/Source/StateTreeModule/Private/Debugger/StateTreeDebugger.cpp
@@ -113,9 +113,6 @@ void FStateTreeDebugger::StopAnalysis()
 void FStateTreeDebugger::Pause()
 {
 	bPaused = true;
-
-	// Force a refresh to make sure most recent traces are processed
-	SyncToCurrentSessionDuration();
 }
 
 void FStateTreeDebugger::Unpause()
@@ -151,9 +148,6 @@ void FStateTreeDebugger::SelectInstance(const FStateTreeInstanceDebugId Instance
 {
 	if (SelectedInstanceId != InstanceId)
 	{
-		HitBreakpointStateIndex = INDEX_NONE;
-		HitBreakpointInstanceId.Reset();
-		
 		SelectedInstanceId = InstanceId;
 
 		// Notify so listener can cleanup anything related to previous instance
@@ -579,6 +573,10 @@ void FStateTreeDebugger::SendNotifications()
 		check(HitBreakpointInstanceId.IsValid());
 		check(StatesWithBreakpoint.IsValidIndex(HitBreakpointStateIndex));
 
+		// Force scrub time to latest read time to reflect most recent events.
+		// This will notify scrub position changed and active states
+		SetScrubTime(LastTraceReadTime);
+
 		// Make sure the instance is selected in case the breakpoint was set for any instances 
 		if (SelectedInstanceId != HitBreakpointInstanceId)
 		{
@@ -642,7 +640,7 @@ bool FStateTreeDebugger::ProcessEvent(const FStateTreeInstanceDebugId InstanceId
 
 			if (SelectedInstanceId.IsValid() || InstanceStateTree == StateTreeAsset)
 			{
-				HitBreakpointStateIndex = StatesWithBreakpoint.Find(FStateTreeStateHandle(StateEvent->Idx));
+				HitBreakpointStateIndex = StatesWithBreakpoint.Find(StateEvent->GetStateHandle());
 				if (HitBreakpointStateIndex != INDEX_NONE)
 				{
 					HitBreakpointInstanceId = InstanceId;
diff --git a/Engine/Plugins/Runtime/StateTree/Source/StateTreeModule/Private/Debugger/StateTreeTrace.cpp b/Engine/Plugins/Runtime/StateTree/Source/StateTreeModule/Private/Debugger/StateTreeTrace.cpp
--- a/Engine/Plugins/Runtime/StateTree/Source/StateTreeModule/Private/Debugger/StateTreeTrace.cpp
+++ b/Engine/Plugins/Runtime/StateTree/Source/StateTreeModule/Private/Debugger/StateTreeTrace.cpp
@@ -7,6 +7,7 @@
 #include "Exporters/Exporter.h"
 #include "Serialization/BufferArchive.h"
 #include "StateTree.h"
+#include "StateTreeExecutionTypes.h"
 #include "UObject/Package.h"
 #include "Trace/Trace.inl"
 
@@ -36,8 +37,9 @@ UE_TRACE_EVENT_BEGIN(StateTreeDebugger, StateEvent)
 	UE_TRACE_EVENT_FIELD(uint32, InstanceId)
 	UE_TRACE_EVENT_FIELD(uint32, InstanceSerial)
 	UE_TRACE_EVENT_FIELD(uint16, Phase)
-	UE_TRACE_EVENT_FIELD(uint16, Idx)
+	UE_TRACE_EVENT_FIELD(uint16, StateIndex)
 	UE_TRACE_EVENT_FIELD(uint8, EventType)
+	UE_TRACE_EVENT_FIELD(uint8, SelectionBehavior)
 UE_TRACE_EVENT_END()
 
 UE_TRACE_EVENT_BEGIN(StateTreeDebugger, TaskEvent)
@@ -45,18 +47,27 @@ UE_TRACE_EVENT_BEGIN(StateTreeDebugger, TaskEvent)
 	UE_TRACE_EVENT_FIELD(uint32, InstanceId)
 	UE_TRACE_EVENT_FIELD(uint32, InstanceSerial)
 	UE_TRACE_EVENT_FIELD(uint16, Phase)
-	UE_TRACE_EVENT_FIELD(uint16, Idx)
+	UE_TRACE_EVENT_FIELD(uint16, NodeIndex)
 	UE_TRACE_EVENT_FIELD(uint8[], DataView)
 	UE_TRACE_EVENT_FIELD(uint8, EventType)
 	UE_TRACE_EVENT_FIELD(uint8, Status)
 UE_TRACE_EVENT_END()
 
+UE_TRACE_EVENT_BEGIN(StateTreeDebugger, TransitionEvent)
+	UE_TRACE_EVENT_FIELD(uint64, Cycle)
+	UE_TRACE_EVENT_FIELD(uint32, InstanceId)
+	UE_TRACE_EVENT_FIELD(uint32, InstanceSerial)
+	UE_TRACE_EVENT_FIELD(uint16, Phase)
+	UE_TRACE_EVENT_FIELD(uint16, TransitionIndex)
+	UE_TRACE_EVENT_FIELD(uint8, EventType)
+UE_TRACE_EVENT_END()
+
 UE_TRACE_EVENT_BEGIN(StateTreeDebugger, ConditionEvent)
 	UE_TRACE_EVENT_FIELD(uint64, Cycle)
 	UE_TRACE_EVENT_FIELD(uint32, InstanceId)
 	UE_TRACE_EVENT_FIELD(uint32, InstanceSerial)
 	UE_TRACE_EVENT_FIELD(uint16, Phase)
-	UE_TRACE_EVENT_FIELD(uint16, Idx)
+	UE_TRACE_EVENT_FIELD(uint16, NodeIndex)
 	UE_TRACE_EVENT_FIELD(uint8[], DataView)
 	UE_TRACE_EVENT_FIELD(uint8, EventType)
 UE_TRACE_EVENT_END()
@@ -65,6 +76,7 @@ UE_TRACE_EVENT_BEGIN(StateTreeDebugger, ActiveStatesEvent)
 	UE_TRACE_EVENT_FIELD(uint64, Cycle)
 	UE_TRACE_EVENT_FIELD(uint32, InstanceId)
 	UE_TRACE_EVENT_FIELD(uint32, InstanceSerial)
+	UE_TRACE_EVENT_FIELD(uint16, Phase)
 	UE_TRACE_EVENT_FIELD(uint16[], ActiveStates)
 UE_TRACE_EVENT_END()
 
@@ -146,7 +158,8 @@ void UE::StateTreeTrace::OutputStateEventTrace(
 	const FStateTreeInstanceDebugId InstanceId,
 	const EStateTreeUpdatePhase Phase,
 	const FStateTreeStateHandle StateHandle,
-	const EStateTreeTraceNodeEventType EventType
+	const EStateTreeTraceNodeEventType EventType,
+	const EStateTreeStateSelectionBehavior SelectionBehavior
 	)
 {
 	UE_TRACE_LOG(StateTreeDebugger, StateEvent, StateTreeDebugChannel)
@@ -154,8 +167,9 @@ void UE::StateTreeTrace::OutputStateEventTrace(
 		<< StateEvent.InstanceId(InstanceId.Id)
 		<< StateEvent.InstanceSerial(InstanceId.SerialNumber)
 		<< StateEvent.Phase(static_cast<std::underlying_type_t<EStateTreeUpdatePhase>>(Phase))
-		<< StateEvent.Idx(StateHandle.Index)
-		<< StateEvent.EventType(static_cast<std::underlying_type_t<EStateTreeTraceNodeEventType>>(EventType));
+		<< StateEvent.StateIndex(StateHandle.Index)
+		<< StateEvent.EventType(static_cast<std::underlying_type_t<EStateTreeTraceNodeEventType>>(EventType))
+		<< StateEvent.SelectionBehavior(static_cast<std::underlying_type_t<EStateTreeStateSelectionBehavior>>(SelectionBehavior));
 }
 
 void UE::StateTreeTrace::OutputTaskEventTrace(
@@ -175,12 +189,31 @@ void UE::StateTreeTrace::OutputTaskEventTrace(
 		<< TaskEvent.InstanceId(InstanceId.Id)
 		<< TaskEvent.InstanceSerial(InstanceId.SerialNumber)
 		<< TaskEvent.Phase(static_cast<std::underlying_type_t<EStateTreeUpdatePhase>>(Phase))
-		<< TaskEvent.Idx(TaskIdx.Get())
+		<< TaskEvent.NodeIndex(TaskIdx.Get())
 		<< TaskEvent.DataView(Archive.GetData(), Archive.Num())
 		<< TaskEvent.EventType(static_cast<std::underlying_type_t<EStateTreeTraceNodeEventType>>(EventType))
 		<< TaskEvent.Status(static_cast<std::underlying_type_t<EStateTreeRunStatus>>(Status));
 }
 
+void UE::StateTreeTrace::OutputTransitionEventTrace(
+	const FStateTreeInstanceDebugId InstanceId,
+	const EStateTreeUpdatePhase Phase,
+	const FStateTreeIndex16 TransitionIdx,
+	const EStateTreeTraceNodeEventType EventType
+	)
+{
+	FBufferArchive Archive;
+	Archive << EventType;
+
+	UE_TRACE_LOG(StateTreeDebugger, TransitionEvent, StateTreeDebugChannel)
+	<< TransitionEvent.Cycle(FPlatformTime::Cycles64())
+	<< TransitionEvent.InstanceId(InstanceId.Id)
+	<< TransitionEvent.InstanceSerial(InstanceId.SerialNumber)
+	<< TransitionEvent.Phase(static_cast<std::underlying_type_t<EStateTreeUpdatePhase>>(Phase))
+	<< TransitionEvent.TransitionIndex(TransitionIdx.Get())
+	<< TransitionEvent.EventType(static_cast<std::underlying_type_t<EStateTreeTraceNodeEventType>>(EventType));
+}
+
 void UE::StateTreeTrace::OutputConditionEventTrace(
 	const FStateTreeInstanceDebugId InstanceId,
 	const EStateTreeUpdatePhase Phase,
@@ -197,13 +230,14 @@ void UE::StateTreeTrace::OutputConditionEventTrace(
 	<< ConditionEvent.InstanceId(InstanceId.Id)
 	<< ConditionEvent.InstanceSerial(InstanceId.SerialNumber)
 	<< ConditionEvent.Phase(static_cast<std::underlying_type_t<EStateTreeUpdatePhase>>(Phase))
-	<< ConditionEvent.Idx(ConditionIdx.Get())
+	<< ConditionEvent.NodeIndex(ConditionIdx.Get())
 	<< ConditionEvent.DataView(Archive.GetData(), Archive.Num())
 	<< ConditionEvent.EventType(static_cast<std::underlying_type_t<EStateTreeTraceNodeEventType>>(EventType));
 }
 
 void UE::StateTreeTrace::OutputActiveStatesEventTrace(
 	const FStateTreeInstanceDebugId InstanceId,
+	const EStateTreeUpdatePhase Phase,
 	const FStateTreeActiveStates& ActiveStates
 	)
 {
@@ -217,6 +251,7 @@ void UE::StateTreeTrace::OutputActiveStatesEventTrace(
 		<< ActiveStatesEvent.Cycle(FPlatformTime::Cycles64())
 		<< ActiveStatesEvent.InstanceId(InstanceId.Id)
 		<< ActiveStatesEvent.InstanceSerial(InstanceId.SerialNumber)
+		<< ActiveStatesEvent.Phase(static_cast<std::underlying_type_t<EStateTreeUpdatePhase>>(Phase))
 		<< ActiveStatesEvent.ActiveStates(StatesIndices.GetData(), StatesIndices.Num());
 }
 
diff --git a/Engine/Plugins/Runtime/StateTree/Source/StateTreeModule/Private/Debugger/StateTreeTraceAnalyzer.cpp b/Engine/Plugins/Runtime/StateTree/Source/StateTreeModule/Private/Debugger/StateTreeTraceAnalyzer.cpp
--- a/Engine/Plugins/Runtime/StateTree/Source/StateTreeModule/Private/Debugger/StateTreeTraceAnalyzer.cpp
+++ b/Engine/Plugins/Runtime/StateTree/Source/StateTreeModule/Private/Debugger/StateTreeTraceAnalyzer.cpp
@@ -23,6 +23,7 @@ void FStateTreeTraceAnalyzer::OnAnalysisBegin(const FOnAnalysisContext& Context)
 	Builder.RouteEvent(RouteId_LogMessage, "StateTreeDebugger", "LogEvent");
 	Builder.RouteEvent(RouteId_State, "StateTreeDebugger", "StateEvent");
 	Builder.RouteEvent(RouteId_Task, "StateTreeDebugger", "TaskEvent");
+	Builder.RouteEvent(RouteId_Transition, "StateTreeDebugger", "TransitionEvent");
 	Builder.RouteEvent(RouteId_Condition, "StateTreeDebugger", "ConditionEvent");
 	Builder.RouteEvent(RouteId_ActiveStates, "StateTreeDebugger", "ActiveStatesEvent");
 }
@@ -92,8 +93,9 @@ bool FStateTreeTraceAnalyzer::OnEvent(const uint16 RouteId, EStyle Style, const
 		{
 			const FStateTreeTraceStateEvent Event(
 				EventData.GetValue<EStateTreeUpdatePhase>("Phase"),
-				EventData.GetValue<int16>("Idx"),
-				EventData.GetValue<EStateTreeTraceNodeEventType>("EventType"));
+				FStateTreeIndex16(EventData.GetValue<uint16>("StateIndex")),
+				EventData.GetValue<EStateTreeTraceNodeEventType>("EventType"),
+				EventData.GetValue<EStateTreeStateSelectionBehavior>("SelectionBehavior"));
 
 			Provider.AppendEvent(FStateTreeInstanceDebugId(EventData.GetValue<uint32>("InstanceId"), EventData.GetValue<uint32>("InstanceSerial")),
 				Context.EventTime.AsSeconds(EventData.GetValue<uint64>("Cycle")),
@@ -109,7 +111,7 @@ bool FStateTreeTraceAnalyzer::OnEvent(const uint16 RouteId, EStyle Style, const
 
 			const FStateTreeTraceTaskEvent Event(
 				EventData.GetValue<EStateTreeUpdatePhase>("Phase"),
-				EventData.GetValue<int16>("Idx"),
+				FStateTreeIndex16(EventData.GetValue<uint16>("NodeIndex")),
 				EventData.GetValue<EStateTreeTraceNodeEventType>("EventType"),
 				EventData.GetValue<EStateTreeRunStatus>("Status"),
 				TypePath, DataAsText);
@@ -128,7 +130,7 @@ bool FStateTreeTraceAnalyzer::OnEvent(const uint16 RouteId, EStyle Style, const
 
 			const FStateTreeTraceConditionEvent Event(
 				EventData.GetValue<EStateTreeUpdatePhase>("Phase"),
-				EventData.GetValue<int16>("Idx"),
+				FStateTreeIndex16(EventData.GetValue<uint16>("NodeIndex")),
 				EventData.GetValue<EStateTreeTraceNodeEventType>("EventType"),
 				TypePath, DataAsText);
 			
@@ -137,9 +139,21 @@ bool FStateTreeTraceAnalyzer::OnEvent(const uint16 RouteId, EStyle Style, const
 				FStateTreeTraceEventVariantType(TInPlaceType<FStateTreeTraceConditionEvent>(), Event));
 			break;
 		}
+	case RouteId_Transition:
+		{
+			const FStateTreeTraceTransitionEvent Event(
+				EventData.GetValue<EStateTreeUpdatePhase>("Phase"),
+				FStateTreeIndex16(EventData.GetValue<uint16>("TransitionIndex")),
+				EventData.GetValue<EStateTreeTraceNodeEventType>("EventType"));
+			
+			Provider.AppendEvent(FStateTreeInstanceDebugId(EventData.GetValue<uint32>("InstanceId"), EventData.GetValue<uint32>("InstanceSerial")),
+				Context.EventTime.AsSeconds(EventData.GetValue<uint64>("Cycle")),
+				FStateTreeTraceEventVariantType(TInPlaceType<FStateTreeTraceTransitionEvent>(), Event));
+			break;
+		}
 	case RouteId_ActiveStates:
 		{
-			FStateTreeTraceActiveStatesEvent Event;
+			FStateTreeTraceActiveStatesEvent Event(EventData.GetValue<EStateTreeUpdatePhase>("Phase"));
 			Event.ActiveStates = EventData.GetArrayView<uint16>("ActiveStates");
 
 			Provider.AppendEvent(FStateTreeInstanceDebugId(EventData.GetValue<uint32>("InstanceId"), EventData.GetValue<uint32>("InstanceSerial")),
diff --git a/Engine/Plugins/Runtime/StateTree/Source/StateTreeModule/Private/Debugger/StateTreeTraceTypes.cpp b/Engine/Plugins/Runtime/StateTree/Source/StateTreeModule/Private/Debugger/StateTreeTraceTypes.cpp
new file mode 100644
--- /dev/null
+++ b/Engine/Plugins/Runtime/StateTree/Source/StateTreeModule/Private/Debugger/StateTreeTraceTypes.cpp
@@ -0,0 +1,131 @@
+// Copyright Epic Games, Inc. All Rights Reserved.
+
+#pragma once
+
+#include "Debugger/StateTreeTraceTypes.h"
+
+#include "StateTree.h"
+#include "StateTreeNodeBase.h"
+
+#if WITH_STATETREE_DEBUGGER
+
+//----------------------------------------------------------------------//
+// FStateTreeTraceLogEvent
+//----------------------------------------------------------------------//
+FString FStateTreeTraceLogEvent::ToString(const UStateTree& StateTree) const
+{
+	return (*Message);
+}
+
+
+//----------------------------------------------------------------------//
+// FStateTreeTraceTransitionEvent
+//----------------------------------------------------------------------//
+FString FStateTreeTraceTransitionEvent::ToString(const UStateTree& StateTree) const
+{
+	if (const FCompactStateTransition* Transition = StateTree.GetTransitionFromIndex(TransitionIndex.Get()))
+	{
+		const FCompactStateTreeState* CompactState = StateTree.GetStateFromHandle(Transition->State);
+		FStringBuilderBase StrBuilder;
+		StrBuilder.Appendf(TEXT("%s Goto %s (Priority: %s)"),
+					*UEnum::GetDisplayValueAsText(Transition->Trigger).ToString(),
+					CompactState != nullptr ? *CompactState->Name.ToString() : *Transition->State.Describe(),
+					*UEnum::GetDisplayValueAsText(Transition->Priority).ToString());
+
+		if (Transition->EventTag.IsValid())
+		{
+			StrBuilder.Appendf(TEXT("\n\t%s"), *Transition->EventTag.ToString()); 
+		}
+
+		return StrBuilder.ToString();
+	}
+
+	return FString::Printf(TEXT("Invalid Transition Index %s for '%s'"), *LexToString(TransitionIndex.Get()), *StateTree.GetFullName()); 
+}
+
+
+//----------------------------------------------------------------------//
+// FStateTreeTraceNodeEvent
+//----------------------------------------------------------------------//
+FString FStateTreeTraceNodeEvent::ToString(const UStateTree& StateTree) const
+{
+	const FConstStructView NodeView = StateTree.GetNode(Index.Get());
+	const FStateTreeNodeBase* Node = NodeView.GetPtr<const FStateTreeNodeBase>();
+
+	return FString::Printf(TEXT("%s '%s (%s)'"),
+			*UEnum::GetDisplayValueAsText(EventType).ToString(),
+			Node != nullptr ? *Node->Name.ToString() : *LexToString(Index.Get()),
+			*NodeView.GetScriptStruct()->GetName());
+}
+
+
+//----------------------------------------------------------------------//
+// FStateTreeTraceStateEvent
+//----------------------------------------------------------------------//
+FStateTreeStateHandle FStateTreeTraceStateEvent::GetStateHandle() const
+{
+	return FStateTreeStateHandle(Index.Get());
+}
+
+FString FStateTreeTraceStateEvent::ToString(const UStateTree& StateTree) const
+{
+	const FStateTreeStateHandle StateHandle(Index.Get());
+	if (const FCompactStateTreeState* CompactState = StateTree.GetStateFromHandle(StateHandle))
+	{
+		if (SelectionBehavior != EStateTreeStateSelectionBehavior::None)
+		{
+			return FString::Printf(TEXT("%s '%s' (%s)"),
+					*UEnum::GetDisplayValueAsText(EventType).ToString(),
+					*CompactState->Name.ToString(),
+					*UEnum::GetDisplayValueAsText(SelectionBehavior).ToString());
+		}
+		else
+		{
+			return FString::Printf(TEXT("%s '%s'"),
+					*UEnum::GetDisplayValueAsText(EventType).ToString(),
+					*CompactState->Name.ToString());
+		}
+	}
+
+	return FString::Printf(TEXT("Invalid State Index %s for '%s'"), *StateHandle.Describe(), *StateTree.GetFullName());
+}
+
+
+//----------------------------------------------------------------------//
+// FStateTreeTraceTaskEvent
+//----------------------------------------------------------------------//
+FString FStateTreeTraceTaskEvent::ToString(const UStateTree& StateTree) const
+{
+	return FString::Printf(TEXT("<%s> %s"),	*UEnum::GetDisplayValueAsText(Status).ToString(), *FStateTreeTraceNodeEvent::ToString(StateTree));
+}
+
+
+//----------------------------------------------------------------------//
+// FStateTreeTraceConditionEvent
+//----------------------------------------------------------------------//
+FString FStateTreeTraceConditionEvent::ToString(const UStateTree& StateTree) const
+{
+	return FStateTreeTraceNodeEvent::ToString(StateTree);
+}
+
+
+//----------------------------------------------------------------------//
+// FStateTreeTraceActiveStatesEvent
+//----------------------------------------------------------------------//
+FStateTreeTraceActiveStatesEvent::FStateTreeTraceActiveStatesEvent(const EStateTreeUpdatePhase Phase): FStateTreeTracePhaseEvent(Phase)
+{
+}
+
+FString FStateTreeTraceActiveStatesEvent::ToString(const UStateTree& StateTree) const
+{
+	FStringBuilderBase StatePath;
+	for (int32 i = 0; i < ActiveStates.Num(); i++)
+	{
+		const FCompactStateTreeState& State = StateTree.GetStates()[ActiveStates[i].Index];
+		StatePath.Appendf(TEXT("%s%s"), i == 0 ? TEXT("") : TEXT("."), *State.Name.ToString());
+	}
+
+	return FString::Printf(TEXT("New active states: '%s'"), *StatePath);
+}
+
+#endif // WITH_STATETREE_DEBUGGER
diff --git a/Engine/Plugins/Runtime/StateTree/Source/StateTreeModule/Private/StateTree.cpp b/Engine/Plugins/Runtime/StateTree/Source/StateTreeModule/Private/StateTree.cpp
--- a/Engine/Plugins/Runtime/StateTree/Source/StateTreeModule/Private/StateTree.cpp
+++ b/Engine/Plugins/Runtime/StateTree/Source/StateTreeModule/Private/StateTree.cpp
@@ -44,6 +44,11 @@ FGuid UStateTree::GetStateIdFromHandle(const FStateTreeStateHandle Handle) const
 	return Entry != nullptr ? Entry->Id : FGuid();
 }
 
+const FCompactStateTransition* UStateTree::GetTransitionFromIndex(const int16 TransitionIndex) const
+{
+	return Transitions.IsValidIndex(TransitionIndex) ? &Transitions[TransitionIndex] : nullptr;
+}
+
 TSharedPtr<FStateTreeInstanceData> UStateTree::GetSharedInstanceData() const
 {
 	// Create a unique index for each thread.
@@ -130,7 +135,7 @@ void UStateTree::OnObjectsReinstanced(const FReplacementObjectMap& ObjectMap)
 
 	for (TMap<UObject*, UObject*>::TConstIterator It(ObjectMap); It; ++It)
 	{
-		if (UObject* ObjectToBeReplaced = It->Value)
+		if (const UObject* ObjectToBeReplaced = It->Value)
 		{
 			if (ObjectToBeReplaced->IsInOuter(this))
 			{
@@ -145,7 +150,7 @@ void UStateTree::OnObjectsReinstanced(const FReplacementObjectMap& ObjectMap)
 		TSet<const UStruct*> Structs;
 		for (TMap<UObject*, UObject*>::TConstIterator It(ObjectMap); It; ++It)
 		{
-			if (UObject* ObjectToBeReplaced = It->Value)
+			if (const UObject* ObjectToBeReplaced = It->Value)
 			{
 				Structs.Add(ObjectToBeReplaced->GetClass());
 			}
@@ -194,10 +199,10 @@ void UStateTree::PostLoadAssetRegistryTags(const FAssetData& InAssetData, TArray
 	Super::PostLoadAssetRegistryTags(InAssetData, OutTagsAndValuesToUpdate);
 
 	static const FName SchemaTag(TEXT("Schema"));
-	FString SchemaTagValue = InAssetData.GetTagValueRef<FString>(SchemaTag);
+	const FString SchemaTagValue = InAssetData.GetTagValueRef<FString>(SchemaTag);
 	if (!SchemaTagValue.IsEmpty() && FPackageName::IsShortPackageName(SchemaTagValue))
 	{
-		FTopLevelAssetPath SchemaTagClassPathName = UClass::TryConvertShortTypeNameToPathName<UStruct>(SchemaTagValue, ELogVerbosity::Warning, TEXT("UStateTree::PostLoadAssetRegistryTags"));
+		const FTopLevelAssetPath SchemaTagClassPathName = UClass::TryConvertShortTypeNameToPathName<UStruct>(SchemaTagValue, ELogVerbosity::Warning, TEXT("UStateTree::PostLoadAssetRegistryTags"));
 		if (!SchemaTagClassPathName.IsNull())
 		{
 			OutTagsAndValuesToUpdate.Add(FAssetRegistryTag(SchemaTag, SchemaTagClassPathName.ToString(), FAssetRegistryTag::TT_Alphabetical));
@@ -285,7 +290,7 @@ void UStateTree::DeclareConstructClasses(TArray<FTopLevelAssetPath>& OutConstruc
 }
 #endif
 
-void UStateTree::Serialize(FStructuredArchiveRecord Record)
+void UStateTree::Serialize(const FStructuredArchiveRecord Record)
 {
 	Super::Serialize(Record);
 
@@ -488,7 +493,7 @@ bool UStateTree::Link()
 
 #if WITH_EDITOR
 
-void FStateTreeMemoryUsage::AddUsage(FConstStructView View)
+void FStateTreeMemoryUsage::AddUsage(const FConstStructView View)
 {
 	if (const UScriptStruct* ScriptStruct = View.GetScriptStruct())
 	{
diff --git a/Engine/Plugins/Runtime/StateTree/Source/StateTreeModule/Private/StateTreeExecutionContext.cpp b/Engine/Plugins/Runtime/StateTree/Source/StateTreeModule/Private/StateTreeExecutionContext.cpp
--- a/Engine/Plugins/Runtime/StateTree/Source/StateTreeModule/Private/StateTreeExecutionContext.cpp
+++ b/Engine/Plugins/Runtime/StateTree/Source/StateTreeModule/Private/StateTreeExecutionContext.cpp
@@ -7,6 +7,7 @@
 #include "Containers/StaticArray.h"
 #include "Debugger/StateTreeTrace.h"
 #include "Debugger/StateTreeTraceTypes.h"
+#include "Misc/ScopeExit.h"
 #include "VisualLogger/VisualLogger.h"
 #include "ProfilingDebugging/CsvProfiler.h"
 #include "Logging/LogScopedVerbosityOverride.h"
@@ -15,17 +16,21 @@
 #define STATETREE_CLOG(Condition, Verbosity, Format, ...) UE_CVLOG_UELOG((Condition), GetOwner(), LogStateTree, Verbosity, TEXT("%s") Format, *GetInstanceDescription(), ##__VA_ARGS__)
 
 #if WITH_STATETREE_DEBUGGER
-	#define STATETREE_SCOPED_PHASE(Phase) TGuardValue<EStateTreeUpdatePhase> PREPROCESSOR_JOIN(ScopedPhaseMask, __LINE__)(CurrentUpdatePhaseMask, (CurrentUpdatePhaseMask | Phase))
-    #define STATETREE_TRACE_LOG_EVENT(Format, ...) TRACE_STATETREE_LOG_EVENT(GetInstanceDebugId(), CurrentUpdatePhaseMask, Format, ##__VA_ARGS__)
-	#define STATETREE_TRACE_STATE_EVENT(Index, EventType) TRACE_STATETREE_STATE_EVENT(GetInstanceDebugId(), CurrentUpdatePhaseMask, FStateTreeStateHandle(Index), EventType);
-	#define STATETREE_TRACE_TASK_EVENT(Index, DataView, EventType, Status) TRACE_STATETREE_TASK_EVENT(GetInstanceDebugId(), CurrentUpdatePhaseMask, FStateTreeIndex16(Index), DataView, EventType, Status);
-	#define STATETREE_TRACE_CONDITION_EVENT(Index, DataViews, EventType) TRACE_STATETREE_CONDITION_EVENT(GetInstanceDebugId(), CurrentUpdatePhaseMask, FStateTreeIndex16(Index), DataView, EventType);	
+	#define STATETREE_SCOPED_PHASE(Phase) TGuardValue<EStateTreeUpdatePhase> PREPROCESSOR_JOIN(ScopedPhaseMask, __LINE__)(UpdatePhaseMask, UpdatePhaseMask | (Phase))
+	#define STATETREE_TRACE_LOG_EVENT(Format, ...) TRACE_STATETREE_LOG_EVENT(GetInstanceDebugId(), UpdatePhaseMask, Format, ##__VA_ARGS__)
+	#define STATETREE_TRACE_STATE_EVENT(StateHandle, EventType) TRACE_STATETREE_STATE_EVENT(GetInstanceDebugId(), UpdatePhaseMask, StateHandle, EventType, EStateTreeStateSelectionBehavior::None);
+	#define STATETREE_TRACE_STATE_SELECTION_EVENT(StateHandle, EventType, SelectionBehavior) TRACE_STATETREE_STATE_EVENT(GetInstanceDebugId(), UpdatePhaseMask, StateHandle, EventType, SelectionBehavior);
+	#define STATETREE_TRACE_TASK_EVENT(Index, DataView, EventType, Status) TRACE_STATETREE_TASK_EVENT(GetInstanceDebugId(), UpdatePhaseMask, FStateTreeIndex16(Index), DataView, EventType, Status);
+	#define STATETREE_TRACE_CONDITION_EVENT(Index, DataViews, EventType) TRACE_STATETREE_CONDITION_EVENT(GetInstanceDebugId(), UpdatePhaseMask, FStateTreeIndex16(Index), DataView, EventType);	
+	#define STATETREE_TRACE_TRANSITION_EVENT(Index, EventType) TRACE_STATETREE_TRANSITION_EVENT(GetInstanceDebugId(), UpdatePhaseMask, FStateTreeIndex16(Index), EventType);
 #else
 	#define STATETREE_SCOPED_PHASE(Phase)
 	#define STATETREE_TRACE_LOG_EVENT(Format, ...)
-	#define STATETREE_TRACE_STATE_EVENT(Index, EventType)
+	#define STATETREE_TRACE_STATE_EVENT(StateHandle, EventType)
+	#define STATETREE_TRACE_STATE_SELECTION_EVENT(StateHandle, EventType, SelectionBehavior)
 	#define STATETREE_TRACE_TASK_EVENT(Index, DataView, EventType, Status)
-	#define STATETREE_TRACE_CONDITION_EVENT(Index, DataViews, EventType)
+	#define STATETREE_TRACE_CONDITION_EVENT(Index, DataView, EventType)
+	#define STATETREE_TRACE_TRANSITION_EVENT(Index, EventType)
 #endif // WITH_STATETREE_DEBUGGER
 
 namespace UE::StateTree
@@ -128,7 +133,7 @@ bool FStateTreeExecutionContext::AreExternalDataViewsValid() const
 
 void FStateTreeExecutionContext::UpdateLinkedStateParameters(const FCompactStateTreeState& State, const int32 ParameterInstanceIndex)
 {
-	FStateTreeDataView StateParamsInstance = InstanceData.GetMutableStruct(ParameterInstanceIndex);
+	const FStateTreeDataView StateParamsInstance = InstanceData.GetMutableStruct(ParameterInstanceIndex);
 	FCompactStateTreeParameters& StateParams = StateParamsInstance.GetMutable<FCompactStateTreeParameters>();
 
 	// Update parameters if the state has any.
@@ -187,7 +192,7 @@ EStateTreeRunStatus FStateTreeExecutionContext::Start()
 
 	// Initialize instance data. No active states yet, so we'll initialize the evals and global tasks.
 	InstanceData.Reset();
-	const FStateTreeActiveStates Empty;
+	constexpr FStateTreeActiveStates Empty;
 	UpdateInstanceData(Empty, Empty);
 	if (!InstanceData.IsValid())
 	{
@@ -281,7 +286,7 @@ EStateTreeRunStatus FStateTreeExecutionContext::Stop()
 		return EStateTreeRunStatus::Failed;
 	}
 
-	FStateTreeExecutionState& Exec = GetExecState();
+	const FStateTreeExecutionState& Exec = GetExecState();
 	EStateTreeRunStatus Result = Exec.TreeRunStatus;
 	
 	// Capture events added between ticks.
@@ -308,7 +313,7 @@ EStateTreeRunStatus FStateTreeExecutionContext::Stop()
 	StopEvaluatorsAndGlobalTasks();
 
 	// Trace before resetting the instance data since it is required to provide all the event information
-	TRACE_STATETREE_ACTIVE_STATES_EVENT(GetInstanceDebugId(), FStateTreeActiveStates());
+	TRACE_STATETREE_ACTIVE_STATES_EVENT(GetInstanceDebugId(), UpdatePhaseMask, FStateTreeActiveStates());
 	TRACE_STATETREE_INSTANCE_EVENT(GetInstanceDebugId(), GetStateTree(), *GetInstanceDescription(), EStateTreeTraceInstanceEventType::Stopped);
 
 	// Destruct all allocated instance data (does not shrink the buffer). This will invalidate Exec too.
@@ -364,8 +369,6 @@ EStateTreeRunStatus FStateTreeExecutionContext::Tick(const float DeltaTime)
 
 	if (Exec->LastTickStatus == EStateTreeRunStatus::Running)
 	{
-		STATETREE_SCOPED_PHASE(EStateTreeUpdatePhase::TickingTasks);
-
 		// Tick tasks on active states.
 		Exec->LastTickStatus = TickTasks(DeltaTime);
 		
@@ -381,8 +384,6 @@ EStateTreeRunStatus FStateTreeExecutionContext::Tick(const float DeltaTime)
 	static constexpr int32 MaxIterations = 5;
 	for (int32 Iter = 0; Iter < MaxIterations; Iter++)
 	{
-		STATETREE_SCOPED_PHASE(EStateTreeUpdatePhase::TriggerTransitions);
-
 		// Append events accumulated during the tick, so that transitions can immediately act on them.
 		// We'll consume the events only if they lead to state change below (EnterState is treated the same as Tick),
 		// or let them be processed next frame if no transition.
@@ -391,6 +392,8 @@ EStateTreeRunStatus FStateTreeExecutionContext::Tick(const float DeltaTime)
 		// Trigger conditional transitions or state succeed/failed transitions. First tick transition is handled here too.
 		if (TriggerTransitions())
 		{
+			STATETREE_SCOPED_PHASE(EStateTreeUpdatePhase::ApplyTransitions);
+			
 			// We have committed to state change, consume events that were accumulated during the tick above.
 			EventQueue.Reset();
 
@@ -499,7 +502,6 @@ void FStateTreeExecutionContext::SendEvent(const FGameplayTag Tag, const FConstS
 void FStateTreeExecutionContext::RequestTransition(const FStateTreeTransitionRequest& Request)
 {
 	CSV_SCOPED_TIMING_STAT_EXCLUSIVE(StateTree_RequestTransition);
-	STATETREE_SCOPED_PHASE(EStateTreeUpdatePhase::RequestTransition);
 
 	if (!IsValid())
 	{
@@ -706,7 +708,7 @@ EStateTreeRunStatus FStateTreeExecutionContext::EnterState(const FStateTreeTrans
 	}
 	
 	STATETREE_LOG(Log, TEXT("Enter state '%s' (%d)"), *DebugGetStatePath(Transition.NextActiveStates), Exec.StateChangeCount);
-	STATETREE_TRACE_STATE_EVENT(Transition.NextActiveStates.Last().Index, EStateTreeTraceNodeEventType::OnEnter);
+	STATETREE_TRACE_STATE_EVENT(Transition.NextActiveStates.Last(), EStateTreeTraceNodeEventType::OnEnter);
 
 	for (int32 Index = 0; Index < Transition.NextActiveStates.Num() && Result != EStateTreeRunStatus::Failed; Index++)
 	{
@@ -788,7 +790,7 @@ EStateTreeRunStatus FStateTreeExecutionContext::EnterState(const FStateTreeTrans
 		}
 	}
 
-	TRACE_STATETREE_ACTIVE_STATES_EVENT(GetInstanceDebugId(), Exec.ActiveStates);
+	TRACE_STATETREE_ACTIVE_STATES_EVENT(GetInstanceDebugId(), UpdatePhaseMask, Exec.ActiveStates);
 	return Result;
 }
 
@@ -884,7 +886,7 @@ void FStateTreeExecutionContext::ExitState(const FStateTreeTransitionResult& Tra
 		CurrentTransition.ChangeType = ExitedStateChangeType[Index];
 
 		STATETREE_LOG(Log, TEXT("%*sState '%s' %s"), Index*UE::StateTree::DebugIndentSize, TEXT(""), *DebugGetStatePath(Transition.CurrentActiveStates, ExitedStateActiveIndex[Index]), *UEnum::GetDisplayValueAsText(CurrentTransition.ChangeType).ToString());
-		STATETREE_TRACE_STATE_EVENT(CurrentHandle.Index, EStateTreeTraceNodeEventType::OnExit);
+		STATETREE_TRACE_STATE_EVENT(CurrentHandle, EStateTreeTraceNodeEventType::OnExit);
 
 		// Tasks
 		for (int32 TaskIndex = (State.TasksBegin + State.TasksNum) - 1; TaskIndex >= State.TasksBegin; TaskIndex--)
@@ -936,7 +938,7 @@ void FStateTreeExecutionContext::StateCompleted()
 		FCurrentlyProcessedStateScope StateScope(*this, CurrentHandle);
 		
 		STATETREE_LOG(Verbose, TEXT("%*sState '%s'"), Index*UE::StateTree::DebugIndentSize, TEXT(""), *DebugGetStatePath(Exec.ActiveStates, Index));
-		STATETREE_TRACE_STATE_EVENT(CurrentHandle.Index, EStateTreeTraceNodeEventType::OnStateCompleted);
+		STATETREE_TRACE_STATE_EVENT(CurrentHandle, EStateTreeTraceNodeEventType::OnStateCompleted);
 
 		// Notify Tasks
 		for (int32 TaskIndex = (State.TasksBegin + State.TasksNum) - 1; TaskIndex >= State.TasksBegin; TaskIndex--)
@@ -955,7 +957,7 @@ void FStateTreeExecutionContext::StateCompleted()
 	}
 }
 
-EStateTreeRunStatus FStateTreeExecutionContext::TickEvaluatorsAndGlobalTasks(const float DeltaTime, bool bTickGlobalTasks)
+EStateTreeRunStatus FStateTreeExecutionContext::TickEvaluatorsAndGlobalTasks(const float DeltaTime, const bool bTickGlobalTasks)
 {
 	CSV_SCOPED_TIMING_STAT_EXCLUSIVE(StateTree_TickEvaluators);
 	STATETREE_SCOPED_PHASE(EStateTreeUpdatePhase::TickingGlobalTasks);
@@ -1167,6 +1169,7 @@ void FStateTreeExecutionContext::StopEvaluatorsAndGlobalTasks(const FStateTreeIn
 EStateTreeRunStatus FStateTreeExecutionContext::TickTasks(const float DeltaTime)
 {
 	CSV_SCOPED_TIMING_STAT_EXCLUSIVE(StateTree_TickTasks);
+	STATETREE_SCOPED_PHASE(EStateTreeUpdatePhase::TickingTasks);
 
 	FStateTreeExecutionState& Exec = GetExecState();
 
@@ -1312,7 +1315,7 @@ bool FStateTreeExecutionContext::TestAllConditions(const int32 ConditionsOffset,
 		}
 
 		const bool bValue = Cond.TestCondition(*this);
-		STATETREE_TRACE_CONDITION_EVENT(ConditionsOffset + Index, DataView, bValue ? EStateTreeTraceNodeEventType::Passed : EStateTreeTraceNodeEventType::Failed);	
+		STATETREE_TRACE_CONDITION_EVENT(ConditionsOffset + Index, DataView, bValue ? EStateTreeTraceNodeEventType::Passed : EStateTreeTraceNodeEventType::Failed);
 
 		// Reset copied properties that might contain object references.
 		if (Cond.BindingsBatch.IsValid())
@@ -1432,6 +1435,7 @@ bool FStateTreeExecutionContext::RequestTransition(const FStateTreeStateHandle N
 bool FStateTreeExecutionContext::TriggerTransitions()
 {
 	CSV_SCOPED_TIMING_STAT_EXCLUSIVE(StateTree_TriggerTransition);
+	STATETREE_SCOPED_PHASE(EStateTreeUpdatePhase::TriggerTransitions);
 
 	FAllowDirectTransitionsScope AllowDirectTransitionsScope(*this); // Set flag for the scope of this function to allow direct transitions without buffering.
 	FStateTreeExecutionState& Exec = GetExecState();
@@ -1521,9 +1525,16 @@ bool FStateTreeExecutionContext::TriggerTransitions()
 			const bool bShouldTrigger = Transition.Trigger == EStateTreeTransitionTrigger::OnTick
 										|| (Transition.Trigger == EStateTreeTransitionTrigger::OnEvent
 											&& HasEventToProcess(Transition.EventTag));
-				
-			if (bShouldTrigger
-				&& TestAllConditions(Transition.ConditionsBegin, Transition.ConditionsNum))
+
+			bool bPassed = false; 
+			if (bShouldTrigger)
+			{
+				STATETREE_TRACE_TRANSITION_EVENT(TransitionIndex, EStateTreeTraceNodeEventType::OnEvaluating);
+				STATETREE_SCOPED_PHASE(EStateTreeUpdatePhase::TransitionConditions);
+				bPassed = TestAllConditions(Transition.ConditionsBegin, Transition.ConditionsNum);
+			}
+
+			if (bPassed)
 			{
 				// If the transitions is delayed, set up the delay. 
 				if (Transition.HasDelay())
@@ -1601,7 +1612,14 @@ bool FStateTreeExecutionContext::TriggerTransitions()
 
 				if (EnumHasAnyFlags(Transition.Trigger, CompletionTrigger))
 				{
-					if (TestAllConditions(Transition.ConditionsBegin, Transition.ConditionsNum))
+					bool bPassed = false;
+					{
+						STATETREE_TRACE_TRANSITION_EVENT(TransitionIndex, EStateTreeTraceNodeEventType::OnEvaluating);
+						STATETREE_SCOPED_PHASE(EStateTreeUpdatePhase::TransitionConditions);
+						bPassed = TestAllConditions(Transition.ConditionsBegin, Transition.ConditionsNum);
+					}
+
+					if (bPassed)
 					{
 						// No delay allowed on completion conditions.
 						// No priority on completion transitions, use the priority to signal that state is selected.
@@ -1621,7 +1639,7 @@ bool FStateTreeExecutionContext::TriggerTransitions()
 
 		if (NextTransition.Priority == EStateTreeTransitionPriority::None)
 		{
-			// Could not trigger completion transition, jump back to start.
+			STATETREE_LOG(Verbose, TEXT("Could not trigger completion transition, jump back to start."));
 			FCurrentlyProcessedStateScope StateScope(*this, FStateTreeStateHandle::Root);
 			RequestTransition(FStateTreeStateHandle::Root, EStateTreeTransitionPriority::Normal);
 		}
@@ -1766,9 +1784,21 @@ bool FStateTreeExecutionContext::SelectStateInternal(const FStateTreeStateHandle
 
 	const FCompactStateTreeState& State = StateTree.States[NextState.Index];
 
+	STATETREE_TRACE_STATE_SELECTION_EVENT(NextState, EStateTreeTraceNodeEventType::PushStateSelection, State.SelectionBehavior);
+	ON_SCOPE_EXIT
+	{
+		STATETREE_TRACE_STATE_EVENT(NextState, EStateTreeTraceNodeEventType::PopStateSelection);
+	};
+	
 	// Check that the state can be entered
-	if (State.SelectionBehavior != EStateTreeStateSelectionBehavior::None
-		&& TestAllConditions(State.EnterConditionsBegin, State.EnterConditionsNum))
+	bool bEnterConditionsPassed = false;
+	if (State.SelectionBehavior != EStateTreeStateSelectionBehavior::None)
+	{
+		STATETREE_SCOPED_PHASE(EStateTreeUpdatePhase::EnterConditions);
+		bEnterConditionsPassed = TestAllConditions(State.EnterConditionsBegin, State.EnterConditionsNum);
+	}
+
+	if (bEnterConditionsPassed)
 	{
 		if (!OutNewActiveState.Push(NextState))
 		{
@@ -1795,6 +1825,7 @@ bool FStateTreeExecutionContext::SelectStateInternal(const FStateTreeStateHandle
 		else if (State.SelectionBehavior == EStateTreeStateSelectionBehavior::TryEnterState)
 		{
 			// Select this state.
+			STATETREE_TRACE_STATE_EVENT(NextState, EStateTreeTraceNodeEventType::OnStateSelected);
 			return true;
 		}
 		else if (State.SelectionBehavior == EStateTreeStateSelectionBehavior::TryFollowTransitions)
@@ -1835,9 +1866,16 @@ bool FStateTreeExecutionContext::SelectStateInternal(const FStateTreeStateHandle
 				const bool bShouldTrigger = Transition.Trigger == EStateTreeTransitionTrigger::OnTick
 											|| (Transition.Trigger == EStateTreeTransitionTrigger::OnEvent
 												&& HasEventToProcess(Transition.EventTag));
-				
-				if (bShouldTrigger
-					&& TestAllConditions(Transition.ConditionsBegin, Transition.ConditionsNum))
+
+				bool bTransitionConditionsPassed = false;
+				if (bShouldTrigger)
+				{
+					STATETREE_TRACE_TRANSITION_EVENT(TransitionIndex, EStateTreeTraceNodeEventType::OnEvaluating);
+					STATETREE_SCOPED_PHASE(EStateTreeUpdatePhase::TransitionConditions);
+					bTransitionConditionsPassed = bShouldTrigger && TestAllConditions(Transition.ConditionsBegin, Transition.ConditionsNum);
+				}
+
+				if (bTransitionConditionsPassed)
 				{
 					// Using SelectState() instead of SelectStateInternal to treat the transitions the same way as regular transitions,
 					// e.g. it may jump to a completely different branch.
@@ -1873,6 +1911,7 @@ bool FStateTreeExecutionContext::SelectStateInternal(const FStateTreeStateHandle
 			else
 			{
 				// Select this state (For backwards compatibility)
+				STATETREE_TRACE_STATE_EVENT(NextState, EStateTreeTraceNodeEventType::OnStateSelected);
 				return true;
 			}
 		}
diff --git a/Engine/Plugins/Runtime/StateTree/Source/StateTreeModule/Private/StateTreeTypes.cpp b/Engine/Plugins/Runtime/StateTree/Source/StateTreeModule/Private/StateTreeTypes.cpp
--- a/Engine/Plugins/Runtime/StateTree/Source/StateTreeModule/Private/StateTreeTypes.cpp
+++ b/Engine/Plugins/Runtime/StateTree/Source/StateTreeModule/Private/StateTreeTypes.cpp
@@ -22,52 +22,6 @@ const FStateTreeExternalDataHandle FStateTreeExternalDataHandle::Invalid = FStat
 const FStateTreeInstanceDebugId FStateTreeInstanceDebugId::Invalid = FStateTreeInstanceDebugId();
 #endif // WITH_STATETREE_DEBUGGER
 
-bool FStateTreeIndex16::SerializeFromMismatchedTag(const FPropertyTag& Tag, FStructuredArchive::FSlot Slot)
-{
-	// Support loading from Index8.
-	if (Tag.Type == NAME_StructProperty
-		&& Tag.StructName == FStateTreeIndex8::StaticStruct()->GetFName())
-	{
-		FStateTreeIndex8 OldValue;
-		FStateTreeIndex8::StaticStruct()->SerializeItem(Slot, &OldValue, nullptr);
-
-		int32 NewValue = OldValue.AsInt32();
-		if (!IsValidIndex(NewValue))
-		{
-			NewValue = INDEX_NONE;
-		}
-		
-		*this = FStateTreeIndex16(NewValue);
-		
-		return true;
-	}
-	
-	return false;
-}
-
-bool FStateTreeIndex8::SerializeFromMismatchedTag(const FPropertyTag& Tag, FStructuredArchive::FSlot Slot)
-{
-	// Support loading from Index16.
-	if (Tag.Type == NAME_StructProperty
-		&& Tag.StructName == FStateTreeIndex16::StaticStruct()->GetFName())
-	{
-		FStateTreeIndex16 OldValue;
-		FStateTreeIndex16::StaticStruct()->SerializeItem(Slot, &OldValue, nullptr);
-
-		int32 NewValue = OldValue.AsInt32();
-		if (!IsValidIndex(NewValue))
-		{
-			NewValue = INDEX_NONE;
-		}
-		
-		*this = FStateTreeIndex8(NewValue);
-		
-		return true;
-	}
-	
-	return false;
-}
-
 //////////////////////////////////////////////////////////////////////////
 // FStateTreeStateLink
 
diff --git a/Engine/Plugins/Runtime/StateTree/Source/StateTreeModule/Public/Debugger/StateTreeDebugger.h b/Engine/Plugins/Runtime/StateTree/Source/StateTreeModule/Public/Debugger/StateTreeDebugger.h
--- a/Engine/Plugins/Runtime/StateTree/Source/StateTreeModule/Public/Debugger/StateTreeDebugger.h
+++ b/Engine/Plugins/Runtime/StateTree/Source/StateTreeModule/Public/Debugger/StateTreeDebugger.h
@@ -82,6 +82,7 @@ struct STATETREEMODULE_API FStateTreeDebugger : FTickableGameObject
 	bool IsActiveInstance(double Time, FStateTreeInstanceDebugId InstanceId) const;
 	FText GetInstanceDescription(FStateTreeInstanceDebugId InstanceId) const;
 	void SelectInstance(const FStateTreeInstanceDebugId InstanceId);
+	FStateTreeInstanceDebugId GetSelectedInstance() const { return SelectedInstanceId; }
 
 	static FText DescribeTrace(const FTraceDescriptor& TraceDescriptor);
 	static FText DescribeInstance(const UE::StateTreeDebugger::FInstanceDescriptor& StateTreeInstanceDesc);
diff --git a/Engine/Plugins/Runtime/StateTree/Source/StateTreeModule/Public/Debugger/StateTreeTrace.h b/Engine/Plugins/Runtime/StateTree/Source/StateTreeModule/Public/Debugger/StateTreeTrace.h
--- a/Engine/Plugins/Runtime/StateTree/Source/StateTreeModule/Public/Debugger/StateTreeTrace.h
+++ b/Engine/Plugins/Runtime/StateTree/Source/StateTreeModule/Public/Debugger/StateTreeTrace.h
@@ -12,6 +12,7 @@ struct FStateTreeActiveStates;
 struct FStateTreeInstanceDebugId;
 struct FStateTreeIndex16;
 struct FStateTreeStateHandle;
+enum class EStateTreeStateSelectionBehavior : uint8;
 enum class EStateTreeRunStatus : uint8;
 enum class EStateTreeTraceNodeEventType : uint8;
 enum class EStateTreeUpdatePhase : uint16;
@@ -21,12 +22,13 @@ UE_TRACE_CHANNEL_EXTERN(StateTreeDebugChannel, STATETREEMODULE_API)
 
 namespace UE::StateTreeTrace
 {
-	void OutputInstanceLifetimeEvent(const FStateTreeInstanceDebugId InstanceId, const UStateTree* StateTree, const TCHAR* InstanceName, const EStateTreeTraceInstanceEventType EventType);
-	void OutputLogEventTrace(const FStateTreeInstanceDebugId InstanceId, const EStateTreeUpdatePhase Phase, const TCHAR* Fmt, ...);
-	void OutputStateEventTrace(const FStateTreeInstanceDebugId InstanceId, const EStateTreeUpdatePhase Phase, const FStateTreeStateHandle StateHandle, const EStateTreeTraceNodeEventType EventType);
-	void OutputTaskEventTrace(const FStateTreeInstanceDebugId InstanceId, const EStateTreeUpdatePhase Phase, const FStateTreeIndex16 TaskIdx, FStateTreeDataView DataView, const EStateTreeTraceNodeEventType EventType, const EStateTreeRunStatus Status);
-	void OutputConditionEventTrace(const FStateTreeInstanceDebugId InstanceId, const EStateTreeUpdatePhase Phase, const FStateTreeIndex16 ConditionIdx, FStateTreeDataView DataView, const EStateTreeTraceNodeEventType EventType);
-	void OutputActiveStatesEventTrace(const FStateTreeInstanceDebugId InstanceId, const FStateTreeActiveStates& ActiveStates);
+	void OutputInstanceLifetimeEvent(FStateTreeInstanceDebugId InstanceId, const UStateTree* StateTree, const TCHAR* InstanceName, EStateTreeTraceInstanceEventType EventType);
+	void OutputLogEventTrace(FStateTreeInstanceDebugId InstanceId, EStateTreeUpdatePhase Phase, const TCHAR* Fmt, ...);
+	void OutputStateEventTrace(FStateTreeInstanceDebugId InstanceId, EStateTreeUpdatePhase Phase, FStateTreeStateHandle StateHandle, EStateTreeTraceNodeEventType EventType, EStateTreeStateSelectionBehavior SelectionBehavior);
+	void OutputTaskEventTrace(FStateTreeInstanceDebugId InstanceId, EStateTreeUpdatePhase Phase, FStateTreeIndex16 TaskIdx, FStateTreeDataView DataView, EStateTreeTraceNodeEventType EventType, EStateTreeRunStatus Status);
+	void OutputConditionEventTrace(FStateTreeInstanceDebugId InstanceId, EStateTreeUpdatePhase Phase, FStateTreeIndex16 ConditionIdx, FStateTreeDataView DataView, EStateTreeTraceNodeEventType EventType);
+	void OutputTransitionEventTrace(FStateTreeInstanceDebugId InstanceId, EStateTreeUpdatePhase Phase, FStateTreeIndex16 TransitionIdx, EStateTreeTraceNodeEventType EventType);
+	void OutputActiveStatesEventTrace(FStateTreeInstanceDebugId InstanceId, EStateTreeUpdatePhase Phase, const FStateTreeActiveStates& ActiveStates);
 }
 
 #define TRACE_STATETREE_INSTANCE_EVENT(InstanceID, StateTree, InstanceName, EventType) \
@@ -41,10 +43,10 @@ namespace UE::StateTreeTrace
 		UE::StateTreeTrace::OutputLogEventTrace(InstanceId, Phase, Format, ##__VA_ARGS__); \
 	}
 
-#define TRACE_STATETREE_STATE_EVENT(InstanceId, Phase, StateIdx, EventType) \
+#define TRACE_STATETREE_STATE_EVENT(InstanceId, Phase, StateHandle, EventType, SelectionBehavior) \
 	if (UE_TRACE_CHANNELEXPR_IS_ENABLED(StateTreeDebugChannel)) \
 	{ \
-		UE::StateTreeTrace::OutputStateEventTrace(InstanceId, Phase, StateIdx, EventType); \
+		UE::StateTreeTrace::OutputStateEventTrace(InstanceId, Phase, StateHandle, EventType, SelectionBehavior); \
 	}
 
 #define TRACE_STATETREE_TASK_EVENT(InstanceId, Phase, TaskIdx, DataView, EventType, Status) \
@@ -59,19 +61,26 @@ namespace UE::StateTreeTrace
 		UE::StateTreeTrace::OutputConditionEventTrace(InstanceId, Phase, ConditionIdx, DataView, EventType); \
 	}
 
-#define TRACE_STATETREE_ACTIVE_STATES_EVENT(InstanceId, ActivateStates) \
+#define TRACE_STATETREE_TRANSITION_EVENT(InstanceId, Phase, TransitionIdx, EventType) \
 	if (UE_TRACE_CHANNELEXPR_IS_ENABLED(StateTreeDebugChannel)) \
 	{ \
-		UE::StateTreeTrace::OutputActiveStatesEventTrace(InstanceId, ActivateStates); \
+		UE::StateTreeTrace::OutputTransitionEventTrace(InstanceId, Phase, TransitionIdx, EventType); \
+	}
+
+#define TRACE_STATETREE_ACTIVE_STATES_EVENT(InstanceId, Phase, ActivateStates) \
+	if (UE_TRACE_CHANNELEXPR_IS_ENABLED(StateTreeDebugChannel)) \
+	{ \
+		UE::StateTreeTrace::OutputActiveStatesEventTrace(InstanceId, Phase, ActivateStates); \
 	}
 
 #else //STATETREE_DEBUG_TRACE_ENABLED
 
 #define TRACE_STATETREE_INSTANCE_EVENT(InstanceID, StateTree, InstanceName, EventType)
 #define TRACE_STATETREE_LOG_EVENT(InstanceId, Phase, Format, ...)
-#define TRACE_STATETREE_STATE_EVENT(InstanceId, Phase, StateIdx, EventType)
+#define TRACE_STATETREE_STATE_EVENT(InstanceId, Phase, StateHandle, EventType, SelectionBehavior)
 #define TRACE_STATETREE_TASK_EVENT(InstanceId, Phase, TaskIdx, DataView, EventType, Status)
 #define TRACE_STATETREE_CONDITION_EVENT(InstanceId, Phase, ConditionIdx, DataView, EventType)
-#define TRACE_STATETREE_ACTIVE_STATES_EVENT(InstanceId, ActivateStates)
+#define TRACE_STATETREE_TRANSITION_EVENT(InstanceId, Phase, TransitionIdx, EventType)
+#define TRACE_STATETREE_ACTIVE_STATES_EVENT(InstanceId, Phase, ActivateStates)
 
 #endif // STATETREE_DEBUG_TRACE_ENABLED
diff --git a/Engine/Plugins/Runtime/StateTree/Source/StateTreeModule/Public/Debugger/StateTreeTraceAnalyzer.h b/Engine/Plugins/Runtime/StateTree/Source/StateTreeModule/Public/Debugger/StateTreeTraceAnalyzer.h
--- a/Engine/Plugins/Runtime/StateTree/Source/StateTreeModule/Public/Debugger/StateTreeTraceAnalyzer.h
+++ b/Engine/Plugins/Runtime/StateTree/Source/StateTreeModule/Public/Debugger/StateTreeTraceAnalyzer.h
@@ -25,6 +25,7 @@ private:
 		RouteId_LogMessage,
 		RouteId_State,
 		RouteId_Task,
+		RouteId_Transition,
 		RouteId_Condition,
 		RouteId_ActiveStates
 	};
diff --git a/Engine/Plugins/Runtime/StateTree/Source/StateTreeModule/Public/Debugger/StateTreeTraceTypes.h b/Engine/Plugins/Runtime/StateTree/Source/StateTreeModule/Public/Debugger/StateTreeTraceTypes.h
--- a/Engine/Plugins/Runtime/StateTree/Source/StateTreeModule/Public/Debugger/StateTreeTraceTypes.h
+++ b/Engine/Plugins/Runtime/StateTree/Source/StateTreeModule/Public/Debugger/StateTreeTraceTypes.h
@@ -1,13 +1,16 @@
-﻿// Copyright Epic Games, Inc. All Rights Reserved.
+// Copyright Epic Games, Inc. All Rights Reserved.
 
 #pragma once
 
 #include "Misc/TVariant.h"
+#include "StateTreeIndexTypes.h"
+#include "StateTreeExecutionTypes.h"
 #include "StateTreeTraceTypes.generated.h"
 
+class UStateTree;
 struct FStateTreeStateHandle;
+enum class EStateTreeStateSelectionBehavior : uint8;
 enum class EStateTreeRunStatus : uint8;
-enum class EStateTreeUpdatePhase : uint16;
 
 UENUM()
 enum class EStateTreeTraceInstanceEventType : uint8
@@ -20,14 +23,18 @@ UENUM()
 enum class EStateTreeTraceNodeEventType : uint8
 {
 	Unset,
-	OnEnter,			// for State + Task events
-	OnExit,				// for State + Task events
-	OnStateCompleted,	// for State events
-	OnTickingTask,		// for Task events
-	OnTaskCompleted,	// for Task events
-	OnTaskTicked,		// for Task events
-	Passed,				// for Condition events
-	Failed				// for Condition events
+	OnEnter				UMETA(DisplayName = "Enter"),		// for State + Task events
+	OnExit				UMETA(DisplayName = "Exit"),		// for State + Task events
+	PushStateSelection	UMETA(DisplayName = "Push"),		// for State events
+	PopStateSelection	UMETA(DisplayName = "Pop"),			// for State events
+	OnStateSelected		UMETA(DisplayName = "Selected"),	// for State events
+	OnStateCompleted	UMETA(DisplayName = "Completed"),	// for State events
+	OnTickingTask		UMETA(DisplayName = "Tick"),		// for Task events
+	OnTaskCompleted		UMETA(DisplayName = "Completed"),	// for Task events
+	OnTaskTicked		UMETA(DisplayName = "Ticked"),		// for Task events
+	Passed				UMETA(DisplayName = "Passed"),		// for Condition events
+	Failed				UMETA(DisplayName = "Failed"),		// for Condition events
+	OnEvaluating		UMETA(DisplayName = "Evaluating")	// for Transition events
 };
 
 #if WITH_STATETREE_DEBUGGER
@@ -39,6 +46,9 @@ struct FStateTreeTracePhaseEvent
 	{
 	}
 
+	static FString GetDataTypePath() { return TEXT(""); }
+	static FString GetDataAsText() { return TEXT(""); }
+
 	EStateTreeUpdatePhase Phase;
 };
 
@@ -50,64 +60,110 @@ struct FStateTreeTraceLogEvent : FStateTreeTracePhaseEvent
 	{
 	}
 
+	STATETREEMODULE_API FString ToString(const UStateTree& StateTree) const;
+
 	FString Message;
 };
 
+struct FStateTreeTraceTransitionEvent : FStateTreeTracePhaseEvent
+{
+	explicit FStateTreeTraceTransitionEvent(const EStateTreeUpdatePhase Phase, const FStateTreeIndex16 Index, const EStateTreeTraceNodeEventType EventType)
+		: FStateTreeTracePhaseEvent(Phase)
+		, TransitionIndex(Index)
+		, EventType(EventType)
+	{
+	}
+
+	STATETREEMODULE_API FString ToString(const UStateTree& StateTree) const;
+
+	FStateTreeIndex16 TransitionIndex;
+	EStateTreeTraceNodeEventType EventType = EStateTreeTraceNodeEventType::Unset;
+};
+
 struct FStateTreeTraceNodeEvent : FStateTreeTracePhaseEvent
 {
-	explicit FStateTreeTraceNodeEvent(const EStateTreeUpdatePhase Phase, const int16 Idx, const EStateTreeTraceNodeEventType EventType)
+	explicit FStateTreeTraceNodeEvent(const EStateTreeUpdatePhase Phase, const FStateTreeIndex16 Index, const EStateTreeTraceNodeEventType EventType)
 		: FStateTreeTracePhaseEvent(Phase)
-		, Idx(Idx)
+		, Index(Index)
 		, EventType(EventType)
 	{
 	}
 
-	int16 Idx = INDEX_NONE;
+	STATETREEMODULE_API FString ToString(const UStateTree& StateTree) const;
+	
+	FStateTreeIndex16 Index;
 	EStateTreeTraceNodeEventType EventType = EStateTreeTraceNodeEventType::Unset;
 };
 
 struct FStateTreeTraceStateEvent : FStateTreeTraceNodeEvent
 {
-	explicit FStateTreeTraceStateEvent(const EStateTreeUpdatePhase Phase, const int16 Idx, const EStateTreeTraceNodeEventType EventType)
-		: FStateTreeTraceNodeEvent(Phase, Idx, EventType)
+	explicit FStateTreeTraceStateEvent(const EStateTreeUpdatePhase Phase, const FStateTreeIndex16 Index, const EStateTreeTraceNodeEventType EventType, const EStateTreeStateSelectionBehavior SelectionBehavior)
+		: FStateTreeTraceNodeEvent(Phase, Index, EventType)
+		, SelectionBehavior(SelectionBehavior)
 	{
 	}
+
+	STATETREEMODULE_API FString ToString(const UStateTree& StateTree) const;
+	STATETREEMODULE_API FStateTreeStateHandle GetStateHandle() const;
+
+	EStateTreeStateSelectionBehavior SelectionBehavior;
 };
 
 struct FStateTreeTraceTaskEvent : FStateTreeTraceNodeEvent
 {
-	explicit FStateTreeTraceTaskEvent(const EStateTreeUpdatePhase Phase, const int16 Idx,	const EStateTreeTraceNodeEventType EventType, const EStateTreeRunStatus Status,	const FString& TypePath, const FString& InstanceDataAsText)
-		: FStateTreeTraceNodeEvent(Phase, Idx, EventType)
+	explicit FStateTreeTraceTaskEvent(const EStateTreeUpdatePhase Phase, const FStateTreeIndex16 Index,	const EStateTreeTraceNodeEventType EventType, const EStateTreeRunStatus Status,	const FString& TypePath, const FString& InstanceDataAsText)
+		: FStateTreeTraceNodeEvent(Phase, Index, EventType)
 		, TypePath(TypePath)
 		, InstanceDataAsText(InstanceDataAsText)
 		, Status(Status)
 	{
 	}
 
+	STATETREEMODULE_API FString ToString(const UStateTree& StateTree) const;
+
+	FString GetDataTypePath() const { return TypePath; }
+	FString GetDataAsText() const { return InstanceDataAsText; }
+
 	FString TypePath;
 	FString InstanceDataAsText;
-	EStateTreeRunStatus Status;	
+	EStateTreeRunStatus Status;
 };
 
 struct FStateTreeTraceConditionEvent : FStateTreeTraceNodeEvent
 {
-	explicit FStateTreeTraceConditionEvent(const EStateTreeUpdatePhase Phase, const int16 Idx, const EStateTreeTraceNodeEventType EventType, const FString& TypePath, const FString& InstanceDataAsText)
-		: FStateTreeTraceNodeEvent(Phase, Idx, EventType)
+	explicit FStateTreeTraceConditionEvent(const EStateTreeUpdatePhase Phase, const FStateTreeIndex16 Index, const EStateTreeTraceNodeEventType EventType, const FString& TypePath, const FString& InstanceDataAsText)
+		: FStateTreeTraceNodeEvent(Phase, Index, EventType)
 		, TypePath(TypePath)
 		, InstanceDataAsText(InstanceDataAsText)
 	{
 	}
 
+	STATETREEMODULE_API FString ToString(const UStateTree& StateTree) const;
+
+	FString GetDataTypePath() const { return TypePath; }
+	FString GetDataAsText() const { return InstanceDataAsText; }
+
 	FString TypePath;
 	FString InstanceDataAsText;
 };
 
-struct FStateTreeTraceActiveStatesEvent
+struct FStateTreeTraceActiveStatesEvent : FStateTreeTracePhaseEvent
 {
+	// Intentionally implemented in source file to compile 'TArray<FStateTreeStateHandle>' using only forward declaration.
+	explicit FStateTreeTraceActiveStatesEvent(const EStateTreeUpdatePhase Phase);
+
+	STATETREEMODULE_API FString ToString(const UStateTree& StateTree) const;
+
 	TArray<FStateTreeStateHandle> ActiveStates;
 };
 
 /** Type aliases for statetree trace events */
-using FStateTreeTraceEventVariantType = TVariant<FStateTreeTraceLogEvent, FStateTreeTraceNodeEvent, FStateTreeTraceStateEvent, FStateTreeTraceTaskEvent, FStateTreeTraceConditionEvent, FStateTreeTraceActiveStatesEvent>;
+using FStateTreeTraceEventVariantType = TVariant<FStateTreeTraceLogEvent,
+												FStateTreeTraceNodeEvent,
+												FStateTreeTraceStateEvent,
+												FStateTreeTraceTaskEvent,
+												FStateTreeTraceTransitionEvent,
+												FStateTreeTraceConditionEvent,
+												FStateTreeTraceActiveStatesEvent>;
 
 #endif // WITH_STATETREE_DEBUGGER
diff --git a/Engine/Plugins/Runtime/StateTree/Source/StateTreeModule/Public/StateTree.h b/Engine/Plugins/Runtime/StateTree/Source/StateTreeModule/Public/StateTree.h
--- a/Engine/Plugins/Runtime/StateTree/Source/StateTreeModule/Public/StateTree.h
+++ b/Engine/Plugins/Runtime/StateTree/Source/StateTreeModule/Public/StateTree.h
@@ -117,6 +117,9 @@ public:
 	/** @return View of all states. */
 	TConstArrayView<FCompactStateTreeState> GetStates() const { return States; }
 
+	/** @return Pointer to the transition at a given index; null if not found. */ 
+	const FCompactStateTransition* GetTransitionFromIndex(const int16 TransitionIndex) const;
+
 #if WITH_EDITOR
 	/** Resets the compiled data to empty. */
 	void ResetCompiled();
@@ -157,7 +160,7 @@ protected:
 	virtual void PostInitProperties() override;
 	virtual void BeginDestroy() override;
 	virtual void GetAssetRegistryTags(TArray<FAssetRegistryTag>& OutTags) const override;
-	virtual void PostLoadAssetRegistryTags(const FAssetData& InAssetData, TArray<FAssetRegistryTag>& OutTagsAndValuesToUpdate) const;
+	virtual void PostLoadAssetRegistryTags(const FAssetData& InAssetData, TArray<FAssetRegistryTag>& OutTagsAndValuesToUpdate) const override;
 	virtual EDataValidationResult IsDataValid(class FDataValidationContext& Context) const override;
 #endif
 
diff --git a/Engine/Plugins/Runtime/StateTree/Source/StateTreeModule/Public/StateTreeExecutionContext.h b/Engine/Plugins/Runtime/StateTree/Source/StateTreeModule/Public/StateTreeExecutionContext.h
--- a/Engine/Plugins/Runtime/StateTree/Source/StateTreeModule/Public/StateTreeExecutionContext.h
+++ b/Engine/Plugins/Runtime/StateTree/Source/StateTreeModule/Public/StateTreeExecutionContext.h
@@ -3,11 +3,10 @@
 #pragma once
 
 #include "StateTree.h"
+#include "StateTreeExecutionTypes.h"
 #include "StateTreeNodeBase.h"
 #include "Experimental/ConcurrentLinearAllocator.h"
 
-#include "StateTreeExecutionContext.generated.h"
-
 struct FGameplayTag;
 struct FInstancedPropertyBag;
 
@@ -18,30 +17,6 @@ struct FStateTreeEvent;
 struct FStateTreeTransitionRequest;
 struct FStateTreeInstanceDebugId;
 
-/**
- * Enumeration for the different update phases.
- * This is used as context information when tracing debug events.
- * The values are ordered based on their hierarchy starting with the "higher" parents.
- */
-UENUM()
-enum class EStateTreeUpdatePhase : uint16
-{
-	Unset				= 0,
-	StartTree			= 1ull << 1,
-	StopTree			= 1ull << 2,
-	StartGlobalTasks	= 1ull << 3,
-	StopGlobalTasks		= 1ull << 4,
-	TickStateTree		= 1ull << 5,
-	TickingGlobalTasks	= 1ull << 6,
-	TickingTasks		= 1ull << 7,
-	TriggerTransitions	= 1ull << 8,
-	StateSelection		= 1ull << 9,
-	StateTransition		= 1ull << 10,
-	RequestTransition	= 1ull << 11,
-};
-ENUM_CLASS_FLAGS(EStateTreeUpdatePhase);
-
-
 /**
  * StateTree Execution Context is a helper that is used to update and access StateTree instance data.
  *
@@ -468,7 +443,7 @@ protected:
 
 #if WITH_STATETREE_DEBUGGER
 	/** Current phase in the update. Used for debugging events */
-	EStateTreeUpdatePhase CurrentUpdatePhaseMask = EStateTreeUpdatePhase::Unset;
+	EStateTreeUpdatePhase UpdatePhaseMask = EStateTreeUpdatePhase::Unset;
 #endif // WITH_STATETREE_DEBUGGER
 
 	/** True if transitions are allowed to be requested directly instead of buffering. */
diff --git a/Engine/Plugins/Runtime/StateTree/Source/StateTreeModule/Public/StateTreeExecutionTypes.h b/Engine/Plugins/Runtime/StateTree/Source/StateTreeModule/Public/StateTreeExecutionTypes.h
new file mode 100644
--- /dev/null
+++ b/Engine/Plugins/Runtime/StateTree/Source/StateTreeModule/Public/StateTreeExecutionTypes.h
@@ -0,0 +1,29 @@
+// Copyright Epic Games, Inc. All Rights Reserved.
+
+#pragma once
+
+#include "StateTreeExecutionTypes.generated.h"
+
+/**
+ * Enumeration for the different update phases.
+ * This is used as context information when tracing debug events.
+ * The values are ordered based on their hierarchy starting with the "higher" parents.
+ */
+UENUM(Flags)
+enum class EStateTreeUpdatePhase : uint16
+{
+	Unset						= 0,
+	StartTree					= 1ull << 1 UMETA(DisplayName = "Start Tree"),
+	StopTree					= 1ull << 2 UMETA(DisplayName = "Stop Tree"),
+	StartGlobalTasks			= 1ull << 3 UMETA(DisplayName = "Start Global Tasks"),
+	StopGlobalTasks				= 1ull << 4 UMETA(DisplayName = "Stop Global Tasks"),
+	TickStateTree				= 1ull << 5 UMETA(DisplayName = "Tick State Tree"),
+	ApplyTransitions			= 1ull << 6 UMETA(DisplayName = "Transition"),
+	TriggerTransitions			= 1ull << 7 UMETA(DisplayName = "Trigger Transitions"),
+	TickingGlobalTasks			= 1ull << 8 UMETA(DisplayName = "Tick Global Tasks"),
+	TickingTasks				= 1ull << 9 UMETA(DisplayName = "Tick Tasks"),
+	TransitionConditions		= 1ull << 10 UMETA(DisplayName = "Transition conditions"),
+	StateSelection				= 1ull << 11 UMETA(DisplayName = "Try Enter"),
+	EnterConditions				= 1ull << 12 UMETA(DisplayName = "Enter conditions"),
+};
+ENUM_CLASS_FLAGS(EStateTreeUpdatePhase);
\ No newline at end of file
diff --git a/Engine/Plugins/Runtime/StateTree/Source/StateTreeModule/Public/StateTreeIndexTypes.cpp b/Engine/Plugins/Runtime/StateTree/Source/StateTreeModule/Public/StateTreeIndexTypes.cpp
new file mode 100644
--- /dev/null
+++ b/Engine/Plugins/Runtime/StateTree/Source/StateTreeModule/Public/StateTreeIndexTypes.cpp
@@ -0,0 +1,52 @@
+﻿// Copyright Epic Games, Inc. All Rights Reserved.
+
+#pragma once
+
+#include "StateTreeIndexTypes.h"
+
+bool FStateTreeIndex16::SerializeFromMismatchedTag(const FPropertyTag& Tag, FStructuredArchive::FSlot Slot)
+{
+	// Support loading from Index8.
+	if (Tag.Type == NAME_StructProperty
+		&& Tag.StructName == FStateTreeIndex8::StaticStruct()->GetFName())
+	{
+		FStateTreeIndex8 OldValue;
+		FStateTreeIndex8::StaticStruct()->SerializeItem(Slot, &OldValue, nullptr);
+
+		int32 NewValue = OldValue.AsInt32();
+		if (!IsValidIndex(NewValue))
+		{
+			NewValue = INDEX_NONE;
+		}
+		
+		*this = FStateTreeIndex16(NewValue);
+		
+		return true;
+	}
+	
+	return false;
+}
+
+
+bool FStateTreeIndex8::SerializeFromMismatchedTag(const FPropertyTag& Tag, FStructuredArchive::FSlot Slot)
+{
+	// Support loading from Index16.
+	if (Tag.Type == NAME_StructProperty
+		&& Tag.StructName == FStateTreeIndex16::StaticStruct()->GetFName())
+	{
+		FStateTreeIndex16 OldValue;
+		FStateTreeIndex16::StaticStruct()->SerializeItem(Slot, &OldValue, nullptr);
+
+		int32 NewValue = OldValue.AsInt32();
+		if (!IsValidIndex(NewValue))
+		{
+			NewValue = INDEX_NONE;
+		}
+		
+		*this = FStateTreeIndex8(NewValue);
+		
+		return true;
+	}
+	
+	return false;
+}
\ No newline at end of file
diff --git a/Engine/Plugins/Runtime/StateTree/Source/StateTreeModule/Public/StateTreeIndexTypes.h b/Engine/Plugins/Runtime/StateTree/Source/StateTreeModule/Public/StateTreeIndexTypes.h
new file mode 100644
--- /dev/null
+++ b/Engine/Plugins/Runtime/StateTree/Source/StateTreeModule/Public/StateTreeIndexTypes.h
@@ -0,0 +1,107 @@
+// Copyright Epic Games, Inc. All Rights Reserved.
+
+#pragma once
+
+#include "StateTreeIndexTypes.generated.h"
+
+/** uint16 index that can be invalid. */
+USTRUCT(BlueprintType)
+struct STATETREEMODULE_API FStateTreeIndex16
+{
+	GENERATED_BODY()
+
+	static constexpr uint16 InvalidValue = MAX_uint16;
+	static const FStateTreeIndex16 Invalid;
+
+	/** @return true if the given index can be represented by the type. */
+	static bool IsValidIndex(const int32 Index)
+	{
+		return Index >= 0 && Index < (int32)MAX_uint16;
+	}
+
+	FStateTreeIndex16() = default;
+	
+	explicit FStateTreeIndex16(const int32 InIndex)
+	{
+		check(InIndex == INDEX_NONE || IsValidIndex(InIndex));
+		Value = InIndex == INDEX_NONE ? InvalidValue : (uint16)InIndex;
+	}
+
+	/** @retrun value of the index. */
+	uint16 Get() const { return Value; }
+	
+	/** @return the index value as int32, mapping invalid value to INDEX_NONE. */
+	int32 AsInt32() const { return Value == InvalidValue ? INDEX_NONE : Value; }
+
+	/** @return true if the index is valid. */
+	bool IsValid() const { return Value != InvalidValue; }
+
+	bool operator==(const FStateTreeIndex16& RHS) const { return Value == RHS.Value; }
+	bool operator!=(const FStateTreeIndex16& RHS) const { return Value != RHS.Value; }
+
+	bool SerializeFromMismatchedTag(const FPropertyTag& Tag, FStructuredArchive::FSlot Slot);
+
+protected:
+	UPROPERTY()
+	uint16 Value = InvalidValue;
+};
+
+template<>
+struct TStructOpsTypeTraits<FStateTreeIndex16> : public TStructOpsTypeTraitsBase2<FStateTreeIndex16>
+{
+	enum
+	{
+		WithStructuredSerializeFromMismatchedTag = true,
+	};
+};
+
+/** uint8 index that can be invalid. */
+USTRUCT(BlueprintType)
+struct STATETREEMODULE_API FStateTreeIndex8
+{
+	GENERATED_BODY()
+
+	static constexpr uint8 InvalidValue = MAX_uint8;
+	static const FStateTreeIndex8 Invalid;
+	
+	/** @return true if the given index can be represented by the type. */
+	static bool IsValidIndex(const int32 Index)
+	{
+		return Index >= 0 && Index < (int32)MAX_uint8;
+	}
+	
+	FStateTreeIndex8() = default;
+	
+	explicit FStateTreeIndex8(const int32 InIndex)
+	{
+		check(InIndex == INDEX_NONE || IsValidIndex(InIndex));
+		Value = InIndex == INDEX_NONE ? InvalidValue : (uint8)InIndex;
+	}
+
+	/** @retrun value of the index. */
+	uint8 Get() const { return Value; }
+
+	/** @return the index value as int32, mapping invalid value to INDEX_NONE. */
+	int32 AsInt32() const { return Value == InvalidValue ? INDEX_NONE : Value; }
+	
+	/** @return true if the index is valid. */
+	bool IsValid() const { return Value != InvalidValue; }
+
+	bool operator==(const FStateTreeIndex8& RHS) const { return Value == RHS.Value; }
+	bool operator!=(const FStateTreeIndex8& RHS) const { return Value != RHS.Value; }
+
+	bool SerializeFromMismatchedTag(const FPropertyTag& Tag, FStructuredArchive::FSlot Slot);
+	
+protected:
+	UPROPERTY()
+	uint8 Value = InvalidValue;
+};
+
+template<>
+struct TStructOpsTypeTraits<FStateTreeIndex8> : public TStructOpsTypeTraitsBase2<FStateTreeIndex8>
+{
+	enum
+	{
+		WithStructuredSerializeFromMismatchedTag = true,
+	};
+};
diff --git a/Engine/Plugins/Runtime/StateTree/Source/StateTreeModule/Public/StateTreeTypes.h b/Engine/Plugins/Runtime/StateTree/Source/StateTreeModule/Public/StateTreeTypes.h
--- a/Engine/Plugins/Runtime/StateTree/Source/StateTreeModule/Public/StateTreeTypes.h
+++ b/Engine/Plugins/Runtime/StateTree/Source/StateTreeModule/Public/StateTreeTypes.h
@@ -4,6 +4,7 @@
 
 #include "PropertyBag.h"
 #include "GameplayTagContainer.h"
+#include "StateTreeIndexTypes.h"
 #include "StateTreeTypes.generated.h"
 
 STATETREEMODULE_API DECLARE_LOG_CATEGORY_EXTERN(LogStateTree, Warning, All);
@@ -113,13 +114,13 @@ enum class EStateTreeStateSelectionBehavior : uint8
 	None,
 
 	/** When state is considered for selection, it is selected even if it has child states. */
-	TryEnterState,
+	TryEnterState UMETA(DisplayName = "Try Enter"),
 
 	/** When state is considered for selection, try to selects the first child state (in order they appear in the child list). If no child states are present, behaves like SelectState. */
-	TrySelectChildrenInOrder,
+	TrySelectChildrenInOrder UMETA(DisplayName = "Try Select Children In Order"),
 	
 	/** When state is considered for selection, try to trigger the transitions instead. */
-	TryFollowTransitions,
+	TryFollowTransitions UMETA(DisplayName = "Try Follow Transitions"),
 };
 
 
@@ -245,108 +246,6 @@ struct STATETREEMODULE_API FStateTreeStateHandle
 	uint16 Index = InvalidIndex;
 };
 
-/** uint16 index that can be invalid. */
-USTRUCT(BlueprintType)
-struct STATETREEMODULE_API FStateTreeIndex16
-{
-	GENERATED_BODY()
-
-	static constexpr uint16 InvalidValue = MAX_uint16;
-	static const FStateTreeIndex16 Invalid;
-
-	/** @return true if the given index can be represented by the type. */
-	static bool IsValidIndex(const int32 Index)
-	{
-		return Index >= 0 && Index < (int32)MAX_uint16;
-	}
-
-	FStateTreeIndex16() = default;
-	
-	explicit FStateTreeIndex16(const int32 InIndex)
-	{
-		check(InIndex == INDEX_NONE || IsValidIndex(InIndex));
-		Value = InIndex == INDEX_NONE ? InvalidValue : (uint16)InIndex;
-	}
-
-	/** @retrun value of the index. */
-	uint16 Get() const { return Value; }
-	
-	/** @return the index value as int32, mapping invalid value to INDEX_NONE. */
-	int32 AsInt32() const { return Value == InvalidValue ? INDEX_NONE : Value; }
-
-	/** @return true if the index is valid. */
-	bool IsValid() const { return Value != InvalidValue; }
-
-	bool operator==(const FStateTreeIndex16& RHS) const { return Value == RHS.Value; }
-	bool operator!=(const FStateTreeIndex16& RHS) const { return Value != RHS.Value; }
-
-	bool SerializeFromMismatchedTag(const FPropertyTag& Tag, FStructuredArchive::FSlot Slot);
-
-protected:
-	UPROPERTY()
-	uint16 Value = InvalidValue;
-};
-
-template<>
-struct TStructOpsTypeTraits<FStateTreeIndex16> : public TStructOpsTypeTraitsBase2<FStateTreeIndex16>
-{
-	enum
-	{
-		WithStructuredSerializeFromMismatchedTag = true,
-	};
-};
-
-/** uint8 index that can be invalid. */
-USTRUCT(BlueprintType)
-struct STATETREEMODULE_API FStateTreeIndex8
-{
-	GENERATED_BODY()
-
-	static constexpr uint8 InvalidValue = MAX_uint8;
-	static const FStateTreeIndex8 Invalid;
-	
-	/** @return true if the given index can be represented by the type. */
-	static bool IsValidIndex(const int32 Index)
-	{
-		return Index >= 0 && Index < (int32)MAX_uint8;
-	}
-	
-	FStateTreeIndex8() = default;
-	
-	explicit FStateTreeIndex8(const int32 InIndex)
-	{
-		check(InIndex == INDEX_NONE || IsValidIndex(InIndex));
-		Value = InIndex == INDEX_NONE ? InvalidValue : (uint8)InIndex;
-	}
-
-	/** @retrun value of the index. */
-	uint8 Get() const { return Value; }
-
-	/** @return the index value as int32, mapping invalid value to INDEX_NONE. */
-	int32 AsInt32() const { return Value == InvalidValue ? INDEX_NONE : Value; }
-	
-	/** @return true if the index is valid. */
-	bool IsValid() const { return Value != InvalidValue; }
-
-	bool operator==(const FStateTreeIndex8& RHS) const { return Value == RHS.Value; }
-	bool operator!=(const FStateTreeIndex8& RHS) const { return Value != RHS.Value; }
-
-	bool SerializeFromMismatchedTag(const FPropertyTag& Tag, FStructuredArchive::FSlot Slot);
-	
-protected:
-	UPROPERTY()
-	uint8 Value = InvalidValue;
-};
-
-template<>
-struct TStructOpsTypeTraits<FStateTreeIndex8> : public TStructOpsTypeTraitsBase2<FStateTreeIndex8>
-{
-	enum
-	{
-		WithStructuredSerializeFromMismatchedTag = true,
-	};
-};
-
 
 /** Transition request */
 USTRUCT(BlueprintType)
