[StateTreeDebugger]
- timeline tracks can be marked as stale (from a previous PIE session) to change their look
- added world simulation timestamp for traces to represent the tracks in simulation time and not trace session time
- create dedicated methods on the StateTree module to Start/Stop traces
- unneeded trace channels are deactivated and restored if possible for a StateTree debug session in order to reduce trace size
- DebuggerView can now auto start new trace on a new PIE session (new settings 'bShouldDebuggerAutoRecordOnPIE'') or user can control with dedicated button (similar to RewindDebugger)
- moved StateTree editor settings under section "Plugins"
#rnx

diff --git a/Engine/Plugins/Runtime/StateTree/Source/StateTreeEditorModule/Private/Debugger/SStateTreeDebuggerInstanceTree.cpp b/Engine/Plugins/Runtime/StateTree/Source/StateTreeEditorModule/Private/Debugger/SStateTreeDebuggerInstanceTree.cpp
--- a/Engine/Plugins/Runtime/StateTree/Source/StateTreeEditorModule/Private/Debugger/SStateTreeDebuggerInstanceTree.cpp
+++ b/Engine/Plugins/Runtime/StateTree/Source/StateTreeEditorModule/Private/Debugger/SStateTreeDebuggerInstanceTree.cpp
@@ -4,6 +4,8 @@
 
 #include "Debugger/SStateTreeDebuggerInstanceTree.h"
 #include "RewindDebuggerTrack.h"
+#include "StateTreeDebuggerTrack.h"
+
 #include "Widgets/Images/SLayeredImage.h"
 
 //----------------------------------------------------------------------//
@@ -45,6 +47,17 @@ TSharedRef<ITableRow> SStateTreeDebuggerInstanceTree::GenerateTreeRow(TSharedPtr
 				SNew(STextBlock)
 				.Text(Item->GetDisplayName())
 				.Font(FAppStyle::GetFontStyle(TEXT("PropertyWindow.NormalFont")))
+				.ColorAndOpacity_Lambda([Item]()
+				{
+					if (const FStateTreeDebuggerTrack* DebuggerTrack = static_cast<FStateTreeDebuggerTrack*>(Item.Get()))
+					{
+						if (DebuggerTrack->IsStale())
+						{
+							return FSlateColor::UseSubduedForeground();
+						}
+					}
+					return FSlateColor::UseForeground();
+				})
 			]
 		];
 }
diff --git a/Engine/Plugins/Runtime/StateTree/Source/StateTreeEditorModule/Private/Debugger/SStateTreeDebuggerTimelines.cpp b/Engine/Plugins/Runtime/StateTree/Source/StateTreeEditorModule/Private/Debugger/SStateTreeDebuggerTimelines.cpp
--- a/Engine/Plugins/Runtime/StateTree/Source/StateTreeEditorModule/Private/Debugger/SStateTreeDebuggerTimelines.cpp
+++ b/Engine/Plugins/Runtime/StateTree/Source/StateTreeEditorModule/Private/Debugger/SStateTreeDebuggerTimelines.cpp
@@ -101,7 +101,7 @@ void SStateTreeDebuggerTimelines::Construct(const FArguments& InArgs)
 		.ScrubPosition(InArgs._ScrubPosition)
 		.OnScrubPositionChanged(InArgs._OnScrubPositionChanged));
 
-	InstanceTracks = InArgs._DebugComponents;
+	InstanceTracks = InArgs._InstanceTracks;
 
 	OnExpansionChanged = InArgs._OnExpansionChanged;
 
diff --git a/Engine/Plugins/Runtime/StateTree/Source/StateTreeEditorModule/Private/Debugger/SStateTreeDebuggerTimelines.h b/Engine/Plugins/Runtime/StateTree/Source/StateTreeEditorModule/Private/Debugger/SStateTreeDebuggerTimelines.h
--- a/Engine/Plugins/Runtime/StateTree/Source/StateTreeEditorModule/Private/Debugger/SStateTreeDebuggerTimelines.h
+++ b/Engine/Plugins/Runtime/StateTree/Source/StateTreeEditorModule/Private/Debugger/SStateTreeDebuggerTimelines.h
@@ -20,7 +20,7 @@ class SStateTreeDebuggerTimelines : public SSimpleTimeSlider
 public:
 	SLATE_BEGIN_ARGS(SStateTreeDebuggerTimelines) {}
 	 
-		SLATE_ARGUMENT( TArray< TSharedPtr< RewindDebugger::FRewindDebuggerTrack > >*, DebugComponents );
+		SLATE_ARGUMENT( TArray< TSharedPtr< RewindDebugger::FRewindDebuggerTrack > >*, InstanceTracks );
 		SLATE_ARGUMENT( TSharedPtr<SScrollBar>, ExternalScrollbar)
 		SLATE_EVENT( FOnSelectionChanged, OnSelectionChanged )
 		SLATE_EVENT( FOnTableViewScrolled, OnScrolled)
diff --git a/Engine/Plugins/Runtime/StateTree/Source/StateTreeEditorModule/Private/Debugger/SStateTreeDebuggerView.cpp b/Engine/Plugins/Runtime/StateTree/Source/StateTreeEditorModule/Private/Debugger/SStateTreeDebuggerView.cpp
--- a/Engine/Plugins/Runtime/StateTree/Source/StateTreeEditorModule/Private/Debugger/SStateTreeDebuggerView.cpp
+++ b/Engine/Plugins/Runtime/StateTree/Source/StateTreeEditorModule/Private/Debugger/SStateTreeDebuggerView.cpp
@@ -11,13 +11,16 @@
 #include "IStructureDetailsView.h"
 #include "Kismet2/DebuggerCommands.h"
 #include "Modules/ModuleManager.h"
+#include "ProfilingDebugging/TraceAuxiliary.h"
 #include "PropertyEditorModule.h"
 #include "SStateTreeDebuggerInstanceTree.h"
 #include "SStateTreeDebuggerTimelines.h"
 #include "StateTree.h"
 #include "StateTreeDebuggerCommands.h"
 #include "StateTreeDebuggerTrack.h"
+#include "StateTreeEditorSettings.h"
 #include "StateTreeExecutionContext.h"
+#include "StateTreeModule.h"
 #include "StateTreeState.h"
 #include "StateTreeViewModel.h"
 #include "Widgets/Input/SComboButton.h"
@@ -167,16 +170,56 @@ SStateTreeDebuggerView::~SStateTreeDebuggerView()
 	FEditorDelegates::SingleStepPIE.RemoveAll(this);
 }
 
-void SStateTreeDebuggerView::OnPIEStarted(const bool bIsSimulating) const
+void SStateTreeDebuggerView::StartRecording()
 {
-	if (!Debugger->IsAnalysisSessionActive())
+	if (!CanStartRecording())
 	{
-		Debugger->StartLastLiveSessionAnalysis();
+		return;
 	}
+
+	IStateTreeModule& StateTreeModule = FModuleManager::GetModuleChecked<IStateTreeModule>("StateTreeModule");
+	StateTreeModule.StartTraces();
+
+	bRecording = true;
 }
 
-void SStateTreeDebuggerView::OnPIEStopped(const bool bIsSimulating) const
+void SStateTreeDebuggerView::StopRecording()
 {
+	if (bRecording)
+	{
+		IStateTreeModule& StateTreeModule = FModuleManager::GetModuleChecked<IStateTreeModule>("StateTreeModule");
+		StateTreeModule.StopTraces();
+		
+		bRecording = false;
+
+		// Mark all tracks from the stopped session as stale to have different look.
+		for (const TSharedPtr<RewindDebugger::FRewindDebuggerTrack>& DebugTrack : InstanceTracks)
+		{
+			if (FStateTreeDebuggerTrack* StateTreeTrack = static_cast<FStateTreeDebuggerTrack*>(DebugTrack.Get()))
+			{
+				StateTreeTrack->MarkAsStale();
+			}
+		}
+	}
+}
+
+void SStateTreeDebuggerView::OnPIEStarted(const bool bIsSimulating)
+{
+	if (UStateTreeEditorSettings::Get().bShouldDebuggerAutoRecordOnPIE)
+	{
+		StartRecording();
+
+		if (!Debugger->IsAnalysisSessionActive())
+		{
+			Debugger->StartLastLiveSessionAnalysis();
+		}
+	}
+}
+
+void SStateTreeDebuggerView::OnPIEStopped(const bool bIsSimulating)
+{
+	StopRecording();
+	Debugger->StopAnalysis();
 	Debugger->Unpause();
 }
 
@@ -201,6 +244,7 @@ void SStateTreeDebuggerView::Construct(const FArguments& InArgs, const UStateTre
 	StateTree = InStateTree;
 
 	Debugger = InStateTreeViewModel->GetDebugger();
+	bRecording = FTraceAuxiliary::IsConnected();
 
 	// Bind callbacks to the debugger delegates
 	Debugger->OnNewInstance.BindSP(this, &SStateTreeDebuggerView::OnNewInstance);
@@ -233,19 +277,23 @@ void SStateTreeDebuggerView::Construct(const FArguments& InArgs, const UStateTre
 	FSlimHorizontalToolBarBuilder ToolbarBuilder(InCommandList, FMultiBoxCustomization::None, /*InExtender*/ nullptr, /*InForceSmallIcons*/ true);
 	ToolbarBuilder.BeginSection(TEXT("Debugging"));
 	{
-		const FPlayWorldCommands& PlayWorldCommand = FPlayWorldCommands::Get();
-		ToolbarBuilder.AddToolBarButton(PlayWorldCommand.RepeatLastPlay);
-		ToolbarBuilder.AddToolBarButton(PlayWorldCommand.PausePlaySession,
+		const FStateTreeDebuggerCommands& DebuggerCommands = FStateTreeDebuggerCommands::Get();
+		ToolbarBuilder.AddToolBarButton(DebuggerCommands.StartRecording);
+		ToolbarBuilder.AddToolBarButton(DebuggerCommands.StopRecording);
+
+		const FPlayWorldCommands& PlayWorldCommands = FPlayWorldCommands::Get();
+		ToolbarBuilder.AddToolBarButton(PlayWorldCommands.RepeatLastPlay);
+		ToolbarBuilder.AddToolBarButton(PlayWorldCommands.PausePlaySession,
 					NAME_None, TAttribute<FText>(), TAttribute<FText>(), FSlateIcon(FAppStyle::GetAppStyleSetName(), "PlayWorld.PausePlaySession.Small"));
-		ToolbarBuilder.AddToolBarButton(PlayWorldCommand.ResumePlaySession,
+		ToolbarBuilder.AddToolBarButton(PlayWorldCommands.ResumePlaySession,
 					NAME_None, TAttribute<FText>(), TAttribute<FText>(), FSlateIcon(FAppStyle::GetAppStyleSetName(), "PlayWorld.ResumePlaySession.Small"));	
-		ToolbarBuilder.AddToolBarButton(PlayWorldCommand.StopPlaySession,
+		ToolbarBuilder.AddToolBarButton(PlayWorldCommands.StopPlaySession,
 					NAME_None, TAttribute<FText>(), TAttribute<FText>(), FSlateIcon(FAppStyle::GetAppStyleSetName(), "PlayWorld.StopPlaySession.Small"));
 		ToolbarBuilder.AddSeparator();
-		ToolbarBuilder.AddToolBarButton(FStateTreeDebuggerCommands::Get().PreviousFrameWithStateChange);
-		ToolbarBuilder.AddToolBarButton(FStateTreeDebuggerCommands::Get().PreviousFrameWithEvents);
-		ToolbarBuilder.AddToolBarButton(FStateTreeDebuggerCommands::Get().NextFrameWithEvents);
-		ToolbarBuilder.AddToolBarButton(FStateTreeDebuggerCommands::Get().NextFrameWithStateChange);
+		ToolbarBuilder.AddToolBarButton(DebuggerCommands.PreviousFrameWithStateChange);
+		ToolbarBuilder.AddToolBarButton(DebuggerCommands.PreviousFrameWithEvents);
+		ToolbarBuilder.AddToolBarButton(DebuggerCommands.NextFrameWithEvents);
+		ToolbarBuilder.AddToolBarButton(DebuggerCommands.NextFrameWithStateChange);
 	}
 	ToolbarBuilder.EndSection();
 
@@ -289,14 +337,14 @@ void SStateTreeDebuggerView::Construct(const FArguments& InArgs, const UStateTre
 	
 	// Timelines TreeView
 	InstanceTimelinesTreeView = SNew(SStateTreeDebuggerTimelines)
-	   .OnExpansionChanged_Lambda([this]() { InstancesTreeView->RestoreExpansion(); })
-	   .OnScrolled_Lambda([this](double ScrollOffset){ InstancesTreeView->ScrollTo(ScrollOffset); })
-	   .DebugComponents(&InstanceTracks)
-	   .ViewRange_Lambda([this](){ return ViewRange; })
-	   .ClampRange_Lambda([this](){ return TRange<double>(0.0f, Debugger->GetRecordingDuration()); })
-	   .OnViewRangeChanged_Lambda([this](TRange<double> NewRange){ ViewRange = NewRange; })
-	   .ScrubPosition(ScrubTimeAttribute)
-	   .OnScrubPositionChanged_Lambda([this](double NewScrubTime, bool bIsScrubbing){ OnTimeLineScrubPositionChanged( NewScrubTime, bIsScrubbing ); });
+		.OnExpansionChanged_Lambda([this]() { InstancesTreeView->RestoreExpansion(); })
+		.OnScrolled_Lambda([this](double ScrollOffset) { InstancesTreeView->ScrollTo(ScrollOffset); })
+		.InstanceTracks(&InstanceTracks)
+		.ViewRange_Lambda([this]() { return ViewRange; })
+		.ClampRange_Lambda([this]() { return TRange<double>(0, Debugger->GetRecordingDuration()); })
+		.OnViewRangeChanged_Lambda([this](TRange<double> NewRange) { ViewRange = NewRange; })
+		.ScrubPosition(ScrubTimeAttribute)
+		.OnScrubPositionChanged_Lambda([this](double NewScrubTime, bool bIsScrubbing) { OnTimeLineScrubPositionChanged(NewScrubTime, bIsScrubbing); });
 
 	// EventsTreeView
 	EventsTreeView = SNew(STreeView<TSharedPtr<UE::StateTreeDebugger::FEventTreeElement>>)
@@ -439,9 +487,9 @@ void SStateTreeDebuggerView::Construct(const FArguments& InArgs, const UStateTre
 						.ClampRangeHighlightColor(FLinearColor::Red.CopyWithNewOpacity(0.5f))
 						.ScrubPosition(ScrubTimeAttribute)
 						.ViewRange_Lambda([this]() { return ViewRange; })
-						.OnViewRangeChanged_Lambda([this](TRange<double> NewRange){ ViewRange = NewRange; })
-						.ClampRange_Lambda([this](){ return TRange<double>(0.0f, Debugger->GetRecordingDuration()); })
-						.OnScrubPositionChanged_Lambda([this](double NewScrubTime, bool bIsScrubbing){ OnTimeLineScrubPositionChanged(NewScrubTime, bIsScrubbing); })
+						.OnViewRangeChanged_Lambda([this](TRange<double> NewRange) { ViewRange = NewRange; })
+						.ClampRange_Lambda([this]() { return TRange<double>(0, Debugger->GetRecordingDuration()); })
+						.OnScrubPositionChanged_Lambda([this](double NewScrubTime, bool bIsScrubbing) { OnTimeLineScrubPositionChanged(NewScrubTime, bIsScrubbing); })
 					]
 				]
 			]
@@ -539,7 +587,7 @@ void SStateTreeDebuggerView::RefreshTracks()
 	{
 		bChanged = DebugTrack->Update() || bChanged;
 	}
-	
+
 	if (bChanged)
 	{
 		InstancesTreeView->Refresh();
@@ -551,6 +599,20 @@ void SStateTreeDebuggerView::RefreshTracks()
 void SStateTreeDebuggerView::BindDebuggerToolbarCommands(const TSharedRef<FUICommandList>& ToolkitCommands)
 {
 	const FStateTreeDebuggerCommands& Commands = FStateTreeDebuggerCommands::Get();
+
+	ToolkitCommands->MapAction(
+		Commands.StartRecording,
+		FExecuteAction::CreateSP(this, &SStateTreeDebuggerView::StartRecording),
+		FCanExecuteAction::CreateSP(this, &SStateTreeDebuggerView::CanStartRecording),
+		FIsActionChecked(),
+		FIsActionButtonVisible::CreateLambda([this]() { return !IsRecording();}));
+
+	ToolkitCommands->MapAction(
+		Commands.StopRecording,
+		FExecuteAction::CreateSP(this, &SStateTreeDebuggerView::StopRecording),
+		FCanExecuteAction::CreateSP(this, &SStateTreeDebuggerView::CanStopRecording),
+		FIsActionChecked(),
+		FIsActionButtonVisible::CreateRaw(this, &SStateTreeDebuggerView::CanStopRecording));
 	
 	ToolkitCommands->MapAction(
 		Commands.PreviousFrameWithStateChange,
@@ -655,6 +717,12 @@ void SStateTreeDebuggerView::OnDebuggerScrubStateChanged(const UE::StateTreeDebu
 {
 	// Rebuild frame details from the events of that frame
 	EventsTreeElements.Reset();
+	ON_SCOPE_EXIT
+	{
+		EventsTreeView->ClearExpandedItems();
+		ExpandAll(EventsTreeElements);
+		EventsTreeView->RequestTreeRefresh();
+	};
 
 	const UE::StateTreeDebugger::FInstanceEventCollection& EventCollection = ScrubState.GetEventCollection();
 	const TConstArrayView<const FStateTreeTraceEventVariantType> Events = EventCollection.Events;
@@ -692,7 +760,7 @@ void SStateTreeDebuggerView::OnDebuggerScrubStateChanged(const UE::StateTreeDebu
 				NodeDesc = FString::Printf(TEXT(" '%s'"), CompactState != nullptr ? *CompactState->Name.ToString() : *StateHandle.Describe());
 			}
 
-			FStateTreeTraceLogEvent LogEvent(NodePhase,	FString::Printf(TEXT("%s%s"), *UEnum::GetDisplayValueAsText(NodePhase).ToString(), *NodeDesc));
+			FStateTreeTraceLogEvent LogEvent(/*RecordingWorldTime*/0, NodePhase, FString::Printf(TEXT("%s%s"), *UEnum::GetDisplayValueAsText(NodePhase).ToString(), *NodeDesc));
 
 			// Create Tree element to hold the event
 			const TSharedPtr<UE::StateTreeDebugger::FEventTreeElement> NewElement = MakeShareable(
@@ -780,7 +848,7 @@ void SStateTreeDebuggerView::OnDebuggerScrubStateChanged(const UE::StateTreeDebu
 		if (bPopState)
 		{
 			ParentStack.Pop();
-			// We use pop state to maintain the state stack but there is no element to visualize				
+			// We use pop state to maintain the state stack but there is no element to visualize
 			continue;
 		}
 
@@ -794,10 +862,6 @@ void SStateTreeDebuggerView::OnDebuggerScrubStateChanged(const UE::StateTreeDebu
 		TArray<TSharedPtr<UE::StateTreeDebugger::FEventTreeElement>>& Elements = ParentStack.IsEmpty() ? EventsTreeElements : ParentStack.Top().Element->Children;
         Elements.Add(MakeShareable(new UE::StateTreeDebugger::FEventTreeElement(Spans[SpanIdx].Frame, Event)));
 	}
-
-	EventsTreeView->ClearExpandedItems();
-	ExpandAll(EventsTreeElements);
-	EventsTreeView->RequestTreeRefresh();
 }
 
 void SStateTreeDebuggerView::ExpandAll(const TArray<TSharedPtr<UE::StateTreeDebugger::FEventTreeElement>>& Items)
@@ -890,7 +954,7 @@ TSharedRef<SWidget> SStateTreeDebuggerView::OnGetDebuggerTracesMenu() const
 
 void SStateTreeDebuggerView::TrackCursor()
 {
-	const double ScrubTime = ScrubTimeAttribute.Get();	
+	const double ScrubTime = ScrubTimeAttribute.Get();
 	TRange<double> CurrentViewRange = ViewRange;
 	const double ViewRangeDuration = CurrentViewRange.GetUpperBoundValue() - CurrentViewRange.GetLowerBoundValue();
 
diff --git a/Engine/Plugins/Runtime/StateTree/Source/StateTreeEditorModule/Private/Debugger/SStateTreeDebuggerView.h b/Engine/Plugins/Runtime/StateTree/Source/StateTreeEditorModule/Private/Debugger/SStateTreeDebuggerView.h
--- a/Engine/Plugins/Runtime/StateTree/Source/StateTreeEditorModule/Private/Debugger/SStateTreeDebuggerView.h
+++ b/Engine/Plugins/Runtime/StateTree/Source/StateTreeEditorModule/Private/Debugger/SStateTreeDebuggerView.h
@@ -53,8 +53,8 @@ private:
 
 	TSharedRef<SWidget> OnGetDebuggerTracesMenu() const;
 
-	void OnPIEStarted(bool bIsSimulating) const;
-	void OnPIEStopped(bool bIsSimulating) const;
+	void OnPIEStarted(bool bIsSimulating);
+	void OnPIEStopped(bool bIsSimulating);
 	void OnPIEPaused(bool bIsSimulating) const;
 	void OnPIEResumed(bool bIsSimulating) const;
 	void OnPIESingleStepped(bool bIsSimulating) const;
@@ -63,6 +63,14 @@ private:
 	void OnNewInstance(FStateTreeInstanceDebugId InstanceId);
 	void OnSelectedInstanceCleared();
 
+	bool CanStartRecording() const { return !IsRecording(); }
+	void StartRecording();
+
+	bool IsRecording() const { return bRecording; }
+
+	bool CanStopRecording() const { return IsRecording(); }	
+	void StopRecording();
+	
 	void BindDebuggerToolbarCommands(const TSharedRef<FUICommandList>& ToolkitCommands);
 
 	bool CanStepBackToPreviousStateWithEvents() const;
@@ -136,6 +144,9 @@ private:
 
 	/** Object created from the event data when statetree node was holding an object. */
 	TWeakObjectPtr<UObject> SelectedNodeDataObject;
+
+	/** Indicates that a live session was started (record button or auto record in PIE) to generate StateTree traces. */
+	bool bRecording = false;
 };
 
 #endif // WITH_STATETREE_DEBUGGER
\ No newline at end of file
diff --git a/Engine/Plugins/Runtime/StateTree/Source/StateTreeEditorModule/Private/Debugger/StateTreeDebuggerCommands.cpp b/Engine/Plugins/Runtime/StateTree/Source/StateTreeEditorModule/Private/Debugger/StateTreeDebuggerCommands.cpp
--- a/Engine/Plugins/Runtime/StateTree/Source/StateTreeEditorModule/Private/Debugger/StateTreeDebuggerCommands.cpp
+++ b/Engine/Plugins/Runtime/StateTree/Source/StateTreeEditorModule/Private/Debugger/StateTreeDebuggerCommands.cpp
@@ -24,6 +24,9 @@ void FStateTreeDebuggerCommands::RegisterCommands()
 {
 	UI_COMMAND(ToggleBreakpoint, "Toggle Breakpoint", "Adds or removes a breakpoint on each selected node", EUserInterfaceActionType::Button, FInputChord(EKeys::F9))
 
+	UI_COMMAND(StartRecording, "Start Recording", "Start a new trace session.", EUserInterfaceActionType::Button, FInputChord(EModifierKey::Control | EModifierKey::Shift, EKeys::R));
+	UI_COMMAND(StopRecording, "Stop Recording", "Stop the current trace session.", EUserInterfaceActionType::Button, FInputChord(EModifierKey::Control | EModifierKey::Shift, EKeys::S));
+	
 	UI_COMMAND(PreviousFrameWithStateChange, "Previous State Change", "Jump to previous state changed frame", EUserInterfaceActionType::Button, FInputChord(EModifierKey::Control, EKeys::Left));
 	UI_COMMAND(PreviousFrameWithEvents, "Previous Frame", "Step one frame back", EUserInterfaceActionType::Button, FInputChord(EKeys::Left));
 	UI_COMMAND(NextFrameWithEvents, "Next Frame", "Step one frame forward", EUserInterfaceActionType::Button, FInputChord(EKeys::Right));
diff --git a/Engine/Plugins/Runtime/StateTree/Source/StateTreeEditorModule/Private/Debugger/StateTreeDebuggerCommands.h b/Engine/Plugins/Runtime/StateTree/Source/StateTreeEditorModule/Private/Debugger/StateTreeDebuggerCommands.h
--- a/Engine/Plugins/Runtime/StateTree/Source/StateTreeEditorModule/Private/Debugger/StateTreeDebuggerCommands.h
+++ b/Engine/Plugins/Runtime/StateTree/Source/StateTreeEditorModule/Private/Debugger/StateTreeDebuggerCommands.h
@@ -18,6 +18,8 @@ public:
 	virtual void RegisterCommands() override;
 
 	TSharedPtr<FUICommandInfo> ToggleBreakpoint;
+	TSharedPtr<FUICommandInfo> StartRecording;
+	TSharedPtr<FUICommandInfo> StopRecording;
 	TSharedPtr<FUICommandInfo> PreviousFrameWithStateChange;
 	TSharedPtr<FUICommandInfo> PreviousFrameWithEvents;
 	TSharedPtr<FUICommandInfo> NextFrameWithEvents;
diff --git a/Engine/Plugins/Runtime/StateTree/Source/StateTreeEditorModule/Private/Debugger/StateTreeDebuggerTrack.cpp b/Engine/Plugins/Runtime/StateTree/Source/StateTreeEditorModule/Private/Debugger/StateTreeDebuggerTrack.cpp
--- a/Engine/Plugins/Runtime/StateTree/Source/StateTreeEditorModule/Private/Debugger/StateTreeDebuggerTrack.cpp
+++ b/Engine/Plugins/Runtime/StateTree/Source/StateTreeEditorModule/Private/Debugger/StateTreeDebuggerTrack.cpp
@@ -40,11 +40,11 @@ bool FStateTreeDebuggerTrack::UpdateInternal()
 	
 	if (StateTree != nullptr && EventCollection.IsValid())
 	{
-		auto MakeRandomColor = [](const uint32 InSeed)->FLinearColor
+		auto MakeRandomColor = [bActive = !bIsStale](const uint32 InSeed)->FLinearColor
 		{
 			const FRandomStream Stream(InSeed);
-			const uint8 Hue = (uint8)(Stream.FRand() * 255.0f);
-			constexpr uint8 SatVal = 196;
+			const uint8 Hue = static_cast<uint8>(Stream.FRand() * 255.0f);
+			const uint8 SatVal = bActive ? 196 : 128;
 			return FLinearColor::MakeFromHSV8(Hue, SatVal, SatVal);
 		};
 		
@@ -52,10 +52,10 @@ bool FStateTreeDebuggerTrack::UpdateInternal()
 		const TConstArrayView<FStateTreeTraceEventVariantType> Events = EventCollection.Events;
 		const uint32 NumStateChanges = EventCollection.ActiveStatesChanges.Num();
 		
-		for (uint32 StatechangeIndex = 0; StatechangeIndex < NumStateChanges; ++StatechangeIndex)
+		for (uint32 StateChangeIndex = 0; StateChangeIndex < NumStateChanges; ++StateChangeIndex)
 		{
-			const uint32 SpanIndex = EventCollection.ActiveStatesChanges[StatechangeIndex].Key;
-			const uint32 EventIndex = EventCollection.ActiveStatesChanges[StatechangeIndex].Value;
+			const uint32 SpanIndex = EventCollection.ActiveStatesChanges[StateChangeIndex].SpanIndex;
+			const uint32 EventIndex = EventCollection.ActiveStatesChanges[StateChangeIndex].EventIndex;
 			const FStateTreeTraceActiveStatesEvent& Event = Events[EventIndex].Get<FStateTreeTraceActiveStatesEvent>();
 				
 			FString StatePath;
@@ -64,21 +64,22 @@ bool FStateTreeDebuggerTrack::UpdateInternal()
 				const FCompactStateTreeState& State = StateTree->GetStates()[Event.ActiveStates[StateIndex].Index];
 				StatePath.Appendf(TEXT("%s%s"), StateIndex == 0 ? TEXT("") : TEXT("."), *State.Name.ToString());
 			}
-			
+
+			UE::StateTreeDebugger::FFrameSpan Span = EventCollection.FrameSpans[SpanIndex];
 			SStateTreeDebuggerEventTimelineView::FTimelineEventData::EventWindow& Window = EventData->Windows.AddDefaulted_GetRef();
 			Window.Color = MakeRandomColor(GetTypeHash(StatePath));
 			Window.Description = FText::FromString(StatePath);
-			Window.TimeStart = EventCollection.FrameSpans[SpanIndex].Frame.StartTime;
+			Window.TimeStart = Span.GetWorldTimeStart();
 
 			// When there is another state change after the current one in the list we use its start time to close the window.
-			if (StatechangeIndex < NumStateChanges-1)
+			if (StateChangeIndex < NumStateChanges-1)
 			{
-				const uint32 NextSpanIndex = EventCollection.ActiveStatesChanges[StatechangeIndex+1].Key;
-				Window.TimeEnd = EventCollection.FrameSpans[NextSpanIndex].Frame.StartTime;
+				const uint32 NextSpanIndex = EventCollection.ActiveStatesChanges[StateChangeIndex+1].SpanIndex;
+				Window.TimeEnd = EventCollection.FrameSpans[NextSpanIndex].GetWorldTimeStart();
 			}
 			else
 			{
-				Window.TimeEnd = Debugger->IsActiveInstance(RecordingDuration, InstanceId) ? RecordingDuration : EventCollection.FrameSpans[SpanIndex].Frame.EndTime;
+				Window.TimeEnd = Debugger->IsActiveInstance(RecordingDuration, InstanceId) ? RecordingDuration : Span.GetWorldTimeEnd();
 			}
 		}
 
@@ -93,9 +94,9 @@ bool FStateTreeDebuggerTrack::UpdateInternal()
 				if (Events[EventIndex].IsType<FStateTreeTraceLogEvent>())
 				{
 					SStateTreeDebuggerEventTimelineView::FTimelineEventData::EventPoint Point;
-					Point.Time = Span.Frame.StartTime;
+					Point.Time = Span.GetWorldTimeStart();
 					Point.Color = FColorList::Salmon;
-					EventData->Points.Add(Point);					
+					EventData->Points.Add(Point);
 				}
 			}
 		}
diff --git a/Engine/Plugins/Runtime/StateTree/Source/StateTreeEditorModule/Private/Debugger/StateTreeDebuggerTrack.h b/Engine/Plugins/Runtime/StateTree/Source/StateTreeEditorModule/Private/Debugger/StateTreeDebuggerTrack.h
--- a/Engine/Plugins/Runtime/StateTree/Source/StateTreeEditorModule/Private/Debugger/StateTreeDebuggerTrack.h
+++ b/Engine/Plugins/Runtime/StateTree/Source/StateTreeEditorModule/Private/Debugger/StateTreeDebuggerTrack.h
@@ -20,6 +20,9 @@ struct FStateTreeDebuggerTrack : RewindDebugger::FRewindDebuggerTrack
 	FStateTreeDebuggerTrack(const TSharedPtr<FStateTreeDebugger>& InDebugger, const FStateTreeInstanceDebugId InInstanceId, const FText InName, const TRange<double>& InViewRange);
 	
 	FStateTreeInstanceDebugId GetInstanceId() const { return InstanceId; }
+	bool IsStale() const { return bIsStale; }
+	void MarkAsStale() { bIsStale = true; }
+
 private:
 	virtual bool UpdateInternal() override;
 	virtual TSharedPtr<SWidget> GetDetailsViewInternal() override;
@@ -35,6 +38,7 @@ private:
 	TSharedPtr<FStateTreeDebugger> StateTreeDebugger;
 	FStateTreeInstanceDebugId InstanceId;
 	const TRange<double>& ViewRange;
+	bool bIsStale = false;
 };
 
 #endif // WITH_STATETREE_DEBUGGER
\ No newline at end of file
diff --git a/Engine/Plugins/Runtime/StateTree/Source/StateTreeEditorModule/Private/StateTreeEditorSettings.cpp b/Engine/Plugins/Runtime/StateTree/Source/StateTreeEditorModule/Private/StateTreeEditorSettings.cpp
--- a/Engine/Plugins/Runtime/StateTree/Source/StateTreeEditorModule/Private/StateTreeEditorSettings.cpp
+++ b/Engine/Plugins/Runtime/StateTree/Source/StateTreeEditorModule/Private/StateTreeEditorSettings.cpp
@@ -6,3 +6,22 @@ UStateTreeEditorSettings::UStateTreeEditorSettings(const FObjectInitializer& Obj
 	: Super(ObjectInitializer)
 {
 }
+
+#if WITH_EDITOR
+
+FText UStateTreeEditorSettings::GetSectionText() const
+{
+	return NSLOCTEXT("StateTreeEditor", "StateTreeEditorSettingsName", "StateTree Editor");
+}
+
+FText UStateTreeEditorSettings::GetSectionDescription() const
+{
+	return NSLOCTEXT("StateTreeEditor", "StateTreeEditorSettingsDescription", "Configure options for the StateTree Editor.");
+}
+
+#endif // WITH_EDITOR
+
+FName UStateTreeEditorSettings::GetCategoryName() const
+{
+	return FName(TEXT("Plugins"));
+}
\ No newline at end of file
diff --git a/Engine/Plugins/Runtime/StateTree/Source/StateTreeEditorModule/Private/StateTreeEditorStyle.cpp b/Engine/Plugins/Runtime/StateTree/Source/StateTreeEditorModule/Private/StateTreeEditorStyle.cpp
--- a/Engine/Plugins/Runtime/StateTree/Source/StateTreeEditorModule/Private/StateTreeEditorStyle.cpp
+++ b/Engine/Plugins/Runtime/StateTree/Source/StateTreeEditorModule/Private/StateTreeEditorStyle.cpp
@@ -178,6 +178,10 @@ FStateTreeEditorStyle::FStateTreeEditorStyle()
 	{
 		// From generic engine editor
 		UE::StateTree::Editor::FContentRootScope Scope(this, EngineEditorSlateContentDir);
+
+		Set("StateTreeEditor.Debugger.StartRecording", new IMAGE_BRUSH("Sequencer/Transport_Bar/Record_24x", CoreStyleConstants::Icon16x16));
+		Set("StateTreeEditor.Debugger.StopRecording", new IMAGE_BRUSH("Sequencer/Transport_Bar/Recording_24x", CoreStyleConstants::Icon16x16));
+		
 		Set("StateTreeEditor.Debugger.PreviousFrameWithStateChange", new IMAGE_BRUSH("Sequencer/Transport_Bar/Go_To_Front_24x", CoreStyleConstants::Icon16x16));
 		Set("StateTreeEditor.Debugger.PreviousFrameWithEvents", new IMAGE_BRUSH("Sequencer/Transport_Bar/Step_Backwards_24x", CoreStyleConstants::Icon16x16));
 		Set("StateTreeEditor.Debugger.NextFrameWithEvents", new IMAGE_BRUSH("Sequencer/Transport_Bar/Step_Forward_24x", CoreStyleConstants::Icon16x16));
diff --git a/Engine/Plugins/Runtime/StateTree/Source/StateTreeEditorModule/Public/StateTreeEditorSettings.h b/Engine/Plugins/Runtime/StateTree/Source/StateTreeEditorModule/Public/StateTreeEditorSettings.h
--- a/Engine/Plugins/Runtime/StateTree/Source/StateTreeEditorModule/Public/StateTreeEditorSettings.h
+++ b/Engine/Plugins/Runtime/StateTree/Source/StateTreeEditorModule/Public/StateTreeEditorSettings.h
@@ -18,9 +18,23 @@ class STATETREEEDITORMODULE_API UStateTreeEditorSettings : public UDeveloperSett
 {
 	GENERATED_UCLASS_BODY()
 public:
+	static UStateTreeEditorSettings& Get() { return *CastChecked<UStateTreeEditorSettings>(UStateTreeEditorSettings::StaticClass()->GetDefaultObject()); }
+
 	/** Determines when to save StateTrees post-compile */
 	UPROPERTY(EditAnywhere, config, Category = "Compiler")
 	EStateTreeSaveOnCompile SaveOnCompile = EStateTreeSaveOnCompile::Never;
+
+	/** If enabled, debugger starts recording information at the start of each PIE session. */
+	UPROPERTY(EditAnywhere, config, Category = "Debugger")
+	bool bShouldDebuggerAutoRecordOnPIE = true;
+
+protected:
+#if WITH_EDITOR
+	virtual FText GetSectionText() const override;
+	virtual FText GetSectionDescription() const override;
+#endif
+
+	virtual FName GetCategoryName() const override;
 };
 
 #if UE_ENABLE_INCLUDE_ORDER_DEPRECATED_IN_5_2
diff --git a/Engine/Plugins/Runtime/StateTree/Source/StateTreeModule/Private/Debugger/StateTreeDebugger.cpp b/Engine/Plugins/Runtime/StateTree/Source/StateTreeModule/Private/Debugger/StateTreeDebugger.cpp
--- a/Engine/Plugins/Runtime/StateTree/Source/StateTreeModule/Private/Debugger/StateTreeDebugger.cpp
+++ b/Engine/Plugins/Runtime/StateTree/Source/StateTreeModule/Private/Debugger/StateTreeDebugger.cpp
@@ -86,19 +86,25 @@ FStateTreeDebugger::~FStateTreeDebugger()
 	StopAnalysis();
 }
 
-void FStateTreeDebugger::Tick(float DeltaTime)
+void FStateTreeDebugger::Tick(const float DeltaTime)
 {
 	TRACE_CPUPROFILER_EVENT_SCOPE(FStateTreeDebugger::Tick);
 
+	if (RetryLoadLastLiveSessionTimer > 0.0f)
+	{
+		// We are still not connected to the last live session.
+		// Update polling timer and retry with remaining time; 0 or less will stop retries.
+		StartLastLiveSessionAnalysis(RetryLoadLastLiveSessionTimer - DeltaTime);
+		if (ActiveSessionTraceDescriptor.IsValid())
+		{
+			RetryLoadLastLiveSessionTimer = 0.0f;
+		}
+	}
+	
 	UpdateInstances();
 	SyncToCurrentSessionDuration();
 }
 
-bool FStateTreeDebugger::IsTickable() const
-{
-	return AnalysisSession.IsValid();
-}
-
 void FStateTreeDebugger::StopAnalysis()
 {
 	if (const TraceServices::IAnalysisSession* Session = GetAnalysisSession())
@@ -106,8 +112,6 @@ void FStateTreeDebugger::StopAnalysis()
 		Session->Stop(true);
 		AnalysisSession.Reset();	
 	}
-
-	EventCollections.Reset();
 }
 
 void FStateTreeDebugger::Pause()
@@ -127,9 +131,9 @@ void FStateTreeDebugger::SyncToCurrentSessionDuration()
 	{
 		{
 			TraceServices::FAnalysisSessionReadScope SessionReadScope(*Session);
-			RecordingDuration = Session->GetDurationSeconds();
+			AnalysisDuration = Session->GetDurationSeconds();
 		}
-		ReadTrace(RecordingDuration);
+		ReadTrace(AnalysisDuration);
 	}
 }
 
@@ -192,7 +196,7 @@ void FStateTreeDebugger::UpdateInstances()
 	}
 }
 
-void FStateTreeDebugger::StartLastLiveSessionAnalysis()
+void FStateTreeDebugger::StartLastLiveSessionAnalysis(const float RetryPollingDuration)
 {
 	TArray<FTraceDescriptor> Traces;
 	GetLiveTraces(Traces);
@@ -201,6 +205,12 @@ void FStateTreeDebugger::StartLastLiveSessionAnalysis()
 	{
 		StartSessionAnalysis(Traces.Last());
 	}
+	else
+	{
+		RetryLoadLastLiveSessionTimer = RetryPollingDuration;
+		ensure(RetryLoadLastLiveSessionTimer > 0);
+		UE_CLOG(RetryLoadLastLiveSessionTimer > 0, LogStateTree, Log, TEXT("Unable to start analysis for the most recent live session."));
+	}
 }
 
 void FStateTreeDebugger::StartSessionAnalysis(const FTraceDescriptor& TraceDescriptor)
@@ -219,6 +229,10 @@ void FStateTreeDebugger::StartSessionAnalysis(const FTraceDescriptor& TraceDescr
 	// Make sure any active analysis is stopped
 	StopAnalysis();
 
+	RecordingDuration = 0;
+	AnalysisDuration = 0;
+	LastTraceReadTime = 0;
+
 	const uint32 TraceId = TraceDescriptor.TraceId;
 
 	// Make sure it is still live
@@ -412,7 +426,7 @@ void FStateTreeDebugger::RefreshActiveStates()
 	if (ScrubState.IsPointingToValidActiveStates())
 	{
 		const UE::StateTreeDebugger::FInstanceEventCollection& EventCollection = EventCollections[ScrubState.EventCollectionIndex];
-		const int32 EventIndex = EventCollection.ActiveStatesChanges[ScrubState.ActiveStatesIndex].Value;
+		const int32 EventIndex = EventCollection.ActiveStatesChanges[ScrubState.ActiveStatesIndex].EventIndex;
 		NewActiveStates = EventCollection.Events[EventIndex].Get<FStateTreeTraceActiveStatesEvent>().ActiveStates;
 	}
 
@@ -598,7 +612,7 @@ void FStateTreeDebugger::ReadTrace(
 {
 	TraceServices::FFrame LastReadFrame;
 	const bool bValidLastReadFrame = FrameProvider.GetFrameFromTime(TraceFrameType_Game, LastTraceReadTime, LastReadFrame);
-	if (LastTraceReadTime == UnsetTime || (bValidLastReadFrame && Frame.Index > LastReadFrame.Index))
+	if (LastTraceReadTime == 0 || (bValidLastReadFrame && Frame.Index > LastReadFrame.Index))
 	{
 		if (const IStateTreeTraceProvider* Provider = Session.ReadProvider<IStateTreeTraceProvider>(FStateTreeTraceProvider::ProviderName))
 		{
@@ -676,25 +690,34 @@ bool FStateTreeDebugger::ProcessEvent(const FStateTreeInstanceDebugId InstanceId
 	// Add new frame span if none added yet or new frame
 	if (ExistingCollection->FrameSpans.IsEmpty() || ExistingCollection->FrameSpans.Last().Frame.Index < Frame.Index)
 	{
-		ExistingCollection->FrameSpans.Add(UE::StateTreeDebugger::FFrameSpan(Frame, Events.Num()));
+		double RecordingWorldTime = 0;
+		Visit([&RecordingWorldTime](auto& TypedEvent)
+			{
+				RecordingWorldTime = TypedEvent.RecordingWorldTime;
+			}, Event);
+
+		// Update global recording duration
+		RecordingDuration = RecordingWorldTime;
+
+		ExistingCollection->FrameSpans.Add(UE::StateTreeDebugger::FFrameSpan(Frame, RecordingWorldTime, Events.Num()));
 	}
 
 	// Add activate states change info
 	if (Event.IsType<FStateTreeTraceActiveStatesEvent>())
 	{
 		checkf(ExistingCollection->FrameSpans.Num() > 0, TEXT("Expecting to always be in a frame span at this point."));
-		const uint32 FrameSpanIndex = ExistingCollection->FrameSpans.Num()-1;
+		const int32 FrameSpanIndex = ExistingCollection->FrameSpans.Num()-1;
 
 		// Add new entry for the first event or if the last event is for a different frame
 		if (ExistingCollection->ActiveStatesChanges.IsEmpty()
-			|| ExistingCollection->ActiveStatesChanges.Last().Key != FrameSpanIndex)
+			|| ExistingCollection->ActiveStatesChanges.Last().SpanIndex != FrameSpanIndex)
 		{
 			ExistingCollection->ActiveStatesChanges.Push({FrameSpanIndex, Events.Num()});
 		}
 		else
 		{
 			// Multiple events for change of active states in the same frame, keep the last one until we implement scrubbing within a frame
-			ExistingCollection->ActiveStatesChanges.Last().Value = Events.Num();
+			ExistingCollection->ActiveStatesChanges.Last().EventIndex = Events.Num();
 		}
 	}
 
diff --git a/Engine/Plugins/Runtime/StateTree/Source/StateTreeModule/Private/Debugger/StateTreeDebuggerTypes.cpp b/Engine/Plugins/Runtime/StateTree/Source/StateTreeModule/Private/Debugger/StateTreeDebuggerTypes.cpp
--- a/Engine/Plugins/Runtime/StateTree/Source/StateTreeModule/Private/Debugger/StateTreeDebuggerTypes.cpp
+++ b/Engine/Plugins/Runtime/StateTree/Source/StateTreeModule/Private/Debugger/StateTreeDebuggerTypes.cpp
@@ -35,45 +35,38 @@ const FInstanceEventCollection FInstanceEventCollection::Invalid;
 //----------------------------------------------------------------//
 void FScrubState::SetScrubTime(const double NewScrubTime)
 {
+	ScrubTimeBoundState = EScrubTimeBoundState::Unset;
+	TraceFrameIndex = INDEX_NONE;
+	FrameSpanIndex = INDEX_NONE;
+	ActiveStatesIndex = INDEX_NONE;
+
 	if (EventCollectionIndex != INDEX_NONE)
 	{
 		const TArray<FFrameSpan>& Spans = EventCollections[EventCollectionIndex].FrameSpans;
-		const uint32 NewFrameSpanIndex = Spans.IndexOfByPredicate([Time = NewScrubTime](const FFrameSpan& Span)
-			{
-				return Span.Frame.StartTime <= Time && Time <= Span.Frame.EndTime;
-			});
-
-		if (NewFrameSpanIndex != INDEX_NONE)
-		{
-			PreviousFrameSpanIndex = NewFrameSpanIndex > 0 ? NewFrameSpanIndex - 1 : NewFrameSpanIndex;
-			SetFrameSpanIndex(NewFrameSpanIndex);
-		}
-		else if (Spans.Num() > 0 && NewScrubTime > Spans.Last().Frame.EndTime)
+		if (Spans.Num() > 0)
 		{
-			// Clamp to last span so it keep the scrub valid when sticking to most recent data.
-			const uint32 LastSpanIndex = EventCollections[EventCollectionIndex].FrameSpans.Num() - 1;
-			SetFrameSpanIndex(LastSpanIndex);
-			PreviousFrameSpanIndex = LastSpanIndex;
-		}
-		else
-		{
-			TraceFrameIndex = INDEX_NONE;
-			FrameSpanIndex = INDEX_NONE;
-
-			// Find the last span before the given time (to allow Prev/Next actions)
-			PreviousFrameSpanIndex = Spans.FindLastByPredicate([Time = NewScrubTime](const FFrameSpan& Span)
-				{
-					return Span.Frame.EndTime < Time;
-				});
-
-			// If previous span found let's update the active states index (to allow Prev/Next actions)
-			if (PreviousFrameSpanIndex != INDEX_NONE)
+			const double SpansLowerBound = Spans[0].GetWorldTimeStart();
+			const double SpansUpperBound =  Spans.Last().GetWorldTimeEnd();
+			
+			if (NewScrubTime < SpansLowerBound)
 			{
-				UpdateActiveStatesIndex(PreviousFrameSpanIndex);
+				ScrubTimeBoundState = EScrubTimeBoundState::BeforeLowerBound;
+			}
+			else if (NewScrubTime > SpansUpperBound)
+			{
+				ScrubTimeBoundState = EScrubTimeBoundState::AfterHigherBound;
+				UpdateActiveStatesIndex(Spans.Num() - 1);
 			}
 			else
 			{
-				ActiveStatesIndex = INDEX_NONE;
+				const uint32 NextFrameSpanIndex = Spans.IndexOfByPredicate([NewScrubTime](const FFrameSpan& Span)
+					{
+						return Span.GetWorldTimeStart() > NewScrubTime;
+					});
+
+				
+				ensure(NextFrameSpanIndex == INDEX_NONE || NextFrameSpanIndex > 0);
+				SetFrameSpanIndex(NextFrameSpanIndex != INDEX_NONE ? NextFrameSpanIndex-1 : Spans.Num() - 1);
 			}
 		}
 	}
@@ -83,19 +76,20 @@ void FScrubState::SetScrubTime(const double NewScrubTime)
 	ScrubTime = NewScrubTime;
 }
 
-void FScrubState::SetFrameSpanIndex(const uint32 NewFrameSpanIndex)
+void FScrubState::SetFrameSpanIndex(const int32 NewFrameSpanIndex)
 {
 	FrameSpanIndex = NewFrameSpanIndex;
 	checkf(EventCollections.IsValidIndex(EventCollectionIndex), TEXT("Internal method expecting validity checks before getting called."));
 	const FInstanceEventCollection& EventCollection = EventCollections[EventCollectionIndex];
 
 	checkf(EventCollections[EventCollectionIndex].FrameSpans.IsValidIndex(FrameSpanIndex), TEXT("Internal method expecting validity checks before getting called."));
-	ScrubTime = EventCollection.FrameSpans[FrameSpanIndex].Frame.StartTime;
+	ScrubTime = EventCollection.FrameSpans[FrameSpanIndex].GetWorldTimeStart();
 	TraceFrameIndex = EventCollection.FrameSpans[FrameSpanIndex].Frame.Index;
+	ScrubTimeBoundState = EScrubTimeBoundState::InBounds;
 	UpdateActiveStatesIndex(NewFrameSpanIndex);
 }
 
-void FScrubState::SetActiveStatesIndex(const uint32 NewActiveStatesIndex)
+void FScrubState::SetActiveStatesIndex(const int32 NewActiveStatesIndex)
 {
 	ActiveStatesIndex = NewActiveStatesIndex;
 
@@ -103,23 +97,24 @@ void FScrubState::SetActiveStatesIndex(const uint32 NewActiveStatesIndex)
 	const FInstanceEventCollection& EventCollection = EventCollections[EventCollectionIndex];
 
 	checkf(EventCollection.ActiveStatesChanges.IsValidIndex(ActiveStatesIndex), TEXT("Internal method expecting validity checks before getting called."));
-	FrameSpanIndex = EventCollection.ActiveStatesChanges[ActiveStatesIndex].Key;
-	ScrubTime = EventCollection.FrameSpans[FrameSpanIndex].Frame.StartTime;
+	FrameSpanIndex = EventCollection.ActiveStatesChanges[ActiveStatesIndex].SpanIndex;
+	ScrubTime = EventCollection.FrameSpans[FrameSpanIndex].GetWorldTimeStart();
 	TraceFrameIndex = EventCollection.FrameSpans[FrameSpanIndex].Frame.Index;
+	ScrubTimeBoundState = EScrubTimeBoundState::InBounds;
 }
 
 bool FScrubState::HasPreviousFrame() const
 {
 	if (EventCollectionIndex != INDEX_NONE)
 	{
-		return IsInBounds() ? EventCollections[EventCollectionIndex].FrameSpans.IsValidIndex(FrameSpanIndex - 1) : PreviousFrameSpanIndex != INDEX_NONE;
+		return IsInBounds() ? EventCollections[EventCollectionIndex].FrameSpans.IsValidIndex(FrameSpanIndex - 1) : ScrubTimeBoundState == EScrubTimeBoundState::AfterHigherBound;
 	}
 	return false;
 }
 
 double FScrubState::GotoPreviousFrame()
 {
-	SetFrameSpanIndex(IsInBounds() ? (FrameSpanIndex - 1) : PreviousFrameSpanIndex);
+	SetFrameSpanIndex(IsInBounds() ? (FrameSpanIndex - 1) : EventCollections[EventCollectionIndex].FrameSpans.Num()-1);
 	return ScrubTime;
 }
 
@@ -127,43 +122,75 @@ bool FScrubState::HasNextFrame() const
 {
 	if (EventCollectionIndex != INDEX_NONE)
 	{
-		const TConstArrayView<FFrameSpan> FrameSpans = EventCollections[EventCollectionIndex].FrameSpans;
-		return IsInBounds()
-				   ? FrameSpans.IsValidIndex(FrameSpanIndex + 1)
-				   : (PreviousFrameSpanIndex != INDEX_NONE && FrameSpans.IsValidIndex(PreviousFrameSpanIndex + 1));
+		return IsInBounds()	? EventCollections[EventCollectionIndex].FrameSpans.IsValidIndex(FrameSpanIndex + 1) : ScrubTimeBoundState == EScrubTimeBoundState::BeforeLowerBound;
 	}
 	return false;
 }
 
 double FScrubState::GotoNextFrame()
 {
-	SetFrameSpanIndex(IsInBounds() ? (FrameSpanIndex + 1) : PreviousFrameSpanIndex + 1);
+	SetFrameSpanIndex(IsInBounds() ? (FrameSpanIndex + 1) : 0);
 	return ScrubTime;
 }
 
 bool FScrubState::HasPreviousActiveStates() const
 {
-	return (EventCollectionIndex != INDEX_NONE && ActiveStatesIndex != INDEX_NONE)
-		? EventCollections[EventCollectionIndex].ActiveStatesChanges.IsValidIndex(ActiveStatesIndex - 1)
-		: false;
+	if (EventCollectionIndex == INDEX_NONE || ActiveStatesIndex == INDEX_NONE)
+	{
+		return false;
+	}
+
+	const TArray<FInstanceEventCollection::FActiveStatesChangePair>& ActiveStatesChanges = EventCollections[EventCollectionIndex].ActiveStatesChanges;
+
+	if (ActiveStatesChanges.IsValidIndex(ActiveStatesIndex) && ActiveStatesChanges[ActiveStatesIndex].SpanIndex < FrameSpanIndex)
+	{
+		return true;
+	}
+
+	return ActiveStatesChanges.IsValidIndex(ActiveStatesIndex - 1);
 }
 
 double FScrubState::GotoPreviousActiveStates()
 {
-	SetActiveStatesIndex(ActiveStatesIndex - 1);
+	const TArray<FInstanceEventCollection::FActiveStatesChangePair>& ActiveStatesChanges = EventCollections[EventCollectionIndex].ActiveStatesChanges;
+	if (ActiveStatesChanges.IsValidIndex(ActiveStatesIndex) && ActiveStatesChanges[ActiveStatesIndex].SpanIndex < FrameSpanIndex)
+	{
+		SetActiveStatesIndex(ActiveStatesIndex);
+	}
+	else
+	{
+		SetActiveStatesIndex(ActiveStatesIndex - 1);
+	}
+
 	return ScrubTime;
 }
 
 bool FScrubState::HasNextActiveStates() const
 {
-	return (EventCollectionIndex != INDEX_NONE && ActiveStatesIndex != INDEX_NONE)
-		? EventCollections[EventCollectionIndex].ActiveStatesChanges.IsValidIndex(ActiveStatesIndex + 1)
-		: false;
+	if (EventCollectionIndex == INDEX_NONE)
+	{
+		return false;
+	}
+
+	const TArray<FInstanceEventCollection::FActiveStatesChangePair>& ActiveStatesChanges = EventCollections[EventCollectionIndex].ActiveStatesChanges;
+	if (ScrubTimeBoundState == EScrubTimeBoundState::BeforeLowerBound && ActiveStatesChanges.Num() > 0)
+	{
+		return true;
+	}
+
+	return ActiveStatesIndex != INDEX_NONE && EventCollections[EventCollectionIndex].ActiveStatesChanges.IsValidIndex(ActiveStatesIndex + 1);
 }
 
 double FScrubState::GotoNextActiveStates()
 {
-	SetActiveStatesIndex(ActiveStatesIndex + 1);
+	if (ScrubTimeBoundState == EScrubTimeBoundState::BeforeLowerBound)
+	{
+		SetActiveStatesIndex(0);
+	}
+	else
+	{
+		SetActiveStatesIndex(ActiveStatesIndex + 1);
+	}
 	return ScrubTime;
 }
 
@@ -172,16 +199,16 @@ const FInstanceEventCollection& FScrubState::GetEventCollection() const
 	return EventCollectionIndex != INDEX_NONE ? EventCollections[EventCollectionIndex] : FInstanceEventCollection::Invalid;
 }
 
-void FScrubState::UpdateActiveStatesIndex(const uint32 SpanIndex)
+void FScrubState::UpdateActiveStatesIndex(const int32 SpanIndex)
 {
 	check(EventCollectionIndex != INDEX_NONE);
 	const FInstanceEventCollection& EventCollection = EventCollections[EventCollectionIndex];
 
 	// Need to find the index of a frame span that contains an active states changed event; either the current one has it otherwise look backward to find the last one
 	ActiveStatesIndex = EventCollection.ActiveStatesChanges.FindLastByPredicate(
-		[SpanIndex](const TTuple<uint32, uint32> SpanAndEventIndices)
+		[SpanIndex](const FInstanceEventCollection::FActiveStatesChangePair& SpanAndEventIndices)
 		{
-			return SpanAndEventIndices.Key <= SpanIndex;
+			return SpanAndEventIndices.SpanIndex <= SpanIndex;
 		});
 }
 } // UE::StateTreeDebugger
diff --git a/Engine/Plugins/Runtime/StateTree/Source/StateTreeModule/Private/Debugger/StateTreeTrace.cpp b/Engine/Plugins/Runtime/StateTree/Source/StateTreeModule/Private/Debugger/StateTreeTrace.cpp
--- a/Engine/Plugins/Runtime/StateTree/Source/StateTreeModule/Private/Debugger/StateTreeTrace.cpp
+++ b/Engine/Plugins/Runtime/StateTree/Source/StateTreeModule/Private/Debugger/StateTreeTrace.cpp
@@ -5,14 +5,23 @@
 #include "Debugger/StateTreeTrace.h"
 #include "Debugger/StateTreeDebugger.h"
 #include "Exporters/Exporter.h"
+#include "ObjectTrace.h"
 #include "Serialization/BufferArchive.h"
 #include "StateTree.h"
 #include "StateTreeExecutionTypes.h"
 #include "UObject/Package.h"
 #include "Trace/Trace.inl"
 
+#if WITH_EDITOR
+#include "Editor.h"
+#endif // WITH_EDITOR
+
 UE_TRACE_CHANNEL_DEFINE(StateTreeDebugChannel)
 
+UE_TRACE_EVENT_BEGIN(StateTreeDebugger, WorldTimestampEvent)
+	UE_TRACE_EVENT_FIELD(double, WorldTime)
+UE_TRACE_EVENT_END()
+
 UE_TRACE_EVENT_BEGIN(StateTreeDebugger, InstanceEvent)
 	UE_TRACE_EVENT_FIELD(uint64, Cycle)
 	UE_TRACE_EVENT_FIELD(UE::Trace::WideString, TreeName)
@@ -82,37 +91,81 @@ UE_TRACE_EVENT_END()
 
 namespace UE::StateTreeTrace
 {
-	void SerializeDataViewToArchive(FBufferArchive Ar, const FStateTreeDataView DataView)
+
+FDelegateHandle GOnWorldTickStartDelegateHandle;
+
+#if WITH_EDITOR
+FDelegateHandle GOnPIEStartDelegateHandle;
+#endif // WITH_EDITOR
+
+double GRecordingWorldTime = -1;
+double GTracedRecordingWorldTime = -1;
+
+void TraceTimeMapping()
+{
+	if (GTracedRecordingWorldTime != GRecordingWorldTime)
 	{
-		if (const UScriptStruct* ScriptStruct = Cast<const UScriptStruct>(DataView.GetStruct()))
-		{
-			TRACE_CPUPROFILER_EVENT_SCOPE(UE::StateTree::ExportStructAsText)
-			FString StructPath = ScriptStruct->GetPathName();
-			FString TextValue;
+		GTracedRecordingWorldTime = GRecordingWorldTime;
+		UE_TRACE_LOG(StateTreeDebugger, WorldTimestampEvent, StateTreeDebugChannel)
+			<< WorldTimestampEvent.WorldTime(GRecordingWorldTime);
+	}
+}
 
-			ScriptStruct->ExportText(TextValue, DataView.GetMemory(), DataView.GetMemory(), /*OwnerObject*/nullptr, PPF_None, /*ExportRootScope*/nullptr);
+void SerializeDataViewToArchive(FBufferArchive Ar, const FStateTreeDataView DataView)
+{
+	if (const UScriptStruct* ScriptStruct = Cast<const UScriptStruct>(DataView.GetStruct()))
+	{
+		TRACE_CPUPROFILER_EVENT_SCOPE(UE::StateTree::ExportStructAsText)
+		FString StructPath = ScriptStruct->GetPathName();
+		FString TextValue;
 
-			Ar << StructPath;		
-			Ar << TextValue;
-		}
-		else if (const UClass* Class = Cast<const UClass>(DataView.GetStruct()))
-		{
-			TRACE_CPUPROFILER_EVENT_SCOPE(UE::StateTree::ExportObjectAsText)
-			FString StructPath = Class->GetPathName();
-			FStringOutputDevice OutputDevice;
-			UObject* Object = DataView.GetMutablePtr<UObject>();
-
-			// Not using on scope FExportObjectInnerContext since it is very costly to build.
-			// Passing a null context will make the export use an already built thread local context.
-			UExporter::ExportToOutputDevice(nullptr, Object, /*Exporter*/nullptr, OutputDevice, TEXT("copy"), 0, PPF_ExportsNotFullyQualified | PPF_Copy | PPF_Delimited, false, Object->GetOuter());
-
-			Ar << StructPath;
-			Ar << OutputDevice;
-		}
+		ScriptStruct->ExportText(TextValue, DataView.GetMemory(), DataView.GetMemory(), /*OwnerObject*/nullptr, PPF_None, /*ExportRootScope*/nullptr);
+
+		Ar << StructPath;
+		Ar << TextValue;
+	}
+	else if (const UClass* Class = Cast<const UClass>(DataView.GetStruct()))
+	{
+		TRACE_CPUPROFILER_EVENT_SCOPE(UE::StateTree::ExportObjectAsText)
+		FString StructPath = Class->GetPathName();
+		FStringOutputDevice OutputDevice;
+		UObject* Object = DataView.GetMutablePtr<UObject>();
+
+		// Not using on scope FExportObjectInnerContext since it is very costly to build.
+		// Passing a null context will make the export use an already built thread local context.
+		UExporter::ExportToOutputDevice(nullptr, Object, /*Exporter*/nullptr, OutputDevice, TEXT("copy"), 0, PPF_ExportsNotFullyQualified | PPF_Copy | PPF_Delimited, false, Object->GetOuter());
+
+		Ar << StructPath;
+		Ar << OutputDevice;
 	}
 }
 
-void UE::StateTreeTrace::OutputInstanceLifetimeEvent(
+void RegisterGlobalDelegates()
+{
+#if WITH_EDITOR
+	GOnPIEStartDelegateHandle = FEditorDelegates::BeginPIE.AddLambda([&LastRecordingWorldTime=GTracedRecordingWorldTime](const bool bIsSimulating)
+		{
+			LastRecordingWorldTime = -1;
+		});
+#endif // WITH_EDITOR
+	
+	GOnWorldTickStartDelegateHandle = FWorldDelegates::OnWorldTickStart.AddLambda([&WorldTime=GRecordingWorldTime](const UWorld* TickedWorld, ELevelTick TickType, float DeltaTime)
+		{
+			WorldTime = FObjectTrace::GetWorldElapsedTime(TickedWorld);
+		});
+}
+
+void UnregisterGlobalDelegates()
+{
+#if WITH_EDITOR
+	FEditorDelegates::BeginPIE.Remove(GOnPIEStartDelegateHandle);
+#endif // WITH_EDITOR
+	
+	FWorldDelegates::OnWorldTickStart.Remove(GOnWorldTickStartDelegateHandle);
+	GOnWorldTickStartDelegateHandle.Reset();
+}
+
+void OutputInstanceLifetimeEvent(
 	const FStateTreeInstanceDebugId InstanceId,
 	const UStateTree* StateTree,
 	const TCHAR* InstanceName,
@@ -126,6 +179,8 @@ void UE::StateTreeTrace::OutputInstanceLifetimeEvent(
 	const FString ObjectName = StateTree->GetName();
 	const FString ObjectPackageName = StateTree->GetPackage()->GetName();
 
+	TraceTimeMapping();
+
 	UE_TRACE_LOG(StateTreeDebugger, InstanceEvent, StateTreeDebugChannel)
 		<< InstanceEvent.Cycle(FPlatformTime::Cycles64())
 		<< InstanceEvent.TreeName(*ObjectName, ObjectName.Len())
@@ -137,7 +192,7 @@ void UE::StateTreeTrace::OutputInstanceLifetimeEvent(
 		<< InstanceEvent.EventType(static_cast<std::underlying_type_t<EStateTreeTraceNodeEventType>>(EventType));
 }
 
-void UE::StateTreeTrace::OutputLogEventTrace(
+void OutputLogEventTrace(
 	const FStateTreeInstanceDebugId InstanceId,
 	const EStateTreeUpdatePhase Phase,
 	const TCHAR* Fmt, ...
@@ -146,6 +201,8 @@ void UE::StateTreeTrace::OutputLogEventTrace(
 	static TCHAR TraceStaticBuffer[8192];
 	GET_VARARGS(TraceStaticBuffer, UE_ARRAY_COUNT(TraceStaticBuffer), UE_ARRAY_COUNT(TraceStaticBuffer) - 1, Fmt, Fmt);
 
+	TraceTimeMapping();
+
 	UE_TRACE_LOG(StateTreeDebugger, LogEvent, StateTreeDebugChannel)
 		<< LogEvent.Cycle(FPlatformTime::Cycles64())
 		<< LogEvent.InstanceId(InstanceId.Id)
@@ -154,7 +211,7 @@ void UE::StateTreeTrace::OutputLogEventTrace(
 		<< LogEvent.Message(TraceStaticBuffer);
 }
 
-void UE::StateTreeTrace::OutputStateEventTrace(
+void OutputStateEventTrace(
 	const FStateTreeInstanceDebugId InstanceId,
 	const EStateTreeUpdatePhase Phase,
 	const FStateTreeStateHandle StateHandle,
@@ -162,6 +219,8 @@ void UE::StateTreeTrace::OutputStateEventTrace(
 	const EStateTreeStateSelectionBehavior SelectionBehavior
 	)
 {
+	TraceTimeMapping();
+
 	UE_TRACE_LOG(StateTreeDebugger, StateEvent, StateTreeDebugChannel)
 		<< StateEvent.Cycle(FPlatformTime::Cycles64())
 		<< StateEvent.InstanceId(InstanceId.Id)
@@ -172,7 +231,7 @@ void UE::StateTreeTrace::OutputStateEventTrace(
 		<< StateEvent.SelectionBehavior(static_cast<std::underlying_type_t<EStateTreeStateSelectionBehavior>>(SelectionBehavior));
 }
 
-void UE::StateTreeTrace::OutputTaskEventTrace(
+void OutputTaskEventTrace(
 	const FStateTreeInstanceDebugId InstanceId,
 	const EStateTreeUpdatePhase Phase,
 	const FStateTreeIndex16 TaskIdx,
@@ -184,6 +243,8 @@ void UE::StateTreeTrace::OutputTaskEventTrace(
 	FBufferArchive Archive;
 	SerializeDataViewToArchive(Archive, DataView);
 
+	TraceTimeMapping();
+
 	UE_TRACE_LOG(StateTreeDebugger, TaskEvent, StateTreeDebugChannel)
 		<< TaskEvent.Cycle(FPlatformTime::Cycles64())
 		<< TaskEvent.InstanceId(InstanceId.Id)
@@ -195,7 +256,7 @@ void UE::StateTreeTrace::OutputTaskEventTrace(
 		<< TaskEvent.Status(static_cast<std::underlying_type_t<EStateTreeRunStatus>>(Status));
 }
 
-void UE::StateTreeTrace::OutputTransitionEventTrace(
+void OutputTransitionEventTrace(
 	const FStateTreeInstanceDebugId InstanceId,
 	const EStateTreeUpdatePhase Phase,
 	const FStateTreeIndex16 TransitionIdx,
@@ -205,6 +266,8 @@ void UE::StateTreeTrace::OutputTransitionEventTrace(
 	FBufferArchive Archive;
 	Archive << EventType;
 
+	TraceTimeMapping();
+
 	UE_TRACE_LOG(StateTreeDebugger, TransitionEvent, StateTreeDebugChannel)
 	<< TransitionEvent.Cycle(FPlatformTime::Cycles64())
 	<< TransitionEvent.InstanceId(InstanceId.Id)
@@ -214,7 +277,7 @@ void UE::StateTreeTrace::OutputTransitionEventTrace(
 	<< TransitionEvent.EventType(static_cast<std::underlying_type_t<EStateTreeTraceNodeEventType>>(EventType));
 }
 
-void UE::StateTreeTrace::OutputConditionEventTrace(
+void OutputConditionEventTrace(
 	const FStateTreeInstanceDebugId InstanceId,
 	const EStateTreeUpdatePhase Phase,
 	const FStateTreeIndex16 ConditionIdx,
@@ -224,7 +287,9 @@ void UE::StateTreeTrace::OutputConditionEventTrace(
 {
 	FBufferArchive Archive;
 	SerializeDataViewToArchive(Archive, DataView);
-	
+
+	TraceTimeMapping();
+
 	UE_TRACE_LOG(StateTreeDebugger, ConditionEvent, StateTreeDebugChannel)
 	<< ConditionEvent.Cycle(FPlatformTime::Cycles64())
 	<< ConditionEvent.InstanceId(InstanceId.Id)
@@ -235,7 +300,7 @@ void UE::StateTreeTrace::OutputConditionEventTrace(
 	<< ConditionEvent.EventType(static_cast<std::underlying_type_t<EStateTreeTraceNodeEventType>>(EventType));
 }
 
-void UE::StateTreeTrace::OutputActiveStatesEventTrace(
+void OutputActiveStatesEventTrace(
 	const FStateTreeInstanceDebugId InstanceId,
 	const EStateTreeUpdatePhase Phase,
 	const FStateTreeActiveStates& ActiveStates
@@ -247,6 +312,8 @@ void UE::StateTreeTrace::OutputActiveStatesEventTrace(
 		StatesIndices.Add(ActiveStates[i].Index);
 	}
 
+	TraceTimeMapping();
+
 	UE_TRACE_LOG(StateTreeDebugger, ActiveStatesEvent, StateTreeDebugChannel)
 		<< ActiveStatesEvent.Cycle(FPlatformTime::Cycles64())
 		<< ActiveStatesEvent.InstanceId(InstanceId.Id)
@@ -255,4 +322,6 @@ void UE::StateTreeTrace::OutputActiveStatesEventTrace(
 		<< ActiveStatesEvent.ActiveStates(StatesIndices.GetData(), StatesIndices.Num());
 }
 
+} // UE::StateTreeTrace
+
 #endif // WITH_STATETREE_DEBUGGER
diff --git a/Engine/Plugins/Runtime/StateTree/Source/StateTreeModule/Private/Debugger/StateTreeTraceAnalyzer.cpp b/Engine/Plugins/Runtime/StateTree/Source/StateTreeModule/Private/Debugger/StateTreeTraceAnalyzer.cpp
--- a/Engine/Plugins/Runtime/StateTree/Source/StateTreeModule/Private/Debugger/StateTreeTraceAnalyzer.cpp
+++ b/Engine/Plugins/Runtime/StateTree/Source/StateTreeModule/Private/Debugger/StateTreeTraceAnalyzer.cpp
@@ -19,6 +19,7 @@ void FStateTreeTraceAnalyzer::OnAnalysisBegin(const FOnAnalysisContext& Context)
 {
 	auto& Builder = Context.InterfaceBuilder;
 
+	Builder.RouteEvent(RouteId_WorldTimestamp, "StateTreeDebugger", "WorldTimestampEvent");
 	Builder.RouteEvent(RouteId_Instance, "StateTreeDebugger", "InstanceEvent");
 	Builder.RouteEvent(RouteId_LogMessage, "StateTreeDebugger", "LogEvent");
 	Builder.RouteEvent(RouteId_State, "StateTreeDebugger", "StateEvent");
@@ -37,6 +38,11 @@ bool FStateTreeTraceAnalyzer::OnEvent(const uint16 RouteId, EStyle Style, const
 	const auto& EventData = Context.EventData;
 	switch (RouteId)
 	{
+	case RouteId_WorldTimestamp:
+		{
+			WorldTime = EventData.GetValue<double>("WorldTime");
+			break;
+		}
 	case RouteId_Instance:
 		{
 			FString ObjectName, ObjectPathName;
@@ -65,6 +71,7 @@ bool FStateTreeTraceAnalyzer::OnEvent(const uint16 RouteId, EStyle Style, const
 						FStateTreeInstanceDebugId(EventData.GetValue<uint32>("InstanceId"), EventData.GetValue<uint32>("InstanceSerial")),
 						*InstanceName,
 						Context.EventTime.AsSeconds(EventData.GetValue<uint64>("Cycle")),
+						WorldTime,
 						EventData.GetValue<EStateTreeTraceInstanceEventType>("EventType"));
 				}
 				else
@@ -82,7 +89,7 @@ bool FStateTreeTraceAnalyzer::OnEvent(const uint16 RouteId, EStyle Style, const
 		{
 			FString Message;
         	EventData.GetString("Message", Message);
-			const FStateTreeTraceLogEvent Event(EventData.GetValue<EStateTreeUpdatePhase>("Phase"), Message);
+			const FStateTreeTraceLogEvent Event(WorldTime, EventData.GetValue<EStateTreeUpdatePhase>("Phase"), Message);
 
 			Provider.AppendEvent(FStateTreeInstanceDebugId(EventData.GetValue<uint32>("InstanceId"), EventData.GetValue<uint32>("InstanceSerial")),
 				Context.EventTime.AsSeconds(EventData.GetValue<uint64>("Cycle")),
@@ -91,7 +98,7 @@ bool FStateTreeTraceAnalyzer::OnEvent(const uint16 RouteId, EStyle Style, const
 		}
 	case RouteId_State:
 		{
-			const FStateTreeTraceStateEvent Event(
+			const FStateTreeTraceStateEvent Event(WorldTime,
 				EventData.GetValue<EStateTreeUpdatePhase>("Phase"),
 				FStateTreeIndex16(EventData.GetValue<uint16>("StateIndex")),
 				EventData.GetValue<EStateTreeTraceNodeEventType>("EventType"),
@@ -109,7 +116,7 @@ bool FStateTreeTraceAnalyzer::OnEvent(const uint16 RouteId, EStyle Style, const
 			Archive << TypePath;
 			Archive << DataAsText;
 
-			const FStateTreeTraceTaskEvent Event(
+			const FStateTreeTraceTaskEvent Event(WorldTime,
 				EventData.GetValue<EStateTreeUpdatePhase>("Phase"),
 				FStateTreeIndex16(EventData.GetValue<uint16>("NodeIndex")),
 				EventData.GetValue<EStateTreeTraceNodeEventType>("EventType"),
@@ -128,7 +135,7 @@ bool FStateTreeTraceAnalyzer::OnEvent(const uint16 RouteId, EStyle Style, const
 			Archive << TypePath;
 			Archive << DataAsText;
 
-			const FStateTreeTraceConditionEvent Event(
+			const FStateTreeTraceConditionEvent Event(WorldTime,
 				EventData.GetValue<EStateTreeUpdatePhase>("Phase"),
 				FStateTreeIndex16(EventData.GetValue<uint16>("NodeIndex")),
 				EventData.GetValue<EStateTreeTraceNodeEventType>("EventType"),
@@ -141,7 +148,7 @@ bool FStateTreeTraceAnalyzer::OnEvent(const uint16 RouteId, EStyle Style, const
 		}
 	case RouteId_Transition:
 		{
-			const FStateTreeTraceTransitionEvent Event(
+			const FStateTreeTraceTransitionEvent Event(WorldTime,
 				EventData.GetValue<EStateTreeUpdatePhase>("Phase"),
 				FStateTreeIndex16(EventData.GetValue<uint16>("TransitionIndex")),
 				EventData.GetValue<EStateTreeTraceNodeEventType>("EventType"));
@@ -153,7 +160,7 @@ bool FStateTreeTraceAnalyzer::OnEvent(const uint16 RouteId, EStyle Style, const
 		}
 	case RouteId_ActiveStates:
 		{
-			FStateTreeTraceActiveStatesEvent Event(EventData.GetValue<EStateTreeUpdatePhase>("Phase"));
+			FStateTreeTraceActiveStatesEvent Event(WorldTime, EventData.GetValue<EStateTreeUpdatePhase>("Phase"));
 			Event.ActiveStates = EventData.GetArrayView<uint16>("ActiveStates");
 
 			Provider.AppendEvent(FStateTreeInstanceDebugId(EventData.GetValue<uint32>("InstanceId"), EventData.GetValue<uint32>("InstanceSerial")),
diff --git a/Engine/Plugins/Runtime/StateTree/Source/StateTreeModule/Private/Debugger/StateTreeTraceModule.cpp b/Engine/Plugins/Runtime/StateTree/Source/StateTreeModule/Private/Debugger/StateTreeTraceModule.cpp
--- a/Engine/Plugins/Runtime/StateTree/Source/StateTreeModule/Private/Debugger/StateTreeTraceModule.cpp
+++ b/Engine/Plugins/Runtime/StateTree/Source/StateTreeModule/Private/Debugger/StateTreeTraceModule.cpp
@@ -5,7 +5,7 @@
 #include "Debugger/StateTreeTraceModule.h"
 #include "Debugger/StateTreeTraceProvider.h"
 #include "Debugger/StateTreeTraceAnalyzer.h"
-#include "Debugger/StateTreeDebugger.h"
+#include "Debugger/StateTreeDebugger.h"  // Required to compile TArray<UE::StateTreeDebugger::FInstanceDescriptor> from StateTreeTraceProvider
 #include "TraceServices/Model/AnalysisSession.h"
 
 FName FStateTreeTraceModule::ModuleName("TraceModule_StateTree");
@@ -13,20 +13,19 @@ FName FStateTreeTraceModule::ModuleName("TraceModule_StateTree");
 void FStateTreeTraceModule::GetModuleInfo(TraceServices::FModuleInfo& OutModuleInfo)
 {
 	OutModuleInfo.Name = ModuleName;
-	OutModuleInfo.DisplayName = TEXT("StateTreee");
+	OutModuleInfo.DisplayName = TEXT("StateTree");
 }
 
 void FStateTreeTraceModule::OnAnalysisBegin(TraceServices::IAnalysisSession& InSession)
 {
 	const TSharedPtr<FStateTreeTraceProvider> Provider = MakeShared<FStateTreeTraceProvider>(InSession);
 	InSession.AddProvider(FStateTreeTraceProvider::ProviderName, Provider);
-
 	InSession.AddAnalyzer(new FStateTreeTraceAnalyzer(InSession, *Provider));
 }
 
 void FStateTreeTraceModule::GetLoggers(TArray<const TCHAR *>& OutLoggers)
 {
-	OutLoggers.Add(TEXT("StateTreee"));
+	OutLoggers.Add(TEXT("StateTree"));
 }
 
 #endif // WITH_STATETREE_DEBUGGER
\ No newline at end of file
diff --git a/Engine/Plugins/Runtime/StateTree/Source/StateTreeModule/Private/Debugger/StateTreeTraceProvider.cpp b/Engine/Plugins/Runtime/StateTree/Source/StateTreeModule/Private/Debugger/StateTreeTraceProvider.cpp
--- a/Engine/Plugins/Runtime/StateTree/Source/StateTreeModule/Private/Debugger/StateTreeTraceProvider.cpp
+++ b/Engine/Plugins/Runtime/StateTree/Source/StateTreeModule/Private/Debugger/StateTreeTraceProvider.cpp
@@ -39,7 +39,7 @@ bool FStateTreeTraceProvider::ReadTimelines(const FStateTreeInstanceDebugId Inst
 			}
 		}
 
-		return EventsTimelines.Num() > 0;	
+		return EventsTimelines.Num() > 0;
 	}
 
 	return false;
@@ -67,8 +67,10 @@ void FStateTreeTraceProvider::AppendEvent(const FStateTreeInstanceDebugId InInst
 {
 	Session.WriteAccessCheck();
 
-	const uint32* IndexPtr = InstanceIdToDebuggerEntryTimelines.Find(InInstanceId);
-	if (ensureMsgf(IndexPtr != nullptr, TEXT("Timeline can't be found: probably caused by a missing instance event 'Started'.")))
+	// It is currently possible to receive events from an instance without receiving event `EStateTreeTraceInstanceEventType::Started` first
+	// (i.e. traces were activated after the statetree instance execution was started).    
+	// We plan to buffer the Instance events (Started/Stopped) to address this but for now we ignore events related to that instance.
+	if (const uint32* IndexPtr = InstanceIdToDebuggerEntryTimelines.Find(InInstanceId))
 	{
 		EventsTimelines[*IndexPtr]->AppendEvent(InTime, InEvent);
 	}
@@ -81,11 +83,12 @@ void FStateTreeTraceProvider::AppendInstanceEvent(
 	const FStateTreeInstanceDebugId InInstanceId,
 	const TCHAR* InInstanceName,
 	const double InTime,
+	const double InWorldRecordingTime,
 	const EStateTreeTraceInstanceEventType InEventType)
 {
 	if (InEventType == EStateTreeTraceInstanceEventType::Started)
 	{
-		Descriptors.Emplace(InStateTree, InInstanceId, InInstanceName, TRange<double>(InTime, std::numeric_limits<double>::max()));
+		Descriptors.Emplace(InStateTree, InInstanceId, InInstanceName, TRange<double>(InWorldRecordingTime, std::numeric_limits<double>::max()));
 
 		check(InstanceIdToDebuggerEntryTimelines.Find(InInstanceId) == nullptr);
 		InstanceIdToDebuggerEntryTimelines.Add(InInstanceId, EventsTimelines.Num());
@@ -94,11 +97,11 @@ void FStateTreeTraceProvider::AppendInstanceEvent(
 	}
 	else if (InEventType == EStateTreeTraceInstanceEventType::Stopped)
 	{
-		const uint32* Index = InstanceIdToDebuggerEntryTimelines.Find(InInstanceId);
-		if (ensureMsgf(Index != nullptr, TEXT("Unable to find matching instance. Looks like we never received event 'Started' for this instance.")))
+		// Process only if timeline can be found. See details in AppendEvent comment.
+		if (const uint32* Index = InstanceIdToDebuggerEntryTimelines.Find(InInstanceId))
 		{
 			check(Descriptors.IsValidIndex(*Index));
-			Descriptors[*Index].Lifetime.SetUpperBound(InTime);
+			Descriptors[*Index].Lifetime.SetUpperBound(InWorldRecordingTime);
 		}
 	}
 
diff --git a/Engine/Plugins/Runtime/StateTree/Source/StateTreeModule/Private/Debugger/StateTreeTraceTypes.cpp b/Engine/Plugins/Runtime/StateTree/Source/StateTreeModule/Private/Debugger/StateTreeTraceTypes.cpp
--- a/Engine/Plugins/Runtime/StateTree/Source/StateTreeModule/Private/Debugger/StateTreeTraceTypes.cpp
+++ b/Engine/Plugins/Runtime/StateTree/Source/StateTreeModule/Private/Debugger/StateTreeTraceTypes.cpp
@@ -110,7 +110,8 @@ FString FStateTreeTraceConditionEvent::ToString(const UStateTree& StateTree) con
 //----------------------------------------------------------------------//
 // FStateTreeTraceActiveStatesEvent
 //----------------------------------------------------------------------//
-FStateTreeTraceActiveStatesEvent::FStateTreeTraceActiveStatesEvent(const EStateTreeUpdatePhase Phase): FStateTreeTracePhaseEvent(Phase)
+FStateTreeTraceActiveStatesEvent::FStateTreeTraceActiveStatesEvent(const double RecordingWorldTime, const EStateTreeUpdatePhase Phase)
+	: FStateTreeTracePhaseEvent(RecordingWorldTime, Phase)
 {
 }
 
diff --git a/Engine/Plugins/Runtime/StateTree/Source/StateTreeModule/Private/StateTreeModule.cpp b/Engine/Plugins/Runtime/StateTree/Source/StateTreeModule/Private/StateTreeModule.cpp
--- a/Engine/Plugins/Runtime/StateTree/Source/StateTreeModule/Private/StateTreeModule.cpp
+++ b/Engine/Plugins/Runtime/StateTree/Source/StateTreeModule/Private/StateTreeModule.cpp
@@ -4,6 +4,7 @@
 #include "StateTreeTypes.h"
 
 #if WITH_STATETREE_DEBUGGER
+#include "Debugger/StateTreeTrace.h"
 #include "Debugger/StateTreeTraceModule.h"
 #include "Features/IModularFeatures.h"
 #include "HAL/IConsoleManager.h"
@@ -15,10 +16,6 @@
 #include "TraceServices/AnalysisService.h"
 #include "TraceServices/ITraceServicesModule.h"
 
-#if WITH_TRACE_STORE
-#include "Misc/Paths.h"
-#endif // WITH_TRACE_STORE
-
 #endif // WITH_STATETREE_DEBUGGER
 
 #define LOCTEXT_NAMESPACE "StateTree"
@@ -29,30 +26,31 @@ class FStateTreeModule : public IStateTreeModule
 	virtual void StartupModule() override;
 	virtual void ShutdownModule() override;
 
-	void StartTraces();
-	void StopTraces();
+	virtual void StartTraces() override;
+	virtual void StopTraces() override;
 
 #if WITH_STATETREE_DEBUGGER
 	/**
 	 * Gets the store client.
 	 */
-	virtual UE::Trace::FStoreClient* GetStoreClient() override  { return StoreClient.Get(); }
+	virtual UE::Trace::FStoreClient* GetStoreClient() override
+	{
+		if (!StoreClient.IsValid())
+		{
+			StoreClient = TUniquePtr<UE::Trace::FStoreClient>(UE::Trace::FStoreClient::Connect(TEXT("localhost")));
+		}
+		return StoreClient.Get();
+	}
 	
 	TSharedPtr<TraceServices::IAnalysisService> TraceAnalysisService;
 	TSharedPtr<TraceServices::IModuleService> TraceModuleService;
 
-	/** The location of the trace files managed by the trace store. */
-	FString StoreDir;
+	TArray<const FString> ChannelsToRestore;
 
 	/** The client used to connect to the trace store. */
 	TUniquePtr<UE::Trace::FStoreClient> StoreClient;
-
-#if WITH_TRACE_STORE
-	TUniquePtr<UE::Trace::FStoreService> StoreService;
-#endif // WITH_TRACE_STORE
 	
 	FStateTreeTraceModule StateTreeTraceModule;
-	FDelegateHandle StoreServiceHandle;
 
 	FAutoConsoleCommand EnableDebugger = FAutoConsoleCommand(
 		TEXT("statetree.enabledebugger"),
@@ -75,12 +73,37 @@ IMPLEMENT_MODULE(FStateTreeModule, StateTreeModule)
 
 void FStateTreeModule::StartupModule()
 {
+#if WITH_STATETREE_DEBUGGER
+	ITraceServicesModule& TraceServicesModule = FModuleManager::LoadModuleChecked<ITraceServicesModule>("TraceServices");
+	TraceAnalysisService = TraceServicesModule.GetAnalysisService();
+	TraceModuleService = TraceServicesModule.GetModuleService();
+
+	IModularFeatures::Get().RegisterModularFeature(TraceServices::ModuleFeatureName, &StateTreeTraceModule);
+
+	UE::StateTreeTrace::RegisterGlobalDelegates();
+#if !WITH_EDITOR
+	// We don't automatically start traces for Editor targets since we rely on the debugger
+	// to start recording either on user action or on PIE session start.
 	StartTraces();
+#endif // !WITH_EDITOR
+
+#endif // WITH_STATETREE_DEBUGGER
 }
 
 void FStateTreeModule::ShutdownModule()
 {
+#if WITH_STATETREE_DEBUGGER
 	StopTraces();
+
+	if (StoreClient.IsValid())
+	{
+		StoreClient.Release();
+	}
+
+	UE::StateTreeTrace::UnregisterGlobalDelegates();
+	
+	IModularFeatures::Get().UnregisterModularFeature(TraceServices::ModuleFeatureName, &StateTreeTraceModule);
+#endif // WITH_STATETREE_DEBUGGER
 }
 
 void FStateTreeModule::StartTraces()
@@ -91,17 +114,23 @@ void FStateTreeModule::StartTraces()
 		return;
 	}
 
-	ITraceServicesModule& TraceServicesModule = FModuleManager::LoadModuleChecked<ITraceServicesModule>("TraceServices");
-	TraceAnalysisService = TraceServicesModule.GetAnalysisService();
-	TraceModuleService = TraceServicesModule.GetModuleService();
-
-	IModularFeatures::Get().RegisterModularFeature(TraceServices::ModuleFeatureName, &StateTreeTraceModule);
+	const bool bAlreadyConnected = FTraceAuxiliary::IsConnected();
 
-	if (!FTraceAuxiliary::IsConnected())
+	// If trace is already connected let's keep track of enabled channels to restore them when we stop recording
+	if (bAlreadyConnected)
+	{
+		UE::Trace::EnumerateChannels([](const ANSICHAR* Name, const bool bIsEnabled, void* Channels)
+		{
+			TArray<FString>* ChannelsToRestore = static_cast<TArray<FString>*>(Channels); 
+			if (bIsEnabled)
+			{
+				ChannelsToRestore->Emplace(ANSI_TO_TCHAR(Name));
+			}		
+		}, &ChannelsToRestore);
+	}
+	else
 	{
-		// cpu tracing is enabled by default, but not connected.
-		// when not connected, disable all channels initially to avoid a whole bunch of extra cpu, memory, and disk overhead from processing all the extra default trace channels
-		////////////////////////////////////////////////////////////////////////////////
+		// Disable all channels and then enable only those we need to minimize trace file size.
 		UE::Trace::EnumerateChannels([](const ANSICHAR* ChannelName, const bool bEnabled, void*)
 			{
 				if (bEnabled)
@@ -110,72 +139,41 @@ void FStateTreeModule::StartTraces()
 					UE::Trace::ToggleChannel(ChannelNameFString.GetCharArray().GetData(), false);
 				}
 			}
-			, nullptr);
+		, nullptr);
 	}
 
 	UE::Trace::ToggleChannel(TEXT("StateTreeDebugChannel"), true);
 	UE::Trace::ToggleChannel(TEXT("FrameChannel"), true);
 
-	// FTraceAuxiliary::FOptions Options;
-	// Options.bExcludeTail = true;
-	// FTraceAuxiliary::Start(FTraceAuxiliary::EConnectionType::Network, TEXT("127.0.0.1"), TEXT(""), &Options, LogStateTree);
-	FTraceAuxiliary::Start(FTraceAuxiliary::EConnectionType::Network, TEXT("127.0.0.1"), TEXT(""), nullptr, LogStateTree);
-
-	// Conditionally create local store service after engine init (if someone doesn't beat us to it).
-	// This is temp until a more formal local server is done by the insights system.
-	StoreServiceHandle = FCoreDelegates::OnFEngineLoopInitComplete.AddLambda([this]
+	if (bAlreadyConnected == false)
 	{
-		LLM_SCOPE_BYNAME(TEXT("StateTree"));
-
-		if (!GetStoreClient())
-		{
-#if WITH_TRACE_STORE
-			UE_LOG(LogCore, Display, TEXT("StateTree module auto-connecting to internal trace server..."));
-
-			// Create the Store Service.
-			StoreDir = FPaths::ProjectSavedDir() / TEXT("TraceSessions");
-			UE::Trace::FStoreService::FDesc StoreServiceDesc;
-			StoreServiceDesc.StoreDir = *StoreDir;
-			StoreServiceDesc.RecorderPort = 0; // Let system decide port
-			StoreServiceDesc.ThreadCount = 2;
-			StoreService = TUniquePtr<UE::Trace::FStoreService>(UE::Trace::FStoreService::Create(StoreServiceDesc));
-			if (StoreService.IsValid())
-			{
-				// Connect to our newly created store
-				StoreClient = TUniquePtr<UE::Trace::FStoreClient>(UE::Trace::FStoreClient::Connect(TEXT("localhost")));
-				UE::Trace::SendTo(TEXT("localhost"), StoreService->GetRecorderPort());
-				
-				FCoreDelegates::OnPreExit.AddLambda([this]() { StoreService.Reset(); });
-			}
-#else
-			UE_LOG(LogCore, Display, TEXT("StateTree module auto-connecting to local trace server..."));
-
-			StoreClient = TUniquePtr<UE::Trace::FStoreClient>(UE::Trace::FStoreClient::Connect(TEXT("localhost")));
-			if (StoreClient.IsValid())
-			{
-				const UE::Trace::FStoreClient::FStatus* Status = StoreClient->GetStatus();
-				StoreDir = FString(Status->GetStoreDir());
-			}
-#endif // WITH_TRACE_STORE
-		}
-	});
+		FTraceAuxiliary::FOptions Options;
+		Options.bExcludeTail = true;
+		FTraceAuxiliary::Start(FTraceAuxiliary::EConnectionType::Network, TEXT("localhost"), TEXT(""), &Options, LogStateTree);
+	}
 #endif // WITH_STATETREE_DEBUGGER
 }
 
 void FStateTreeModule::StopTraces()
 {
 #if WITH_STATETREE_DEBUGGER
-	if (StoreServiceHandle.IsValid())
+
+	UE::Trace::ToggleChannel(TEXT("StateTreeDebugChannel"), false);
+	UE::Trace::ToggleChannel(TEXT("FrameChannel"), false);
+
+	// When we have channels to restore it also indicates that the trace were active
+	// so we only toggle the channels back (i.e. not calling FTraceAuxiliary::Stop)
+	if (ChannelsToRestore.Num() > 0)
 	{
-		if (FTraceAuxiliary::IsConnected())
+		for (const FString& ChannelName : ChannelsToRestore)
 		{
-			FTraceAuxiliary::Stop();
+			UE::Trace::ToggleChannel(ChannelName.GetCharArray().GetData(), true);
 		}
-		
-		IModularFeatures::Get().UnregisterModularFeature(TraceServices::ModuleFeatureName, &StateTreeTraceModule);
-		
-		FCoreDelegates::OnFEngineLoopInitComplete.Remove(StoreServiceHandle);
-		StoreServiceHandle.Reset();
+		ChannelsToRestore.Reset();
+	}
+	else
+	{
+		FTraceAuxiliary::Stop();
 	}
 #endif // WITH_STATETREE_DEBUGGER
 }
diff --git a/Engine/Plugins/Runtime/StateTree/Source/StateTreeModule/Public/Debugger/StateTreeDebugger.h b/Engine/Plugins/Runtime/StateTree/Source/StateTreeModule/Public/Debugger/StateTreeDebugger.h
--- a/Engine/Plugins/Runtime/StateTree/Source/StateTreeModule/Public/Debugger/StateTreeDebugger.h
+++ b/Engine/Plugins/Runtime/StateTree/Source/StateTreeModule/Public/Debugger/StateTreeDebugger.h
@@ -94,14 +94,28 @@ struct STATETREEMODULE_API FStateTreeDebugger : FTickableGameObject
 	 * @return Event collection associated to the provided Id or an invalid one if not found.
 	 */
 	const UE::StateTreeDebugger::FInstanceEventCollection& GetEventCollection(FStateTreeInstanceDebugId InstanceId) const;
-	
+
+	/** Returns the recording duration in world recorded time. */
 	double GetRecordingDuration() const { return RecordingDuration; }
+
+	/** Returns the duration of the analysis session. This is not related to the world simulation time. */
+	double GetAnalysisDuration() const { return AnalysisDuration; }
+
+	/** Returns the time (based on the recording duration) associated to the selected frame. */
 	double GetScrubTime() const	{ return ScrubState.ScrubTime; }
 	void SetScrubTime(double ScrubTime);
 
 	void GetLiveTraces(TArray<FTraceDescriptor>& OutTraceDescriptors) const;
-	void StartLastLiveSessionAnalysis();
+	
+	/**
+	 * Creates a new analysis session instance and loads the latest available trace that is live.
+	 * Replaces the current analysis session.
+	 * On failure, if 'RetryPollingDuration is > 0', will retry connecting every frame for 'RetryPollingDuration' seconds 
+	 * @param RetryPollingDuration - On failure, how many seconds to retry connecting during FStateTreeDebugger::Tick
+	 */
+	void StartLastLiveSessionAnalysis(float RetryPollingDuration = 1.0f);
 	void StartSessionAnalysis(const FTraceDescriptor& TraceDescriptor);
+	void StopAnalysis();
 	FTraceDescriptor GetSelectedTraceDescriptor() const { return ActiveSessionTraceDescriptor; }
 	FText GetSelectedTraceDescription() const;
 	
@@ -116,14 +130,12 @@ struct STATETREEMODULE_API FStateTreeDebugger : FTickableGameObject
 
 protected:
 	virtual void Tick(float DeltaTime) override;
-	virtual bool IsTickable() const override;
+	virtual bool IsTickable() const override { return true; }
 	virtual bool IsTickableWhenPaused() const override { return true; }
 	virtual bool IsTickableInEditor() const override { return true; }
 	virtual TStatId GetStatId() const override { RETURN_QUICK_DECLARE_CYCLE_STAT(FStateTreeDebugger, STATGROUP_Tickables); }
 	
 private:
-	void StopAnalysis();
-
 	void ReadTrace(double ScrubTime);
 	void ReadTrace(uint64 FrameIndex);
 	void ReadTrace(
@@ -177,13 +189,20 @@ private:
 	/** List of currently active states in the selected instance */
 	TArray<FStateTreeStateHandle> ActiveStates;
 
-	inline static constexpr double UnsetTime = -1;
+	/**
+	 * When auto-connecting on session start it is possible that a few frames are required for the tracing session to be accessible and connected to.
+	 * This is to keep track of the time window where we will retry.
+	 */
+	float RetryLoadLastLiveSessionTimer = 0.0f;
+	
+	/** Recording duration of the analysis session in world recorded time. */
+	double RecordingDuration = 0;
 
-	/** Recording duration of the analysis session. This is not related to the world simulation time. */
-	double RecordingDuration = UnsetTime;
+	/** Duration of the analysis session. This is not related to the world simulation time. */
+	double AnalysisDuration = 0;
 
 	/** Last time in the recording that we use to fetch events and we will use for the next read. */
-	double LastTraceReadTime = UnsetTime;
+	double LastTraceReadTime = 0;
 
 	/** Combined information regarding current scrub time (e.g. frame index, event collection index, etc.) */ 
 	UE::StateTreeDebugger::FScrubState ScrubState;
diff --git a/Engine/Plugins/Runtime/StateTree/Source/StateTreeModule/Public/Debugger/StateTreeDebuggerTypes.h b/Engine/Plugins/Runtime/StateTree/Source/StateTreeModule/Public/Debugger/StateTreeDebuggerTypes.h
--- a/Engine/Plugins/Runtime/StateTree/Source/StateTreeModule/Public/Debugger/StateTreeDebuggerTypes.h
+++ b/Engine/Plugins/Runtime/StateTree/Source/StateTreeModule/Public/Debugger/StateTreeDebuggerTypes.h
@@ -19,15 +19,22 @@ namespace UE::StateTreeDebugger
 struct FFrameSpan
 {
 	FFrameSpan() = default;
-	FFrameSpan(const TraceServices::FFrame& Frame, const uint32 EventIdx)
+	FFrameSpan(const TraceServices::FFrame& Frame, const double RecordingWorldTime, const int32 EventIdx)
 		: Frame(Frame)
+		, WorldTime(RecordingWorldTime)
 		, EventIdx(EventIdx)
 	{
 	}
 
+	double GetWorldTimeStart() const { return WorldTime; }
+	double GetWorldTimeEnd() const { return WorldTime + (Frame.EndTime - Frame.StartTime); }
+
 	/** Frame index in the analysis session */
 	TraceServices::FFrame Frame;
 
+	/** World simulation time associated to that Frame index */
+	double WorldTime = 0;
+
 	/** Index of the first event for that Frame index */
 	int32 EventIdx = INDEX_NONE;
 };
@@ -97,6 +104,18 @@ struct STATETREEMODULE_API FInstanceEventCollection
 	bool IsValid() const { return InstanceId.IsValid(); }
 	bool IsInvalid() const { return !IsValid(); }
 
+	struct FActiveStatesChangePair
+	{
+		FActiveStatesChangePair(const int32 SpanIndex, const int32 EventIndex)
+			: SpanIndex(SpanIndex),
+			  EventIndex(EventIndex)
+		{
+		}
+
+		int32 SpanIndex = INDEX_NONE;
+		int32 EventIndex = INDEX_NONE;
+	};
+
 	/** Id of the instance associated to the stored events. */
 	FStateTreeInstanceDebugId InstanceId;
 
@@ -107,7 +126,7 @@ struct STATETREEMODULE_API FInstanceEventCollection
 	TArray<FFrameSpan> FrameSpans;
 
 	/** Indices of span and event for frames with a change of activate states. */
-	TArray<TTuple<uint32, uint32>> ActiveStatesChanges;
+	TArray<FActiveStatesChangePair> ActiveStatesChanges;
 
 	/**
 	 * Returns the event collection associated to the currently selected instance.
@@ -119,6 +138,17 @@ struct STATETREEMODULE_API FInstanceEventCollection
 
 struct STATETREEMODULE_API FScrubState
 {
+	enum class EScrubTimeBoundState : uint8
+	{
+		Unset,
+		/** There are events but current time is before the first frame. */
+		BeforeLowerBound,
+		/** There are events and current time is within the frames received. */
+		InBounds,
+		/** There are events but current time is after the last frame. */
+		AfterHigherBound
+	};
+
 	explicit FScrubState(const TArray<FInstanceEventCollection>& EventCollections)
 		: EventCollections(EventCollections)
 	{
@@ -129,11 +159,11 @@ private:
 	
 public:
 	double ScrubTime = 0;
-	uint32 EventCollectionIndex = INDEX_NONE;
+	int32 EventCollectionIndex = INDEX_NONE;
 	uint64 TraceFrameIndex = INDEX_NONE;
-	uint32 FrameSpanIndex = INDEX_NONE;
-	uint32 PreviousFrameSpanIndex = INDEX_NONE;
-	uint32 ActiveStatesIndex = INDEX_NONE;
+	int32 FrameSpanIndex = INDEX_NONE;
+	int32 ActiveStatesIndex = INDEX_NONE;
+	EScrubTimeBoundState ScrubTimeBoundState = EScrubTimeBoundState::Unset;
 	
 
 	void SetScrubTime(double NewScrubTime);
@@ -142,7 +172,7 @@ public:
 	 * Indicates if the current scrub state points to a valid frame.
 	 * @return True if the frame index is set
 	 */
-	bool IsInBounds() const { return FrameSpanIndex != INDEX_NONE; }
+	bool IsInBounds() const { return ScrubTimeBoundState == EScrubTimeBoundState::InBounds; }
 
 	/**
 	 * Indicates if the current scrub state points to an active states entry in the event collection.
@@ -198,9 +228,9 @@ public:
 	const FInstanceEventCollection& GetEventCollection() const;
 
 private:
-	void SetFrameSpanIndex(uint32 NewFrameSpanIndex);
-	void SetActiveStatesIndex(uint32 NewActiveStatesIndex);
-	void UpdateActiveStatesIndex(uint32 SpanIndex);
+	void SetFrameSpanIndex(int32 NewFrameSpanIndex);
+	void SetActiveStatesIndex(int32 NewActiveStatesIndex);
+	void UpdateActiveStatesIndex(int32 SpanIndex);
 };
 
 } // UE::StateTreeDebugger
diff --git a/Engine/Plugins/Runtime/StateTree/Source/StateTreeModule/Public/Debugger/StateTreeTrace.h b/Engine/Plugins/Runtime/StateTree/Source/StateTreeModule/Public/Debugger/StateTreeTrace.h
--- a/Engine/Plugins/Runtime/StateTree/Source/StateTreeModule/Public/Debugger/StateTreeTrace.h
+++ b/Engine/Plugins/Runtime/StateTree/Source/StateTreeModule/Public/Debugger/StateTreeTrace.h
@@ -22,6 +22,8 @@ UE_TRACE_CHANNEL_EXTERN(StateTreeDebugChannel, STATETREEMODULE_API)
 
 namespace UE::StateTreeTrace
 {
+	void RegisterGlobalDelegates();
+	void UnregisterGlobalDelegates();
 	void OutputInstanceLifetimeEvent(FStateTreeInstanceDebugId InstanceId, const UStateTree* StateTree, const TCHAR* InstanceName, EStateTreeTraceInstanceEventType EventType);
 	void OutputLogEventTrace(FStateTreeInstanceDebugId InstanceId, EStateTreeUpdatePhase Phase, const TCHAR* Fmt, ...);
 	void OutputStateEventTrace(FStateTreeInstanceDebugId InstanceId, EStateTreeUpdatePhase Phase, FStateTreeStateHandle StateHandle, EStateTreeTraceNodeEventType EventType, EStateTreeStateSelectionBehavior SelectionBehavior);
diff --git a/Engine/Plugins/Runtime/StateTree/Source/StateTreeModule/Public/Debugger/StateTreeTraceAnalyzer.h b/Engine/Plugins/Runtime/StateTree/Source/StateTreeModule/Public/Debugger/StateTreeTraceAnalyzer.h
--- a/Engine/Plugins/Runtime/StateTree/Source/StateTreeModule/Public/Debugger/StateTreeTraceAnalyzer.h
+++ b/Engine/Plugins/Runtime/StateTree/Source/StateTreeModule/Public/Debugger/StateTreeTraceAnalyzer.h
@@ -21,6 +21,7 @@ public:
 private:
 	enum : uint16
 	{
+		RouteId_WorldTimestamp,
 		RouteId_Instance,
 		RouteId_LogMessage,
 		RouteId_State,
@@ -32,5 +33,6 @@ private:
 
 	TraceServices::IAnalysisSession& Session;
 	FStateTreeTraceProvider& Provider;
+	double WorldTime = 0;
 };
 #endif // WITH_STATETREE_DEBUGGER
\ No newline at end of file
diff --git a/Engine/Plugins/Runtime/StateTree/Source/StateTreeModule/Public/Debugger/StateTreeTraceProvider.h b/Engine/Plugins/Runtime/StateTree/Source/StateTreeModule/Public/Debugger/StateTreeTraceProvider.h
--- a/Engine/Plugins/Runtime/StateTree/Source/StateTreeModule/Public/Debugger/StateTreeTraceProvider.h
+++ b/Engine/Plugins/Runtime/StateTree/Source/StateTreeModule/Public/Debugger/StateTreeTraceProvider.h
@@ -25,8 +25,9 @@ public:
 		const UStateTree* InStateTree,
 		const FStateTreeInstanceDebugId InInstanceId,
 		const TCHAR* InInstanceName,
-		double Time,
-		EStateTreeTraceInstanceEventType EventType);
+		double InTime,
+		double InWorldRecordingTime,
+		EStateTreeTraceInstanceEventType InEventType);
 
 protected:
 	/** IStateTreeDebuggerProvider interface */
diff --git a/Engine/Plugins/Runtime/StateTree/Source/StateTreeModule/Public/Debugger/StateTreeTraceTypes.h b/Engine/Plugins/Runtime/StateTree/Source/StateTreeModule/Public/Debugger/StateTreeTraceTypes.h
--- a/Engine/Plugins/Runtime/StateTree/Source/StateTreeModule/Public/Debugger/StateTreeTraceTypes.h
+++ b/Engine/Plugins/Runtime/StateTree/Source/StateTreeModule/Public/Debugger/StateTreeTraceTypes.h
@@ -41,21 +41,23 @@ enum class EStateTreeTraceNodeEventType : uint8
 
 struct FStateTreeTracePhaseEvent
 {
-	explicit FStateTreeTracePhaseEvent(const EStateTreeUpdatePhase Phase)
-		: Phase(Phase)
+	explicit FStateTreeTracePhaseEvent(const double RecordingWorldTime, const EStateTreeUpdatePhase Phase)
+		: RecordingWorldTime(RecordingWorldTime)
+		, Phase(Phase)
 	{
 	}
 
 	static FString GetDataTypePath() { return TEXT(""); }
 	static FString GetDataAsText() { return TEXT(""); }
 
+	double RecordingWorldTime = 0;
 	EStateTreeUpdatePhase Phase;
 };
 
 struct FStateTreeTraceLogEvent : FStateTreeTracePhaseEvent
 {
-	explicit FStateTreeTraceLogEvent(const EStateTreeUpdatePhase Phase, const FString& Message)
-		: FStateTreeTracePhaseEvent(Phase)
+	explicit FStateTreeTraceLogEvent(const double RecordingWorldTime, const EStateTreeUpdatePhase Phase, const FString& Message)
+		: FStateTreeTracePhaseEvent(RecordingWorldTime, Phase)
 		, Message(Message)
 	{
 	}
@@ -67,8 +69,8 @@ struct FStateTreeTraceLogEvent : FStateTreeTracePhaseEvent
 
 struct FStateTreeTraceTransitionEvent : FStateTreeTracePhaseEvent
 {
-	explicit FStateTreeTraceTransitionEvent(const EStateTreeUpdatePhase Phase, const FStateTreeIndex16 Index, const EStateTreeTraceNodeEventType EventType)
-		: FStateTreeTracePhaseEvent(Phase)
+	explicit FStateTreeTraceTransitionEvent(const double RecordingWorldTime, const EStateTreeUpdatePhase Phase, const FStateTreeIndex16 Index, const EStateTreeTraceNodeEventType EventType)
+		: FStateTreeTracePhaseEvent(RecordingWorldTime, Phase)
 		, TransitionIndex(Index)
 		, EventType(EventType)
 	{
@@ -82,8 +84,8 @@ struct FStateTreeTraceTransitionEvent : FStateTreeTracePhaseEvent
 
 struct FStateTreeTraceNodeEvent : FStateTreeTracePhaseEvent
 {
-	explicit FStateTreeTraceNodeEvent(const EStateTreeUpdatePhase Phase, const FStateTreeIndex16 Index, const EStateTreeTraceNodeEventType EventType)
-		: FStateTreeTracePhaseEvent(Phase)
+	explicit FStateTreeTraceNodeEvent(const double RecordingWorldTime, const EStateTreeUpdatePhase Phase, const FStateTreeIndex16 Index, const EStateTreeTraceNodeEventType EventType)
+		: FStateTreeTracePhaseEvent(RecordingWorldTime, Phase)
 		, Index(Index)
 		, EventType(EventType)
 	{
@@ -97,8 +99,8 @@ struct FStateTreeTraceNodeEvent : FStateTreeTracePhaseEvent
 
 struct FStateTreeTraceStateEvent : FStateTreeTraceNodeEvent
 {
-	explicit FStateTreeTraceStateEvent(const EStateTreeUpdatePhase Phase, const FStateTreeIndex16 Index, const EStateTreeTraceNodeEventType EventType, const EStateTreeStateSelectionBehavior SelectionBehavior)
-		: FStateTreeTraceNodeEvent(Phase, Index, EventType)
+	explicit FStateTreeTraceStateEvent(const double RecordingWorldTime, const EStateTreeUpdatePhase Phase, const FStateTreeIndex16 Index, const EStateTreeTraceNodeEventType EventType, const EStateTreeStateSelectionBehavior SelectionBehavior)
+		: FStateTreeTraceNodeEvent(RecordingWorldTime, Phase, Index, EventType)
 		, SelectionBehavior(SelectionBehavior)
 	{
 	}
@@ -111,8 +113,8 @@ struct FStateTreeTraceStateEvent : FStateTreeTraceNodeEvent
 
 struct FStateTreeTraceTaskEvent : FStateTreeTraceNodeEvent
 {
-	explicit FStateTreeTraceTaskEvent(const EStateTreeUpdatePhase Phase, const FStateTreeIndex16 Index,	const EStateTreeTraceNodeEventType EventType, const EStateTreeRunStatus Status,	const FString& TypePath, const FString& InstanceDataAsText)
-		: FStateTreeTraceNodeEvent(Phase, Index, EventType)
+	explicit FStateTreeTraceTaskEvent(const double RecordingWorldTime, const EStateTreeUpdatePhase Phase, const FStateTreeIndex16 Index,	const EStateTreeTraceNodeEventType EventType, const EStateTreeRunStatus Status,	const FString& TypePath, const FString& InstanceDataAsText)
+		: FStateTreeTraceNodeEvent(RecordingWorldTime, Phase, Index, EventType)
 		, TypePath(TypePath)
 		, InstanceDataAsText(InstanceDataAsText)
 		, Status(Status)
@@ -131,8 +133,8 @@ struct FStateTreeTraceTaskEvent : FStateTreeTraceNodeEvent
 
 struct FStateTreeTraceConditionEvent : FStateTreeTraceNodeEvent
 {
-	explicit FStateTreeTraceConditionEvent(const EStateTreeUpdatePhase Phase, const FStateTreeIndex16 Index, const EStateTreeTraceNodeEventType EventType, const FString& TypePath, const FString& InstanceDataAsText)
-		: FStateTreeTraceNodeEvent(Phase, Index, EventType)
+	explicit FStateTreeTraceConditionEvent(const double RecordingWorldTime, const EStateTreeUpdatePhase Phase, const FStateTreeIndex16 Index, const EStateTreeTraceNodeEventType EventType, const FString& TypePath, const FString& InstanceDataAsText)
+		: FStateTreeTraceNodeEvent(RecordingWorldTime, Phase, Index, EventType)
 		, TypePath(TypePath)
 		, InstanceDataAsText(InstanceDataAsText)
 	{
@@ -150,7 +152,7 @@ struct FStateTreeTraceConditionEvent : FStateTreeTraceNodeEvent
 struct FStateTreeTraceActiveStatesEvent : FStateTreeTracePhaseEvent
 {
 	// Intentionally implemented in source file to compile 'TArray<FStateTreeStateHandle>' using only forward declaration.
-	explicit FStateTreeTraceActiveStatesEvent(const EStateTreeUpdatePhase Phase);
+	explicit FStateTreeTraceActiveStatesEvent(const double RecordingWorldTime, const EStateTreeUpdatePhase Phase);
 
 	STATETREEMODULE_API FString ToString(const UStateTree& StateTree) const;
 
diff --git a/Engine/Plugins/Runtime/StateTree/Source/StateTreeModule/Public/StateTreeModule.h b/Engine/Plugins/Runtime/StateTree/Source/StateTreeModule/Public/StateTreeModule.h
--- a/Engine/Plugins/Runtime/StateTree/Source/StateTreeModule/Public/StateTreeModule.h
+++ b/Engine/Plugins/Runtime/StateTree/Source/StateTreeModule/Public/StateTreeModule.h
@@ -38,6 +38,18 @@ public:
 		return FModuleManager::Get().IsModuleLoaded("StateTreeModule");
 	}
 
+	/**
+	 * Start tracing and enables StateTree debugging related channels (frame + statetree).
+	 * If traces are already active we keep track of all channels previously activated to restore them on stop.
+	 */
+	virtual void StartTraces() = 0;
+
+	/**
+	 * Stops the trace service if it was not already connected when StartTraces was called.
+	 * Restores previously enabled channels if necessary.
+	 */
+	virtual void StopTraces() = 0;
+
 #if WITH_STATETREE_DEBUGGER
 	/**
 	 * Gets the store client.
