Fix [[nodiscard]] error
#rnx

diff --git a/Engine/Plugins/Runtime/StateTree/Source/StateTreeModule/Private/StateTreeModule.cpp b/Engine/Plugins/Runtime/StateTree/Source/StateTreeModule/Private/StateTreeModule.cpp
--- a/Engine/Plugins/Runtime/StateTree/Source/StateTreeModule/Private/StateTreeModule.cpp
+++ b/Engine/Plugins/Runtime/StateTree/Source/StateTreeModule/Private/StateTreeModule.cpp
@@ -97,7 +97,7 @@ void FStateTreeModule::ShutdownModule()
 
 	if (StoreClient.IsValid())
 	{
-		StoreClient.Release();
+		StoreClient.Reset();
 	}
 
 	UE::StateTreeTrace::UnregisterGlobalDelegates();
