Fix missing dependency
#rnx

diff --git a/Engine/Plugins/Runtime/StateTree/Source/StateTreeModule/StateTreeModule.Build.cs b/Engine/Plugins/Runtime/StateTree/Source/StateTreeModule/StateTreeModule.Build.cs
--- a/Engine/Plugins/Runtime/StateTree/Source/StateTreeModule/StateTreeModule.Build.cs
+++ b/Engine/Plugins/Runtime/StateTree/Source/StateTreeModule/StateTreeModule.Build.cs
@@ -21,6 +21,7 @@ namespace UnrealBuildTool.Rules
 					"GameplayTags",
 					"StructUtils",
 					"StructUtilsEngine",
+					"UnrealEd"
 				}
 			);
 
