Fix non-editor builds
#rnx

diff --git a/Engine/Plugins/Runtime/StateTree/Source/StateTreeModule/StateTreeModule.Build.cs b/Engine/Plugins/Runtime/StateTree/Source/StateTreeModule/StateTreeModule.Build.cs
--- a/Engine/Plugins/Runtime/StateTree/Source/StateTreeModule/StateTreeModule.Build.cs
+++ b/Engine/Plugins/Runtime/StateTree/Source/StateTreeModule/StateTreeModule.Build.cs
@@ -20,8 +20,7 @@ namespace UnrealBuildTool.Rules
 					"AIModule",
 					"GameplayTags",
 					"StructUtils",
-					"StructUtilsEngine",
-					"UnrealEd"
+					"StructUtilsEngine"
 				}
 			);
 
@@ -33,6 +32,11 @@ namespace UnrealBuildTool.Rules
 
 			UnsafeTypeCastWarningLevel = WarningLevel.Error;
 
+			if (Target.bBuildEditor)
+			{
+				PublicDependencyModuleNames.Add("UnrealEd");
+			}
+
 			if (Target.Platform == UnrealTargetPlatform.Win64 && 
 				(Target.Configuration != UnrealTargetConfiguration.Shipping || Target.bBuildEditor))
 			{
