[StateTreeDebugger]
- Added trace event for Phase (Push + Pop) to make it easier to recreate the event hierarchy when analyzing the traces.
- Merged enums EStateTreeTraceInstanceEventType and EStateTreeTraceNodeEventType to EStateTreeTraceEventType which is mainly a list of verbs that could be reused for different events (States, Instances, Tasks, etc.)
- EStateTreeUpdatePhase is no longer used as flags and reduced to uint8
- Phase events are stacked when producing traces and sent only if meaningful events (Task, State, Transition, etc.) are sent during their scope. This is to avoid sending useless events when the StateTree is ticked without any changes
#rnx

diff --git a/Engine/Plugins/Runtime/StateTree/Source/StateTreeEditorModule/Private/Debugger/SStateTreeDebuggerView.cpp b/Engine/Plugins/Runtime/StateTree/Source/StateTreeEditorModule/Private/Debugger/SStateTreeDebuggerView.cpp
--- a/Engine/Plugins/Runtime/StateTree/Source/StateTreeEditorModule/Private/Debugger/SStateTreeDebuggerView.cpp
+++ b/Engine/Plugins/Runtime/StateTree/Source/StateTreeEditorModule/Private/Debugger/SStateTreeDebuggerView.cpp
@@ -739,128 +739,69 @@ void SStateTreeDebuggerView::OnDebuggerScrubStateChanged(const UE::StateTreeDebu
 	struct FParentInfo
 	{
 		TSharedPtr<UE::StateTreeDebugger::FEventTreeElement> Element;
-		FStateTreeStateHandle State;
+		EStateTreeUpdatePhase Phase;
 	};
-	TArray<FParentInfo, TInlineAllocator<8>> ParentStack;
-	EStateTreeUpdatePhase LastPhase = EStateTreeUpdatePhase::Unset;
-	EStateTreeUpdatePhase EventPhase = EStateTreeUpdatePhase::Unset;
+	TArray<FParentInfo, TInlineAllocator<8>> PhaseStack;
+	TArray<FStateTreeStateHandle, TInlineAllocator<8>> StateStack;
 
 	const int32 SpanIdx = ScrubState.FrameSpanIndex;
 	const int32 FirstEventIdx = Spans[SpanIdx].EventIdx;
 	const TraceServices::FFrame Frame = Spans[SpanIdx].Frame;
 	const int32 MaxEventIdx = Spans.IsValidIndex(SpanIdx+1) ? Spans[SpanIdx+1].EventIdx : Events.Num();
 	
-	auto CreateParentFunc = [this, &ParentStack, Frame](const EStateTreeUpdatePhase NodePhase, const FStateTreeStateHandle StateHandle)
-		{
-			// Create log event to describe the phase
-			FString NodeDesc;
-			if (StateHandle.IsValid())
-			{
-				const FCompactStateTreeState* CompactState = StateTree->GetStateFromHandle(StateHandle);
-				NodeDesc = FString::Printf(TEXT(" '%s'"), CompactState != nullptr ? *CompactState->Name.ToString() : *StateHandle.Describe());
-			}
-
-			FStateTreeTraceLogEvent LogEvent(/*RecordingWorldTime*/0, NodePhase, FString::Printf(TEXT("%s%s"), *UEnum::GetDisplayValueAsText(NodePhase).ToString(), *NodeDesc));
-
-			// Create Tree element to hold the event
-			const TSharedPtr<UE::StateTreeDebugger::FEventTreeElement> NewElement = MakeShareable(
-				new UE::StateTreeDebugger::FEventTreeElement(Frame, FStateTreeTraceEventVariantType(TInPlaceType<FStateTreeTraceLogEvent>(), LogEvent)));
-
-			// Push tree element to the proper hierarchy level
-			TArray<TSharedPtr<UE::StateTreeDebugger::FEventTreeElement>>& TreeElements = (ParentStack.IsEmpty()) ? EventsTreeElements : ParentStack.Top().Element->Children;
-			return ParentStack.Push({TreeElements.Add_GetRef(NewElement), StateHandle});
-		};
-	
 	for (int32 EventIdx = FirstEventIdx; EventIdx < MaxEventIdx; EventIdx++)
 	{
 		const FStateTreeTraceEventVariantType& Event = Events[EventIdx];
-		Visit([&EventPhase](auto& TypedEvent)
-		{
-			EventPhase = TypedEvent.Phase;
-		}, Event);
-
-		// Keep track of required state push/pop to be processed after phases
-		FStateTreeStateHandle StateToPush;
-		bool bPopState = false;
 
 		if (const FStateTreeTraceStateEvent* StateEvent = Event.TryGet<FStateTreeTraceStateEvent>())
 		{
-			if (StateEvent->EventType == EStateTreeTraceNodeEventType::PushStateSelection)
+			if (StateEvent->EventType == EStateTreeTraceEventType::Push)
 			{
-				StateToPush = StateEvent->GetStateHandle();
+				StateStack.Push(StateEvent->GetStateHandle());
+				continue;
 			}
-			else if (StateEvent->EventType == EStateTreeTraceNodeEventType::PopStateSelection)
+			
+			if (StateEvent->EventType == EStateTreeTraceEventType::Pop)
 			{
-				bPopState = true;
+				StateStack.Pop();
+				continue;
 			}
+			// Otherwise (other EventType) we add as a child
 		}
-
-		// Create a hierarchy level for each update phase
-		if (EventPhase != LastPhase)
+		else if (const FStateTreeTracePhaseEvent* PhaseEvent = Event.TryGet<FStateTreeTracePhaseEvent>())
 		{
-			const UEnum* PhaseEnum = StaticEnum<EStateTreeUpdatePhase>();
-			EStateTreeUpdatePhase PhasesDiff = EventPhase ^ LastPhase;
-			const int32 NumEnum = PhaseEnum->NumEnums();
-			check(NumEnum > 0);
-
-			// Pop phases first from last enum (every bit in the previous phase that differs must be popped)
-			if (EnumHasAnyFlags(LastPhase, PhasesDiff))
+			if (PhaseEvent->EventType == EStateTreeTraceEventType::Push)
 			{
-				for (int i = NumEnum - 1; i >= 0; --i)
+				// Create log event to describe the phase
+				FString NodeDesc;
+				const FStateTreeStateHandle StateHandle = StateStack.IsEmpty() ? FStateTreeStateHandle::Invalid : StateStack.Top();
+				if (StateHandle.IsValid())
 				{
-					const EStateTreeUpdatePhase Phase = static_cast<EStateTreeUpdatePhase>(PhaseEnum->GetValueByIndex(i));
-					if (EnumHasAnyFlags(LastPhase & PhasesDiff, Phase))
-					{
-						check(ParentStack.Num());
-						ParentStack.Pop();
-						EnumRemoveFlags(PhasesDiff, Phase);
-						if (EnumHasAnyFlags(LastPhase, PhasesDiff) == false)
-						{
-							break;
-						}
-					}
+					const FCompactStateTreeState* CompactState = StateTree->GetStateFromHandle(StateHandle);
+					NodeDesc = FString::Printf(TEXT(" '%s'"), CompactState != nullptr ? *CompactState->Name.ToString() : *StateHandle.Describe());
 				}
-			}
 
-			// Push required phases from first enum
-			if (EnumHasAnyFlags(EventPhase, PhasesDiff))
-			{
-				const FStateTreeStateHandle StateHandle = ParentStack.IsEmpty() ? FStateTreeStateHandle() : ParentStack.Top().State; 
-				for (int i = 0; i < NumEnum; ++i)
-				{
-					const EStateTreeUpdatePhase Phase = static_cast<EStateTreeUpdatePhase>(PhaseEnum->GetValueByIndex(i));
-					if (EnumHasAnyFlags(EventPhase & PhasesDiff, Phase))
-					{
-						CreateParentFunc(Phase, StateHandle);
-						
-						EnumRemoveFlags(PhasesDiff, Phase);
-						if (EnumHasAnyFlags(EventPhase, PhasesDiff) == false)
-						{
-							break;
-						}
-					}
-				}
-			}
-		}
+				FStateTreeTraceLogEvent LogEvent(/*RecordingWorldTime*/0, FString::Printf(TEXT("%s%s"), *UEnum::GetDisplayValueAsText(PhaseEvent->Phase).ToString(), *NodeDesc));
 
-		LastPhase = EventPhase;
+				// Create Tree element to hold the event
+				const TSharedPtr<UE::StateTreeDebugger::FEventTreeElement> NewElement = MakeShareable(
+					new UE::StateTreeDebugger::FEventTreeElement(Frame, FStateTreeTraceEventVariantType(TInPlaceType<FStateTreeTraceLogEvent>(), LogEvent)));
 
-		if (bPopState)
-		{
-			ParentStack.Pop();
-			// We use pop state to maintain the state stack but there is no element to visualize
-			continue;
-		}
-
-		if (StateToPush.IsValid())
-		{
-			// We use pushed state to create a new parent only
-			CreateParentFunc(EStateTreeUpdatePhase::StateSelection, StateToPush);
+				// Push tree element to the proper hierarchy level
+				TArray<TSharedPtr<UE::StateTreeDebugger::FEventTreeElement>>& TreeElements = (PhaseStack.IsEmpty()) ? EventsTreeElements : PhaseStack.Top().Element->Children;
+				PhaseStack.Push({TreeElements.Add_GetRef(NewElement), PhaseEvent->Phase});
+			}
+			else if (PhaseEvent->EventType == EStateTreeTraceEventType::Pop)
+			{
+				PhaseStack.Pop();
+			}
+			
+			// We don't want to create a child for phase events since they were already used to create a custom element.  
 			continue;
 		}
 
-		TArray<TSharedPtr<UE::StateTreeDebugger::FEventTreeElement>>& Elements = ParentStack.IsEmpty() ? EventsTreeElements : ParentStack.Top().Element->Children;
-        Elements.Add(MakeShareable(new UE::StateTreeDebugger::FEventTreeElement(Spans[SpanIdx].Frame, Event)));
+		TArray<TSharedPtr<UE::StateTreeDebugger::FEventTreeElement>>& TreeElements = PhaseStack.IsEmpty() ? EventsTreeElements : PhaseStack.Top().Element->Children;
+		TreeElements.Add(MakeShareable(new UE::StateTreeDebugger::FEventTreeElement(Spans[SpanIdx].Frame, Event)));
 	}
 }
 
diff --git a/Engine/Plugins/Runtime/StateTree/Source/StateTreeModule/Private/Debugger/StateTreeDebugger.cpp b/Engine/Plugins/Runtime/StateTree/Source/StateTreeModule/Private/Debugger/StateTreeDebugger.cpp
--- a/Engine/Plugins/Runtime/StateTree/Source/StateTreeModule/Private/Debugger/StateTreeDebugger.cpp
+++ b/Engine/Plugins/Runtime/StateTree/Source/StateTreeModule/Private/Debugger/StateTreeDebugger.cpp
@@ -634,7 +634,7 @@ bool FStateTreeDebugger::ProcessEvent(const FStateTreeInstanceDebugId InstanceId
 		&& (SelectedInstanceId == InstanceId || SelectedInstanceId.IsInvalid())) // allow breakpoints on any instances if not specified
 	{
 		const FStateTreeTraceStateEvent* StateEvent = Event.TryGet<FStateTreeTraceStateEvent>();
-		if (StateEvent != nullptr && StateEvent->EventType == EStateTreeTraceNodeEventType::OnEnter)
+		if (StateEvent != nullptr && StateEvent->EventType == EStateTreeTraceEventType::OnEnter)
 		{
 			const UStateTree* InstanceStateTree = nullptr;
 			if (SelectedInstanceId.IsInvalid())
diff --git a/Engine/Plugins/Runtime/StateTree/Source/StateTreeModule/Private/Debugger/StateTreeTrace.cpp b/Engine/Plugins/Runtime/StateTree/Source/StateTreeModule/Private/Debugger/StateTreeTrace.cpp
--- a/Engine/Plugins/Runtime/StateTree/Source/StateTreeModule/Private/Debugger/StateTreeTrace.cpp
+++ b/Engine/Plugins/Runtime/StateTree/Source/StateTreeModule/Private/Debugger/StateTreeTrace.cpp
@@ -33,11 +33,18 @@ UE_TRACE_EVENT_BEGIN(StateTreeDebugger, InstanceEvent)
 	UE_TRACE_EVENT_FIELD(uint8, EventType)
 UE_TRACE_EVENT_END()
 
-UE_TRACE_EVENT_BEGIN(StateTreeDebugger, LogEvent)
+UE_TRACE_EVENT_BEGIN(StateTreeDebugger, PhaseEvent)
 	UE_TRACE_EVENT_FIELD(uint64, Cycle)
 	UE_TRACE_EVENT_FIELD(uint32, InstanceId)
 	UE_TRACE_EVENT_FIELD(uint32, InstanceSerial)
 	UE_TRACE_EVENT_FIELD(uint16, Phase)
+	UE_TRACE_EVENT_FIELD(uint8, EventType)
+UE_TRACE_EVENT_END()
+
+UE_TRACE_EVENT_BEGIN(StateTreeDebugger, LogEvent)
+	UE_TRACE_EVENT_FIELD(uint64, Cycle)
+	UE_TRACE_EVENT_FIELD(uint32, InstanceId)
+	UE_TRACE_EVENT_FIELD(uint32, InstanceSerial)
 	UE_TRACE_EVENT_FIELD(UE::Trace::WideString, Message)
 UE_TRACE_EVENT_END()
 
@@ -45,7 +52,6 @@ UE_TRACE_EVENT_BEGIN(StateTreeDebugger, StateEvent)
 	UE_TRACE_EVENT_FIELD(uint64, Cycle)
 	UE_TRACE_EVENT_FIELD(uint32, InstanceId)
 	UE_TRACE_EVENT_FIELD(uint32, InstanceSerial)
-	UE_TRACE_EVENT_FIELD(uint16, Phase)
 	UE_TRACE_EVENT_FIELD(uint16, StateIndex)
 	UE_TRACE_EVENT_FIELD(uint8, EventType)
 	UE_TRACE_EVENT_FIELD(uint8, SelectionBehavior)
@@ -55,7 +61,6 @@ UE_TRACE_EVENT_BEGIN(StateTreeDebugger, TaskEvent)
 	UE_TRACE_EVENT_FIELD(uint64, Cycle)
 	UE_TRACE_EVENT_FIELD(uint32, InstanceId)
 	UE_TRACE_EVENT_FIELD(uint32, InstanceSerial)
-	UE_TRACE_EVENT_FIELD(uint16, Phase)
 	UE_TRACE_EVENT_FIELD(uint16, NodeIndex)
 	UE_TRACE_EVENT_FIELD(uint8[], DataView)
 	UE_TRACE_EVENT_FIELD(uint8, EventType)
@@ -66,7 +71,6 @@ UE_TRACE_EVENT_BEGIN(StateTreeDebugger, TransitionEvent)
 	UE_TRACE_EVENT_FIELD(uint64, Cycle)
 	UE_TRACE_EVENT_FIELD(uint32, InstanceId)
 	UE_TRACE_EVENT_FIELD(uint32, InstanceSerial)
-	UE_TRACE_EVENT_FIELD(uint16, Phase)
 	UE_TRACE_EVENT_FIELD(uint16, TransitionIndex)
 	UE_TRACE_EVENT_FIELD(uint8, EventType)
 UE_TRACE_EVENT_END()
@@ -75,7 +79,6 @@ UE_TRACE_EVENT_BEGIN(StateTreeDebugger, ConditionEvent)
 	UE_TRACE_EVENT_FIELD(uint64, Cycle)
 	UE_TRACE_EVENT_FIELD(uint32, InstanceId)
 	UE_TRACE_EVENT_FIELD(uint32, InstanceSerial)
-	UE_TRACE_EVENT_FIELD(uint16, Phase)
 	UE_TRACE_EVENT_FIELD(uint16, NodeIndex)
 	UE_TRACE_EVENT_FIELD(uint8[], DataView)
 	UE_TRACE_EVENT_FIELD(uint8, EventType)
@@ -85,7 +88,6 @@ UE_TRACE_EVENT_BEGIN(StateTreeDebugger, ActiveStatesEvent)
 	UE_TRACE_EVENT_FIELD(uint64, Cycle)
 	UE_TRACE_EVENT_FIELD(uint32, InstanceId)
 	UE_TRACE_EVENT_FIELD(uint32, InstanceSerial)
-	UE_TRACE_EVENT_FIELD(uint16, Phase)
 	UE_TRACE_EVENT_FIELD(uint16[], ActiveStates)
 UE_TRACE_EVENT_END()
 
@@ -101,6 +103,109 @@ FDelegateHandle GOnPIEStartDelegateHandle;
 double GRecordingWorldTime = -1;
 double GTracedRecordingWorldTime = -1;
 
+/** Struct to keep track if a given phase was traced or not. */
+struct FPhaseTraceStatusPair
+{
+	explicit FPhaseTraceStatusPair(const EStateTreeUpdatePhase Phase) : Phase(Phase) {}
+
+	EStateTreeUpdatePhase Phase = EStateTreeUpdatePhase::Unset;
+	bool bTraced = false;
+};
+
+/** Struct to keep track of the list of stacked phases for a given statetree instance. */
+struct FPhaseStack
+{
+	FStateTreeInstanceDebugId InstanceId;
+	TArray<FPhaseTraceStatusPair> Stack;
+};
+
+/**
+ * Stacks to keep track of all received phase events so other events will control when and if a given phase trace will be sent.
+ * This is per thread since it is possible to update execution contexts on multiple threads.
+ */
+thread_local TArray<FPhaseStack> GPhaseStacks;
+
+/**
+ * Called by TraceBufferedEvents from the OutputXYZ methods to flush pending phase events.
+ * Phases popped before TraceStackedPhases gets called will never produce any trace since
+ * they will not be required for the analysis.
+ */
+void TraceStackedPhases(const FStateTreeInstanceDebugId InstanceId)
+{
+	for (FPhaseStack& PhaseStack : GPhaseStacks)
+	{
+		if (PhaseStack.InstanceId == InstanceId)
+		{
+			for (FPhaseTraceStatusPair& StackEntry : PhaseStack.Stack)
+			{
+				if (StackEntry.bTraced == false)
+				{
+					UE_TRACE_LOG(StateTreeDebugger, PhaseEvent, StateTreeDebugChannel)
+						<< PhaseEvent.Cycle(FPlatformTime::Cycles64())
+						<< PhaseEvent.InstanceId(InstanceId.Id)
+						<< PhaseEvent.InstanceSerial(InstanceId.SerialNumber)
+						<< PhaseEvent.Phase(static_cast<std::underlying_type_t<EStateTreeUpdatePhase>>(StackEntry.Phase))
+						<< PhaseEvent.EventType(static_cast<std::underlying_type_t<EStateTreeTraceEventType>>(EStateTreeTraceEventType::Push));
+			
+					StackEntry.bTraced = true;
+				}
+			}
+			break;
+		}
+	}
+}
+
+/**
+ * Pushed or pops an entry on the Phase stack for a given Instance.
+ * Will send the Pop events for phases popped if their associated Push events were sent.
+ */
+void ProcessPhaseScopeEvent(const FStateTreeInstanceDebugId InstanceId, const EStateTreeUpdatePhase Phase, const EStateTreeTraceEventType EventType)
+{
+	int32 ExistingStackIndex = GPhaseStacks.IndexOfByPredicate([InstanceId](const FPhaseStack& PhaseStack){ return PhaseStack.InstanceId == InstanceId; });
+
+	if (EventType == EStateTreeTraceEventType::Push)
+	{
+		if (ExistingStackIndex == INDEX_NONE)
+		{
+			ExistingStackIndex = GPhaseStacks.AddDefaulted();
+		}
+		FPhaseStack& PhaseStack = GPhaseStacks[ExistingStackIndex];
+		PhaseStack.InstanceId = InstanceId;
+		PhaseStack.Stack.Push(FPhaseTraceStatusPair(Phase));
+	}
+	else if (ensureMsgf(ExistingStackIndex != INDEX_NONE, TEXT("Not expected to pop phases for an instance that never pushed a phase.")))
+	{
+		FPhaseStack& PhaseStack = GPhaseStacks[ExistingStackIndex];
+
+		if (ensureMsgf(PhaseStack.Stack.IsEmpty() == false, TEXT("Not expected to pop phases that never got pushed.")) &&
+			ensureMsgf(PhaseStack.InstanceId == InstanceId, TEXT("Not expected to pop phases for an instance that is not the one currently assigned to the stack.")))
+		{
+			const FPhaseTraceStatusPair RemovedPair = PhaseStack.Stack.Pop();
+			ensureMsgf(RemovedPair.Phase == Phase, TEXT("Not expected to pop a phase that is not on the top of the stack."));
+
+			// Clear associated InstanceId when removing last entry from the stack.
+			if (PhaseStack.Stack.IsEmpty())
+			{
+				GPhaseStacks.RemoveAt(ExistingStackIndex);
+			}
+
+			// Phase was previously traced (i.e. other events were traced in that scope so we need to trace the closing (i.e. Pop) event.
+			if (RemovedPair.bTraced)
+			{
+				UE_TRACE_LOG(StateTreeDebugger, PhaseEvent, StateTreeDebugChannel)
+				<< PhaseEvent.Cycle(FPlatformTime::Cycles64())
+				<< PhaseEvent.InstanceId(InstanceId.Id)
+				<< PhaseEvent.InstanceSerial(InstanceId.SerialNumber)
+				<< PhaseEvent.Phase(static_cast<std::underlying_type_t<EStateTreeUpdatePhase>>(Phase))
+				<< PhaseEvent.EventType(static_cast<std::underlying_type_t<EStateTreeTraceEventType>>(EStateTreeTraceEventType::Pop));
+			}
+		}
+	}
+}
+
+/**
+ * Called by TraceBufferedEvents from the OutputXYZ methods to make sure we have the current world time was sent.
+ */
 void TraceTimeMapping()
 {
 	if (GTracedRecordingWorldTime != GRecordingWorldTime)
@@ -111,6 +216,15 @@ void TraceTimeMapping()
 	}
 }
 
+/**
+ * Called by the OutputXYZ methods to flush pending events (e.g. Push or WorldTime).
+ */
+void TraceBufferedEvents(const FStateTreeInstanceDebugId InstanceId)
+{
+	TraceTimeMapping();
+	TraceStackedPhases(InstanceId);
+}
+
 void SerializeDataViewToArchive(FBufferArchive Ar, const FStateTreeDataView DataView)
 {
 	if (const UScriptStruct* ScriptStruct = Cast<const UScriptStruct>(DataView.GetStruct()))
@@ -171,7 +285,7 @@ void OutputInstanceLifetimeEvent(
 	const FStateTreeInstanceDebugId InstanceId,
 	const UStateTree* StateTree,
 	const TCHAR* InstanceName,
-	const EStateTreeTraceInstanceEventType EventType
+	const EStateTreeTraceEventType EventType
 	)
 {
 	if (StateTree == nullptr)
@@ -181,7 +295,7 @@ void OutputInstanceLifetimeEvent(
 	const FString ObjectName = StateTree->GetName();
 	const FString ObjectPackageName = StateTree->GetPackage()->GetName();
 
-	TraceTimeMapping();
+	TraceBufferedEvents(InstanceId);
 
 	UE_TRACE_LOG(StateTreeDebugger, InstanceEvent, StateTreeDebugChannel)
 		<< InstanceEvent.Cycle(FPlatformTime::Cycles64())
@@ -191,120 +305,109 @@ void OutputInstanceLifetimeEvent(
 		<< InstanceEvent.InstanceId(InstanceId.Id)
 		<< InstanceEvent.InstanceSerial(InstanceId.SerialNumber)
 		<< InstanceEvent.InstanceName(InstanceName)
-		<< InstanceEvent.EventType(static_cast<std::underlying_type_t<EStateTreeTraceNodeEventType>>(EventType));
+		<< InstanceEvent.EventType(static_cast<std::underlying_type_t<EStateTreeTraceEventType>>(EventType));
 }
 
 void OutputLogEventTrace(
 	const FStateTreeInstanceDebugId InstanceId,
-	const EStateTreeUpdatePhase Phase,
 	const TCHAR* Fmt, ...
 	)
 {
 	static TCHAR TraceStaticBuffer[8192];
 	GET_VARARGS(TraceStaticBuffer, UE_ARRAY_COUNT(TraceStaticBuffer), UE_ARRAY_COUNT(TraceStaticBuffer) - 1, Fmt, Fmt);
 
-	TraceTimeMapping();
+	TraceBufferedEvents(InstanceId);
 
 	UE_TRACE_LOG(StateTreeDebugger, LogEvent, StateTreeDebugChannel)
 		<< LogEvent.Cycle(FPlatformTime::Cycles64())
 		<< LogEvent.InstanceId(InstanceId.Id)
 		<< LogEvent.InstanceSerial(InstanceId.SerialNumber)
-		<< LogEvent.Phase(static_cast<std::underlying_type_t<EStateTreeUpdatePhase>>(Phase))
 		<< LogEvent.Message(TraceStaticBuffer);
 }
 
 void OutputStateEventTrace(
 	const FStateTreeInstanceDebugId InstanceId,
-	const EStateTreeUpdatePhase Phase,
 	const FStateTreeStateHandle StateHandle,
-	const EStateTreeTraceNodeEventType EventType,
+	const EStateTreeTraceEventType EventType,
 	const EStateTreeStateSelectionBehavior SelectionBehavior
 	)
 {
-	TraceTimeMapping();
+	TraceBufferedEvents(InstanceId);
 
 	UE_TRACE_LOG(StateTreeDebugger, StateEvent, StateTreeDebugChannel)
 		<< StateEvent.Cycle(FPlatformTime::Cycles64())
 		<< StateEvent.InstanceId(InstanceId.Id)
 		<< StateEvent.InstanceSerial(InstanceId.SerialNumber)
-		<< StateEvent.Phase(static_cast<std::underlying_type_t<EStateTreeUpdatePhase>>(Phase))
 		<< StateEvent.StateIndex(StateHandle.Index)
-		<< StateEvent.EventType(static_cast<std::underlying_type_t<EStateTreeTraceNodeEventType>>(EventType))
+		<< StateEvent.EventType(static_cast<std::underlying_type_t<EStateTreeTraceEventType>>(EventType))
 		<< StateEvent.SelectionBehavior(static_cast<std::underlying_type_t<EStateTreeStateSelectionBehavior>>(SelectionBehavior));
 }
 
 void OutputTaskEventTrace(
 	const FStateTreeInstanceDebugId InstanceId,
-	const EStateTreeUpdatePhase Phase,
 	const FStateTreeIndex16 TaskIdx,
 	const FStateTreeDataView DataView,
-	const EStateTreeTraceNodeEventType EventType,
+	const EStateTreeTraceEventType EventType,
 	const EStateTreeRunStatus Status
 	)
 {
 	FBufferArchive Archive;
 	SerializeDataViewToArchive(Archive, DataView);
 
-	TraceTimeMapping();
+	TraceBufferedEvents(InstanceId);
 
 	UE_TRACE_LOG(StateTreeDebugger, TaskEvent, StateTreeDebugChannel)
 		<< TaskEvent.Cycle(FPlatformTime::Cycles64())
 		<< TaskEvent.InstanceId(InstanceId.Id)
 		<< TaskEvent.InstanceSerial(InstanceId.SerialNumber)
-		<< TaskEvent.Phase(static_cast<std::underlying_type_t<EStateTreeUpdatePhase>>(Phase))
 		<< TaskEvent.NodeIndex(TaskIdx.Get())
 		<< TaskEvent.DataView(Archive.GetData(), Archive.Num())
-		<< TaskEvent.EventType(static_cast<std::underlying_type_t<EStateTreeTraceNodeEventType>>(EventType))
+		<< TaskEvent.EventType(static_cast<std::underlying_type_t<EStateTreeTraceEventType>>(EventType))
 		<< TaskEvent.Status(static_cast<std::underlying_type_t<EStateTreeRunStatus>>(Status));
 }
 
 void OutputTransitionEventTrace(
 	const FStateTreeInstanceDebugId InstanceId,
-	const EStateTreeUpdatePhase Phase,
 	const FStateTreeIndex16 TransitionIdx,
-	const EStateTreeTraceNodeEventType EventType
+	const EStateTreeTraceEventType EventType
 	)
 {
 	FBufferArchive Archive;
 	Archive << EventType;
 
-	TraceTimeMapping();
+	TraceBufferedEvents(InstanceId);
 
 	UE_TRACE_LOG(StateTreeDebugger, TransitionEvent, StateTreeDebugChannel)
 	<< TransitionEvent.Cycle(FPlatformTime::Cycles64())
 	<< TransitionEvent.InstanceId(InstanceId.Id)
 	<< TransitionEvent.InstanceSerial(InstanceId.SerialNumber)
-	<< TransitionEvent.Phase(static_cast<std::underlying_type_t<EStateTreeUpdatePhase>>(Phase))
 	<< TransitionEvent.TransitionIndex(TransitionIdx.Get())
-	<< TransitionEvent.EventType(static_cast<std::underlying_type_t<EStateTreeTraceNodeEventType>>(EventType));
+	<< TransitionEvent.EventType(static_cast<std::underlying_type_t<EStateTreeTraceEventType>>(EventType));
 }
 
 void OutputConditionEventTrace(
 	const FStateTreeInstanceDebugId InstanceId,
-	const EStateTreeUpdatePhase Phase,
 	const FStateTreeIndex16 ConditionIdx,
 	const FStateTreeDataView DataView,
-	const EStateTreeTraceNodeEventType EventType
+	const EStateTreeTraceEventType EventType
 	)
 {
 	FBufferArchive Archive;
 	SerializeDataViewToArchive(Archive, DataView);
 
-	TraceTimeMapping();
+	TraceBufferedEvents(InstanceId);
 
 	UE_TRACE_LOG(StateTreeDebugger, ConditionEvent, StateTreeDebugChannel)
 	<< ConditionEvent.Cycle(FPlatformTime::Cycles64())
 	<< ConditionEvent.InstanceId(InstanceId.Id)
 	<< ConditionEvent.InstanceSerial(InstanceId.SerialNumber)
-	<< ConditionEvent.Phase(static_cast<std::underlying_type_t<EStateTreeUpdatePhase>>(Phase))
 	<< ConditionEvent.NodeIndex(ConditionIdx.Get())
 	<< ConditionEvent.DataView(Archive.GetData(), Archive.Num())
-	<< ConditionEvent.EventType(static_cast<std::underlying_type_t<EStateTreeTraceNodeEventType>>(EventType));
+	<< ConditionEvent.EventType(static_cast<std::underlying_type_t<EStateTreeTraceEventType>>(EventType));
 }
 
 void OutputActiveStatesEventTrace(
 	const FStateTreeInstanceDebugId InstanceId,
-	const EStateTreeUpdatePhase Phase,
 	const FStateTreeActiveStates& ActiveStates
 	)
 {
@@ -314,13 +417,12 @@ void OutputActiveStatesEventTrace(
 		StatesIndices.Add(ActiveStates[i].Index);
 	}
 
-	TraceTimeMapping();
+	TraceBufferedEvents(InstanceId);
 
 	UE_TRACE_LOG(StateTreeDebugger, ActiveStatesEvent, StateTreeDebugChannel)
 		<< ActiveStatesEvent.Cycle(FPlatformTime::Cycles64())
 		<< ActiveStatesEvent.InstanceId(InstanceId.Id)
 		<< ActiveStatesEvent.InstanceSerial(InstanceId.SerialNumber)
-		<< ActiveStatesEvent.Phase(static_cast<std::underlying_type_t<EStateTreeUpdatePhase>>(Phase))
 		<< ActiveStatesEvent.ActiveStates(StatesIndices.GetData(), StatesIndices.Num());
 }
 
diff --git a/Engine/Plugins/Runtime/StateTree/Source/StateTreeModule/Private/Debugger/StateTreeTraceAnalyzer.cpp b/Engine/Plugins/Runtime/StateTree/Source/StateTreeModule/Private/Debugger/StateTreeTraceAnalyzer.cpp
--- a/Engine/Plugins/Runtime/StateTree/Source/StateTreeModule/Private/Debugger/StateTreeTraceAnalyzer.cpp
+++ b/Engine/Plugins/Runtime/StateTree/Source/StateTreeModule/Private/Debugger/StateTreeTraceAnalyzer.cpp
@@ -21,6 +21,7 @@ void FStateTreeTraceAnalyzer::OnAnalysisBegin(const FOnAnalysisContext& Context)
 
 	Builder.RouteEvent(RouteId_WorldTimestamp, "StateTreeDebugger", "WorldTimestampEvent");
 	Builder.RouteEvent(RouteId_Instance, "StateTreeDebugger", "InstanceEvent");
+	Builder.RouteEvent(RouteId_Phase, "StateTreeDebugger", "PhaseEvent");
 	Builder.RouteEvent(RouteId_LogMessage, "StateTreeDebugger", "LogEvent");
 	Builder.RouteEvent(RouteId_State, "StateTreeDebugger", "StateEvent");
 	Builder.RouteEvent(RouteId_Task, "StateTreeDebugger", "TaskEvent");
@@ -72,7 +73,7 @@ bool FStateTreeTraceAnalyzer::OnEvent(const uint16 RouteId, EStyle Style, const
 						*InstanceName,
 						Context.EventTime.AsSeconds(EventData.GetValue<uint64>("Cycle")),
 						WorldTime,
-						EventData.GetValue<EStateTreeTraceInstanceEventType>("EventType"));
+						EventData.GetValue<EStateTreeTraceEventType>("EventType"));
 				}
 				else
 				{
@@ -85,11 +86,22 @@ bool FStateTreeTraceAnalyzer::OnEvent(const uint16 RouteId, EStyle Style, const
 			}
 			break;
 		}
+	case RouteId_Phase:
+		{
+			const FStateTreeTracePhaseEvent Event(WorldTime,
+				EventData.GetValue<EStateTreeUpdatePhase>("Phase"),
+				EventData.GetValue<EStateTreeTraceEventType>("EventType"));
+
+			Provider.AppendEvent(FStateTreeInstanceDebugId(EventData.GetValue<uint32>("InstanceId"), EventData.GetValue<uint32>("InstanceSerial")),
+				Context.EventTime.AsSeconds(EventData.GetValue<uint64>("Cycle")),
+				FStateTreeTraceEventVariantType(TInPlaceType<FStateTreeTracePhaseEvent>(), Event));
+			break;
+		}
 	case RouteId_LogMessage:
 		{
 			FString Message;
         	EventData.GetString("Message", Message);
-			const FStateTreeTraceLogEvent Event(WorldTime, EventData.GetValue<EStateTreeUpdatePhase>("Phase"), Message);
+			const FStateTreeTraceLogEvent Event(WorldTime, Message);
 
 			Provider.AppendEvent(FStateTreeInstanceDebugId(EventData.GetValue<uint32>("InstanceId"), EventData.GetValue<uint32>("InstanceSerial")),
 				Context.EventTime.AsSeconds(EventData.GetValue<uint64>("Cycle")),
@@ -99,9 +111,8 @@ bool FStateTreeTraceAnalyzer::OnEvent(const uint16 RouteId, EStyle Style, const
 	case RouteId_State:
 		{
 			const FStateTreeTraceStateEvent Event(WorldTime,
-				EventData.GetValue<EStateTreeUpdatePhase>("Phase"),
 				FStateTreeIndex16(EventData.GetValue<uint16>("StateIndex")),
-				EventData.GetValue<EStateTreeTraceNodeEventType>("EventType"),
+				EventData.GetValue<EStateTreeTraceEventType>("EventType"),
 				EventData.GetValue<EStateTreeStateSelectionBehavior>("SelectionBehavior"));
 
 			Provider.AppendEvent(FStateTreeInstanceDebugId(EventData.GetValue<uint32>("InstanceId"), EventData.GetValue<uint32>("InstanceSerial")),
@@ -117,9 +128,8 @@ bool FStateTreeTraceAnalyzer::OnEvent(const uint16 RouteId, EStyle Style, const
 			Archive << DataAsText;
 
 			const FStateTreeTraceTaskEvent Event(WorldTime,
-				EventData.GetValue<EStateTreeUpdatePhase>("Phase"),
 				FStateTreeIndex16(EventData.GetValue<uint16>("NodeIndex")),
-				EventData.GetValue<EStateTreeTraceNodeEventType>("EventType"),
+				EventData.GetValue<EStateTreeTraceEventType>("EventType"),
 				EventData.GetValue<EStateTreeRunStatus>("Status"),
 				TypePath, DataAsText);
 
@@ -136,9 +146,8 @@ bool FStateTreeTraceAnalyzer::OnEvent(const uint16 RouteId, EStyle Style, const
 			Archive << DataAsText;
 
 			const FStateTreeTraceConditionEvent Event(WorldTime,
-				EventData.GetValue<EStateTreeUpdatePhase>("Phase"),
 				FStateTreeIndex16(EventData.GetValue<uint16>("NodeIndex")),
-				EventData.GetValue<EStateTreeTraceNodeEventType>("EventType"),
+				EventData.GetValue<EStateTreeTraceEventType>("EventType"),
 				TypePath, DataAsText);
 			
 			Provider.AppendEvent(FStateTreeInstanceDebugId(EventData.GetValue<uint32>("InstanceId"), EventData.GetValue<uint32>("InstanceSerial")),
@@ -149,9 +158,8 @@ bool FStateTreeTraceAnalyzer::OnEvent(const uint16 RouteId, EStyle Style, const
 	case RouteId_Transition:
 		{
 			const FStateTreeTraceTransitionEvent Event(WorldTime,
-				EventData.GetValue<EStateTreeUpdatePhase>("Phase"),
 				FStateTreeIndex16(EventData.GetValue<uint16>("TransitionIndex")),
-				EventData.GetValue<EStateTreeTraceNodeEventType>("EventType"));
+				EventData.GetValue<EStateTreeTraceEventType>("EventType"));
 			
 			Provider.AppendEvent(FStateTreeInstanceDebugId(EventData.GetValue<uint32>("InstanceId"), EventData.GetValue<uint32>("InstanceSerial")),
 				Context.EventTime.AsSeconds(EventData.GetValue<uint64>("Cycle")),
@@ -160,7 +168,7 @@ bool FStateTreeTraceAnalyzer::OnEvent(const uint16 RouteId, EStyle Style, const
 		}
 	case RouteId_ActiveStates:
 		{
-			FStateTreeTraceActiveStatesEvent Event(WorldTime, EventData.GetValue<EStateTreeUpdatePhase>("Phase"));
+			FStateTreeTraceActiveStatesEvent Event(WorldTime);
 			Event.ActiveStates = EventData.GetArrayView<uint16>("ActiveStates");
 
 			Provider.AppendEvent(FStateTreeInstanceDebugId(EventData.GetValue<uint32>("InstanceId"), EventData.GetValue<uint32>("InstanceSerial")),
diff --git a/Engine/Plugins/Runtime/StateTree/Source/StateTreeModule/Private/Debugger/StateTreeTraceProvider.cpp b/Engine/Plugins/Runtime/StateTree/Source/StateTreeModule/Private/Debugger/StateTreeTraceProvider.cpp
--- a/Engine/Plugins/Runtime/StateTree/Source/StateTreeModule/Private/Debugger/StateTreeTraceProvider.cpp
+++ b/Engine/Plugins/Runtime/StateTree/Source/StateTreeModule/Private/Debugger/StateTreeTraceProvider.cpp
@@ -67,7 +67,7 @@ void FStateTreeTraceProvider::AppendEvent(const FStateTreeInstanceDebugId InInst
 {
 	Session.WriteAccessCheck();
 
-	// It is currently possible to receive events from an instance without receiving event `EStateTreeTraceInstanceEventType::Started` first
+	// It is currently possible to receive events from an instance without receiving event `EStateTreeTraceEventType::Push` first
 	// (i.e. traces were activated after the statetree instance execution was started).    
 	// We plan to buffer the Instance events (Started/Stopped) to address this but for now we ignore events related to that instance.
 	if (const uint32* IndexPtr = InstanceIdToDebuggerEntryTimelines.Find(InInstanceId))
@@ -84,9 +84,9 @@ void FStateTreeTraceProvider::AppendInstanceEvent(
 	const TCHAR* InInstanceName,
 	const double InTime,
 	const double InWorldRecordingTime,
-	const EStateTreeTraceInstanceEventType InEventType)
+	const EStateTreeTraceEventType InEventType)
 {
-	if (InEventType == EStateTreeTraceInstanceEventType::Started)
+	if (InEventType == EStateTreeTraceEventType::Push)
 	{
 		Descriptors.Emplace(InStateTree, InInstanceId, InInstanceName, TRange<double>(InWorldRecordingTime, std::numeric_limits<double>::max()));
 
@@ -95,7 +95,7 @@ void FStateTreeTraceProvider::AppendInstanceEvent(
 		
 		EventsTimelines.Emplace(MakeShared<TraceServices::TPointTimeline<FStateTreeTraceEventVariantType>>(Session.GetLinearAllocator()));
 	}
-	else if (InEventType == EStateTreeTraceInstanceEventType::Stopped)
+	else if (InEventType == EStateTreeTraceEventType::Pop)
 	{
 		// Process only if timeline can be found. See details in AppendEvent comment.
 		if (const uint32* Index = InstanceIdToDebuggerEntryTimelines.Find(InInstanceId))
diff --git a/Engine/Plugins/Runtime/StateTree/Source/StateTreeModule/Private/Debugger/StateTreeTraceTypes.cpp b/Engine/Plugins/Runtime/StateTree/Source/StateTreeModule/Private/Debugger/StateTreeTraceTypes.cpp
--- a/Engine/Plugins/Runtime/StateTree/Source/StateTreeModule/Private/Debugger/StateTreeTraceTypes.cpp
+++ b/Engine/Plugins/Runtime/StateTree/Source/StateTreeModule/Private/Debugger/StateTreeTraceTypes.cpp
@@ -7,6 +7,18 @@
 
 #if WITH_STATETREE_DEBUGGER
 
+
+//----------------------------------------------------------------------//
+// FStateTreeTracePhaseEvent
+//----------------------------------------------------------------------//
+FString FStateTreeTracePhaseEvent::ToString(const UStateTree& StateTree) const
+{
+	return FString::Printf(TEXT("%s %s"),
+			*UEnum::GetDisplayValueAsText(EventType).ToString(),
+			*UEnum::GetDisplayValueAsText(Phase).ToString());
+}
+
+
 //----------------------------------------------------------------------//
 // FStateTreeTraceLogEvent
 //----------------------------------------------------------------------//
@@ -110,8 +122,8 @@ FString FStateTreeTraceConditionEvent::ToString(const UStateTree& StateTree) con
 //----------------------------------------------------------------------//
 // FStateTreeTraceActiveStatesEvent
 //----------------------------------------------------------------------//
-FStateTreeTraceActiveStatesEvent::FStateTreeTraceActiveStatesEvent(const double RecordingWorldTime, const EStateTreeUpdatePhase Phase)
-	: FStateTreeTracePhaseEvent(RecordingWorldTime, Phase)
+FStateTreeTraceActiveStatesEvent::FStateTreeTraceActiveStatesEvent(const double RecordingWorldTime)
+	: FStateTreeTraceBaseEvent(RecordingWorldTime)
 {
 }
 
diff --git a/Engine/Plugins/Runtime/StateTree/Source/StateTreeModule/Private/StateTreeExecutionContext.cpp b/Engine/Plugins/Runtime/StateTree/Source/StateTreeModule/Private/StateTreeExecutionContext.cpp
--- a/Engine/Plugins/Runtime/StateTree/Source/StateTreeModule/Private/StateTreeExecutionContext.cpp
+++ b/Engine/Plugins/Runtime/StateTree/Source/StateTreeModule/Private/StateTreeExecutionContext.cpp
@@ -16,15 +16,30 @@
 #define STATETREE_CLOG(Condition, Verbosity, Format, ...) UE_CVLOG_UELOG((Condition), GetOwner(), LogStateTree, Verbosity, TEXT("%s") Format, *GetInstanceDescription(), ##__VA_ARGS__)
 
 #if WITH_STATETREE_DEBUGGER
-	#define STATETREE_SCOPED_PHASE(Phase) TGuardValue<EStateTreeUpdatePhase> PREPROCESSOR_JOIN(ScopedPhaseMask, __LINE__)(UpdatePhaseMask, UpdatePhaseMask | (Phase))
-	#define STATETREE_TRACE_LOG_EVENT(Format, ...) TRACE_STATETREE_LOG_EVENT(GetInstanceDebugId(), UpdatePhaseMask, Format, ##__VA_ARGS__)
-	#define STATETREE_TRACE_STATE_EVENT(StateHandle, EventType) TRACE_STATETREE_STATE_EVENT(GetInstanceDebugId(), UpdatePhaseMask, StateHandle, EventType, EStateTreeStateSelectionBehavior::None);
-	#define STATETREE_TRACE_STATE_SELECTION_EVENT(StateHandle, EventType, SelectionBehavior) TRACE_STATETREE_STATE_EVENT(GetInstanceDebugId(), UpdatePhaseMask, StateHandle, EventType, SelectionBehavior);
-	#define STATETREE_TRACE_TASK_EVENT(Index, DataView, EventType, Status) TRACE_STATETREE_TASK_EVENT(GetInstanceDebugId(), UpdatePhaseMask, FStateTreeIndex16(Index), DataView, EventType, Status);
-	#define STATETREE_TRACE_CONDITION_EVENT(Index, DataViews, EventType) TRACE_STATETREE_CONDITION_EVENT(GetInstanceDebugId(), UpdatePhaseMask, FStateTreeIndex16(Index), DataView, EventType);	
-	#define STATETREE_TRACE_TRANSITION_EVENT(Index, EventType) TRACE_STATETREE_TRANSITION_EVENT(GetInstanceDebugId(), UpdatePhaseMask, FStateTreeIndex16(Index), EventType);
+	#define ID_NAME PREPROCESSOR_JOIN(InstanceId,__LINE__) \
+
+	#define STATETREE_TRACE_SCOPED_PHASE(Phase) \
+		FStateTreeInstanceDebugId ID_NAME = GetInstanceDebugId(); \
+		TRACE_STATETREE_PHASE_EVENT(ID_NAME, Phase, EStateTreeTraceEventType::Push) \
+		ON_SCOPE_EXIT { TRACE_STATETREE_PHASE_EVENT(ID_NAME, Phase, EStateTreeTraceEventType::Pop) }
+
+	#define STATETREE_TRACE_SCOPED_STATE_SELECTION(StateHandle, SelectionBehavior) \
+		FStateTreeInstanceDebugId ID_NAME = GetInstanceDebugId(); \
+		TRACE_STATETREE_STATE_EVENT(ID_NAME, StateHandle, EStateTreeTraceEventType::Push, SelectionBehavior) \
+		ON_SCOPE_EXIT { TRACE_STATETREE_STATE_EVENT(ID_NAME, StateHandle, EStateTreeTraceEventType::Pop, SelectionBehavior); };
+
+	#define STATETREE_TRACE_INSTANCE_EVENT(EventType)						TRACE_STATETREE_INSTANCE_EVENT(GetInstanceDebugId(), GetStateTree(), *GetInstanceDescription(), EventType);
+	#define STATETREE_TRACE_ACTIVE_STATES_EVENT(ActiveStates)				TRACE_STATETREE_ACTIVE_STATES_EVENT(GetInstanceDebugId(), ActiveStates);
+	#define STATETREE_TRACE_LOG_EVENT(Format, ...)							TRACE_STATETREE_LOG_EVENT(GetInstanceDebugId(), Format, ##__VA_ARGS__)
+	#define STATETREE_TRACE_STATE_EVENT(StateHandle, EventType)				TRACE_STATETREE_STATE_EVENT(GetInstanceDebugId(), StateHandle, EventType, EStateTreeStateSelectionBehavior::None);
+	#define STATETREE_TRACE_TASK_EVENT(Index, DataView, EventType, Status)	TRACE_STATETREE_TASK_EVENT(GetInstanceDebugId(), FStateTreeIndex16(Index), DataView, EventType, Status);
+	#define STATETREE_TRACE_CONDITION_EVENT(Index, DataViews, EventType)	TRACE_STATETREE_CONDITION_EVENT(GetInstanceDebugId(), FStateTreeIndex16(Index), DataView, EventType);	
+	#define STATETREE_TRACE_TRANSITION_EVENT(Index, EventType)				TRACE_STATETREE_TRANSITION_EVENT(GetInstanceDebugId(), FStateTreeIndex16(Index), EventType);
 #else
-	#define STATETREE_SCOPED_PHASE(Phase)
+	#define STATETREE_TRACE_SCOPED_PHASE(Phase)
+	#define STATETREE_TRACE_SCOPED_STATE_SELECTION(StateHandle, SelectionBehavior)
+	#define STATETREE_TRACE_INSTANCE_EVENT(EventType)
+	#define STATETREE_TRACE_ACTIVE_STATES_EVENT(ActiveStates)
 	#define STATETREE_TRACE_LOG_EVENT(Format, ...)
 	#define STATETREE_TRACE_STATE_EVENT(StateHandle, EventType)
 	#define STATETREE_TRACE_STATE_SELECTION_EVENT(StateHandle, EventType, SelectionBehavior)
@@ -178,7 +193,6 @@ void FStateTreeExecutionContext::UpdateSubtreeStateParameters(const FCompactStat
 EStateTreeRunStatus FStateTreeExecutionContext::Start()
 {
 	CSV_SCOPED_TIMING_STAT_EXCLUSIVE(StateTree_Start);
-	STATETREE_SCOPED_PHASE(EStateTreeUpdatePhase::StartTree);
 
 	if (!IsValid())
 	{
@@ -201,7 +215,12 @@ EStateTreeRunStatus FStateTreeExecutionContext::Start()
 		return EStateTreeRunStatus::Failed;
 	}
 
-	TRACE_STATETREE_INSTANCE_EVENT(GetInstanceDebugId(), GetStateTree(), *GetInstanceDescription(), EStateTreeTraceInstanceEventType::Started);
+	// Must sent instance creation event first 
+	STATETREE_TRACE_INSTANCE_EVENT(EStateTreeTraceEventType::Push);
+	
+	// Set scoped phase only for properly initialized context with valid Instance data
+	// since we need it to output the InstanceId
+	STATETREE_TRACE_SCOPED_PHASE(EStateTreeUpdatePhase::StartTree);
 	
 	// Start evaluators and global tasks. Fail the execution if any global task fails.
 	FStateTreeIndex16 LastInitializedTaskIndex;
@@ -272,7 +291,6 @@ EStateTreeRunStatus FStateTreeExecutionContext::Start()
 EStateTreeRunStatus FStateTreeExecutionContext::Stop()
 {
 	CSV_SCOPED_TIMING_STAT_EXCLUSIVE(StateTree_Stop);
-	STATETREE_SCOPED_PHASE(EStateTreeUpdatePhase::StopTree);
 
 	if (!IsValid())
 	{
@@ -286,6 +304,10 @@ EStateTreeRunStatus FStateTreeExecutionContext::Stop()
 		return EStateTreeRunStatus::Failed;
 	}
 
+	// Set scoped phase only for properly initialized context with valid Instance data
+	// since we need it to output the InstanceId 
+	STATETREE_TRACE_SCOPED_PHASE(EStateTreeUpdatePhase::StopTree);
+
 	const FStateTreeExecutionState& Exec = GetExecState();
 	EStateTreeRunStatus Result = Exec.TreeRunStatus;
 	
@@ -313,8 +335,8 @@ EStateTreeRunStatus FStateTreeExecutionContext::Stop()
 	StopEvaluatorsAndGlobalTasks();
 
 	// Trace before resetting the instance data since it is required to provide all the event information
-	TRACE_STATETREE_ACTIVE_STATES_EVENT(GetInstanceDebugId(), UpdatePhaseMask, FStateTreeActiveStates());
-	TRACE_STATETREE_INSTANCE_EVENT(GetInstanceDebugId(), GetStateTree(), *GetInstanceDescription(), EStateTreeTraceInstanceEventType::Stopped);
+	STATETREE_TRACE_ACTIVE_STATES_EVENT(FStateTreeActiveStates());
+	STATETREE_TRACE_INSTANCE_EVENT(EStateTreeTraceEventType::Pop);
 
 	// Destruct all allocated instance data (does not shrink the buffer). This will invalidate Exec too.
 	InstanceData.Reset();
@@ -325,7 +347,7 @@ EStateTreeRunStatus FStateTreeExecutionContext::Stop()
 EStateTreeRunStatus FStateTreeExecutionContext::Tick(const float DeltaTime)
 {
 	CSV_SCOPED_TIMING_STAT_EXCLUSIVE(StateTree_Tick);
-	STATETREE_SCOPED_PHASE(EStateTreeUpdatePhase::TickStateTree);
+	STATETREE_TRACE_SCOPED_PHASE(EStateTreeUpdatePhase::TickStateTree);
 
 	if (!IsValid())
 	{
@@ -392,7 +414,7 @@ EStateTreeRunStatus FStateTreeExecutionContext::Tick(const float DeltaTime)
 		// Trigger conditional transitions or state succeed/failed transitions. First tick transition is handled here too.
 		if (TriggerTransitions())
 		{
-			STATETREE_SCOPED_PHASE(EStateTreeUpdatePhase::ApplyTransitions);
+			STATETREE_TRACE_SCOPED_PHASE(EStateTreeUpdatePhase::ApplyTransitions);
 			
 			// We have committed to state change, consume events that were accumulated during the tick above.
 			EventQueue.Reset();
@@ -708,7 +730,6 @@ EStateTreeRunStatus FStateTreeExecutionContext::EnterState(const FStateTreeTrans
 	}
 	
 	STATETREE_LOG(Log, TEXT("Enter state '%s' (%d)"), *DebugGetStatePath(Transition.NextActiveStates), Exec.StateChangeCount);
-	STATETREE_TRACE_STATE_EVENT(Transition.NextActiveStates.Last(), EStateTreeTraceNodeEventType::OnEnter);
 
 	for (int32 Index = 0; Index < Transition.NextActiveStates.Num() && Result != EStateTreeRunStatus::Failed; Index++)
 	{
@@ -722,6 +743,8 @@ EStateTreeRunStatus FStateTreeExecutionContext::EnterState(const FStateTreeTrans
 				__FUNCTION__, *GetStateStatusString(Exec), *GetNameSafe(&Owner), *GetFullNameSafe(&StateTree));
 			break;
 		}
+		
+		STATETREE_TRACE_STATE_EVENT(CurrentHandle, EStateTreeTraceEventType::OnEnter);
 
 		if (State.Type == EStateTreeStateType::Linked)
 		{
@@ -771,7 +794,7 @@ EStateTreeRunStatus FStateTreeExecutionContext::EnterState(const FStateTreeTrans
 
 				if (Status != EStateTreeRunStatus::Running)
 				{
-					STATETREE_TRACE_TASK_EVENT(TaskIndex, DataViews[Task.DataViewIndex.Get()], EStateTreeTraceNodeEventType::OnTaskCompleted, Status);
+					STATETREE_TRACE_TASK_EVENT(TaskIndex, DataViews[Task.DataViewIndex.Get()], EStateTreeTraceEventType::OnTaskCompleted, Status);
 					// Store the first state that completed, will be used to decide where to trigger transitions.
 					if (!Exec.CompletedStateHandle.IsValid())
 					{
@@ -790,7 +813,7 @@ EStateTreeRunStatus FStateTreeExecutionContext::EnterState(const FStateTreeTrans
 		}
 	}
 
-	TRACE_STATETREE_ACTIVE_STATES_EVENT(GetInstanceDebugId(), UpdatePhaseMask, Exec.ActiveStates);
+	STATETREE_TRACE_ACTIVE_STATES_EVENT(Exec.ActiveStates);
 	return Result;
 }
 
@@ -886,7 +909,7 @@ void FStateTreeExecutionContext::ExitState(const FStateTreeTransitionResult& Tra
 		CurrentTransition.ChangeType = ExitedStateChangeType[Index];
 
 		STATETREE_LOG(Log, TEXT("%*sState '%s' %s"), Index*UE::StateTree::DebugIndentSize, TEXT(""), *DebugGetStatePath(Transition.CurrentActiveStates, ExitedStateActiveIndex[Index]), *UEnum::GetDisplayValueAsText(CurrentTransition.ChangeType).ToString());
-		STATETREE_TRACE_STATE_EVENT(CurrentHandle, EStateTreeTraceNodeEventType::OnExit);
+		STATETREE_TRACE_STATE_EVENT(CurrentHandle, EStateTreeTraceEventType::OnExit);
 
 		// Tasks
 		for (int32 TaskIndex = (State.TasksBegin + State.TasksNum) - 1; TaskIndex >= State.TasksBegin; TaskIndex--)
@@ -938,7 +961,7 @@ void FStateTreeExecutionContext::StateCompleted()
 		FCurrentlyProcessedStateScope StateScope(*this, CurrentHandle);
 		
 		STATETREE_LOG(Verbose, TEXT("%*sState '%s'"), Index*UE::StateTree::DebugIndentSize, TEXT(""), *DebugGetStatePath(Exec.ActiveStates, Index));
-		STATETREE_TRACE_STATE_EVENT(CurrentHandle, EStateTreeTraceNodeEventType::OnStateCompleted);
+		STATETREE_TRACE_STATE_EVENT(CurrentHandle, EStateTreeTraceEventType::OnStateCompleted);
 
 		// Notify Tasks
 		for (int32 TaskIndex = (State.TasksBegin + State.TasksNum) - 1; TaskIndex >= State.TasksBegin; TaskIndex--)
@@ -960,7 +983,7 @@ void FStateTreeExecutionContext::StateCompleted()
 EStateTreeRunStatus FStateTreeExecutionContext::TickEvaluatorsAndGlobalTasks(const float DeltaTime, const bool bTickGlobalTasks)
 {
 	CSV_SCOPED_TIMING_STAT_EXCLUSIVE(StateTree_TickEvaluators);
-	STATETREE_SCOPED_PHASE(EStateTreeUpdatePhase::TickingGlobalTasks);
+	STATETREE_TRACE_SCOPED_PHASE(EStateTreeUpdatePhase::TickingGlobalTasks);
 
 	STATETREE_CLOG(StateTree.EvaluatorsNum > 0, VeryVerbose, TEXT("Ticking Evaluators"));
 
@@ -1024,7 +1047,7 @@ EStateTreeRunStatus FStateTreeExecutionContext::TickEvaluatorsAndGlobalTasks(con
 			// If a global task succeeds or fails, it will stop the whole tree.
 			if (TaskResult != EStateTreeRunStatus::Running)
 			{
-				STATETREE_TRACE_TASK_EVENT(TaskIndex, DataViews[Task.DataViewIndex.Get()], EStateTreeTraceNodeEventType::OnTaskCompleted, TaskResult);
+				STATETREE_TRACE_TASK_EVENT(TaskIndex, DataViews[Task.DataViewIndex.Get()], EStateTreeTraceEventType::OnTaskCompleted, TaskResult);
 				Result = TaskResult;
 			}
 				
@@ -1041,7 +1064,7 @@ EStateTreeRunStatus FStateTreeExecutionContext::TickEvaluatorsAndGlobalTasks(con
 bool FStateTreeExecutionContext::StartEvaluatorsAndGlobalTasks(FStateTreeIndex16& OutLastInitializedTaskIndex)
 {
 	CSV_SCOPED_TIMING_STAT_EXCLUSIVE(StateTree_StartEvaluators);
-	STATETREE_SCOPED_PHASE(EStateTreeUpdatePhase::StartGlobalTasks);
+	STATETREE_TRACE_SCOPED_PHASE(EStateTreeUpdatePhase::StartGlobalTasks);
 
 	STATETREE_CLOG(StateTree.EvaluatorsNum > 0 || StateTree.GlobalTasksNum > 0, Verbose, TEXT("Start Evaluators & Global tasks"));
 
@@ -1101,7 +1124,7 @@ bool FStateTreeExecutionContext::StartEvaluatorsAndGlobalTasks(FStateTreeIndex16
 void FStateTreeExecutionContext::StopEvaluatorsAndGlobalTasks(const FStateTreeIndex16 LastInitializedTaskIndex)
 {
 	CSV_SCOPED_TIMING_STAT_EXCLUSIVE(StateTree_StopEvaluators);
-	STATETREE_SCOPED_PHASE(EStateTreeUpdatePhase::StopGlobalTasks);
+	STATETREE_TRACE_SCOPED_PHASE(EStateTreeUpdatePhase::StopGlobalTasks);
 
 	STATETREE_CLOG(StateTree.EvaluatorsNum > 0, Verbose, TEXT("Stop Evaluators & Global Tasks"));
 
@@ -1169,7 +1192,7 @@ void FStateTreeExecutionContext::StopEvaluatorsAndGlobalTasks(const FStateTreeIn
 EStateTreeRunStatus FStateTreeExecutionContext::TickTasks(const float DeltaTime)
 {
 	CSV_SCOPED_TIMING_STAT_EXCLUSIVE(StateTree_TickTasks);
-	STATETREE_SCOPED_PHASE(EStateTreeUpdatePhase::TickingTasks);
+	STATETREE_TRACE_SCOPED_PHASE(EStateTreeUpdatePhase::TickingTasks);
 
 	FStateTreeExecutionState& Exec = GetExecState();
 
@@ -1235,7 +1258,7 @@ EStateTreeRunStatus FStateTreeExecutionContext::TickTasks(const float DeltaTime)
 				StateTree.PropertyBindings.CopyTo(DataViews, Task.BindingsBatch, TaskDataView);
 			}
 
-			STATETREE_TRACE_TASK_EVENT(TaskIndex, TaskDataView, EStateTreeTraceNodeEventType::OnTickingTask, EStateTreeRunStatus::Running);
+			STATETREE_TRACE_TASK_EVENT(TaskIndex, TaskDataView, EStateTreeTraceEventType::OnTickingTask, EStateTreeRunStatus::Running);
 			EStateTreeRunStatus TaskResult = EStateTreeRunStatus::Unset;
 			{
 				QUICK_SCOPE_CYCLE_COUNTER(StateTree_Task_Tick);
@@ -1245,7 +1268,7 @@ EStateTreeRunStatus FStateTreeExecutionContext::TickTasks(const float DeltaTime)
 			}
 
 			STATETREE_TRACE_TASK_EVENT(TaskIndex, TaskDataView,
-				TaskResult != EStateTreeRunStatus::Running ? EStateTreeTraceNodeEventType::OnTaskCompleted : EStateTreeTraceNodeEventType::OnTaskTicked,
+				TaskResult != EStateTreeRunStatus::Running ? EStateTreeTraceEventType::OnTaskCompleted : EStateTreeTraceEventType::OnTaskTicked,
 				TaskResult);
 			
 			// TODO: Add more control over which states can control the failed/succeeded result.
@@ -1315,7 +1338,7 @@ bool FStateTreeExecutionContext::TestAllConditions(const int32 ConditionsOffset,
 		}
 
 		const bool bValue = Cond.TestCondition(*this);
-		STATETREE_TRACE_CONDITION_EVENT(ConditionsOffset + Index, DataView, bValue ? EStateTreeTraceNodeEventType::Passed : EStateTreeTraceNodeEventType::Failed);
+		STATETREE_TRACE_CONDITION_EVENT(ConditionsOffset + Index, DataView, bValue ? EStateTreeTraceEventType::Passed : EStateTreeTraceEventType::Failed);
 
 		// Reset copied properties that might contain object references.
 		if (Cond.BindingsBatch.IsValid())
@@ -1435,7 +1458,7 @@ bool FStateTreeExecutionContext::RequestTransition(const FStateTreeStateHandle N
 bool FStateTreeExecutionContext::TriggerTransitions()
 {
 	CSV_SCOPED_TIMING_STAT_EXCLUSIVE(StateTree_TriggerTransition);
-	STATETREE_SCOPED_PHASE(EStateTreeUpdatePhase::TriggerTransitions);
+	STATETREE_TRACE_SCOPED_PHASE(EStateTreeUpdatePhase::TriggerTransitions);
 
 	FAllowDirectTransitionsScope AllowDirectTransitionsScope(*this); // Set flag for the scope of this function to allow direct transitions without buffering.
 	FStateTreeExecutionState& Exec = GetExecState();
@@ -1529,8 +1552,8 @@ bool FStateTreeExecutionContext::TriggerTransitions()
 			bool bPassed = false; 
 			if (bShouldTrigger)
 			{
-				STATETREE_TRACE_TRANSITION_EVENT(TransitionIndex, EStateTreeTraceNodeEventType::OnEvaluating);
-				STATETREE_SCOPED_PHASE(EStateTreeUpdatePhase::TransitionConditions);
+				STATETREE_TRACE_TRANSITION_EVENT(TransitionIndex, EStateTreeTraceEventType::OnEvaluating);
+				STATETREE_TRACE_SCOPED_PHASE(EStateTreeUpdatePhase::TransitionConditions);
 				bPassed = TestAllConditions(Transition.ConditionsBegin, Transition.ConditionsNum);
 			}
 
@@ -1614,8 +1637,8 @@ bool FStateTreeExecutionContext::TriggerTransitions()
 				{
 					bool bPassed = false;
 					{
-						STATETREE_TRACE_TRANSITION_EVENT(TransitionIndex, EStateTreeTraceNodeEventType::OnEvaluating);
-						STATETREE_SCOPED_PHASE(EStateTreeUpdatePhase::TransitionConditions);
+						STATETREE_TRACE_TRANSITION_EVENT(TransitionIndex, EStateTreeTraceEventType::OnEvaluating);
+						STATETREE_TRACE_SCOPED_PHASE(EStateTreeUpdatePhase::TransitionConditions);
 						bPassed = TestAllConditions(Transition.ConditionsBegin, Transition.ConditionsNum);
 					}
 
@@ -1770,7 +1793,6 @@ bool FStateTreeExecutionContext::SelectState(const FStateTreeStateHandle NextSta
 bool FStateTreeExecutionContext::SelectStateInternal(const FStateTreeStateHandle NextState, FStateTreeActiveStates& OutNewActiveState, FStateTreeActiveStates& VisitedStates)
 {
 	CSV_SCOPED_TIMING_STAT_EXCLUSIVE(StateTree_SelectState);
-	STATETREE_SCOPED_PHASE(EStateTreeUpdatePhase::StateSelection);
 
 	const FStateTreeExecutionState& Exec = GetExecState();
 
@@ -1784,17 +1806,15 @@ bool FStateTreeExecutionContext::SelectStateInternal(const FStateTreeStateHandle
 
 	const FCompactStateTreeState& State = StateTree.States[NextState.Index];
 
-	STATETREE_TRACE_STATE_SELECTION_EVENT(NextState, EStateTreeTraceNodeEventType::PushStateSelection, State.SelectionBehavior);
-	ON_SCOPE_EXIT
-	{
-		STATETREE_TRACE_STATE_EVENT(NextState, EStateTreeTraceNodeEventType::PopStateSelection);
-	};
+	// Push State before the scoped phase so the debugger can create the hierarchy parent using the state description.
+	STATETREE_TRACE_SCOPED_STATE_SELECTION(NextState, State.SelectionBehavior);
+	STATETREE_TRACE_SCOPED_PHASE(EStateTreeUpdatePhase::StateSelection);
 	
 	// Check that the state can be entered
 	bool bEnterConditionsPassed = false;
 	if (State.SelectionBehavior != EStateTreeStateSelectionBehavior::None)
 	{
-		STATETREE_SCOPED_PHASE(EStateTreeUpdatePhase::EnterConditions);
+		STATETREE_TRACE_SCOPED_PHASE(EStateTreeUpdatePhase::EnterConditions);
 		bEnterConditionsPassed = TestAllConditions(State.EnterConditionsBegin, State.EnterConditionsNum);
 	}
 
@@ -1825,7 +1845,7 @@ bool FStateTreeExecutionContext::SelectStateInternal(const FStateTreeStateHandle
 		else if (State.SelectionBehavior == EStateTreeStateSelectionBehavior::TryEnterState)
 		{
 			// Select this state.
-			STATETREE_TRACE_STATE_EVENT(NextState, EStateTreeTraceNodeEventType::OnStateSelected);
+			STATETREE_TRACE_STATE_EVENT(NextState, EStateTreeTraceEventType::OnStateSelected);
 			return true;
 		}
 		else if (State.SelectionBehavior == EStateTreeStateSelectionBehavior::TryFollowTransitions)
@@ -1870,8 +1890,8 @@ bool FStateTreeExecutionContext::SelectStateInternal(const FStateTreeStateHandle
 				bool bTransitionConditionsPassed = false;
 				if (bShouldTrigger)
 				{
-					STATETREE_TRACE_TRANSITION_EVENT(TransitionIndex, EStateTreeTraceNodeEventType::OnEvaluating);
-					STATETREE_SCOPED_PHASE(EStateTreeUpdatePhase::TransitionConditions);
+					STATETREE_TRACE_TRANSITION_EVENT(TransitionIndex, EStateTreeTraceEventType::OnEvaluating);
+					STATETREE_TRACE_SCOPED_PHASE(EStateTreeUpdatePhase::TransitionConditions);
 					bTransitionConditionsPassed = bShouldTrigger && TestAllConditions(Transition.ConditionsBegin, Transition.ConditionsNum);
 				}
 
@@ -1911,7 +1931,7 @@ bool FStateTreeExecutionContext::SelectStateInternal(const FStateTreeStateHandle
 			else
 			{
 				// Select this state (For backwards compatibility)
-				STATETREE_TRACE_STATE_EVENT(NextState, EStateTreeTraceNodeEventType::OnStateSelected);
+				STATETREE_TRACE_STATE_EVENT(NextState, EStateTreeTraceEventType::OnStateSelected);
 				return true;
 			}
 		}
diff --git a/Engine/Plugins/Runtime/StateTree/Source/StateTreeModule/Public/Debugger/StateTreeTrace.h b/Engine/Plugins/Runtime/StateTree/Source/StateTreeModule/Public/Debugger/StateTreeTrace.h
--- a/Engine/Plugins/Runtime/StateTree/Source/StateTreeModule/Public/Debugger/StateTreeTrace.h
+++ b/Engine/Plugins/Runtime/StateTree/Source/StateTreeModule/Public/Debugger/StateTreeTrace.h
@@ -14,9 +14,8 @@ struct FStateTreeIndex16;
 struct FStateTreeStateHandle;
 enum class EStateTreeStateSelectionBehavior : uint8;
 enum class EStateTreeRunStatus : uint8;
-enum class EStateTreeTraceNodeEventType : uint8;
-enum class EStateTreeUpdatePhase : uint16;
-enum class EStateTreeTraceInstanceEventType : uint8;
+enum class EStateTreeTraceEventType : uint8;
+enum class EStateTreeUpdatePhase : uint8;
 
 UE_TRACE_CHANNEL_EXTERN(StateTreeDebugChannel, STATETREEMODULE_API)
 
@@ -24,13 +23,14 @@ namespace UE::StateTreeTrace
 {
 	void RegisterGlobalDelegates();
 	void UnregisterGlobalDelegates();
-	void OutputInstanceLifetimeEvent(FStateTreeInstanceDebugId InstanceId, const UStateTree* StateTree, const TCHAR* InstanceName, EStateTreeTraceInstanceEventType EventType);
-	void OutputLogEventTrace(FStateTreeInstanceDebugId InstanceId, EStateTreeUpdatePhase Phase, const TCHAR* Fmt, ...);
-	void OutputStateEventTrace(FStateTreeInstanceDebugId InstanceId, EStateTreeUpdatePhase Phase, FStateTreeStateHandle StateHandle, EStateTreeTraceNodeEventType EventType, EStateTreeStateSelectionBehavior SelectionBehavior);
-	void OutputTaskEventTrace(FStateTreeInstanceDebugId InstanceId, EStateTreeUpdatePhase Phase, FStateTreeIndex16 TaskIdx, FStateTreeDataView DataView, EStateTreeTraceNodeEventType EventType, EStateTreeRunStatus Status);
-	void OutputConditionEventTrace(FStateTreeInstanceDebugId InstanceId, EStateTreeUpdatePhase Phase, FStateTreeIndex16 ConditionIdx, FStateTreeDataView DataView, EStateTreeTraceNodeEventType EventType);
-	void OutputTransitionEventTrace(FStateTreeInstanceDebugId InstanceId, EStateTreeUpdatePhase Phase, FStateTreeIndex16 TransitionIdx, EStateTreeTraceNodeEventType EventType);
-	void OutputActiveStatesEventTrace(FStateTreeInstanceDebugId InstanceId, EStateTreeUpdatePhase Phase, const FStateTreeActiveStates& ActiveStates);
+	void ProcessPhaseScopeEvent(FStateTreeInstanceDebugId InstanceId, EStateTreeUpdatePhase Phase, EStateTreeTraceEventType EventType);
+	void OutputInstanceLifetimeEvent(FStateTreeInstanceDebugId InstanceId, const UStateTree* StateTree, const TCHAR* InstanceName, EStateTreeTraceEventType EventType);
+	void OutputLogEventTrace(FStateTreeInstanceDebugId InstanceId, const TCHAR* Fmt, ...);
+	void OutputStateEventTrace(FStateTreeInstanceDebugId InstanceId, FStateTreeStateHandle StateHandle, EStateTreeTraceEventType EventType, EStateTreeStateSelectionBehavior SelectionBehavior);
+	void OutputTaskEventTrace(FStateTreeInstanceDebugId InstanceId, FStateTreeIndex16 TaskIdx, FStateTreeDataView DataView, EStateTreeTraceEventType EventType, EStateTreeRunStatus Status);
+	void OutputConditionEventTrace(FStateTreeInstanceDebugId InstanceId, FStateTreeIndex16 ConditionIdx, FStateTreeDataView DataView, EStateTreeTraceEventType EventType);
+	void OutputTransitionEventTrace(FStateTreeInstanceDebugId InstanceId, FStateTreeIndex16 TransitionIdx, EStateTreeTraceEventType EventType);
+	void OutputActiveStatesEventTrace(FStateTreeInstanceDebugId InstanceId, const FStateTreeActiveStates& ActiveStates);
 }
 
 #define TRACE_STATETREE_INSTANCE_EVENT(InstanceID, StateTree, InstanceName, EventType) \
@@ -39,50 +39,54 @@ namespace UE::StateTreeTrace
 		UE::StateTreeTrace::OutputInstanceLifetimeEvent(InstanceID, StateTree, InstanceName, EventType); \
 	}
 
-#define TRACE_STATETREE_LOG_EVENT(InstanceId, Phase, Format, ...) \
+#define TRACE_STATETREE_PHASE_EVENT(InstanceID, Phase, EventType) \
+	UE::StateTreeTrace::ProcessPhaseScopeEvent(InstanceID, Phase, EventType); \
+
+#define TRACE_STATETREE_LOG_EVENT(InstanceId, Format, ...) \
 	if (UE_TRACE_CHANNELEXPR_IS_ENABLED(StateTreeDebugChannel)) \
 	{ \
-		UE::StateTreeTrace::OutputLogEventTrace(InstanceId, Phase, Format, ##__VA_ARGS__); \
+		UE::StateTreeTrace::OutputLogEventTrace(InstanceId, Format, ##__VA_ARGS__); \
 	}
 
-#define TRACE_STATETREE_STATE_EVENT(InstanceId, Phase, StateHandle, EventType, SelectionBehavior) \
+#define TRACE_STATETREE_STATE_EVENT(InstanceId, StateHandle, EventType, SelectionBehavior) \
 	if (UE_TRACE_CHANNELEXPR_IS_ENABLED(StateTreeDebugChannel)) \
 	{ \
-		UE::StateTreeTrace::OutputStateEventTrace(InstanceId, Phase, StateHandle, EventType, SelectionBehavior); \
+		UE::StateTreeTrace::OutputStateEventTrace(InstanceId, StateHandle, EventType, SelectionBehavior); \
 	}
 
-#define TRACE_STATETREE_TASK_EVENT(InstanceId, Phase, TaskIdx, DataView, EventType, Status) \
+#define TRACE_STATETREE_TASK_EVENT(InstanceId, TaskIdx, DataView, EventType, Status) \
 	if (UE_TRACE_CHANNELEXPR_IS_ENABLED(StateTreeDebugChannel)) \
 	{ \
-		UE::StateTreeTrace::OutputTaskEventTrace(InstanceId, Phase, TaskIdx, DataView, EventType, Status); \
+		UE::StateTreeTrace::OutputTaskEventTrace(InstanceId, TaskIdx, DataView, EventType, Status); \
 	}
 
-#define TRACE_STATETREE_CONDITION_EVENT(InstanceId, Phase, ConditionIdx, DataView, EventType) \
+#define TRACE_STATETREE_CONDITION_EVENT(InstanceId, ConditionIdx, DataView, EventType) \
 	if (UE_TRACE_CHANNELEXPR_IS_ENABLED(StateTreeDebugChannel)) \
 	{ \
-		UE::StateTreeTrace::OutputConditionEventTrace(InstanceId, Phase, ConditionIdx, DataView, EventType); \
+		UE::StateTreeTrace::OutputConditionEventTrace(InstanceId, ConditionIdx, DataView, EventType); \
 	}
 
-#define TRACE_STATETREE_TRANSITION_EVENT(InstanceId, Phase, TransitionIdx, EventType) \
+#define TRACE_STATETREE_TRANSITION_EVENT(InstanceId, TransitionIdx, EventType) \
 	if (UE_TRACE_CHANNELEXPR_IS_ENABLED(StateTreeDebugChannel)) \
 	{ \
-		UE::StateTreeTrace::OutputTransitionEventTrace(InstanceId, Phase, TransitionIdx, EventType); \
+		UE::StateTreeTrace::OutputTransitionEventTrace(InstanceId, TransitionIdx, EventType); \
 	}
 
-#define TRACE_STATETREE_ACTIVE_STATES_EVENT(InstanceId, Phase, ActivateStates) \
+#define TRACE_STATETREE_ACTIVE_STATES_EVENT(InstanceId, ActivateStates) \
 	if (UE_TRACE_CHANNELEXPR_IS_ENABLED(StateTreeDebugChannel)) \
 	{ \
-		UE::StateTreeTrace::OutputActiveStatesEventTrace(InstanceId, Phase, ActivateStates); \
+		UE::StateTreeTrace::OutputActiveStatesEventTrace(InstanceId, ActivateStates); \
 	}
 
 #else //STATETREE_DEBUG_TRACE_ENABLED
 
 #define TRACE_STATETREE_INSTANCE_EVENT(InstanceID, StateTree, InstanceName, EventType)
-#define TRACE_STATETREE_LOG_EVENT(InstanceId, Phase, Format, ...)
-#define TRACE_STATETREE_STATE_EVENT(InstanceId, Phase, StateHandle, EventType, SelectionBehavior)
-#define TRACE_STATETREE_TASK_EVENT(InstanceId, Phase, TaskIdx, DataView, EventType, Status)
-#define TRACE_STATETREE_CONDITION_EVENT(InstanceId, Phase, ConditionIdx, DataView, EventType)
-#define TRACE_STATETREE_TRANSITION_EVENT(InstanceId, Phase, TransitionIdx, EventType)
-#define TRACE_STATETREE_ACTIVE_STATES_EVENT(InstanceId, Phase, ActivateStates)
+#define TRACE_STATETREE_PHASE_EVENT(InstanceID, Phase, EventType)
+#define TRACE_STATETREE_LOG_EVENT(InstanceId, Format, ...)
+#define TRACE_STATETREE_STATE_EVENT(InstanceId, StateHandle, EventType, SelectionBehavior)
+#define TRACE_STATETREE_TASK_EVENT(InstanceId, TaskIdx, DataView, EventType, Status)
+#define TRACE_STATETREE_CONDITION_EVENT(InstanceId, ConditionIdx, DataView, EventType)
+#define TRACE_STATETREE_TRANSITION_EVENT(InstanceId, TransitionIdx, EventType)
+#define TRACE_STATETREE_ACTIVE_STATES_EVENT(InstanceId, ActivateStates)
 
 #endif // STATETREE_DEBUG_TRACE_ENABLED
diff --git a/Engine/Plugins/Runtime/StateTree/Source/StateTreeModule/Public/Debugger/StateTreeTraceAnalyzer.h b/Engine/Plugins/Runtime/StateTree/Source/StateTreeModule/Public/Debugger/StateTreeTraceAnalyzer.h
--- a/Engine/Plugins/Runtime/StateTree/Source/StateTreeModule/Public/Debugger/StateTreeTraceAnalyzer.h
+++ b/Engine/Plugins/Runtime/StateTree/Source/StateTreeModule/Public/Debugger/StateTreeTraceAnalyzer.h
@@ -23,6 +23,7 @@ private:
 	{
 		RouteId_WorldTimestamp,
 		RouteId_Instance,
+		RouteId_Phase,
 		RouteId_LogMessage,
 		RouteId_State,
 		RouteId_Task,
diff --git a/Engine/Plugins/Runtime/StateTree/Source/StateTreeModule/Public/Debugger/StateTreeTraceProvider.h b/Engine/Plugins/Runtime/StateTree/Source/StateTreeModule/Public/Debugger/StateTreeTraceProvider.h
--- a/Engine/Plugins/Runtime/StateTree/Source/StateTreeModule/Public/Debugger/StateTreeTraceProvider.h
+++ b/Engine/Plugins/Runtime/StateTree/Source/StateTreeModule/Public/Debugger/StateTreeTraceProvider.h
@@ -9,7 +9,6 @@
 #include "StateTreeTypes.h" // required to compile TMap<FStateTreeInstanceDebugId, ...>
 
 namespace TraceServices { class IAnalysisSession; }
-enum class EStateTreeTraceInstanceEventType : uint8;
 class UStateTree;
 namespace UE::StateTreeDebugger { struct FInstanceDescriptor; }
 
@@ -27,7 +26,7 @@ public:
 		const TCHAR* InInstanceName,
 		double InTime,
 		double InWorldRecordingTime,
-		EStateTreeTraceInstanceEventType InEventType);
+		EStateTreeTraceEventType InEventType);
 
 protected:
 	/** IStateTreeDebuggerProvider interface */
diff --git a/Engine/Plugins/Runtime/StateTree/Source/StateTreeModule/Public/Debugger/StateTreeTraceTypes.h b/Engine/Plugins/Runtime/StateTree/Source/StateTreeModule/Public/Debugger/StateTreeTraceTypes.h
--- a/Engine/Plugins/Runtime/StateTree/Source/StateTreeModule/Public/Debugger/StateTreeTraceTypes.h
+++ b/Engine/Plugins/Runtime/StateTree/Source/StateTreeModule/Public/Debugger/StateTreeTraceTypes.h
@@ -12,21 +12,15 @@ struct FStateTreeStateHandle;
 enum class EStateTreeStateSelectionBehavior : uint8;
 enum class EStateTreeRunStatus : uint8;
 
-UENUM()
-enum class EStateTreeTraceInstanceEventType : uint8
-{
-	Started,
-	Stopped
-};
 
 UENUM()
-enum class EStateTreeTraceNodeEventType : uint8
+enum class EStateTreeTraceEventType : uint8
 {
 	Unset,
 	OnEnter				UMETA(DisplayName = "Enter"),		// for State + Task events
 	OnExit				UMETA(DisplayName = "Exit"),		// for State + Task events
-	PushStateSelection	UMETA(DisplayName = "Push"),		// for State events
-	PopStateSelection	UMETA(DisplayName = "Pop"),			// for State events
+	Push				UMETA(DisplayName = "Push"),		// for State + Phase + Instance events
+	Pop					UMETA(DisplayName = "Pop"),			// for State + Phase + Instance events
 	OnStateSelected		UMETA(DisplayName = "Selected"),	// for State events
 	OnStateCompleted	UMETA(DisplayName = "Completed"),	// for State events
 	OnTickingTask		UMETA(DisplayName = "Tick"),		// for Task events
@@ -39,11 +33,10 @@ enum class EStateTreeTraceNodeEventType : uint8
 
 #if WITH_STATETREE_DEBUGGER
 
-struct FStateTreeTracePhaseEvent
+struct FStateTreeTraceBaseEvent
 {
-	explicit FStateTreeTracePhaseEvent(const double RecordingWorldTime, const EStateTreeUpdatePhase Phase)
+	explicit FStateTreeTraceBaseEvent(const double RecordingWorldTime)
 		: RecordingWorldTime(RecordingWorldTime)
-		, Phase(Phase)
 	{
 	}
 
@@ -51,13 +44,27 @@ struct FStateTreeTracePhaseEvent
 	static FString GetDataAsText() { return TEXT(""); }
 
 	double RecordingWorldTime = 0;
+};
+
+struct FStateTreeTracePhaseEvent : FStateTreeTraceBaseEvent
+{
+	explicit FStateTreeTracePhaseEvent(const double RecordingWorldTime, const EStateTreeUpdatePhase Phase, const EStateTreeTraceEventType EventType)
+		: FStateTreeTraceBaseEvent(RecordingWorldTime)
+		, Phase(Phase)
+		, EventType(EventType)
+	{
+	}
+
+	STATETREEMODULE_API FString ToString(const UStateTree& StateTree) const;
+
 	EStateTreeUpdatePhase Phase;
+	EStateTreeTraceEventType EventType;
 };
 
-struct FStateTreeTraceLogEvent : FStateTreeTracePhaseEvent
+struct FStateTreeTraceLogEvent : FStateTreeTraceBaseEvent
 {
-	explicit FStateTreeTraceLogEvent(const double RecordingWorldTime, const EStateTreeUpdatePhase Phase, const FString& Message)
-		: FStateTreeTracePhaseEvent(RecordingWorldTime, Phase)
+	explicit FStateTreeTraceLogEvent(const double RecordingWorldTime, const FString& Message)
+		: FStateTreeTraceBaseEvent(RecordingWorldTime)
 		, Message(Message)
 	{
 	}
@@ -67,10 +74,10 @@ struct FStateTreeTraceLogEvent : FStateTreeTracePhaseEvent
 	FString Message;
 };
 
-struct FStateTreeTraceTransitionEvent : FStateTreeTracePhaseEvent
+struct FStateTreeTraceTransitionEvent : FStateTreeTraceBaseEvent
 {
-	explicit FStateTreeTraceTransitionEvent(const double RecordingWorldTime, const EStateTreeUpdatePhase Phase, const FStateTreeIndex16 Index, const EStateTreeTraceNodeEventType EventType)
-		: FStateTreeTracePhaseEvent(RecordingWorldTime, Phase)
+	explicit FStateTreeTraceTransitionEvent(const double RecordingWorldTime, const FStateTreeIndex16 Index, const EStateTreeTraceEventType EventType)
+		: FStateTreeTraceBaseEvent(RecordingWorldTime)
 		, TransitionIndex(Index)
 		, EventType(EventType)
 	{
@@ -79,13 +86,13 @@ struct FStateTreeTraceTransitionEvent : FStateTreeTracePhaseEvent
 	STATETREEMODULE_API FString ToString(const UStateTree& StateTree) const;
 
 	FStateTreeIndex16 TransitionIndex;
-	EStateTreeTraceNodeEventType EventType = EStateTreeTraceNodeEventType::Unset;
+	EStateTreeTraceEventType EventType = EStateTreeTraceEventType::Unset;
 };
 
-struct FStateTreeTraceNodeEvent : FStateTreeTracePhaseEvent
+struct FStateTreeTraceNodeEvent : FStateTreeTraceBaseEvent
 {
-	explicit FStateTreeTraceNodeEvent(const double RecordingWorldTime, const EStateTreeUpdatePhase Phase, const FStateTreeIndex16 Index, const EStateTreeTraceNodeEventType EventType)
-		: FStateTreeTracePhaseEvent(RecordingWorldTime, Phase)
+	explicit FStateTreeTraceNodeEvent(const double RecordingWorldTime, const FStateTreeIndex16 Index, const EStateTreeTraceEventType EventType)
+		: FStateTreeTraceBaseEvent(RecordingWorldTime)
 		, Index(Index)
 		, EventType(EventType)
 	{
@@ -94,13 +101,13 @@ struct FStateTreeTraceNodeEvent : FStateTreeTracePhaseEvent
 	STATETREEMODULE_API FString ToString(const UStateTree& StateTree) const;
 	
 	FStateTreeIndex16 Index;
-	EStateTreeTraceNodeEventType EventType = EStateTreeTraceNodeEventType::Unset;
+	EStateTreeTraceEventType EventType = EStateTreeTraceEventType::Unset;
 };
 
 struct FStateTreeTraceStateEvent : FStateTreeTraceNodeEvent
 {
-	explicit FStateTreeTraceStateEvent(const double RecordingWorldTime, const EStateTreeUpdatePhase Phase, const FStateTreeIndex16 Index, const EStateTreeTraceNodeEventType EventType, const EStateTreeStateSelectionBehavior SelectionBehavior)
-		: FStateTreeTraceNodeEvent(RecordingWorldTime, Phase, Index, EventType)
+	explicit FStateTreeTraceStateEvent(const double RecordingWorldTime, const FStateTreeIndex16 Index, const EStateTreeTraceEventType EventType, const EStateTreeStateSelectionBehavior SelectionBehavior)
+		: FStateTreeTraceNodeEvent(RecordingWorldTime, Index, EventType)
 		, SelectionBehavior(SelectionBehavior)
 	{
 	}
@@ -113,8 +120,8 @@ struct FStateTreeTraceStateEvent : FStateTreeTraceNodeEvent
 
 struct FStateTreeTraceTaskEvent : FStateTreeTraceNodeEvent
 {
-	explicit FStateTreeTraceTaskEvent(const double RecordingWorldTime, const EStateTreeUpdatePhase Phase, const FStateTreeIndex16 Index,	const EStateTreeTraceNodeEventType EventType, const EStateTreeRunStatus Status,	const FString& TypePath, const FString& InstanceDataAsText)
-		: FStateTreeTraceNodeEvent(RecordingWorldTime, Phase, Index, EventType)
+	explicit FStateTreeTraceTaskEvent(const double RecordingWorldTime, const FStateTreeIndex16 Index,	const EStateTreeTraceEventType EventType, const EStateTreeRunStatus Status,	const FString& TypePath, const FString& InstanceDataAsText)
+		: FStateTreeTraceNodeEvent(RecordingWorldTime, Index, EventType)
 		, TypePath(TypePath)
 		, InstanceDataAsText(InstanceDataAsText)
 		, Status(Status)
@@ -133,8 +140,8 @@ struct FStateTreeTraceTaskEvent : FStateTreeTraceNodeEvent
 
 struct FStateTreeTraceConditionEvent : FStateTreeTraceNodeEvent
 {
-	explicit FStateTreeTraceConditionEvent(const double RecordingWorldTime, const EStateTreeUpdatePhase Phase, const FStateTreeIndex16 Index, const EStateTreeTraceNodeEventType EventType, const FString& TypePath, const FString& InstanceDataAsText)
-		: FStateTreeTraceNodeEvent(RecordingWorldTime, Phase, Index, EventType)
+	explicit FStateTreeTraceConditionEvent(const double RecordingWorldTime, const FStateTreeIndex16 Index, const EStateTreeTraceEventType EventType, const FString& TypePath, const FString& InstanceDataAsText)
+		: FStateTreeTraceNodeEvent(RecordingWorldTime, Index, EventType)
 		, TypePath(TypePath)
 		, InstanceDataAsText(InstanceDataAsText)
 	{
@@ -149,10 +156,10 @@ struct FStateTreeTraceConditionEvent : FStateTreeTraceNodeEvent
 	FString InstanceDataAsText;
 };
 
-struct FStateTreeTraceActiveStatesEvent : FStateTreeTracePhaseEvent
+struct FStateTreeTraceActiveStatesEvent : FStateTreeTraceBaseEvent
 {
 	// Intentionally implemented in source file to compile 'TArray<FStateTreeStateHandle>' using only forward declaration.
-	explicit FStateTreeTraceActiveStatesEvent(const double RecordingWorldTime, const EStateTreeUpdatePhase Phase);
+	explicit FStateTreeTraceActiveStatesEvent(const double RecordingWorldTime);
 
 	STATETREEMODULE_API FString ToString(const UStateTree& StateTree) const;
 
@@ -160,7 +167,8 @@ struct FStateTreeTraceActiveStatesEvent : FStateTreeTracePhaseEvent
 };
 
 /** Type aliases for statetree trace events */
-using FStateTreeTraceEventVariantType = TVariant<FStateTreeTraceLogEvent,
+using FStateTreeTraceEventVariantType = TVariant<FStateTreeTracePhaseEvent,
+												FStateTreeTraceLogEvent,
 												FStateTreeTraceNodeEvent,
 												FStateTreeTraceStateEvent,
 												FStateTreeTraceTaskEvent,
diff --git a/Engine/Plugins/Runtime/StateTree/Source/StateTreeModule/Public/StateTreeExecutionContext.h b/Engine/Plugins/Runtime/StateTree/Source/StateTreeModule/Public/StateTreeExecutionContext.h
--- a/Engine/Plugins/Runtime/StateTree/Source/StateTreeModule/Public/StateTreeExecutionContext.h
+++ b/Engine/Plugins/Runtime/StateTree/Source/StateTreeModule/Public/StateTreeExecutionContext.h
@@ -441,11 +441,6 @@ protected:
 	/** Current state we're processing, or invalid if not applicable. */
 	FStateTreeStateHandle CurrentlyProcessedState;
 
-#if WITH_STATETREE_DEBUGGER
-	/** Current phase in the update. Used for debugging events */
-	EStateTreeUpdatePhase UpdatePhaseMask = EStateTreeUpdatePhase::Unset;
-#endif // WITH_STATETREE_DEBUGGER
-
 	/** True if transitions are allowed to be requested directly instead of buffering. */
 	bool bAllowDirectTransitions = false;
 
diff --git a/Engine/Plugins/Runtime/StateTree/Source/StateTreeModule/Public/StateTreeExecutionTypes.h b/Engine/Plugins/Runtime/StateTree/Source/StateTreeModule/Public/StateTreeExecutionTypes.h
--- a/Engine/Plugins/Runtime/StateTree/Source/StateTreeModule/Public/StateTreeExecutionTypes.h
+++ b/Engine/Plugins/Runtime/StateTree/Source/StateTreeModule/Public/StateTreeExecutionTypes.h
@@ -7,23 +7,21 @@
 /**
  * Enumeration for the different update phases.
  * This is used as context information when tracing debug events.
- * The values are ordered based on their hierarchy starting with the "higher" parents.
  */
-UENUM(Flags)
-enum class EStateTreeUpdatePhase : uint16
+UENUM()
+enum class EStateTreeUpdatePhase : uint8
 {
-	Unset						= 0,
-	StartTree					= 1ull << 1 UMETA(DisplayName = "Start Tree"),
-	StopTree					= 1ull << 2 UMETA(DisplayName = "Stop Tree"),
-	StartGlobalTasks			= 1ull << 3 UMETA(DisplayName = "Start Global Tasks"),
-	StopGlobalTasks				= 1ull << 4 UMETA(DisplayName = "Stop Global Tasks"),
-	TickStateTree				= 1ull << 5 UMETA(DisplayName = "Tick State Tree"),
-	ApplyTransitions			= 1ull << 6 UMETA(DisplayName = "Transition"),
-	TriggerTransitions			= 1ull << 7 UMETA(DisplayName = "Trigger Transitions"),
-	TickingGlobalTasks			= 1ull << 8 UMETA(DisplayName = "Tick Global Tasks"),
-	TickingTasks				= 1ull << 9 UMETA(DisplayName = "Tick Tasks"),
-	TransitionConditions		= 1ull << 10 UMETA(DisplayName = "Transition conditions"),
-	StateSelection				= 1ull << 11 UMETA(DisplayName = "Try Enter"),
-	EnterConditions				= 1ull << 12 UMETA(DisplayName = "Enter conditions"),
-};
-ENUM_CLASS_FLAGS(EStateTreeUpdatePhase);
\ No newline at end of file
+	Unset					= 0,
+	StartTree				UMETA(DisplayName = "Start Tree"),
+	StopTree				UMETA(DisplayName = "Stop Tree"),
+	StartGlobalTasks		UMETA(DisplayName = "Start Global Tasks"),
+	StopGlobalTasks			UMETA(DisplayName = "Stop Global Tasks"),
+	TickStateTree			UMETA(DisplayName = "Tick State Tree"),
+	ApplyTransitions		UMETA(DisplayName = "Transition"),
+	TriggerTransitions		UMETA(DisplayName = "Trigger Transitions"),
+	TickingGlobalTasks		UMETA(DisplayName = "Tick Global Tasks"),
+	TickingTasks			UMETA(DisplayName = "Tick Tasks"),
+	TransitionConditions	UMETA(DisplayName = "Transition conditions"),
+	StateSelection			UMETA(DisplayName = "Try Enter"),
+	EnterConditions			UMETA(DisplayName = "Enter conditions"),
+};
\ No newline at end of file
