Fix build break in UE5\main
Fix future CIS error in the StateTreeModule

diff --git a/Engine/Plugins/Runtime/StateTree/Source/StateTreeModule/Private/StateTreeModule.cpp b/Engine/Plugins/Runtime/StateTree/Source/StateTreeModule/Private/StateTreeModule.cpp
--- a/Engine/Plugins/Runtime/StateTree/Source/StateTreeModule/Private/StateTreeModule.cpp
+++ b/Engine/Plugins/Runtime/StateTree/Source/StateTreeModule/Private/StateTreeModule.cpp
@@ -121,10 +121,10 @@ void FStateTreeModule::StartTraces()
 	{
 		UE::Trace::EnumerateChannels([](const ANSICHAR* Name, const bool bIsEnabled, void* Channels)
 		{
-			TArray<FString>* ChannelsToRestore = static_cast<TArray<FString>*>(Channels); 
+			TArray<FString>* EnabledChannels = static_cast<TArray<FString>*>(Channels); 
 			if (bIsEnabled)
 			{
-				ChannelsToRestore->Emplace(ANSI_TO_TCHAR(Name));
+				EnabledChannels->Emplace(ANSI_TO_TCHAR(Name));
 			}		
 		}, &ChannelsToRestore);
 	}
