Static Analysis fix

diff --git a/Engine/Plugins/Runtime/StateTree/Source/StateTreeModule/Private/Debugger/StateTreeTraceTypes.cpp b/Engine/Plugins/Runtime/StateTree/Source/StateTreeModule/Private/Debugger/StateTreeTraceTypes.cpp
--- a/Engine/Plugins/Runtime/StateTree/Source/StateTreeModule/Private/Debugger/StateTreeTraceTypes.cpp
+++ b/Engine/Plugins/Runtime/StateTree/Source/StateTreeModule/Private/Debugger/StateTreeTraceTypes.cpp
@@ -132,7 +132,8 @@ FString FStateTreeTraceActiveStatesEvent::ToString(const UStateTree& StateTree)
 	FStringBuilderBase StatePath;
 	for (int32 i = 0; i < ActiveStates.Num(); i++)
 	{
-		const FCompactStateTreeState& State = StateTree.GetStates()[ActiveStates[i].Index];
+		TConstArrayView<FCompactStateTreeState> StatesView = StateTree.GetStates();
+		const FCompactStateTreeState& State = StatesView[ActiveStates[i].Index];
 		StatePath.Appendf(TEXT("%s%s"), i == 0 ? TEXT("") : TEXT("."), *State.Name.ToString());
 	}
 
