StateTree: Fix parameter sync in State Tree reference
- Remove transaction that was needed when instanced struct UI was buffered (causes infinite refresh loop now)

diff --git a/Engine/Plugins/Runtime/StateTree/Source/StateTreeEditorModule/Private/Customizations/StateTreeReferenceDetails.cpp b/Engine/Plugins/Runtime/StateTree/Source/StateTreeEditorModule/Private/Customizations/StateTreeReferenceDetails.cpp
--- a/Engine/Plugins/Runtime/StateTree/Source/StateTreeEditorModule/Private/Customizations/StateTreeReferenceDetails.cpp
+++ b/Engine/Plugins/Runtime/StateTree/Source/StateTreeEditorModule/Private/Customizations/StateTreeReferenceDetails.cpp
@@ -20,6 +20,7 @@ TSharedRef<IPropertyTypeCustomization> FStateTreeReferenceDetails::MakeInstance(
 void FStateTreeReferenceDetails::CustomizeHeader(const TSharedRef<IPropertyHandle> InStructPropertyHandle, FDetailWidgetRow& InHeaderRow, IPropertyTypeCustomizationUtils& InCustomizationUtils)
 {
 	StructPropertyHandle = InStructPropertyHandle;
+	PropUtils = InCustomizationUtils.GetPropertyUtilities();
 
 	// Make sure parameters are synced when displaying StateTreeReference.
 	// Associated StateTree asset might have been recompiled after the StateTreeReference was loaded.
@@ -91,9 +92,6 @@ void FStateTreeReferenceDetails::SyncParameters(const UStateTree* StateTreeToSyn
 	// - auto update when the property customization is created
 	// - auto update when a state tree is compiled.
 	// - when the associate state tree asset is changed
-	// The auto updates probably should not create undoable transactions, but we need to
-	// handle the change via the property notification system or else the data might not update
-	// if we're modifying cached data (e.g. instanced struct).
 	
 	TArray<UObject*> OuterObjects;
 	StructPropertyHandle->GetOuterObjects(OuterObjects);
@@ -101,7 +99,7 @@ void FStateTreeReferenceDetails::SyncParameters(const UStateTree* StateTreeToSyn
 	TArray<void*> RawData;
 	StructPropertyHandle->AccessRawData(RawData);
 
-	bool bNeedsSync = false;
+	bool bDidSync = false;
 	
 	for (int32 i = 0; i < RawData.Num(); i++)
 	{
@@ -109,38 +107,21 @@ void FStateTreeReferenceDetails::SyncParameters(const UStateTree* StateTreeToSyn
 		{
 			const UStateTree* StateTreeAsset = StateTreeReference->GetStateTree();
 
-			if (StateTreeAsset
-				&& (StateTreeToSync == nullptr || StateTreeAsset == StateTreeToSync) 
+			if ((StateTreeToSync == nullptr || StateTreeAsset == StateTreeToSync) 
 				&& StateTreeReference->RequiresParametersSync())
 			{
-				bNeedsSync = true;
+				// Changing the data without Modify().
+				// When called on property row creation, we don't expect to dirty the owner.
+				// In other cases we expect the outer to already been modified.
+				StateTreeReference->SyncParameters();
+				bDidSync = true;
 			}
 		}
 	}
 
-	if (bNeedsSync)
+	if (bDidSync && PropUtils)
 	{
-		FScopedTransaction Transaction(LOCTEXT("SyncParameters", "Update StateTree reference"));
-		StructPropertyHandle->NotifyPreChange();
-
-		for (int32 i = 0; i < RawData.Num(); i++)
-		{
-			if (FStateTreeReference* StateTreeReference = static_cast<FStateTreeReference*>(RawData[i]))
-			{
-				const UStateTree* StateTreeAsset = StateTreeReference->GetStateTree();
-
-				if (StateTreeAsset
-					&& (StateTreeToSync == nullptr || StateTreeAsset == StateTreeToSync) 
-					&& StateTreeReference->RequiresParametersSync())
-				{
-					OuterObjects[i]->Modify();
-					StateTreeReference->SyncParameters();
-				}
-			}
-		}
-
-		StructPropertyHandle->NotifyPostChange(EPropertyChangeType::ValueSet);
-		StructPropertyHandle->NotifyFinishedChangingProperties();
+		PropUtils->RequestRefresh();
 	}
 }
 
