[StateTreeDebugger] fixed TreeView refreshed too often
#rnx

diff --git a/Engine/Plugins/Runtime/StateTree/Source/StateTreeModule/Private/Debugger/StateTreeDebugger.cpp b/Engine/Plugins/Runtime/StateTree/Source/StateTreeModule/Private/Debugger/StateTreeDebugger.cpp
--- a/Engine/Plugins/Runtime/StateTree/Source/StateTreeModule/Private/Debugger/StateTreeDebugger.cpp
+++ b/Engine/Plugins/Runtime/StateTree/Source/StateTreeModule/Private/Debugger/StateTreeDebugger.cpp
@@ -369,10 +369,12 @@ FText FStateTreeDebugger::GetSelectedTraceDescription() const
 
 void FStateTreeDebugger::SetScrubTime(const double ScrubTime)
 {
-	ScrubState.SetScrubTime(ScrubTime);
-	OnScrubStateChanged.Execute(ScrubState);
+	if (ScrubState.SetScrubTime(ScrubTime))
+	{
+		OnScrubStateChanged.Execute(ScrubState);
 
-	RefreshActiveStates();
+		RefreshActiveStates();
+	}
 }
 
 bool FStateTreeDebugger::IsActiveInstance(const double Time, const FStateTreeInstanceDebugId InstanceId) const
diff --git a/Engine/Plugins/Runtime/StateTree/Source/StateTreeModule/Private/Debugger/StateTreeDebuggerTypes.cpp b/Engine/Plugins/Runtime/StateTree/Source/StateTreeModule/Private/Debugger/StateTreeDebuggerTypes.cpp
--- a/Engine/Plugins/Runtime/StateTree/Source/StateTreeModule/Private/Debugger/StateTreeDebuggerTypes.cpp
+++ b/Engine/Plugins/Runtime/StateTree/Source/StateTreeModule/Private/Debugger/StateTreeDebuggerTypes.cpp
@@ -33,8 +33,13 @@ const FInstanceEventCollection FInstanceEventCollection::Invalid;
 //----------------------------------------------------------------//
 // FScrubState
 //----------------------------------------------------------------//
-void FScrubState::SetScrubTime(const double NewScrubTime)
+bool FScrubState::SetScrubTime(const double NewScrubTime)
 {
+	if (NewScrubTime == ScrubTime)
+	{
+		return false;
+	}
+
 	ScrubTimeBoundState = EScrubTimeBoundState::Unset;
 	TraceFrameIndex = INDEX_NONE;
 	FrameSpanIndex = INDEX_NONE;
@@ -74,6 +79,8 @@ void FScrubState::SetScrubTime(const double NewScrubTime)
 	// This will set back to the exact value provided since SetFrameSpanIndex will snap it to the start time of the matching frame.
 	// It will be consistent with the case where EventCollectionIndex is not set.
 	ScrubTime = NewScrubTime;
+
+	return true;
 }
 
 void FScrubState::SetFrameSpanIndex(const int32 NewFrameSpanIndex)
diff --git a/Engine/Plugins/Runtime/StateTree/Source/StateTreeModule/Public/Debugger/StateTreeDebuggerTypes.h b/Engine/Plugins/Runtime/StateTree/Source/StateTreeModule/Public/Debugger/StateTreeDebuggerTypes.h
--- a/Engine/Plugins/Runtime/StateTree/Source/StateTreeModule/Public/Debugger/StateTreeDebuggerTypes.h
+++ b/Engine/Plugins/Runtime/StateTree/Source/StateTreeModule/Public/Debugger/StateTreeDebuggerTypes.h
@@ -165,8 +165,11 @@ public:
 	int32 ActiveStatesIndex = INDEX_NONE;
 	EScrubTimeBoundState ScrubTimeBoundState = EScrubTimeBoundState::Unset;
 	
-
-	void SetScrubTime(double NewScrubTime);
+	/**
+	 * Updates internal indices based on the new time.
+	 * @return true if values were updated; false otherwise (i.e. no changes)
+	 */
+	bool SetScrubTime(double NewScrubTime);
 
 	/**
 	 * Indicates if the current scrub state points to a valid frame.
