StateTree: Fixed and improved transition status reporting
- Allow to specify the completion status and terminal state when calling Stop() on State Tree exection context
- Global task's execution status directly translates to stree state on global EnterState (previously could just fail)
- Report terminal state transitions for global tasks (allows to use same task erro handling on states and global)
- Tie global task exection to the tree status, global tasks get ExitState() when transitioning to terminal state
- RequestTransition() was errorneuously reporting TreeRunStatus (instead of LastTickStatus) in Transition.CurrentRunStatus, this made it impossible to detect failed transitions on ExitState()

diff --git a/Engine/Plugins/Runtime/StateTree/Source/StateTreeModule/Private/StateTreeExecutionContext.cpp b/Engine/Plugins/Runtime/StateTree/Source/StateTreeModule/Private/StateTreeExecutionContext.cpp
--- a/Engine/Plugins/Runtime/StateTree/Source/StateTreeModule/Private/StateTreeExecutionContext.cpp
+++ b/Engine/Plugins/Runtime/StateTree/Source/StateTreeModule/Private/StateTreeExecutionContext.cpp
@@ -224,12 +224,13 @@ EStateTreeRunStatus FStateTreeExecutionContext::Start()
 	
 	// Start evaluators and global tasks. Fail the execution if any global task fails.
 	FStateTreeIndex16 LastInitializedTaskIndex;
-	if (!StartEvaluatorsAndGlobalTasks(LastInitializedTaskIndex))
+	const EStateTreeRunStatus GlobalTasksRunStatus = StartEvaluatorsAndGlobalTasks(LastInitializedTaskIndex);
+	if (GlobalTasksRunStatus != EStateTreeRunStatus::Running)
 	{
-		StopEvaluatorsAndGlobalTasks(LastInitializedTaskIndex);
-		STATETREE_LOG(Warning, TEXT("%hs: Failed to start evaluators and global tasks '%s' using StateTree '%s'. Try to recompile the StateTree asset."),
-			__FUNCTION__, *GetNameSafe(&Owner), *GetFullNameSafe(&StateTree));
-		return EStateTreeRunStatus::Failed;
+		StopEvaluatorsAndGlobalTasks(GlobalTasksRunStatus, LastInitializedTaskIndex);
+		STATETREE_LOG(VeryVerbose, TEXT("%hs: Global tasks completed the StateTree %s on start in status '%s'. "),
+			__FUNCTION__, *GetNameSafe(&Owner), *GetFullNameSafe(&StateTree), *UEnum::GetDisplayValueAsText(GlobalTasksRunStatus).ToString());
+		return GlobalTasksRunStatus;
 	}
 
 	// First tick.
@@ -288,7 +289,7 @@ EStateTreeRunStatus FStateTreeExecutionContext::Start()
 	return Exec->TreeRunStatus;
 }
 
-EStateTreeRunStatus FStateTreeExecutionContext::Stop()
+EStateTreeRunStatus FStateTreeExecutionContext::Stop(const EStateTreeRunStatus CompletionStatus)
 {
 	CSV_SCOPED_TIMING_STAT_EXCLUSIVE(StateTree_Stop);
 
@@ -317,22 +318,25 @@ EStateTreeRunStatus FStateTreeExecutionContext::Stop()
 	EventQueue.Reset();
 
 	// Exit states if still in some valid state.
-	if (Exec.TreeRunStatus == EStateTreeRunStatus::Running && !Exec.ActiveStates.IsEmpty())
+	if (Exec.TreeRunStatus == EStateTreeRunStatus::Running)
 	{
 		// Transition to Succeeded state.
 		FStateTreeTransitionResult Transition;
-		Transition.TargetState = FStateTreeStateHandle::Stopped;
+		Transition.TargetState = FStateTreeStateHandle::FromCompletionStatus(CompletionStatus);
 		Transition.CurrentActiveStates = Exec.ActiveStates;
-		Transition.CurrentRunStatus = Exec.LastTickStatus;
-		Transition.NextActiveStates = FStateTreeActiveStates(FStateTreeStateHandle::Stopped);
+		Transition.CurrentRunStatus = CompletionStatus;
+		Transition.NextActiveStates = FStateTreeActiveStates(Transition.TargetState);
 
-		ExitState(Transition);
+		if (!Exec.ActiveStates.IsEmpty())
+		{
+			ExitState(Transition);
+		}
 
-		Result = EStateTreeRunStatus::Stopped;
-	}
+		// Stop evaluators and global tasks.
+		StopEvaluatorsAndGlobalTasks(CompletionStatus);
 
-	// Stop evaluators and global tasks.
-	StopEvaluatorsAndGlobalTasks();
+		Result = CompletionStatus;
+	}
 
 	// Trace before resetting the instance data since it is required to provide all the event information
 	STATETREE_TRACE_ACTIVE_STATES_EVENT(FStateTreeActiveStates());
@@ -386,7 +390,7 @@ EStateTreeRunStatus FStateTreeExecutionContext::Tick(const float DeltaTime)
 	const EStateTreeRunStatus EvalAndGlobalTaskStatus = TickEvaluatorsAndGlobalTasks(DeltaTime);
 	if (EvalAndGlobalTaskStatus != EStateTreeRunStatus::Running)
 	{
-		return Stop();
+		return Stop(EvalAndGlobalTaskStatus);
 	}
 
 	if (Exec->LastTickStatus == EStateTreeRunStatus::Running)
@@ -427,6 +431,10 @@ EStateTreeRunStatus FStateTreeExecutionContext::Tick(const float DeltaTime)
 				// Transition to a terminal state (succeeded/failed), or default transition failed.
 				Exec->TreeRunStatus = NextTransition.TargetState.ToCompletionStatus();
 				Exec->ActiveStates.Reset();
+
+				// Stop evaluators and global tasks.
+				StopEvaluatorsAndGlobalTasks(Exec->TreeRunStatus);
+
 				return Exec->TreeRunStatus;
 			}
 
@@ -465,7 +473,12 @@ EStateTreeRunStatus FStateTreeExecutionContext::Tick(const float DeltaTime)
 		// Should not happen. This may happen if a state completion transition could not be selected. 
 		STATETREE_LOG(Error, TEXT("%hs: Failed to select state on '%s' using StateTree '%s'. This should not happen, state completion transition is likely missing."),
 			__FUNCTION__, *GetNameSafe(&Owner), *GetFullNameSafe(&StateTree));
+
 		Exec->TreeRunStatus = EStateTreeRunStatus::Failed;
+
+		// Stop evaluators and global tasks.
+		StopEvaluatorsAndGlobalTasks(Exec->TreeRunStatus);
+
 		return Exec->TreeRunStatus;
 	}
 
@@ -1061,7 +1074,7 @@ EStateTreeRunStatus FStateTreeExecutionContext::TickEvaluatorsAndGlobalTasks(con
 	return Result;
 }
 
-bool FStateTreeExecutionContext::StartEvaluatorsAndGlobalTasks(FStateTreeIndex16& OutLastInitializedTaskIndex)
+EStateTreeRunStatus FStateTreeExecutionContext::StartEvaluatorsAndGlobalTasks(FStateTreeIndex16& OutLastInitializedTaskIndex)
 {
 	CSV_SCOPED_TIMING_STAT_EXCLUSIVE(StateTree_StartEvaluators);
 	STATETREE_TRACE_SCOPED_PHASE(EStateTreeUpdatePhase::StartGlobalTasks);
@@ -1069,7 +1082,7 @@ bool FStateTreeExecutionContext::StartEvaluatorsAndGlobalTasks(FStateTreeIndex16
 	STATETREE_CLOG(StateTree.EvaluatorsNum > 0 || StateTree.GlobalTasksNum > 0, Verbose, TEXT("Start Evaluators & Global tasks"));
 
 	OutLastInitializedTaskIndex = FStateTreeIndex16();
-	bool bResult = true;
+	EStateTreeRunStatus Result = EStateTreeRunStatus::Running;
 	
 	// Start evaluators
 	int32 InstanceStructIndex = 1; // Exec is at index 0
@@ -1109,19 +1122,20 @@ bool FStateTreeExecutionContext::StartEvaluatorsAndGlobalTasks(FStateTreeIndex16
 		STATETREE_LOG(Verbose, TEXT("  Start: '%s'"), *Task.Name.ToString());
 		{
 			QUICK_SCOPE_CYCLE_COUNTER(StateTree_Task_TreeStart);
-			if (Task.EnterState(*this, Transition) != EStateTreeRunStatus::Running)
+			EStateTreeRunStatus TaskStatus = Task.EnterState(*this, Transition); 
+			if (TaskStatus != EStateTreeRunStatus::Running)
 			{
 				OutLastInitializedTaskIndex = FStateTreeIndex16(TaskIndex);
-				bResult = false;
+				Result = TaskStatus;
 				break;
 			}
 		}
 	}
 
-	return bResult;
+	return Result;
 }
 
-void FStateTreeExecutionContext::StopEvaluatorsAndGlobalTasks(const FStateTreeIndex16 LastInitializedTaskIndex)
+void FStateTreeExecutionContext::StopEvaluatorsAndGlobalTasks(const EStateTreeRunStatus CompletionStatus, const FStateTreeIndex16 LastInitializedTaskIndex)
 {
 	CSV_SCOPED_TIMING_STAT_EXCLUSIVE(StateTree_StopEvaluators);
 	STATETREE_TRACE_SCOPED_PHASE(EStateTreeUpdatePhase::StopGlobalTasks);
@@ -1158,15 +1172,17 @@ void FStateTreeExecutionContext::StopEvaluatorsAndGlobalTasks(const FStateTreeIn
 	}
 
 
-	// Call in reverse order.	
-	const FStateTreeTransitionResult Transition = {}; // Empty transition
+	// Call in reverse order.
+	FStateTreeTransitionResult Transition;
+	Transition.TargetState = FStateTreeStateHandle::FromCompletionStatus(CompletionStatus);
+	Transition.CurrentActiveStates = {};
+	Transition.CurrentRunStatus = CompletionStatus;
+	Transition.NextActiveStates = FStateTreeActiveStates(Transition.TargetState);
 
 	for (int32 TaskIndex = (StateTree.GlobalTasksBegin + StateTree.GlobalTasksNum) - 1;  TaskIndex >= StateTree.GlobalTasksBegin ; TaskIndex--)
 	{
 		const FStateTreeTaskBase& Task =  StateTree.Nodes[TaskIndex].Get<const FStateTreeTaskBase>();
 
-		UE_LOG(LogTemp, Error, TEXT("**** %s: Exit Global Task %s exec=%s"), *GetNameSafe(&Owner), *Task.Name.ToString(), *LexToString(TaskIndex <= LastInitializedTaskIndex.Get()));
-
 		// Relying here that invalid value of LastInitializedTaskIndex == MAX_uint16.
 		if (TaskIndex <= LastInitializedTaskIndex.Get())
 		{
@@ -1412,7 +1428,7 @@ bool FStateTreeExecutionContext::RequestTransition(const FStateTreeStateHandle N
 	if (NextState.IsCompletionState())
 	{
 		NextTransition.CurrentActiveStates = Exec.ActiveStates;
-		NextTransition.CurrentRunStatus = Exec.TreeRunStatus;
+		NextTransition.CurrentRunStatus = Exec.LastTickStatus;
 		NextTransition.SourceState = CurrentlyProcessedState;
 		NextTransition.TargetState = NextState;
 		NextTransition.NextActiveStates = FStateTreeActiveStates(NextState);
@@ -1427,7 +1443,7 @@ bool FStateTreeExecutionContext::RequestTransition(const FStateTreeStateHandle N
 	{
 		// NotSet is no-operation, but can be used to mask a transition at parent state. Returning unset keeps updating current state.
 		NextTransition.CurrentActiveStates = Exec.ActiveStates;
-		NextTransition.CurrentRunStatus = Exec.TreeRunStatus;
+		NextTransition.CurrentRunStatus = Exec.LastTickStatus;
 		NextTransition.SourceState = CurrentlyProcessedState;
 		NextTransition.TargetState = FStateTreeStateHandle::Invalid;
 		NextTransition.NextActiveStates.Reset();
@@ -1440,7 +1456,7 @@ bool FStateTreeExecutionContext::RequestTransition(const FStateTreeStateHandle N
 	if (SelectState(NextState, NewActiveState, VisitedStates))
 	{
 		NextTransition.CurrentActiveStates = Exec.ActiveStates;
-		NextTransition.CurrentRunStatus = Exec.TreeRunStatus;
+		NextTransition.CurrentRunStatus = Exec.LastTickStatus;
 		NextTransition.SourceState = CurrentlyProcessedState;
 		NextTransition.TargetState = NextState;
 		NextTransition.NextActiveStates = NewActiveState;
diff --git a/Engine/Plugins/Runtime/StateTree/Source/StateTreeModule/Public/StateTreeExecutionContext.h b/Engine/Plugins/Runtime/StateTree/Source/StateTreeModule/Public/StateTreeExecutionContext.h
--- a/Engine/Plugins/Runtime/StateTree/Source/StateTreeModule/Public/StateTreeExecutionContext.h
+++ b/Engine/Plugins/Runtime/StateTree/Source/StateTreeModule/Public/StateTreeExecutionContext.h
@@ -83,8 +83,12 @@ public:
 	/** Start executing. */
 	EStateTreeRunStatus Start();
 	
-	/** Stop executing. */
-	EStateTreeRunStatus Stop();
+	/**
+	 * Stop executing if the tree is running.
+	 * @param CompletionStatus Status (and terminal state) reported in the transition when the tree is stopped.
+	 * @return Tree execution status at stop, can be CompletionStatus, or earlier status if the tree is not running. 
+	 */
+	EStateTreeRunStatus Stop(const EStateTreeRunStatus CompletionStatus = EStateTreeRunStatus::Stopped);
 
 	/**
 	 * Tick the state tree logic.
@@ -310,14 +314,14 @@ protected:
 
 	/**
 	 * Starts evaluators and global tasks.
-	 * @return true if all evaluators and tasks were started, or false if any failed.
+	 * @return run status returned by the global tasks.
 	 */
-	bool StartEvaluatorsAndGlobalTasks(FStateTreeIndex16& OutLastInitializedTaskIndex);
+	EStateTreeRunStatus StartEvaluatorsAndGlobalTasks(FStateTreeIndex16& OutLastInitializedTaskIndex);
 
 	/**
 	 * Stops evaluators and global tasks.
 	 */
-	void StopEvaluatorsAndGlobalTasks(const FStateTreeIndex16 LastInitializedTaskIndex = FStateTreeIndex16());
+	void StopEvaluatorsAndGlobalTasks(const EStateTreeRunStatus CompletionStatus, const FStateTreeIndex16 LastInitializedTaskIndex = FStateTreeIndex16());
 
 	/**
 	 * Ticks tasks of all active states starting from current state by delta time.
diff --git a/Engine/Plugins/Runtime/StateTree/Source/StateTreeModule/Public/StateTreeTypes.h b/Engine/Plugins/Runtime/StateTree/Source/StateTreeModule/Public/StateTreeTypes.h
--- a/Engine/Plugins/Runtime/StateTree/Source/StateTreeModule/Public/StateTreeTypes.h
+++ b/Engine/Plugins/Runtime/StateTree/Source/StateTreeModule/Public/StateTreeTypes.h
@@ -227,6 +227,23 @@ struct STATETREEMODULE_API FStateTreeStateHandle
 		return EStateTreeRunStatus::Unset;
 	}
 
+	static FStateTreeStateHandle FromCompletionStatus(const EStateTreeRunStatus Status)
+	{
+		if (Status == EStateTreeRunStatus::Succeeded)
+		{
+			return FStateTreeStateHandle::Succeeded;
+		}
+		else if (Status == EStateTreeRunStatus::Failed)
+		{
+			return FStateTreeStateHandle::Failed;
+		}
+		else if (Status == EStateTreeRunStatus::Stopped)
+		{
+			return FStateTreeStateHandle::Stopped;
+		}
+		return {};
+	}
+	
 	bool operator==(const FStateTreeStateHandle& RHS) const { return Index == RHS.Index; }
 	bool operator!=(const FStateTreeStateHandle& RHS) const { return Index != RHS.Index; }
 
diff --git a/Engine/Plugins/Runtime/StateTree/Source/StateTreeTestSuite/Private/StateTreeTest.cpp b/Engine/Plugins/Runtime/StateTree/Source/StateTreeTestSuite/Private/StateTreeTest.cpp
--- a/Engine/Plugins/Runtime/StateTree/Source/StateTreeTestSuite/Private/StateTreeTest.cpp
+++ b/Engine/Plugins/Runtime/StateTree/Source/StateTreeTestSuite/Private/StateTreeTest.cpp
@@ -1485,6 +1485,449 @@ struct FStateTreeTest_InfiniteLoop : FAITestBase
 IMPLEMENT_AI_INSTANT_TEST(FStateTreeTest_InfiniteLoop, "System.StateTree.InfiniteLoop");
 
 
+//
+// The stop tests test how the combinations of execution path to stop the tree are reported on ExitState() transition.  
+//
+struct FStateTreeTest_Stop : FAITestBase
+{
+	UStateTree& SetupTree()
+	{
+		UStateTree& StateTree = UE::StateTree::Tests::NewStateTree(&GetWorld());
+		UStateTreeEditorData& EditorData = *Cast<UStateTreeEditorData>(StateTree.EditorData);
+
+		UStateTreeState& Root = EditorData.AddSubTree(FName(TEXT("Root")));
+		UStateTreeState& StateA = Root.AddChildState(FName(TEXT("A")));
+		TStateTreeEditorNode<FTestTask_Stand>& TaskA = StateA.AddTask<FTestTask_Stand>(TaskAName);
+		TStateTreeEditorNode<FTestTask_Stand>& GlobalTask = EditorData.AddGlobalTask<FTestTask_Stand>(GlobalTaskName);
+
+		// Transition success 
+		StateA.AddTransition(EStateTreeTransitionTrigger::OnStateSucceeded, EStateTreeTransitionType::Succeeded);
+		StateA.AddTransition(EStateTreeTransitionTrigger::OnStateFailed, EStateTreeTransitionType::Failed);
+
+		GlobalTask.GetNode().TicksToCompletion = GlobalTaskTicks;
+		GlobalTask.GetNode().TickCompletionResult = GlobalTaskStatus;
+		GlobalTask.GetNode().EnterStateResult = GlobalTaskEnterStatus;
+
+		TaskA.GetNode().TicksToCompletion = NormalTaskTicks;
+		TaskA.GetNode().TickCompletionResult = NormalTaskStatus;
+		TaskA.GetNode().EnterStateResult = NormalTaskEnterStatus;
+
+		return StateTree;
+	}
+	
+	virtual bool InstantTest() override
+	{
+		UStateTree& StateTree = SetupTree();
+		
+		FStateTreeCompilerLog Log;
+		FStateTreeCompiler Compiler(Log);
+		const bool bResult = Compiler.Compile(StateTree);
+
+		AITEST_TRUE("StateTree should get compiled", bResult);
+
+		EStateTreeRunStatus Status = EStateTreeRunStatus::Unset;
+		FStateTreeInstanceData InstanceData;
+		FTestStateTreeExecutionContext Exec(StateTree, StateTree, InstanceData);
+		const bool bInitSucceeded = Exec.IsValid();
+		AITEST_TRUE("StateTree should init", bInitSucceeded);
+
+		const FString TickStr(TEXT("Tick"));
+		const FString EnterStateStr(TEXT("EnterState"));
+		const FString ExitStateStr(TEXT("ExitState"));
+
+		Status = Exec.Start();
+		AITEST_EQUAL("Start should be running", Status, EStateTreeRunStatus::Running);
+		AITEST_TRUE("StateTree GlobalTask should enter state", Exec.Expect(GlobalTaskName, EnterStateStr));
+		AITEST_TRUE("StateTree TaskA should enter state", Exec.Expect(TaskAName, EnterStateStr));
+		Exec.LogClear();
+
+		Status = Exec.Tick(0.1f);
+		AITEST_EQUAL("Tree should end expectedly", Status, ExpectedStatusAfterTick);
+		AITEST_TRUE("StateTree GlobalTask should get exit state expectedly", Exec.Expect(GlobalTaskName, ExpectedExitStatusStr));
+		AITEST_TRUE("StateTree TaskA should get exit state expectedly", Exec.Expect(TaskAName, ExpectedExitStatusStr));
+		Exec.LogClear();
+		
+		return true;
+	}
+
+protected:
+
+	const FName GlobalTaskName = FName(TEXT("GlobalTask"));
+	const FName TaskAName = FName(TEXT("TaskA"));
+	
+	EStateTreeRunStatus NormalTaskStatus = EStateTreeRunStatus::Succeeded;
+	EStateTreeRunStatus NormalTaskEnterStatus = EStateTreeRunStatus::Running;
+	int32 NormalTaskTicks = 1;
+
+	EStateTreeRunStatus GlobalTaskStatus = EStateTreeRunStatus::Succeeded;
+	EStateTreeRunStatus GlobalTaskEnterStatus = EStateTreeRunStatus::Running;
+	int32 GlobalTaskTicks = 1;
+
+	EStateTreeRunStatus ExpectedStatusAfterTick = EStateTreeRunStatus::Succeeded;
+
+	FString ExpectedExitStatusStr = TEXT("ExitSucceeded");
+};
+
+struct FStateTreeTest_Stop_NormalSucceeded : FStateTreeTest_Stop
+{
+	virtual bool SetUp() override
+	{
+		// Normal task completes as succeeded.
+		NormalTaskStatus = EStateTreeRunStatus::Succeeded;
+		NormalTaskTicks = 1;
+
+		// Global task completes later
+		GlobalTaskTicks = 2;
+
+		// Tree should complete as succeeded.
+		ExpectedStatusAfterTick = EStateTreeRunStatus::Succeeded;
+		
+		// Tasks should have Transition.CurrentRunStatus as succeeded 
+		ExpectedExitStatusStr = TEXT("ExitSucceeded");
+
+		return true;
+	}
+};
+IMPLEMENT_AI_INSTANT_TEST(FStateTreeTest_Stop_NormalSucceeded, "System.StateTree.Stop.NormalSucceeded");
+
+struct FStateTreeTest_Stop_NormalFailed : FStateTreeTest_Stop
+{
+	virtual bool SetUp() override
+	{
+		// Normal task completes as failed.
+		NormalTaskStatus = EStateTreeRunStatus::Failed;
+		NormalTaskTicks = 1;
+
+		// Global task completes later.
+		GlobalTaskTicks = 2;
+
+		// Tree should complete as failed.
+		ExpectedStatusAfterTick = EStateTreeRunStatus::Failed;
+		
+		// Tasks should have Transition.CurrentRunStatus as failed.
+		ExpectedExitStatusStr = TEXT("ExitFailed");
+
+		return true;
+	}
+};
+IMPLEMENT_AI_INSTANT_TEST(FStateTreeTest_Stop_NormalFailed, "System.StateTree.Stop.NormalFailed");
+
+
+struct FStateTreeTest_Stop_GlobalSucceeded : FStateTreeTest_Stop
+{
+	virtual bool SetUp() override
+	{
+		// Normal task completes later.
+		NormalTaskTicks = 2;
+
+		// Global task completes as succeeded.
+		GlobalTaskStatus = EStateTreeRunStatus::Succeeded;
+		GlobalTaskTicks = 1;
+
+		// Tree should complete as succeeded.
+		ExpectedStatusAfterTick = EStateTreeRunStatus::Succeeded;
+		
+		// Tasks should have Transition.CurrentRunStatus as succeeded.
+		ExpectedExitStatusStr = TEXT("ExitSucceeded");
+
+		return true;
+	}
+};
+IMPLEMENT_AI_INSTANT_TEST(FStateTreeTest_Stop_GlobalSucceeded, "System.StateTree.Stop.GlobalSucceeded");
+
+struct FStateTreeTest_Stop_GlobalFailed : FStateTreeTest_Stop
+{
+	virtual bool SetUp() override
+	{
+		// Normal task completes later
+		NormalTaskTicks = 2;
+
+		// Global task completes as failed.
+		GlobalTaskStatus = EStateTreeRunStatus::Failed;
+		GlobalTaskTicks = 1;
+
+		// Tree should complete as failed.
+		ExpectedStatusAfterTick = EStateTreeRunStatus::Failed;
+		
+		// Tasks should have Transition.CurrentRunStatus as failed.
+		ExpectedExitStatusStr = TEXT("ExitFailed");
+
+		return true;
+	}
+};
+IMPLEMENT_AI_INSTANT_TEST(FStateTreeTest_Stop_GlobalFailed, "System.StateTree.Stop.GlobalFailed");
+
+
+//
+// Tests combinations of completing the tree on EnterState.
+//
+struct FStateTreeTest_StopEnterNormal : FStateTreeTest_Stop
+{
+	virtual bool InstantTest() override
+	{
+		UStateTree& StateTree = SetupTree();
+		
+		FStateTreeCompilerLog Log;
+		FStateTreeCompiler Compiler(Log);
+		const bool bResult = Compiler.Compile(StateTree);
+
+		AITEST_TRUE("StateTree should get compiled", bResult);
+
+		EStateTreeRunStatus Status = EStateTreeRunStatus::Unset;
+		FStateTreeInstanceData InstanceData;
+		FTestStateTreeExecutionContext Exec(StateTree, StateTree, InstanceData);
+		const bool bInitSucceeded = Exec.IsValid();
+		AITEST_TRUE("StateTree should init", bInitSucceeded);
+
+		const FString TickStr(TEXT("Tick"));
+		const FString EnterStateStr(TEXT("EnterState"));
+		const FString ExitStateStr(TEXT("ExitState"));
+
+		// If a normal task fails at start, the last tick status will be failed, but transition handling (and final execution status) will take place next tick. 
+		Status = Exec.Start();
+		AITEST_EQUAL("Tree should be running after start", Status, EStateTreeRunStatus::Running);
+		AITEST_EQUAL("Last execution status should be expected value", Exec.GetLastTickStatus(), ExpectedStatusAfterStart);
+
+		// Handles any transitions from failed transition
+		Status = Exec.Tick(0.1f);
+		AITEST_EQUAL("Start should be expected value", Status, ExpectedStatusAfterStart);
+		AITEST_TRUE("StateTree GlobalTask should get exit state expectedly", Exec.Expect(GlobalTaskName, ExpectedExitStatusStr));
+
+		AITEST_TRUE("StateTree TaskA should enter state", Exec.Expect(TaskAName, EnterStateStr));
+		AITEST_TRUE("StateTree TaskA should report exit status", Exec.Expect(TaskAName, ExpectedExitStatusStr));
+
+		Exec.LogClear();
+		
+		return true;
+	}
+
+	EStateTreeRunStatus ExpectedStatusAfterStart = EStateTreeRunStatus::Succeeded;
+	FString ExpectedExitStatusStr = TEXT("ExitSucceeded");
+	bool bExpectNormalTaskToRun = true; 
+};
+
+struct FStateTreeTest_Stop_NormalEnterSucceeded : FStateTreeTest_StopEnterNormal
+{
+	virtual bool SetUp() override
+	{
+		// Tasks should complete later.
+		NormalTaskTicks = 2;
+		GlobalTaskTicks = 2;
+
+		// Normal task EnterState as succeeded, completion is handled using completion transitions.
+		NormalTaskEnterStatus = EStateTreeRunStatus::Succeeded;
+
+		// Tree should complete as succeeded.
+		ExpectedStatusAfterStart = EStateTreeRunStatus::Succeeded;
+		
+		// Tasks should have Transition.CurrentRunStatus as succeeded.
+		ExpectedExitStatusStr = TEXT("ExitSucceeded");
+
+		return true;
+	}
+};
+IMPLEMENT_AI_INSTANT_TEST(FStateTreeTest_Stop_NormalEnterSucceeded, "System.StateTree.Stop.NormalEnterSucceeded");
+
+struct FStateTreeTest_Stop_NormalEnterFailed : FStateTreeTest_StopEnterNormal
+{
+	virtual bool SetUp() override
+	{
+		// Tasks should complete later.
+		NormalTaskTicks = 2;
+		GlobalTaskTicks = 2;
+
+		// Normal task EnterState as failed, completion is handled using completion transitions.
+		NormalTaskEnterStatus = EStateTreeRunStatus::Failed;
+
+		// Tree should complete as failed.
+		ExpectedStatusAfterStart = EStateTreeRunStatus::Failed;
+		
+		// Tasks should have Transition.CurrentRunStatus as failed.
+		ExpectedExitStatusStr = TEXT("ExitFailed");
+
+		return true;
+	}
+};
+IMPLEMENT_AI_INSTANT_TEST(FStateTreeTest_Stop_NormalEnterFailed, "System.StateTree.Stop.NormalEnterFailed");
+
+
+
+
+struct FStateTreeTest_StopEnterGlobal : FStateTreeTest_Stop
+{
+	virtual bool InstantTest() override
+	{
+		UStateTree& StateTree = SetupTree();
+		
+		FStateTreeCompilerLog Log;
+		FStateTreeCompiler Compiler(Log);
+		const bool bResult = Compiler.Compile(StateTree);
+
+		AITEST_TRUE("StateTree should get compiled", bResult);
+
+		EStateTreeRunStatus Status = EStateTreeRunStatus::Unset;
+		FStateTreeInstanceData InstanceData;
+		FTestStateTreeExecutionContext Exec(StateTree, StateTree, InstanceData);
+		const bool bInitSucceeded = Exec.IsValid();
+		AITEST_TRUE("StateTree should init", bInitSucceeded);
+
+		const FString TickStr(TEXT("Tick"));
+		const FString EnterStateStr(TEXT("EnterState"));
+		const FString ExitStateStr(TEXT("ExitState"));
+
+		Status = Exec.Start();
+		AITEST_EQUAL("Start should be expected value", Status, ExpectedStatusAfterStart);
+		AITEST_TRUE("StateTree GlobalTask should get exit state expectedly", Exec.Expect(GlobalTaskName, ExpectedExitStatusStr));
+
+		// Normal tasks should not run
+		AITEST_FALSE("StateTree TaskA should not enter state", Exec.Expect(TaskAName, EnterStateStr));
+		AITEST_FALSE("StateTree TaskA should not report exit status", Exec.Expect(TaskAName, ExpectedExitStatusStr));
+
+		Exec.LogClear();
+		
+		return true;
+	}
+
+	EStateTreeRunStatus ExpectedStatusAfterStart = EStateTreeRunStatus::Succeeded;
+	FString ExpectedExitStatusStr = TEXT("ExitSucceeded");
+};
+
+struct FStateTreeTest_Stop_GlobalEnterSucceeded : FStateTreeTest_StopEnterGlobal
+{
+	virtual bool SetUp() override
+	{
+		// Tasks should complete later.
+		NormalTaskTicks = 2;
+		GlobalTaskTicks = 2;
+
+		// Global task EnterState as succeeded, completion is handled directly based on the global task status.
+		GlobalTaskEnterStatus = EStateTreeRunStatus::Succeeded;
+
+		// Tree should complete as succeeded.
+		ExpectedStatusAfterStart = EStateTreeRunStatus::Succeeded;
+		
+		// Tasks should have Transition.CurrentRunStatus as Succeeded.
+		ExpectedExitStatusStr = TEXT("ExitSucceeded");
+
+		return true;
+	}
+};
+IMPLEMENT_AI_INSTANT_TEST(FStateTreeTest_Stop_GlobalEnterSucceeded, "System.StateTree.Stop.GlobalEnterSucceeded");
+
+struct FStateTreeTest_Stop_GlobalEnterFailed : FStateTreeTest_StopEnterGlobal
+{
+	virtual bool SetUp() override
+	{
+		// Tasks should complete later.
+		NormalTaskTicks = 2;
+		GlobalTaskTicks = 2;
+
+		// Global task EnterState as failed, completion is handled directly based on the global task status.
+		GlobalTaskEnterStatus = EStateTreeRunStatus::Failed;
+
+		// Tree should complete as failed.
+		ExpectedStatusAfterStart = EStateTreeRunStatus::Failed;
+		
+		// Tasks should have Transition.CurrentRunStatus as failed.
+		ExpectedExitStatusStr = TEXT("ExitFailed");
+
+		return true;
+	}
+};
+IMPLEMENT_AI_INSTANT_TEST(FStateTreeTest_Stop_GlobalEnterFailed, "System.StateTree.Stop.GlobalEnterFailed");
+
+struct FStateTreeTest_Stop_ExternalStop : FStateTreeTest_Stop
+{
+	virtual bool SetUp() override
+	{
+		// Tasks should complete later.
+		NormalTaskTicks = 2;
+		GlobalTaskTicks = 2;
+
+		// Tree should tick and keep on running.
+		ExpectedStatusAfterTick = EStateTreeRunStatus::Running;
+
+		// Tree should stop as stopped.
+		ExpectedStatusAfterStop = EStateTreeRunStatus::Stopped;
+		
+		// Tasks should have Transition.CurrentRunStatus as stopped. 
+		ExpectedExitStatusStr = TEXT("ExitStopped");
+
+		return true;
+	}
+	
+	virtual bool InstantTest() override
+	{
+		UStateTree& StateTree = SetupTree();
+		
+		FStateTreeCompilerLog Log;
+		FStateTreeCompiler Compiler(Log);
+		const bool bResult = Compiler.Compile(StateTree);
+
+		AITEST_TRUE("StateTree should get compiled", bResult);
+
+		EStateTreeRunStatus Status = EStateTreeRunStatus::Unset;
+		FStateTreeInstanceData InstanceData;
+		FTestStateTreeExecutionContext Exec(StateTree, StateTree, InstanceData);
+		const bool bInitSucceeded = Exec.IsValid();
+		AITEST_TRUE("StateTree should init", bInitSucceeded);
+
+		const FString TickStr(TEXT("Tick"));
+		const FString EnterStateStr(TEXT("EnterState"));
+		const FString ExitStateStr(TEXT("ExitState"));
+
+		Status = Exec.Start();
+		AITEST_EQUAL("Start should be running", Status, EStateTreeRunStatus::Running);
+		AITEST_TRUE("StateTree GlobalTask should enter state", Exec.Expect(GlobalTaskName, EnterStateStr));
+		AITEST_TRUE("StateTree TaskA should enter state", Exec.Expect(TaskAName, EnterStateStr));
+		Exec.LogClear();
+
+		Status = Exec.Tick(0.1f);
+		AITEST_EQUAL("Tree should end expectedly", Status, ExpectedStatusAfterTick);
+		Exec.LogClear();
+
+		Status = Exec.Stop(EStateTreeRunStatus::Stopped);
+		AITEST_EQUAL("Start should be running", Status, ExpectedStatusAfterStop);
+		if (!ExpectedExitStatusStr.IsEmpty())
+		{
+			AITEST_TRUE("StateTree GlobalTask should get exit state expectedly", Exec.Expect(GlobalTaskName, ExpectedExitStatusStr));
+			AITEST_TRUE("StateTree TaskA should get exit state expectedly", Exec.Expect(TaskAName, ExpectedExitStatusStr));
+		}
+		
+		return true;
+	}
+
+	EStateTreeRunStatus ExpectedStatusAfterStop = EStateTreeRunStatus::Stopped;
+	
+};
+IMPLEMENT_AI_INSTANT_TEST(FStateTreeTest_Stop_ExternalStop, "System.StateTree.Stop.ExternalStop");
+
+struct FStateTreeTest_Stop_AlreadyStopped : FStateTreeTest_Stop_ExternalStop
+{
+	virtual bool SetUp() override
+	{
+		// Normal task completes before stop.
+		NormalTaskTicks = 1;
+		NormalTaskStatus = EStateTreeRunStatus::Succeeded;
+
+		// Global task completes later
+		GlobalTaskTicks = 2;
+
+		// Tree should tick stop as succeeded.
+		ExpectedStatusAfterTick = EStateTreeRunStatus::Succeeded;
+
+		// Tree is already stopped, should keep the status (not Stopped).
+		ExpectedStatusAfterStop = EStateTreeRunStatus::Succeeded;
+		
+		// Skip exit status check.
+		ExpectedExitStatusStr = TEXT("");
+
+		return true;
+	}
+};
+IMPLEMENT_AI_INSTANT_TEST(FStateTreeTest_Stop_AlreadyStopped, "System.StateTree.Stop.AlreadyStopped");
+
+
 UE_ENABLE_OPTIMIZATION_SHIP
 
 #undef LOCTEXT_NAMESPACE
diff --git a/Engine/Plugins/Runtime/StateTree/Source/StateTreeTestSuite/Private/StateTreeTestTypes.h b/Engine/Plugins/Runtime/StateTree/Source/StateTreeTestSuite/Private/StateTreeTestTypes.h
--- a/Engine/Plugins/Runtime/StateTree/Source/StateTreeTestSuite/Private/StateTreeTestTypes.h
+++ b/Engine/Plugins/Runtime/StateTree/Source/StateTreeTestSuite/Private/StateTreeTestTypes.h
@@ -238,6 +238,20 @@ struct FTestTask_Stand : public FStateTreeTaskBase
 	virtual void ExitState(FStateTreeExecutionContext& Context, const FStateTreeTransitionResult& Transition) const override
 	{
 		FTestStateTreeExecutionContext& TestContext = static_cast<FTestStateTreeExecutionContext&>(Context);
+
+		if (Transition.CurrentRunStatus == EStateTreeRunStatus::Succeeded)
+		{
+			TestContext.Log(Name, TEXT("ExitSucceeded"));
+		}
+		else if (Transition.CurrentRunStatus == EStateTreeRunStatus::Failed)
+		{
+			TestContext.Log(Name, TEXT("ExitFailed"));
+		}
+		else if (Transition.CurrentRunStatus == EStateTreeRunStatus::Stopped)
+		{
+			TestContext.Log(Name, TEXT("ExitStopped"));
+		}
+		
 		TestContext.Log(Name, TEXT("ExitState"));
 	}
 
@@ -256,14 +270,14 @@ struct FTestTask_Stand : public FStateTreeTaskBase
 		
 		InstanceData.CurrentTick++;
 		
-		return (InstanceData.CurrentTick >= TicksToCompletion) ? TickResult : EStateTreeRunStatus::Running;
+		return (InstanceData.CurrentTick >= TicksToCompletion) ? TickCompletionResult : EStateTreeRunStatus::Running;
 	};
 
 	UPROPERTY(EditAnywhere, Category = Parameter)
 	int32 TicksToCompletion = 1;
 
 	UPROPERTY(EditAnywhere, Category = Parameter)
-	EStateTreeRunStatus TickResult = EStateTreeRunStatus::Succeeded;
+	EStateTreeRunStatus TickCompletionResult = EStateTreeRunStatus::Succeeded;
 
 	UPROPERTY(EditAnywhere, Category = Parameter)
 	EStateTreeRunStatus EnterStateResult = EStateTreeRunStatus::Running;
