Updated names and descriptions for GameplayTag StateTree conditions so they are more consistent with the actual BP library
#rnx

diff --git a/Engine/Plugins/Runtime/StateTree/Source/StateTreeModule/Public/Conditions/StateTreeGameplayTagConditions.h b/Engine/Plugins/Runtime/StateTree/Source/StateTreeModule/Public/Conditions/StateTreeGameplayTagConditions.h
--- a/Engine/Plugins/Runtime/StateTree/Source/StateTreeModule/Public/Conditions/StateTreeGameplayTagConditions.h
+++ b/Engine/Plugins/Runtime/StateTree/Source/StateTreeModule/Public/Conditions/StateTreeGameplayTagConditions.h
@@ -9,18 +9,26 @@ USTRUCT()
 struct STATETREEMODULE_API FGameplayTagMatchConditionInstanceData
 {
 	GENERATED_BODY()
-	
+
+	/** Container to check for the tag. */
 	UPROPERTY(EditAnywhere, Category = Parameter)
 	FGameplayTagContainer TagContainer;
 
+	/** Tag to check for in the container. */
 	UPROPERTY(EditAnywhere, Category = Parameter)
 	FGameplayTag Tag;
 };
 
 /**
- * Gameplay Tag match condition.
+ * HasTag condition
+ * Succeeds if the tag container has the specified tag.
+ * 
+ * Condition can be used with multiple configurations:
+ *	Does TagContainer {"A.1"} has Tag "A" ?
+ *		exact match 'false' will SUCCEED
+ *		exact match 'true' will FAIL
  */
-USTRUCT(DisplayName="Gameplay Tag Match", Category="Gameplay Tags")
+USTRUCT(DisplayName="Has Tag", Category="Gameplay Tags")
 struct STATETREEMODULE_API FGameplayTagMatchCondition : public FStateTreeConditionCommonBase
 {
 	GENERATED_BODY()
@@ -32,6 +40,7 @@ struct STATETREEMODULE_API FGameplayTagMatchCondition : public FStateTreeConditi
 	virtual const UStruct* GetInstanceDataType() const override { return FInstanceDataType::StaticStruct(); }
 	virtual bool TestCondition(FStateTreeExecutionContext& Context) const override;
 
+	/** If true, the tag has to be exactly present, if false then TagContainer will include it's parent tags while matching */
 	UPROPERTY(EditAnywhere, Category = Condition)
 	bool bExactMatch = false;
 
@@ -44,18 +53,43 @@ USTRUCT()
 struct STATETREEMODULE_API FGameplayTagContainerMatchConditionInstanceData
 {
 	GENERATED_BODY()
-	
+
+	/** Container to check if it matches any of the tags in the other container. */
 	UPROPERTY(EditAnywhere, Category = Input)
 	FGameplayTagContainer TagContainer;
 
+	/** Container to check against. */
 	UPROPERTY(EditAnywhere, Category = Parameter)
 	FGameplayTagContainer OtherContainer;
 };
 
 /**
- * Gameplay Tag Container match condition.
+ * HasAnyTags / HasAllTags condition
+ * Succeeds  if the specified tag container has ANY or ALL of the tags in the other container.
+ *
+ * Condition can be used with multiple configurations:
+ * 
+ *	Has ANY Tags:
+ *		exact match 'false':
+ *			TagContainer {"A.1"} has any of OtherContainer {"A","B"} will SUCCEED
+ * 			TagContainer {"A"} has any of OtherContainer {"A.1","B"} will FAIL
+ *		
+ *		exact match 'true':
+* 			TagContainer {"A","B"} has any of OtherContainer {"A.1"} will FAIL
+ *
+ *		If TagContainer is empty/invalid it will always fail.
+ *
+ *	Has ALL Tags:
+ *		exact match 'false':
+ *			TagContainer {"A.1","B.1"} has all of OtherContainer {"A","B"} will SUCCEED
+ *			TagContainer {"A","B"} has all of OtherContainer {"A.1","B.1"} will FAIL
+ *
+ *		exact match 'true':
+ *			TagContainer {"A.1","B.1"} has all of OtherContainer {"A","B"} will FAIL
+ *
+ *		If TagContainer is empty/invalid it will always SUCCEED, because there were no failed checks.
  */
-USTRUCT(DisplayName="Gameplay Tag Container Match", Category="Gameplay Tags")
+USTRUCT(DisplayName="Has Tags", Category="Gameplay Tags")
 struct STATETREEMODULE_API FGameplayTagContainerMatchCondition : public FStateTreeConditionCommonBase
 {
 	GENERATED_BODY()
@@ -70,6 +104,7 @@ struct STATETREEMODULE_API FGameplayTagContainerMatchCondition : public FStateTr
 	UPROPERTY(EditAnywhere, Category = Condition)
 	EGameplayContainerMatchType MatchType = EGameplayContainerMatchType::Any;
 
+	/** If true, the tag has to be exactly present, if false then TagContainer will include it's parent tags while matching */
 	UPROPERTY(EditAnywhere, Category = Condition)
 	bool bExactMatch = false;
 
@@ -82,15 +117,17 @@ USTRUCT()
 struct STATETREEMODULE_API FGameplayTagQueryConditionInstanceData
 {
 	GENERATED_BODY()
-	
+
+	/** Container that needs to match the query. */
 	UPROPERTY(EditAnywhere, Category = Input)
 	FGameplayTagContainer TagContainer;
 };
 
 /**
- * Gameplay Tag Query match condition.
+ * DoesContainerMatchTagQuery condition
+ * Succeeds if the specified tag container matches the given Tag Query.
  */
-USTRUCT(DisplayName="Gameplay Tag Query", Category="Gameplay Tags")
+USTRUCT(DisplayName="Does Container Match Tag Query", Category="Gameplay Tags")
 struct STATETREEMODULE_API FGameplayTagQueryCondition : public FStateTreeConditionCommonBase
 {
 	GENERATED_BODY()
@@ -102,6 +139,7 @@ struct STATETREEMODULE_API FGameplayTagQueryCondition : public FStateTreeConditi
 	virtual const UStruct* GetInstanceDataType() const override { return FInstanceDataType::StaticStruct(); }
 	virtual bool TestCondition(FStateTreeExecutionContext& Context) const override;
 
+	/** Query to match against */
 	UPROPERTY(EditAnywhere, Category = Condition)
 	FGameplayTagQuery TagQuery;
 
