[StateTree] added possibility to disable States and Tasks. Also possible to force condition result to true/false.

diff --git a/Engine/Plugins/Runtime/StateTree/Source/StateTreeEditorModule/Private/Customizations/StateTreeEditorNodeDetails.cpp b/Engine/Plugins/Runtime/StateTree/Source/StateTreeEditorModule/Private/Customizations/StateTreeEditorNodeDetails.cpp
--- a/Engine/Plugins/Runtime/StateTree/Source/StateTreeEditorModule/Private/Customizations/StateTreeEditorNodeDetails.cpp
+++ b/Engine/Plugins/Runtime/StateTree/Source/StateTreeEditorModule/Private/Customizations/StateTreeEditorNodeDetails.cpp
@@ -28,8 +28,10 @@
 #include "ScopedTransaction.h"
 #include "Widgets/Input/SSearchBox.h"
 #include "EditorFontGlyphs.h"
+#include "StateTreeEditorNodeUtils.h"
 #include "Framework/Notifications/NotificationManager.h"
 #include "Widgets/Notifications/SNotificationList.h"
+#include "Debugger/StateTreeDebuggerUIExtensions.h"
 
 #define LOCTEXT_NAMESPACE "StateTreeEditor"
 
@@ -288,7 +290,7 @@ FStateTreeEditorNodeDetails::~FStateTreeEditorNodeDetails()
 void FStateTreeEditorNodeDetails::CustomizeHeader(TSharedRef<class IPropertyHandle> StructPropertyHandle, class FDetailWidgetRow& HeaderRow, IPropertyTypeCustomizationUtils& StructCustomizationUtils)
 {
 	StructProperty = StructPropertyHandle;
-	PropUtils = StructCustomizationUtils.GetPropertyUtilities().Get();
+	PropUtils = StructCustomizationUtils.GetPropertyUtilities();
 
 	NodeProperty = StructProperty->GetChildHandle(TEXT("Node"));
 	InstanceProperty = StructProperty->GetChildHandle(TEXT("Instance"));
@@ -414,16 +416,7 @@ void FStateTreeEditorNodeDetails::CustomizeHeader(TSharedRef<class IPropertyHand
 							SNew(STextBlock)
 							.Text(FEditorFontGlyphs::Paper_Plane)
 							.TextStyle(FStateTreeEditorStyle::Get(), "StateTree.DetailsIcon")
-							.Visibility_Lambda([this]()
-							{
-								// Only show on tasks
-								const UScriptStruct* ScriptStruct = nullptr;
-								if (const FStateTreeEditorNode* Node = GetCommonNode())
-								{
-									ScriptStruct = Node->Node.GetScriptStruct();
-								}
-								return ScriptStruct != nullptr && ScriptStruct->IsChildOf(FStateTreeTaskBase::StaticStruct()) ? EVisibility::Visible : EVisibility::Collapsed;
-							})
+							.Visibility(this, &FStateTreeEditorNodeDetails::IsTaskVisible)
 						]
 						+ SHorizontalBox::Slot()
 						.VAlign(VAlign_Center)
@@ -456,6 +449,13 @@ void FStateTreeEditorNodeDetails::CustomizeHeader(TSharedRef<class IPropertyHand
 				[
 					StructPropertyHandle->CreateDefaultPropertyButtonWidgets()
 				]
+
+				+ SHorizontalBox::Slot()
+				.HAlign(HAlign_Right)
+				.FillWidth(1.0f)
+				[
+					UE::StateTreeEditor::DebuggerExtensions::CreateEditorNodeWidget(StructPropertyHandle)
+				]
 			]
 		]
 		.OverrideResetToDefault(ResetOverride)
@@ -620,26 +620,20 @@ bool FStateTreeEditorNodeDetails::ShouldResetToDefault(TSharedPtr<IPropertyHandl
 
 void FStateTreeEditorNodeDetails::ResetToDefault(TSharedPtr<IPropertyHandle> PropertyHandle)
 {
-	check(StructProperty);
-	
-	GEditor->BeginTransaction(LOCTEXT("OnResetToDefault", "Reset to default"));
-
-	StructProperty->NotifyPreChange();
-	
-	TArray<void*> RawNodeData;
-	StructProperty->AccessRawData(RawNodeData);
-	for (void* Data : RawNodeData)
-	{
-		if (FStateTreeEditorNode* Node = static_cast<FStateTreeEditorNode*>(Data))
+	UE::StateTreeEditor::EditorNodeUtils::ModifyNodeInTransaction(LOCTEXT("OnTaskEnableToggled", "Toggled Task Enabled"),
+		StructProperty,
+		[](IPropertyHandle& StructPropertyHandle)
 		{
-			Node->Reset();
-		}
-	}
-
-	StructProperty->NotifyPostChange(EPropertyChangeType::ValueSet);
-	StructProperty->NotifyFinishedChangingProperties();
-
-	GEditor->EndTransaction();
+			TArray<void*> RawNodeData;
+			StructPropertyHandle.AccessRawData(RawNodeData);
+			for (void* Data : RawNodeData)
+			{
+				if (FStateTreeEditorNode* Node = static_cast<FStateTreeEditorNode*>(Data))
+				{
+					Node->Reset();
+				}
+			}
+		});
 
 	if (PropUtils)
 	{
@@ -1039,16 +1033,14 @@ void FStateTreeEditorNodeDetails::SetOperand(const EStateTreeConditionOperand Op
 	OperandProperty->SetValue((uint8)Operand);
 }
 
-
 EVisibility FStateTreeEditorNodeDetails::IsConditionVisible() const
 {
-	const UScriptStruct* ScriptStruct = nullptr;
-	if (const FStateTreeEditorNode* Node = GetCommonNode())
-	{
-		ScriptStruct = Node->Node.GetScriptStruct();
-	}
+	return UE::StateTreeEditor::EditorNodeUtils::IsConditionVisible(StructProperty);
+}
 
-	return ScriptStruct != nullptr && ScriptStruct->IsChildOf(FStateTreeConditionBase::StaticStruct()) ? EVisibility::Visible : EVisibility::Collapsed;
+EVisibility FStateTreeEditorNodeDetails::IsTaskVisible() const
+{
+	return UE::StateTreeEditor::EditorNodeUtils::IsTaskVisible(StructProperty);
 }
 
 FText FStateTreeEditorNodeDetails::GetName() const
@@ -1090,7 +1082,7 @@ FText FStateTreeEditorNodeDetails::GetName() const
 EVisibility FStateTreeEditorNodeDetails::IsNameVisible() const
 {
 	const UScriptStruct* ScriptStruct = nullptr;
-	if (const FStateTreeEditorNode* Node = GetCommonNode())
+	if (const FStateTreeEditorNode* Node = UE::StateTreeEditor::EditorNodeUtils::GetCommonNode(StructProperty))
 	{
 		ScriptStruct = Node->Node.GetScriptStruct();
 	}
@@ -1165,43 +1157,13 @@ void FStateTreeEditorNodeDetails::OnNameCommitted(const FText& NewText, ETextCom
 bool FStateTreeEditorNodeDetails::IsNameEnabled() const
 {
 	// Can only edit if we have valid instantiated type.
-	const FStateTreeEditorNode* Node = GetCommonNode();
+	const FStateTreeEditorNode* Node = UE::StateTreeEditor::EditorNodeUtils::GetCommonNode(StructProperty);
 	return Node && Node->Node.IsValid();
 }
 
-const FStateTreeEditorNode* FStateTreeEditorNodeDetails::GetCommonNode() const
-{
-	check(StructProperty);
-
-	const UScriptStruct* CommonScriptStruct = nullptr;
-	bool bMultipleValues = false;
-	TArray<void*> RawNodeData;
-	StructProperty->AccessRawData(RawNodeData);
-
-	const FStateTreeEditorNode* CommonNode = nullptr;
-	
-	for (void* Data : RawNodeData)
-	{
-		if (const FStateTreeEditorNode* Node = static_cast<FStateTreeEditorNode*>(Data))
-		{
-			if (!bMultipleValues && !CommonNode)
-			{
-				CommonNode = Node;
-			}
-			else if (CommonNode != Node)
-			{
-				CommonNode = nullptr;
-				bMultipleValues = true;
-			}
-		}
-	}
-
-	return CommonNode;
-}
-
 FText FStateTreeEditorNodeDetails::GetDisplayValueString() const
 {
-	if (const FStateTreeEditorNode* Node = GetCommonNode())
+	if (const FStateTreeEditorNode* Node = UE::StateTreeEditor::EditorNodeUtils::GetCommonNode(StructProperty))
 	{
 		if (const UScriptStruct* ScriptStruct = Node->Node.GetScriptStruct())
 		{
@@ -1226,7 +1188,7 @@ FText FStateTreeEditorNodeDetails::GetDisplayValueString() const
 
 const FSlateBrush* FStateTreeEditorNodeDetails::GetDisplayValueIcon() const
 {
-	if (const FStateTreeEditorNode* Node = GetCommonNode())
+	if (const FStateTreeEditorNode* Node = UE::StateTreeEditor::EditorNodeUtils::GetCommonNode(StructProperty))
 	{
 		if (const UScriptStruct* ScriptStruct = Node->Node.GetScriptStruct())
 		{
@@ -1777,7 +1739,7 @@ TSharedRef<SWidget> FStateTreeEditorNodeDetails::GeneratePicker()
 	
 	// Expand and select currently selected item.
 	const UStruct* CommonStruct  = nullptr;
-	if (const FStateTreeEditorNode* Node = GetCommonNode())
+	if (const FStateTreeEditorNode* Node = UE::StateTreeEditor::EditorNodeUtils::GetCommonNode(StructProperty))
 	{
 		if (const UScriptStruct* ScriptStruct = Node->Node.GetScriptStruct())
 		{
diff --git a/Engine/Plugins/Runtime/StateTree/Source/StateTreeEditorModule/Private/Customizations/StateTreeEditorNodeDetails.h b/Engine/Plugins/Runtime/StateTree/Source/StateTreeEditorModule/Private/Customizations/StateTreeEditorNodeDetails.h
--- a/Engine/Plugins/Runtime/StateTree/Source/StateTreeEditorModule/Private/Customizations/StateTreeEditorNodeDetails.h
+++ b/Engine/Plugins/Runtime/StateTree/Source/StateTreeEditorModule/Private/Customizations/StateTreeEditorNodeDetails.h
@@ -8,13 +8,18 @@
 struct FObjectKey;
 struct FOptionalSize;
 
+class FDetailWidgetRow;
+class IDetailChildrenBuilder;
 class IPropertyHandle;
+class IPropertyUtilities;
+class SComboButton;
 class SWidget;
 class SSearchBox;
 class UStateTree;
 class UStateTreeState;
 class UStateTreeEditorData;
 struct EVisibility;
+struct FStateTreeEditorNode;
 struct FStateTreePropertyPath;
 enum class EStateTreeConditionOperand : uint8;
 
@@ -30,8 +35,8 @@ public:
 	virtual ~FStateTreeEditorNodeDetails();
 	
 	/** IPropertyTypeCustomization interface */
-	virtual void CustomizeHeader(TSharedRef<IPropertyHandle> StructPropertyHandle, class FDetailWidgetRow& HeaderRow, IPropertyTypeCustomizationUtils& StructCustomizationUtils) override;
-	virtual void CustomizeChildren(TSharedRef<IPropertyHandle> StructPropertyHandle, class IDetailChildrenBuilder& StructBuilder, IPropertyTypeCustomizationUtils& StructCustomizationUtils) override;
+	virtual void CustomizeHeader(TSharedRef<IPropertyHandle> StructPropertyHandle, FDetailWidgetRow& HeaderRow, IPropertyTypeCustomizationUtils& StructCustomizationUtils) override;
+	virtual void CustomizeChildren(TSharedRef<IPropertyHandle> StructPropertyHandle, IDetailChildrenBuilder& StructBuilder, IPropertyTypeCustomizationUtils& StructCustomizationUtils) override;
 
 private:
 
@@ -64,13 +69,13 @@ private:
 	FText GetCloseParens() const;
 
 	EVisibility IsConditionVisible() const;
+	EVisibility IsTaskVisible() const;
 
 	FText GetName() const;
 	EVisibility IsNameVisible() const;
 	void OnNameCommitted(const FText& NewText, ETextCommit::Type InTextCommit) const;
 	bool IsNameEnabled() const;
 
-	const struct FStateTreeEditorNode* GetCommonNode() const;
 	FText GetDisplayValueString() const;
 	const FSlateBrush* GetDisplayValueIcon() const;
 
@@ -122,7 +127,7 @@ private:
 	
 	UScriptStruct* BaseScriptStruct = nullptr;
 	UClass* BaseClass = nullptr;
-	TSharedPtr<class SComboButton> ComboButton;
+	TSharedPtr<SComboButton> ComboButton;
 	TSharedPtr<SSearchBox> SearchBox;
 	TSharedPtr<STreeView<TSharedPtr<FStateTreeNodeTypeItem>>> NodeTypeTree;
 	bool bIsRestoringExpansion = false;
@@ -133,7 +138,7 @@ private:
 	// Save expansion state for each base node type. The expansion state does not persist between editor sessions. 
 	static TMap<FObjectKey, FCategoryExpansionState> CategoryExpansionStates;
 	
-	class IPropertyUtilities* PropUtils = nullptr;
+	TSharedPtr<IPropertyUtilities> PropUtils;
 	TSharedPtr<IPropertyHandle> StructProperty;
 	TSharedPtr<IPropertyHandle> NodeProperty;
 	TSharedPtr<IPropertyHandle> InstanceProperty;
diff --git a/Engine/Plugins/Runtime/StateTree/Source/StateTreeEditorModule/Private/Customizations/StateTreeEditorNodeUtils.cpp b/Engine/Plugins/Runtime/StateTree/Source/StateTreeEditorModule/Private/Customizations/StateTreeEditorNodeUtils.cpp
new file mode 100644
--- /dev/null
+++ b/Engine/Plugins/Runtime/StateTree/Source/StateTreeEditorModule/Private/Customizations/StateTreeEditorNodeUtils.cpp
@@ -0,0 +1,149 @@
+// Copyright Epic Games, Inc. All Rights Reserved.
+
+#include "StateTreeEditorNodeUtils.h"
+
+#include "Layout/Visibility.h"
+#include "PropertyHandle.h"
+#include "ScopedTransaction.h"
+#include "StateTreeConditionBase.h"
+#include "StateTreeEditorNode.h"
+#include "StateTreeTaskBase.h"
+
+namespace UE::StateTreeEditor::EditorNodeUtils
+{
+
+EStateTreeConditionEvaluationMode GetConditionEvaluationMode(TSharedPtr<IPropertyHandle> StructProperty)
+{
+	if (const FStateTreeEditorNode* Node = GetCommonNode(StructProperty))
+	{
+		if (const FStateTreeConditionBase* ConditionBase = Node->Node.GetPtr<FStateTreeConditionBase>())
+		{
+			return ConditionBase->EvaluationMode;
+		}
+	}
+	// Evaluate as default value
+	return EStateTreeConditionEvaluationMode::Evaluated;
+}
+
+EVisibility IsConditionVisible(TSharedPtr<IPropertyHandle> StructProperty)
+{
+	const UScriptStruct* ScriptStruct = nullptr;
+	if (const FStateTreeEditorNode* Node = GetCommonNode(StructProperty))
+	{
+		ScriptStruct = Node->Node.GetScriptStruct();
+	}
+
+	return ScriptStruct != nullptr && ScriptStruct->IsChildOf(FStateTreeConditionBase::StaticStruct()) ? EVisibility::Visible : EVisibility::Collapsed;
+}
+
+EVisibility IsTaskVisible(TSharedPtr<IPropertyHandle> StructProperty)
+{
+	const UScriptStruct* ScriptStruct = nullptr;
+	if (const FStateTreeEditorNode* Node = GetCommonNode(StructProperty))
+	{
+		ScriptStruct = Node->Node.GetScriptStruct();
+	}
+
+	return ScriptStruct != nullptr && ScriptStruct->IsChildOf(FStateTreeTaskBase::StaticStruct()) ? EVisibility::Visible : EVisibility::Collapsed;
+}
+
+bool IsTaskDisabled(TSharedPtr<IPropertyHandle> StructProperty)
+{
+	if (const FStateTreeEditorNode* Node = GetCommonNode(StructProperty))
+	{
+		if (const FStateTreeTaskBase* TaskBase = Node->Node.GetPtr<FStateTreeTaskBase>())
+		{
+			return !TaskBase->bTaskEnabled;
+		}
+	}
+
+	return false;
+}
+
+void ModifyNodeInTransaction(FText Description, TSharedPtr<IPropertyHandle> StructProperty, TFunctionRef<void(IPropertyHandle&)> Func)
+{
+	check(StructProperty);
+
+	FScopedTransaction ScopedTransaction(Description);
+
+	StructProperty->NotifyPreChange();
+
+	Func(*StructProperty);
+
+	StructProperty->NotifyPostChange(EPropertyChangeType::ValueSet);
+	StructProperty->NotifyFinishedChangingProperties();
+}
+
+const FStateTreeEditorNode* GetCommonNode(const TSharedPtr<IPropertyHandle> StructProperty)
+{
+	if (const IPropertyHandle* PropertyHandle = StructProperty.Get())
+	{
+		return GetCommonNode(*PropertyHandle);
+	}
+
+	return nullptr;
+}
+
+FStateTreeEditorNode* GetMutableCommonNode(const TSharedPtr<IPropertyHandle> StructProperty)
+{
+	if (IPropertyHandle* PropertyHandle = StructProperty.Get())
+	{
+		return GetMutableCommonNode(*PropertyHandle);
+	}
+
+	return nullptr;
+}
+
+const FStateTreeEditorNode* GetCommonNode(const IPropertyHandle& InStructProperty)
+{
+	TArray<const void*> RawNodeData;
+	InStructProperty.AccessRawData(RawNodeData);
+
+	const FStateTreeEditorNode* CommonNode = nullptr;
+
+	for (const void* Data : RawNodeData)
+	{
+		if (const FStateTreeEditorNode* Node = static_cast<const FStateTreeEditorNode*>(Data))
+		{
+			if (CommonNode == nullptr)
+			{
+				CommonNode = Node;
+			}
+			else if (CommonNode != Node)
+			{
+				CommonNode = nullptr;
+				break;
+			}
+		}
+	}
+
+	return CommonNode;
+}
+
+FStateTreeEditorNode* GetMutableCommonNode(IPropertyHandle& InStructProperty)
+{
+	TArray<void*> RawNodeData;
+	InStructProperty.AccessRawData(RawNodeData);
+
+	FStateTreeEditorNode* CommonNode = nullptr;
+
+	for (void* Data : RawNodeData)
+	{
+		if (FStateTreeEditorNode* Node = static_cast<FStateTreeEditorNode*>(Data))
+		{
+			if (CommonNode == nullptr)
+			{
+				CommonNode = Node;
+			}
+			else if (CommonNode != Node)
+			{
+				CommonNode = nullptr;
+				break;
+			}
+		}
+	}
+
+	return CommonNode;
+}
+
+} // namespace UE::StateTreeEditor::EditorNodeUtils
diff --git a/Engine/Plugins/Runtime/StateTree/Source/StateTreeEditorModule/Private/Customizations/StateTreeEditorNodeUtils.h b/Engine/Plugins/Runtime/StateTree/Source/StateTreeEditorModule/Private/Customizations/StateTreeEditorNodeUtils.h
new file mode 100644
--- /dev/null
+++ b/Engine/Plugins/Runtime/StateTree/Source/StateTreeEditorModule/Private/Customizations/StateTreeEditorNodeUtils.h
@@ -0,0 +1,37 @@
+// Copyright Epic Games, Inc. All Rights Reserved.
+
+#pragma once
+
+#include "Templates/SharedPointerFwd.h"
+
+enum class EStateTreeConditionEvaluationMode : uint8;
+struct FStateTreeEditorNode;
+struct EVisibility;
+class FText;
+class IPropertyHandle;
+template <typename FuncType> class TFunctionRef;
+
+namespace UE::StateTreeEditor::EditorNodeUtils
+{
+
+const FStateTreeEditorNode* GetCommonNode(const TSharedPtr<IPropertyHandle> StructProperty);
+FStateTreeEditorNode* GetMutableCommonNode(const TSharedPtr<IPropertyHandle> StructProperty);
+
+const FStateTreeEditorNode* GetCommonNode(const IPropertyHandle& StructProperty);
+FStateTreeEditorNode* GetMutableCommonNode(IPropertyHandle& StructProperty);
+
+EVisibility IsConditionVisible(TSharedPtr<IPropertyHandle> StructProperty);
+EStateTreeConditionEvaluationMode GetConditionEvaluationMode(TSharedPtr<IPropertyHandle> StructProperty);
+	
+EVisibility IsTaskVisible(TSharedPtr<IPropertyHandle> StructProperty);
+bool IsTaskDisabled(TSharedPtr<IPropertyHandle> StructProperty);
+
+/**
+ * Execute the provided function within a Transaction 
+ * @param Description Description to associate to the transaction
+ * @param StructProperty Property handle of the StateTreeEditorNode(s) 
+ * @param Func Function that will be execute in the transaction on a valid Editor node property
+ */
+void ModifyNodeInTransaction(FText Description, TSharedPtr<IPropertyHandle> StructProperty, TFunctionRef<void(IPropertyHandle&)> Func);
+
+}; // UE::StateTreeEditor::EditorNodeUtils
diff --git a/Engine/Plugins/Runtime/StateTree/Source/StateTreeEditorModule/Private/Customizations/StateTreeStateDetails.cpp b/Engine/Plugins/Runtime/StateTree/Source/StateTreeEditorModule/Private/Customizations/StateTreeStateDetails.cpp
--- a/Engine/Plugins/Runtime/StateTree/Source/StateTreeEditorModule/Private/Customizations/StateTreeStateDetails.cpp
+++ b/Engine/Plugins/Runtime/StateTree/Source/StateTreeEditorModule/Private/Customizations/StateTreeStateDetails.cpp
@@ -7,7 +7,7 @@
 #include "IPropertyUtilities.h"
 #include "StateTreeEditorData.h"
 #include "StateTreeSchema.h"
-
+#include "Debugger/StateTreeDebuggerUIExtensions.h"
 
 #define LOCTEXT_NAMESPACE "StateTreeEditor"
 
@@ -36,14 +36,17 @@ void FStateTreeStateDetails::CustomizeDetails(IDetailLayoutBuilder& DetailBuilde
 	}
 	const UStateTreeSchema* Schema = EditorData ? EditorData->Schema : nullptr;
 
-	TSharedPtr<IPropertyHandle> TasksProperty = DetailBuilder.GetProperty(TEXT("Tasks"));
-	TSharedPtr<IPropertyHandle> SingleTaskProperty = DetailBuilder.GetProperty(TEXT("SingleTask"));
-	TSharedPtr<IPropertyHandle> EnterConditionsProperty = DetailBuilder.GetProperty(TEXT("EnterConditions"));
-	TSharedPtr<IPropertyHandle> TransitionsProperty = DetailBuilder.GetProperty(TEXT("Transitions"));
-	TSharedPtr<IPropertyHandle> TypeProperty = DetailBuilder.GetProperty(TEXT("Type"));
-	TSharedPtr<IPropertyHandle> LinkedSubtreeProperty = DetailBuilder.GetProperty(TEXT("LinkedSubtree"));
-	TSharedPtr<IPropertyHandle> ParametersProperty = DetailBuilder.GetProperty(TEXT("Parameters"));
-	TSharedPtr<IPropertyHandle> SelectionBehaviorProperty = DetailBuilder.GetProperty(TEXT("SelectionBehavior"));
+	const TSharedPtr<IPropertyHandle> EnabledProperty = DetailBuilder.GetProperty(GET_MEMBER_NAME_CHECKED(UStateTreeState, bEnabled));
+	const TSharedPtr<IPropertyHandle> TasksProperty = DetailBuilder.GetProperty(GET_MEMBER_NAME_CHECKED(UStateTreeState, Tasks));
+	const TSharedPtr<IPropertyHandle> SingleTaskProperty = DetailBuilder.GetProperty(GET_MEMBER_NAME_CHECKED(UStateTreeState, SingleTask));
+	const TSharedPtr<IPropertyHandle> EnterConditionsProperty = DetailBuilder.GetProperty(GET_MEMBER_NAME_CHECKED(UStateTreeState, EnterConditions));
+	const TSharedPtr<IPropertyHandle> TransitionsProperty = DetailBuilder.GetProperty(GET_MEMBER_NAME_CHECKED(UStateTreeState, Transitions));
+	const TSharedPtr<IPropertyHandle> TypeProperty = DetailBuilder.GetProperty(GET_MEMBER_NAME_CHECKED(UStateTreeState, Type));
+	const TSharedPtr<IPropertyHandle> LinkedSubtreeProperty = DetailBuilder.GetProperty(GET_MEMBER_NAME_CHECKED(UStateTreeState, LinkedSubtree));
+	const TSharedPtr<IPropertyHandle> ParametersProperty = DetailBuilder.GetProperty(GET_MEMBER_NAME_CHECKED(UStateTreeState, Parameters));
+	const TSharedPtr<IPropertyHandle> SelectionBehaviorProperty = DetailBuilder.GetProperty(GET_MEMBER_NAME_CHECKED(UStateTreeState, SelectionBehavior));
+
+	EnabledProperty->MarkHiddenByCustomization();
 
 	uint8 StateTypeValue = 0;
 	TypeProperty->GetValue(StateTypeValue);
@@ -52,6 +55,8 @@ void FStateTreeStateDetails::CustomizeDetails(IDetailLayoutBuilder& DetailBuilde
 	IDetailCategoryBuilder& StateCategory = DetailBuilder.EditCategory(TEXT("State"), LOCTEXT("StateDetailsState", "State"));
 	StateCategory.SetSortOrder(0);
 
+	StateCategory.HeaderContent(UE::StateTreeEditor::DebuggerExtensions::CreateStateWidget(EnabledProperty));
+
 	if (StateType != EStateTreeStateType::Linked)
 	{
 		LinkedSubtreeProperty->MarkHiddenByCustomization();
@@ -121,7 +126,7 @@ void FStateTreeStateDetails::MakeArrayCategory(IDetailLayoutBuilder& DetailBuild
 	IDetailCategoryBuilder& Category = DetailBuilder.EditCategory(CategoryName, DisplayName);
 	Category.SetSortOrder(SortOrder);
 
-	TSharedRef<SHorizontalBox> HeaderContentWidget = SNew(SHorizontalBox);
+	const TSharedRef<SHorizontalBox> HeaderContentWidget = SNew(SHorizontalBox);
 	HeaderContentWidget->AddSlot()
 	.HAlign(HAlign_Right)
 	.VAlign(VAlign_Center)
@@ -131,7 +136,7 @@ void FStateTreeStateDetails::MakeArrayCategory(IDetailLayoutBuilder& DetailBuild
 	Category.HeaderContent(HeaderContentWidget);
 
 	// Add items inline
-	TSharedRef<FDetailArrayBuilder> Builder = MakeShareable(new FDetailArrayBuilder(PropertyHandle.ToSharedRef(), /*InGenerateHeader*/ false, /*InDisplayResetToDefault*/ true, /*InDisplayElementNum*/ false));
+	const TSharedRef<FDetailArrayBuilder> Builder = MakeShareable(new FDetailArrayBuilder(PropertyHandle.ToSharedRef(), /*InGenerateHeader*/ false, /*InDisplayResetToDefault*/ true, /*InDisplayElementNum*/ false));
 	Builder->OnGenerateArrayElementWidget(FOnGenerateArrayElementWidget::CreateLambda([](TSharedRef<IPropertyHandle> PropertyHandle, int32 ArrayIndex, IDetailChildrenBuilder& ChildrenBuilder)
 	{
 		ChildrenBuilder.AddProperty(PropertyHandle);
diff --git a/Engine/Plugins/Runtime/StateTree/Source/StateTreeEditorModule/Private/Customizations/StateTreeStateDetails.h b/Engine/Plugins/Runtime/StateTree/Source/StateTreeEditorModule/Private/Customizations/StateTreeStateDetails.h
--- a/Engine/Plugins/Runtime/StateTree/Source/StateTreeEditorModule/Private/Customizations/StateTreeStateDetails.h
+++ b/Engine/Plugins/Runtime/StateTree/Source/StateTreeEditorModule/Private/Customizations/StateTreeStateDetails.h
@@ -19,5 +19,4 @@ public:
 
 protected:
 	void MakeArrayCategory(IDetailLayoutBuilder& DetailBuilder, FName CategoryName, const FText& DisplayName, int32 SortOrder, TSharedPtr<IPropertyHandle> PropertyHandle);
-	void GenerateArrayElementWidget(TSharedRef<IPropertyHandle> PropertyHandle, int32 ArrayIndex, IDetailChildrenBuilder& ChildrenBuilder);
 };
diff --git a/Engine/Plugins/Runtime/StateTree/Source/StateTreeEditorModule/Private/Debugger/StateTreeDebuggerUIExtensions.cpp b/Engine/Plugins/Runtime/StateTree/Source/StateTreeEditorModule/Private/Debugger/StateTreeDebuggerUIExtensions.cpp
new file mode 100644
--- /dev/null
+++ b/Engine/Plugins/Runtime/StateTree/Source/StateTreeEditorModule/Private/Debugger/StateTreeDebuggerUIExtensions.cpp
@@ -0,0 +1,388 @@
+// Copyright Epic Games, Inc. All Rights Reserved.
+
+#include "StateTreeDebuggerUIExtensions.h"
+#include "Customizations/StateTreeEditorNodeUtils.h"
+#include "Framework/MultiBox/MultiBoxBuilder.h"
+#include "PropertyHandle.h"
+#include "ScopedTransaction.h"
+#include "StateTreeConditionBase.h"
+#include "StateTreeEditorNode.h"
+#include "StateTreeEditorStyle.h"
+#include "StateTreeTaskBase.h"
+#include "Styling/StyleColors.h"
+#include "Widgets/Images/SImage.h"
+#include "Widgets/Input/SComboButton.h"
+#include "Widgets/SOverlay.h"
+
+
+#define LOCTEXT_NAMESPACE "StateTreeEditor"
+
+namespace UE::StateTreeEditor::DebuggerExtensions
+{
+TSharedRef<SWidget> CreateStateWidget(TSharedPtr<IPropertyHandle> EnabledProperty)
+{
+	FMenuBuilder DebuggingActions(/*CloseWindowAfterMenuSelection*/ true, /*CommandList*/nullptr);
+
+	DebuggingActions.BeginSection(FName("Options"), LOCTEXT("StateOptions", "State Debug Options"));
+	DebuggingActions.AddMenuEntry
+	(
+		LOCTEXT("ToggleStateEnabled", "State Enabled"),
+		LOCTEXT("ToggleStateEnabled_ToolTip", "Enables or disables selected state(s). StateTree must be recompiled to take effect."),
+		FSlateIcon(),
+		FUIAction
+		(
+			FExecuteAction::CreateLambda([EnabledProperty] { OnStateEnableToggled(EnabledProperty); }),
+			FCanExecuteAction(),
+			FGetActionCheckState::CreateLambda([EnabledProperty]
+				{
+					bool bEnabled = false;
+					const FPropertyAccess::Result Result = EnabledProperty->GetValue(bEnabled);
+					if (Result == FPropertyAccess::MultipleValues)
+					{
+						return ECheckBoxState::Undetermined;
+					}
+
+					return bEnabled ? ECheckBoxState::Checked : ECheckBoxState::Unchecked;
+				})
+		),
+		NAME_None,
+		EUserInterfaceActionType::Check
+	);
+	DebuggingActions.EndSection();
+
+	const TSharedRef<SHorizontalBox> HeaderContentWidget = SNew(SHorizontalBox)
+		+ SHorizontalBox::Slot()
+		.FillWidth(1.0f)
+		.HAlign(HAlign_Right)
+		[
+			SNew(SHorizontalBox)
+			.Clipping(EWidgetClipping::OnDemand)
+			
+			// Disabled Label
+			+ SHorizontalBox::Slot()
+			.AutoWidth()
+			.VAlign(VAlign_Center)
+			.Padding(4.0f, 0.0f)
+			[
+				SNew(SBorder)
+				.Padding(FMargin(6, 1))
+				.BorderImage(FStateTreeEditorStyle::Get().GetBrush("StateTree.Node.Label"))
+				.Visibility_Lambda([Property=EnabledProperty]
+					{
+						bool bEnabled = false;
+						Property->GetValue(bEnabled);
+						return bEnabled ? EVisibility::Collapsed : EVisibility::Visible;
+					})
+				[
+					SNew(STextBlock)
+					.TextStyle(FStateTreeEditorStyle::Get(), "StateTree.Param.Label")
+					.ColorAndOpacity(FStyleColors::Foreground)
+					.Text(LOCTEXT("LabelDisabled", "DISABLED"))
+					.ToolTipText(LOCTEXT("DisabledStateTooltip", "This state has been disabled."))
+				]
+			]
+		]
+		// Debug Options
+		+ SHorizontalBox::Slot()
+		.Padding(0)
+		.HAlign(HAlign_Right)
+		.VAlign(VAlign_Center)
+		.AutoWidth()
+		[
+			SNew(SComboButton)
+			.ToolTipText(LOCTEXT("DebugOptions", "Debug Options"))
+			.HasDownArrow(false)
+			.ContentPadding(0)
+			.ForegroundColor(FSlateColor::UseForeground())
+			.ButtonStyle(FAppStyle::Get(), "SimpleButton")
+			.MenuContent()
+			[
+				DebuggingActions.MakeWidget()
+			]
+			.ButtonContent()
+			[
+				SNew(SImage)
+				.Image(FStateTreeEditorStyle::Get().GetBrush("StateTreeEditor.DebugOptions"))
+			]
+		];
+
+	return HeaderContentWidget;
+}
+
+TSharedRef<SWidget> CreateEditorNodeWidget(TSharedPtr<IPropertyHandle> InPropertyHandle)
+{
+	using namespace EditorNodeUtils;
+	
+	bool bIsTask = false;
+	bool bIsCondition = false;
+	if (const FStructProperty* TmpStructProperty = CastField<FStructProperty>(InPropertyHandle->GetProperty()))
+	{
+		if (TmpStructProperty->Struct->IsChildOf(TBaseStructure<FStateTreeEditorNode>::Get()))
+		{
+			const UScriptStruct* ScriptStruct = nullptr;
+			if (const FStateTreeEditorNode* Node = GetCommonNode(InPropertyHandle))
+			{
+				ScriptStruct = Node->Node.GetScriptStruct();
+			}
+
+			if (ScriptStruct != nullptr)
+			{
+				bIsTask = ScriptStruct->IsChildOf(TBaseStructure<FStateTreeTaskBase>::Get());
+				bIsCondition = ScriptStruct->IsChildOf(TBaseStructure<FStateTreeConditionBase>::Get());
+			}
+		}
+	}
+
+	// In case there is no common Editor node or not associated to a condition or a task we don't need to add any widget.	
+	if (bIsCondition == false && bIsTask == false)
+	{
+		return SNullWidget::NullWidget;
+	}
+
+	FMenuBuilder DebuggingActions(/*CloseWindowAfterMenuSelection*/ true, /*CommandList*/nullptr);
+
+	if (bIsTask)
+	{
+		DebuggingActions.BeginSection(FName("Options"), LOCTEXT("TaskOptions", "Task Debug Options"));
+		DebuggingActions.AddMenuEntry
+		(
+			LOCTEXT("ToggleTaskEnabled", "Task Enabled"),
+			LOCTEXT("ToggleTaskEnabled_ToolTip", "Enables or disables selected task(s). StateTree must be recompiled to take effect."),
+			FSlateIcon(),
+			FUIAction
+			(
+				FExecuteAction::CreateLambda([InPropertyHandle] { OnTaskEnableToggled(InPropertyHandle); }),
+				FCanExecuteAction(),
+				FGetActionCheckState::CreateLambda([InPropertyHandle]
+					{
+						return IsTaskDisabled(InPropertyHandle) ? ECheckBoxState::Unchecked : ECheckBoxState::Checked;
+					})
+			),
+			NAME_None,
+			EUserInterfaceActionType::Check
+		);
+		DebuggingActions.EndSection();
+	}
+
+	if (bIsCondition)
+	{
+		DebuggingActions.BeginSection(FName("Options"), LOCTEXT("ConditionOptions", "Condition Debug Options"));
+		DebuggingActions.AddMenuEntry
+		(
+			LOCTEXT("Evaluate", "Evaluate"),
+			LOCTEXT("Evaluate_ToolTip", "Condition result is evaluated (normal behavior)"),
+			FSlateIcon(),
+			FUIAction
+			(
+				FExecuteAction::CreateLambda([InPropertyHandle]
+					{
+						OnConditionEvaluationModeChanged(InPropertyHandle, EStateTreeConditionEvaluationMode::Evaluated);
+					}),
+				FCanExecuteAction(),
+				FIsActionChecked::CreateLambda([InPropertyHandle]
+					{
+						return GetConditionEvaluationMode(InPropertyHandle) == EStateTreeConditionEvaluationMode::Evaluated;
+					})
+			),
+			NAME_None,
+			EUserInterfaceActionType::RadioButton
+		);
+		DebuggingActions.AddMenuEntry
+		(
+			LOCTEXT("ForceTrue", "Force True"),
+			LOCTEXT("ForceTrue_ToolTip", "Result is forced to 'true' (condition is not evaluated)"),
+			FSlateIcon(),
+			FUIAction
+			(
+				FExecuteAction::CreateLambda([InPropertyHandle]
+					{
+						OnConditionEvaluationModeChanged(InPropertyHandle, EStateTreeConditionEvaluationMode::ForcedTrue);
+					}),
+				FCanExecuteAction(),
+				FIsActionChecked::CreateLambda([InPropertyHandle]
+					{
+						return GetConditionEvaluationMode(InPropertyHandle) == EStateTreeConditionEvaluationMode::ForcedTrue;
+					})
+			),
+			NAME_None,
+			EUserInterfaceActionType::RadioButton
+		);
+		DebuggingActions.AddMenuEntry
+		(
+			LOCTEXT("ForceFalse", "Force False"),
+			LOCTEXT("ForceFalse_ToolTip", "Result is forced to 'false' (condition is not evaluated)"),
+			FSlateIcon(),
+			FUIAction
+			(
+				FExecuteAction::CreateLambda([InPropertyHandle]
+					{
+						OnConditionEvaluationModeChanged(InPropertyHandle, EStateTreeConditionEvaluationMode::ForcedFalse);
+					}),
+				FCanExecuteAction(),
+				FIsActionChecked::CreateLambda([InPropertyHandle]
+					{
+						return GetConditionEvaluationMode(InPropertyHandle) == EStateTreeConditionEvaluationMode::ForcedFalse;
+					})
+			),
+			NAME_None,
+			EUserInterfaceActionType::RadioButton
+		);
+
+		DebuggingActions.EndSection();
+	}
+
+
+	return SNew(SHorizontalBox)
+		+ SHorizontalBox::Slot()
+		.FillWidth(1.0f)
+		.HAlign(HAlign_Right)
+		[
+			SNew(SHorizontalBox)
+			.Clipping(EWidgetClipping::OnDemand)
+
+			// Disabled Label
+			+ SHorizontalBox::Slot()
+			.AutoWidth()
+			.HAlign(HAlign_Left)
+			.VAlign(VAlign_Center)
+			.Padding(4.0f, 0.0f)
+			[
+				SNew(SBorder)
+				.Padding(FMargin(6, 2))
+				.BorderImage(FStateTreeEditorStyle::Get().GetBrush("StateTree.Node.Label"))
+				.Visibility_Lambda([InPropertyHandle]
+					{
+						return IsTaskDisabled(InPropertyHandle) ? EVisibility::Visible : EVisibility::Collapsed;
+					})
+				[
+					SNew(STextBlock)
+					.TextStyle(FStateTreeEditorStyle::Get(), "StateTree.Param.Label")
+					.ColorAndOpacity(FStyleColors::Foreground)
+					.Text(LOCTEXT("LabelDisabled", "DISABLED"))
+					.ToolTipText(LOCTEXT("DisabledTaskTooltip", "This task has been disabled."))
+				]
+			]
+
+			// Force True / Force False Label
+			+ SHorizontalBox::Slot()
+			.AutoWidth()
+			.HAlign(HAlign_Left)
+			.VAlign(VAlign_Center)
+			.Padding(4.0f, 0.0f)
+			[
+				SNew(SBorder)
+				.Padding(FMargin(6, 2))
+				.BorderImage(FStateTreeEditorStyle::Get().GetBrush("StateTree.Node.Label"))
+				.BorderBackgroundColor_Lambda([InPropertyHandle]
+					{
+						return GetConditionEvaluationMode(InPropertyHandle) == EStateTreeConditionEvaluationMode::ForcedTrue
+								   ? FStyleColors::AccentGreen.GetSpecifiedColor()
+								   : FStyleColors::AccentRed.GetSpecifiedColor();
+					})
+				.Visibility_Lambda([InPropertyHandle]
+					{
+						return GetConditionEvaluationMode(InPropertyHandle) != EStateTreeConditionEvaluationMode::Evaluated ? EVisibility::Visible : EVisibility::Collapsed;
+					})
+				[
+					SNew(SOverlay)
+					+ SOverlay::Slot()
+					[
+						SNew(STextBlock)
+						.TextStyle(FStateTreeEditorStyle::Get(), "StateTree.Param.Label")
+						.Text(LOCTEXT("ConditionForcedTrue", "TRUE"))
+						.ToolTipText(LOCTEXT("ForcedTrueConditionTooltip", "This condition is not evaluated and result forced to 'true'."))
+						.Visibility_Lambda([InPropertyHandle]
+							{
+								return GetConditionEvaluationMode(InPropertyHandle) == EStateTreeConditionEvaluationMode::ForcedTrue ? EVisibility::Visible : EVisibility::Collapsed;
+							})
+					]
+					+ SOverlay::Slot()
+					[
+						SNew(STextBlock)
+						.TextStyle(FStateTreeEditorStyle::Get(), "StateTree.Param.Label")
+						.Text(LOCTEXT("ConditionForcedFalse", "FALSE"))
+						.ToolTipText(LOCTEXT("ForcedFalseConditionTooltip", "This condition is not evaluated and result forced to 'false'."))
+						.Visibility_Lambda([InPropertyHandle]
+							{
+								return GetConditionEvaluationMode(InPropertyHandle) == EStateTreeConditionEvaluationMode::ForcedFalse ? EVisibility::Visible : EVisibility::Collapsed;
+							})
+					]
+				]
+			]
+		]
+		+ SHorizontalBox::Slot()
+		.Padding(0)
+		.HAlign(HAlign_Right)
+		.VAlign(VAlign_Center)
+		.AutoWidth()
+		[
+			SNew(SComboButton)
+			.ToolTipText(LOCTEXT("DebugOptions", "Debug Options"))
+			.HasDownArrow(false)
+			.ContentPadding(0)
+			.ForegroundColor(FSlateColor::UseForeground())
+			.ButtonStyle(FAppStyle::Get(), "SimpleButton")
+			.MenuContent()
+			[
+				DebuggingActions.MakeWidget()
+			]
+			.ButtonContent()
+			[
+				SNew(SImage)
+				.Image(FStateTreeEditorStyle::Get().GetBrush("StateTreeEditor.DebugOptions"))
+			]
+		];
+}
+
+void OnConditionEvaluationModeChanged(TSharedPtr<IPropertyHandle> StructProperty, const EStateTreeConditionEvaluationMode Mode)
+{
+	EditorNodeUtils::ModifyNodeInTransaction(LOCTEXT("OnEvaluationModeChanged", "Condition Evaluation Mode Changed"),
+		StructProperty,
+		[Mode](IPropertyHandle& StructProperty)
+		{
+			if (FStateTreeEditorNode* Node = EditorNodeUtils::GetMutableCommonNode(StructProperty))
+			{
+				if (FStateTreeConditionBase* ConditionBase = Node->Node.GetMutablePtr<FStateTreeConditionBase>())
+				{
+					ConditionBase->EvaluationMode = Mode;
+				}
+			}
+		});
+}
+
+void OnTaskEnableToggled(TSharedPtr<IPropertyHandle> StructProperty)
+{
+	EditorNodeUtils::ModifyNodeInTransaction(LOCTEXT("OnTaskEnableToggled", "Toggled Task Enabled"),
+		StructProperty,
+		[](IPropertyHandle& StructProperty)
+		{
+			if (FStateTreeEditorNode* Node = EditorNodeUtils::GetMutableCommonNode(StructProperty))
+			{
+				if (FStateTreeTaskBase* TaskBase = Node->Node.GetMutablePtr<FStateTreeTaskBase>())
+				{
+					TaskBase->bTaskEnabled = !TaskBase->bTaskEnabled;
+				}
+			}
+		});
+}
+
+void OnStateEnableToggled(TSharedPtr<IPropertyHandle> EnabledProperty)
+{
+	FScopedTransaction Transaction(LOCTEXT("OnStateEnableToggled", "Toggled State Enabled"));
+
+	bool bEnabled = false;
+	const FPropertyAccess::Result Result = EnabledProperty->GetValue(bEnabled);
+
+	if (Result == FPropertyAccess::MultipleValues || bEnabled == false)
+	{
+		EnabledProperty->SetValue(true);
+	}
+	else
+	{
+		EnabledProperty->SetValue(false);
+	}
+}
+
+} // namespace UE::StateTreeEditor::DebuggerExtensions
+
+#undef LOCTEXT_NAMESPACE
diff --git a/Engine/Plugins/Runtime/StateTree/Source/StateTreeEditorModule/Private/Debugger/StateTreeDebuggerUIExtensions.h b/Engine/Plugins/Runtime/StateTree/Source/StateTreeEditorModule/Private/Debugger/StateTreeDebuggerUIExtensions.h
new file mode 100644
--- /dev/null
+++ b/Engine/Plugins/Runtime/StateTree/Source/StateTreeEditorModule/Private/Debugger/StateTreeDebuggerUIExtensions.h
@@ -0,0 +1,26 @@
+// Copyright Epic Games, Inc. All Rights Reserved.
+
+#pragma once
+
+#include "Templates/SharedPointerFwd.h"
+
+struct FStateTreeEditorNode;
+enum class EStateTreeConditionEvaluationMode : uint8;
+
+class SWidget;
+class IPropertyHandle;
+class IDetailLayoutBuilder;
+class FDetailWidgetRow;
+class FStateTreeViewModel;
+
+namespace UE::StateTreeEditor::DebuggerExtensions
+{
+
+TSharedRef<SWidget> CreateStateWidget(TSharedPtr<IPropertyHandle> EnabledProperty);
+TSharedRef<SWidget> CreateEditorNodeWidget(TSharedPtr<IPropertyHandle> Shared);
+
+void OnConditionEvaluationModeChanged(TSharedPtr<IPropertyHandle> StructProperty, EStateTreeConditionEvaluationMode Mode);
+void OnStateEnableToggled(TSharedPtr<IPropertyHandle> StructProperty);
+void OnTaskEnableToggled(TSharedPtr<IPropertyHandle> StructProperty);
+
+}; // UE::StateTreeEditor::DebuggerExtensions
diff --git a/Engine/Plugins/Runtime/StateTree/Source/StateTreeEditorModule/Private/SStateTreeView.cpp b/Engine/Plugins/Runtime/StateTree/Source/StateTreeEditorModule/Private/SStateTreeView.cpp
--- a/Engine/Plugins/Runtime/StateTree/Source/StateTreeEditorModule/Private/SStateTreeView.cpp
+++ b/Engine/Plugins/Runtime/StateTree/Source/StateTreeEditorModule/Private/SStateTreeView.cpp
@@ -182,6 +182,34 @@ void SStateTreeView::BindCommands()
 		Commands.RenameState,
 		FExecuteAction::CreateSP(this, &SStateTreeView::HandleRenameState),
 		FCanExecuteAction::CreateSP(this, &SStateTreeView::HasSelection));
+
+	CommandList->MapAction(
+		Commands.EnableStates,
+		FExecuteAction::CreateSP(this, &SStateTreeView::HandleEnableSelectedStates),
+		FCanExecuteAction(),
+		FGetActionCheckState::CreateLambda([this]
+			{
+				const bool bCanEnable = CanEnableStates();
+				const bool bCanDisable = CanDisableStates();
+				if (bCanEnable && bCanDisable)
+				{
+					return ECheckBoxState::Undetermined;
+				}
+				
+				if (bCanDisable)
+				{
+					return ECheckBoxState::Checked;
+				}
+
+				if (bCanEnable)
+				{
+					return ECheckBoxState::Unchecked;
+				}
+
+				// Should not happen since action is not visible in this case
+				return ECheckBoxState::Undetermined;
+			} )),
+		FIsActionButtonVisible::CreateLambda([this] { return CanEnableStates() || CanDisableStates(); });
 }
 
 bool SStateTreeView::HasSelection() const
@@ -196,6 +224,20 @@ bool SStateTreeView::CanPaste() const
 			&& StateTreeViewModel->CanPasteStatesFromClipboard();
 }
 
+bool SStateTreeView::CanEnableStates() const
+{
+	return StateTreeViewModel
+			&& StateTreeViewModel->HasSelection()
+			&& StateTreeViewModel->CanEnableStates();
+}
+
+bool SStateTreeView::CanDisableStates() const
+{
+	return StateTreeViewModel
+			&& StateTreeViewModel->HasSelection()
+			&& StateTreeViewModel->CanDisableStates();
+}
+
 FReply SStateTreeView::OnKeyDown(const FGeometry& MyGeometry, const FKeyEvent& InKeyEvent)
 {
 	if(CommandList->ProcessCommandBindings(InKeyEvent))
@@ -317,7 +359,10 @@ void SStateTreeView::HandleModelStatesChanged(const TSet<UStateTreeState*>& Affe
 	// This method is called when anything in a state changes, make sure to only rebuild when needed.
 	if (PropertyChangedEvent.MemberProperty->GetFName() == GET_MEMBER_NAME_CHECKED(UStateTreeState, Tasks))
 	{
-		bArraysChanged = true;
+		if (PropertyChangedEvent.ChangeType != EPropertyChangeType::ValueSet)
+		{
+			bArraysChanged = true;
+		}
 	}
 		
 	if (bArraysChanged)
@@ -426,6 +471,8 @@ TSharedPtr<SWidget> SStateTreeView::HandleContextMenuOpening()
 	MenuBuilder.AddMenuEntry(FStateTreeEditorCommands::Get().DuplicateStates);
 	MenuBuilder.AddMenuEntry(FStateTreeEditorCommands::Get().DeleteStates);
 	MenuBuilder.AddMenuEntry(FStateTreeEditorCommands::Get().RenameState);
+	MenuBuilder.AddSeparator();
+	MenuBuilder.AddMenuEntry(FStateTreeEditorCommands::Get().EnableStates);
 
 #if WITH_STATETREE_DEBUGGER
 	if (UStateTreeSettings::Get().bUseDebugger)
@@ -557,4 +604,28 @@ void SStateTreeView::HandleRenameState()
 	RequestedRenameState = GetFirstSelectedState();
 }
 
+void SStateTreeView::HandleEnableSelectedStates()
+{
+	if (StateTreeViewModel)
+	{
+		// Process CanEnable first so in case of undetermined state (mixed selection) we Enable by default. 
+		if (CanEnableStates())
+		{
+			StateTreeViewModel->SetSelectedStatesEnabled(true);	
+		}
+		else if (CanDisableStates())
+		{
+			StateTreeViewModel->SetSelectedStatesEnabled(false);
+		}
+	}
+}
+
+void SStateTreeView::HandleDisableSelectedStates()
+{
+	if (StateTreeViewModel)
+	{
+		StateTreeViewModel->SetSelectedStatesEnabled(false);
+	}
+}
+
 #undef LOCTEXT_NAMESPACE
diff --git a/Engine/Plugins/Runtime/StateTree/Source/StateTreeEditorModule/Private/SStateTreeView.h b/Engine/Plugins/Runtime/StateTree/Source/StateTreeEditorModule/Private/SStateTreeView.h
--- a/Engine/Plugins/Runtime/StateTree/Source/StateTreeEditorModule/Private/SStateTreeView.h
+++ b/Engine/Plugins/Runtime/StateTree/Source/StateTreeEditorModule/Private/SStateTreeView.h
@@ -85,9 +85,13 @@ private:
 	void HandleDuplicateSelectedStates();
 	void HandleRenameState();
 	void HandleDeleteStates();
+	void HandleEnableSelectedStates();
+	void HandleDisableSelectedStates();
 
 	bool HasSelection() const;
 	bool CanPaste() const;
+	bool CanEnableStates() const;
+	bool CanDisableStates() const;
 
 	void BindCommands();
 
diff --git a/Engine/Plugins/Runtime/StateTree/Source/StateTreeEditorModule/Private/SStateTreeViewRow.cpp b/Engine/Plugins/Runtime/StateTree/Source/StateTreeEditorModule/Private/SStateTreeViewRow.cpp
--- a/Engine/Plugins/Runtime/StateTree/Source/StateTreeEditorModule/Private/SStateTreeViewRow.cpp
+++ b/Engine/Plugins/Runtime/StateTree/Source/StateTreeEditorModule/Private/SStateTreeViewRow.cpp
@@ -22,7 +22,7 @@ void SStateTreeViewRow::Construct(const FArguments& InArgs, const TSharedRef<STa
 	StateTreeViewModel = InStateTreeViewModel;
 	WeakState = InState;
 
-	STableRow<TWeakObjectPtr<UStateTreeState>>::ConstructInternal(STableRow::FArguments()
+	ConstructInternal(STableRow::FArguments()
 		.Padding(5.0f)
 		.OnDragDetected(this, &SStateTreeViewRow::HandleDragDetected)
 		.OnCanAcceptDrop(this, &SStateTreeViewRow::HandleCanAcceptDrop)
@@ -30,9 +30,8 @@ void SStateTreeViewRow::Construct(const FArguments& InArgs, const TSharedRef<STa
 		.Style(&FStateTreeEditorStyle::Get().GetWidgetStyle<FTableRowStyle>("StateTree.Selection"))
 		, InOwnerTableView);
 
-	static const FLinearColor TasksBackground = FLinearColor(FColor(17, 117, 131));
 	static const FLinearColor LinkBackground = FLinearColor(FColor(84, 84, 84));
-	static const FLinearColor IconTint = FLinearColor(1, 1, 1, 0.5f);
+	static constexpr FLinearColor IconTint = FLinearColor(1, 1, 1, 0.5f);
 
 	this->ChildSlot
 	.HAlign(HAlign_Fill)
@@ -73,7 +72,7 @@ void SStateTreeViewRow::Construct(const FArguments& InArgs, const TSharedRef<STa
 				.VAlign(VAlign_Fill)
 				[
 					SNew(SBorder)
-					.BorderImage(FAppStyle::GetBrush("WhiteBrush"))
+					.BorderImage(FStateTreeEditorStyle::Get().GetBrush("StateTree.State.Border"))
 					.BorderBackgroundColor(this, &SStateTreeViewRow::GetActiveStateColor)
 					[
 						SNew(SHorizontalBox)
@@ -106,6 +105,11 @@ void SStateTreeViewRow::Construct(const FArguments& InArgs, const TSharedRef<STa
 								.BorderImage(FAppStyle::GetBrush("WhiteBrush"))
 								.BorderBackgroundColor(this, &SStateTreeViewRow::GetTitleColor)
 								.Padding(FMargin(4.0f, 0.0f, 12.0f, 0.0f))
+								.IsEnabled_Lambda([InState]
+									{
+										const UStateTreeState* State = InState.Get();
+										return State != nullptr && State->bEnabled;
+									})
 								[
 									SNew(SHorizontalBox)
 
@@ -176,46 +180,6 @@ void SStateTreeViewRow::Construct(const FArguments& InArgs, const TSharedRef<STa
 								]
 							]
 						]
-
-						// Tasks Box
-						+ SHorizontalBox::Slot()
-						.VAlign(VAlign_Center)
-						.AutoWidth()
-						[
-							SNew(SBox)
-							.HeightOverride(28.0f)
-							.VAlign(VAlign_Fill)
-							.Visibility(this, &SStateTreeViewRow::GetTasksVisibility)
-							[
-								SNew(SBorder)
-								.BorderImage(FAppStyle::GetBrush("WhiteBrush"))
-								.BorderBackgroundColor(TasksBackground)
-								.Padding(FMargin(6.0f, 0.0f, 12.0f, 0.0f))
-								[
-									// Task icon
-									SNew(SHorizontalBox)
-									+ SHorizontalBox::Slot()
-									.VAlign(VAlign_Center)
-									.Padding(FMargin(0.0f, 0.0f, 4.0f, 0.0f))
-									.AutoWidth()
-									[
-										SNew(STextBlock)
-										.Text(FEditorFontGlyphs::Paper_Plane)
-										.TextStyle(FStateTreeEditorStyle::Get(), "StateTree.DetailsIcon")
-									]
-
-									// Tasks list
-									+ SHorizontalBox::Slot()
-									.VAlign(VAlign_Center)
-									.AutoWidth()
-									[
-										SNew(STextBlock)
-										.Text(this, &SStateTreeViewRow::GetTasksDesc)
-										.TextStyle(FStateTreeEditorStyle::Get(), "StateTree.Details")
-									]
-								]
-							]
-						]
 						
 						// Linked State box
 						+ SHorizontalBox::Slot()
@@ -256,6 +220,18 @@ void SStateTreeViewRow::Construct(const FArguments& InArgs, const TSharedRef<STa
 								]
 							]
 						]
+						// Tasks
+						+ SHorizontalBox::Slot()
+						.VAlign(VAlign_Fill)
+						.AutoWidth()
+						[
+							SNew(SBox)
+							.VAlign(VAlign_Fill)
+							.Visibility(this, &SStateTreeViewRow::GetTasksVisibility)
+							[
+								CreateTasksWidget()
+							]
+						]
 					]
 				]
 			]
@@ -406,6 +382,56 @@ void SStateTreeViewRow::Construct(const FArguments& InArgs, const TSharedRef<STa
 	];
 }
 
+TSharedRef<SHorizontalBox> SStateTreeViewRow::CreateTasksWidget()
+{
+	const TSharedRef<SHorizontalBox> TasksBox = SNew(SHorizontalBox);
+	
+	const UStateTreeState* State = WeakState.Get();
+	if (State == nullptr || State->Tasks.IsEmpty())
+	{
+		return TasksBox;
+	}
+
+	for (int32 TaskIndex = 0; TaskIndex < State->Tasks.Num(); TaskIndex++)
+	{
+		if (const FStateTreeTaskBase* Task = State->Tasks[TaskIndex].Node.GetPtr<FStateTreeTaskBase>())
+		{
+			auto IsTaskEnabledFunc = [WeakState=WeakState, TaskIndex]
+				{
+					const UStateTreeState* State = WeakState.Get();
+					if (State != nullptr && State->Tasks.IsValidIndex(TaskIndex))
+					{
+						if (const FStateTreeTaskBase* Task = State->Tasks[TaskIndex].Node.GetPtr<FStateTreeTaskBase>())
+						{
+							return (State->bEnabled && Task->bTaskEnabled);
+						}
+					}					
+					return true;
+				};
+			
+			TasksBox->AddSlot()
+				.AutoWidth()
+				.VAlign(VAlign_Fill)
+				[
+					SNew(SBorder)
+					.VAlign(VAlign_Center)
+					.BorderImage(FStateTreeEditorStyle::Get().GetBrush("StateTree.Task.Rect"))
+					.IsEnabled_Lambda(IsTaskEnabledFunc)					
+					[
+						SNew(STextBlock)
+						.Margin(FMargin(4.f, 0.f))
+						.Text(FText::FromName(Task->Name))
+						.TextStyle(FStateTreeEditorStyle::Get(), "StateTree.Task.Title")
+						.IsEnabled_Lambda(IsTaskEnabledFunc)						
+						.ToolTipText(FText::FromName(Task->Name))
+					]
+				];
+		}
+	}
+
+	return TasksBox;
+}
+
 void SStateTreeViewRow::RequestRename() const
 {
 	if (NameTextBlock)
diff --git a/Engine/Plugins/Runtime/StateTree/Source/StateTreeEditorModule/Private/SStateTreeViewRow.h b/Engine/Plugins/Runtime/StateTree/Source/StateTreeEditorModule/Private/SStateTreeViewRow.h
--- a/Engine/Plugins/Runtime/StateTree/Source/StateTreeEditorModule/Private/SStateTreeViewRow.h
+++ b/Engine/Plugins/Runtime/StateTree/Source/StateTreeEditorModule/Private/SStateTreeViewRow.h
@@ -28,6 +28,7 @@ public:
 	void RequestRename() const;
 
 private:
+	TSharedRef<SHorizontalBox> CreateTasksWidget();
 
 	FSlateColor GetTitleColor() const;
 	FSlateColor GetActiveStateColor() const;
diff --git a/Engine/Plugins/Runtime/StateTree/Source/StateTreeEditorModule/Private/StateTreeCompiler.cpp b/Engine/Plugins/Runtime/StateTree/Source/StateTreeEditorModule/Private/StateTreeCompiler.cpp
--- a/Engine/Plugins/Runtime/StateTree/Source/StateTreeEditorModule/Private/StateTreeCompiler.cpp
+++ b/Engine/Plugins/Runtime/StateTree/Source/StateTreeEditorModule/Private/StateTreeCompiler.cpp
@@ -325,6 +325,7 @@ bool FStateTreeCompiler::CreateStateRecursive(UStateTreeState& State, const FSta
 	FCompactStateTreeState& CompactState = StateTree->States[StateIdx];
 	CompactState.Name = State.Name;
 	CompactState.Parent = Parent;
+	CompactState.bEnabled = State.bEnabled;
 
 	CompactState.Type = State.Type;
 	CompactState.SelectionBehavior = State.SelectionBehavior;
diff --git a/Engine/Plugins/Runtime/StateTree/Source/StateTreeEditorModule/Private/StateTreeEditorCommands.cpp b/Engine/Plugins/Runtime/StateTree/Source/StateTreeEditorModule/Private/StateTreeEditorCommands.cpp
--- a/Engine/Plugins/Runtime/StateTree/Source/StateTreeEditorModule/Private/StateTreeEditorCommands.cpp
+++ b/Engine/Plugins/Runtime/StateTree/Source/StateTreeEditorModule/Private/StateTreeEditorCommands.cpp
@@ -1,13 +1,12 @@
 // Copyright Epic Games, Inc. All Rights Reserved.
 
 #include "StateTreeEditorCommands.h"
-#include "StateTreeEditorStyle.h"
 #include "Styling/AppStyle.h"
 
 #define LOCTEXT_NAMESPACE "StateTreeEditor"
 
 FStateTreeEditorCommands::FStateTreeEditorCommands() 
-	: TCommands<FStateTreeEditorCommands>(
+	: TCommands(
 		TEXT("StateTreeEditor"), // Context name for fast lookup
 		LOCTEXT("StateTreeEditor", "StateTree Editor"), // Localized context name for displaying
 		NAME_None,
@@ -32,6 +31,7 @@ void FStateTreeEditorCommands::RegisterCommands()
 	UI_COMMAND(DuplicateStates, "Duplicate", "Duplicate Selected States", EUserInterfaceActionType::Button, FInputChord(EModifierKey::Control, EKeys::D));
 	UI_COMMAND(DeleteStates, "Delete", "Delete Selected States", EUserInterfaceActionType::Button, FInputChord(EKeys::Delete));
 	UI_COMMAND(RenameState, "Rename", "Rename Selected State", EUserInterfaceActionType::Button, FInputChord(EKeys::F2));
+	UI_COMMAND(EnableStates, "Enable", "Enables Selected States", EUserInterfaceActionType::Check, FInputChord());
 }
 
 #undef LOCTEXT_NAMESPACE
diff --git a/Engine/Plugins/Runtime/StateTree/Source/StateTreeEditorModule/Private/StateTreeEditorCommands.h b/Engine/Plugins/Runtime/StateTree/Source/StateTreeEditorModule/Private/StateTreeEditorCommands.h
--- a/Engine/Plugins/Runtime/StateTree/Source/StateTreeEditorModule/Private/StateTreeEditorCommands.h
+++ b/Engine/Plugins/Runtime/StateTree/Source/StateTreeEditorModule/Private/StateTreeEditorCommands.h
@@ -30,4 +30,5 @@ public:
 	TSharedPtr<FUICommandInfo> PasteStatesAsChildren;
 	TSharedPtr<FUICommandInfo> DuplicateStates;
 	TSharedPtr<FUICommandInfo> DeleteStates;
+	TSharedPtr<FUICommandInfo> EnableStates;
 };
\ No newline at end of file
diff --git a/Engine/Plugins/Runtime/StateTree/Source/StateTreeEditorModule/Private/StateTreeEditorStyle.cpp b/Engine/Plugins/Runtime/StateTree/Source/StateTreeEditorModule/Private/StateTreeEditorStyle.cpp
--- a/Engine/Plugins/Runtime/StateTree/Source/StateTreeEditorModule/Private/StateTreeEditorStyle.cpp
+++ b/Engine/Plugins/Runtime/StateTree/Source/StateTreeEditorModule/Private/StateTreeEditorStyle.cpp
@@ -84,14 +84,28 @@ FStateTreeEditorStyle::FStateTreeEditorStyle()
 		Set("StateTree.State.TitleInlineEditableText", FInlineEditableTextBlockStyle()
 			.SetTextStyle(StateTitle)
 			.SetEditableTextBoxStyle(StateTitleEditableText));
+
+		Set("StateTree.State.Border", new FSlateBorderBrush(NAME_None, FMargin(2.0f)));
 	}
 
 	// Details
 	{
-		FTextBlockStyle StateTitle = FTextBlockStyle(NormalText)
+		const FTextBlockStyle Details = FTextBlockStyle(NormalText)
 			.SetFont(CORE_FONT("Fonts/Roboto-Regular", 10))
 			.SetColorAndOpacity(FLinearColor(230.0f / 255.0f, 230.0f / 255.0f, 230.0f / 255.0f, 0.75f));
-		Set("StateTree.Details", StateTitle);
+		Set("StateTree.Details", Details);
+
+		Set("StateTree.Node.Label", new FSlateRoundedBoxBrush(FStyleColors::AccentGray, 6.f));
+	}
+
+	// Task
+	{
+		const FTextBlockStyle TaskTitle = FTextBlockStyle(NormalText)
+			.SetFont(CORE_FONT("Fonts/Roboto-Regular", 10))
+			.SetColorAndOpacity(FLinearColor(230.0f / 255.0f, 230.0f / 255.0f, 230.0f / 255.0f, 0.85f));
+		Set("StateTree.Task.Title", TaskTitle);
+
+		Set("StateTree.Task.Rect", new FSlateRoundedBoxBrush(FLinearColor(FColor(21, 104, 116)), 0.0f, FLinearColor(FColor(27, 138, 154)), 1.0f));
 	}
 
 	// Details rich text
@@ -173,7 +187,7 @@ FStateTreeEditorStyle::FStateTreeEditorStyle()
 		Set("StateTreeEditor.DuplicateStates", new IMAGE_BRUSH_SVG("Starship/Common/Duplicate", CoreStyleConstants::Icon16x16));
 		Set("StateTreeEditor.DeleteStates", new IMAGE_BRUSH_SVG("Starship/Common/Delete", CoreStyleConstants::Icon16x16));
 		Set("StateTreeEditor.RenameState", new IMAGE_BRUSH_SVG("Starship/Common/Rename", CoreStyleConstants::Icon16x16));
-		Set("StateTreeEditor.Debugger.ToggleBreakpoint", new IMAGE_BRUSH_SVG("Starship/Common/filled-circle", CoreStyleConstants::Icon16x16));
+		Set("StateTreeEditor.EnableStates", new IMAGE_BRUSH_SVG("Starship/Common/check-circle-large", CoreStyleConstants::Icon16x16));
 	}
 
 	{
@@ -187,6 +201,9 @@ FStateTreeEditorStyle::FStateTreeEditorStyle()
 		Set("StateTreeEditor.Debugger.PreviousFrameWithEvents", new IMAGE_BRUSH("Sequencer/Transport_Bar/Step_Backwards_24x", CoreStyleConstants::Icon16x16));
 		Set("StateTreeEditor.Debugger.NextFrameWithEvents", new IMAGE_BRUSH("Sequencer/Transport_Bar/Step_Forward_24x", CoreStyleConstants::Icon16x16));
 		Set("StateTreeEditor.Debugger.NextFrameWithStateChange", new IMAGE_BRUSH("Sequencer/Transport_Bar/Go_To_End_24x", CoreStyleConstants::Icon16x16));
+
+		Set("StateTreeEditor.Debugger.ToggleBreakpoint", new IMAGE_BRUSH_SVG("Starship/Blueprints/Breakpoint_Valid", CoreStyleConstants::Icon16x16));
+		Set("StateTreeEditor.DebugOptions", new IMAGE_BRUSH_SVG("Starship/Common/Bug", CoreStyleConstants::Icon16x16));
 	}
 
 	{
diff --git a/Engine/Plugins/Runtime/StateTree/Source/StateTreeEditorModule/Private/StateTreeViewModel.cpp b/Engine/Plugins/Runtime/StateTree/Source/StateTreeEditorModule/Private/StateTreeViewModel.cpp
--- a/Engine/Plugins/Runtime/StateTree/Source/StateTreeEditorModule/Private/StateTreeViewModel.cpp
+++ b/Engine/Plugins/Runtime/StateTree/Source/StateTreeEditorModule/Private/StateTreeViewModel.cpp
@@ -381,10 +381,10 @@ bool FStateTreeViewModel::IsChildOfSelection(const UStateTreeState* State) const
 	return false;
 }
 
-void FStateTreeViewModel::GetSelectedStates(TArray<UStateTreeState*>& OutSelectedStates)
+void FStateTreeViewModel::GetSelectedStates(TArray<UStateTreeState*>& OutSelectedStates) const
 {
 	OutSelectedStates.Reset();
-	for (TWeakObjectPtr<UStateTreeState>& WeakState : SelectedStates)
+	for (const TWeakObjectPtr<UStateTreeState>& WeakState : SelectedStates)
 	{
 		if (UStateTreeState* State = WeakState.Get())
 		{
@@ -393,10 +393,10 @@ void FStateTreeViewModel::GetSelectedStates(TArray<UStateTreeState*>& OutSelecte
 	}
 }
 
-void FStateTreeViewModel::GetSelectedStates(TArray<TWeakObjectPtr<UStateTreeState>>& OutSelectedStates)
+void FStateTreeViewModel::GetSelectedStates(TArray<TWeakObjectPtr<UStateTreeState>>& OutSelectedStates) const
 {
 	OutSelectedStates.Reset();
-	for (TWeakObjectPtr<UStateTreeState>& WeakState : SelectedStates)
+	for (const TWeakObjectPtr<UStateTreeState>& WeakState : SelectedStates)
 	{
 		if (WeakState.Get())
 		{
@@ -746,6 +746,57 @@ void FStateTreeViewModel::MoveSelectedStatesInto(UStateTreeState* TargetState)
 	MoveSelectedStates(TargetState, FStateTreeViewModelInsert::Into);
 }
 
+bool FStateTreeViewModel::CanEnableStates() const
+{
+	TArray<UStateTreeState*> States;
+	GetSelectedStates(States);
+
+	for (const UStateTreeState* State : States)
+	{
+		// Stop if at least one state can be enabled
+		if (State->bEnabled == false)
+		{
+			return true;
+		}
+	}
+
+	return false;
+}
+
+bool FStateTreeViewModel::CanDisableStates() const
+{
+	TArray<UStateTreeState*> States;
+	GetSelectedStates(States);
+
+	for (const UStateTreeState* State : States)
+	{
+		// Stop if at least one state can be disabled
+		if (State->bEnabled)
+		{
+			return true;
+		}
+	}
+
+	return false;
+}
+
+void FStateTreeViewModel::SetSelectedStatesEnabled(const bool bEnable)
+{
+	TArray<UStateTreeState*> States;
+	GetSelectedStates(States);
+
+	if (States.Num() > 0)
+	{
+		const FScopedTransaction Transaction(LOCTEXT("SetStatesEnabledTransaction", "Set State Enabled"));
+
+		for (UStateTreeState* State : States)
+		{
+			State->Modify();
+			State->bEnabled = bEnable;
+		}
+	}
+}
+
 void FStateTreeViewModel::MoveSelectedStates(UStateTreeState* TargetState, const FStateTreeViewModelInsert RelativeLocation)
 {
 	UStateTreeEditorData* TreeData = TreeDataWeak.Get();
diff --git a/Engine/Plugins/Runtime/StateTree/Source/StateTreeEditorModule/Private/StateTreeViewModel.h b/Engine/Plugins/Runtime/StateTree/Source/StateTreeEditorModule/Private/StateTreeViewModel.h
--- a/Engine/Plugins/Runtime/StateTree/Source/StateTreeEditorModule/Private/StateTreeViewModel.h
+++ b/Engine/Plugins/Runtime/StateTree/Source/StateTreeEditorModule/Private/StateTreeViewModel.h
@@ -51,8 +51,8 @@ public:
 	void SetSelection(const TArray<TWeakObjectPtr<UStateTreeState>>& InSelection);
 	bool IsSelected(const UStateTreeState* State) const;
 	bool IsChildOfSelection(const UStateTreeState* State) const;
-	void GetSelectedStates(TArray<UStateTreeState*>& OutSelectedStates);
-	void GetSelectedStates(TArray<TWeakObjectPtr<UStateTreeState>>& OutSelectedStates);
+	void GetSelectedStates(TArray<UStateTreeState*>& OutSelectedStates) const;
+	void GetSelectedStates(TArray<TWeakObjectPtr<UStateTreeState>>& OutSelectedStates) const;
 	bool HasSelection() const;
 
 	// Returns associated state tree asset.
@@ -79,6 +79,9 @@ public:
 	void MoveSelectedStatesBefore(UStateTreeState* TargetState);
 	void MoveSelectedStatesAfter(UStateTreeState* TargetState);
 	void MoveSelectedStatesInto(UStateTreeState* TargetState);
+	bool CanEnableStates() const;
+	bool CanDisableStates() const;
+	void SetSelectedStatesEnabled(bool bEnable);
 
 	// Force to update the view externally.
 	void NotifyAssetChangedExternally() const;
diff --git a/Engine/Plugins/Runtime/StateTree/Source/StateTreeEditorModule/Public/StateTreeState.h b/Engine/Plugins/Runtime/StateTree/Source/StateTreeEditorModule/Public/StateTreeState.h
--- a/Engine/Plugins/Runtime/StateTree/Source/StateTreeEditorModule/Public/StateTreeState.h
+++ b/Engine/Plugins/Runtime/StateTree/Source/StateTreeEditorModule/Public/StateTreeState.h
@@ -224,6 +224,9 @@ public:
 	UPROPERTY(meta = (ExcludeFromHash))
 	bool bExpanded = true;
 
+	UPROPERTY(EditDefaultsOnly, Category = "State")
+	bool bEnabled = true;
+
 	UPROPERTY(meta = (ExcludeFromHash))
 	TObjectPtr<UStateTreeState> Parent = nullptr;
 };
diff --git a/Engine/Plugins/Runtime/StateTree/Source/StateTreeModule/Private/StateTreeExecutionContext.cpp b/Engine/Plugins/Runtime/StateTree/Source/StateTreeModule/Private/StateTreeExecutionContext.cpp
--- a/Engine/Plugins/Runtime/StateTree/Source/StateTreeModule/Private/StateTreeExecutionContext.cpp
+++ b/Engine/Plugins/Runtime/StateTree/Source/StateTreeModule/Private/StateTreeExecutionContext.cpp
@@ -750,6 +750,8 @@ EStateTreeRunStatus FStateTreeExecutionContext::EnterState(const FStateTreeTrans
 		const FStateTreeStateHandle PreviousHandle = Transition.CurrentActiveStates.GetStateSafe(Index);
 		const FCompactStateTreeState& State = StateTree.States[CurrentHandle.Index];
 
+		ensureMsgf(State.bEnabled, TEXT("Should never enter disabled state. This indicates an unhandled case."));
+		
 		if (!Exec.ActiveStates.Push(CurrentHandle))
 		{
 			STATETREE_LOG(Error, TEXT("%hs: Reached max execution depth when trying to enter state '%s'.  '%s' using StateTree '%s'."),
@@ -790,6 +792,13 @@ EStateTreeRunStatus FStateTreeExecutionContext::EnterState(const FStateTreeTrans
 				StateTree.PropertyBindings.CopyTo(DataViews, Task.BindingsBatch, DataViews[Task.DataViewIndex.Get()]);
 			}
 
+			// Ignore disabled task
+			if (Task.bTaskEnabled == false)
+			{
+				STATETREE_LOG(VeryVerbose, TEXT("%*sSkipped 'EnterState' for disabled Task: '%s'"), UE::StateTree::DebugIndentSize, TEXT(""), *Task.Name.ToString());
+				continue;
+			}
+
 			const bool bShouldCallStateChange = CurrentTransition.ChangeType == EStateTreeStateChangeType::Changed
 												|| (CurrentTransition.ChangeType == EStateTreeStateChangeType::Sustained && Task.bShouldStateChangeOnReselect);
 
@@ -934,6 +943,13 @@ void FStateTreeExecutionContext::ExitState(const FStateTreeTransitionResult& Tra
 			{
 				const FStateTreeTaskBase& Task = StateTree.Nodes[TaskIndex].Get<const FStateTreeTaskBase>();
 
+				// Ignore disabled task
+				if (Task.bTaskEnabled == false)
+				{
+					STATETREE_LOG(VeryVerbose, TEXT("%*sSkipped 'ExitState' for disabled Task: '%s'"), UE::StateTree::DebugIndentSize, TEXT(""), *Task.Name.ToString());
+					continue;
+				}
+
 				const bool bShouldCallStateChange = CurrentTransition.ChangeType == EStateTreeStateChangeType::Changed
 							|| (CurrentTransition.ChangeType == EStateTreeStateChangeType::Sustained && Task.bShouldStateChangeOnReselect);
 
@@ -986,6 +1002,13 @@ void FStateTreeExecutionContext::StateCompleted()
 			{
 				const FStateTreeTaskBase& Task = StateTree.Nodes[TaskIndex].Get<const FStateTreeTaskBase>();
 
+				// Ignore disabled task
+				if (Task.bTaskEnabled == false)
+				{
+					STATETREE_LOG(VeryVerbose, TEXT("%*sSkipped 'StateCompleted' for disabled Task: '%s'"), UE::StateTree::DebugIndentSize, TEXT(""), *Task.Name.ToString());
+					continue;
+				}
+				
 				STATETREE_LOG(Verbose, TEXT("%*s  Task '%s'"), Index*UE::StateTree::DebugIndentSize, TEXT(""), *Task.Name.ToString());
 				Task.StateCompleted(*this, Exec.LastTickStatus, Exec.ActiveStates);
 			}
@@ -1033,6 +1056,14 @@ EStateTreeRunStatus FStateTreeExecutionContext::TickEvaluatorsAndGlobalTasks(con
 		for (int32 TaskIndex = StateTree.GlobalTasksBegin; TaskIndex < (StateTree.GlobalTasksBegin + StateTree.GlobalTasksNum); TaskIndex++)
 		{
 			const FStateTreeTaskBase& Task =  StateTree.Nodes[TaskIndex].Get<const FStateTreeTaskBase>();
+
+			// Ignore disabled task
+			if (Task.bTaskEnabled == false)
+			{
+				STATETREE_LOG(VeryVerbose, TEXT("%*sSkipped 'Tick' for disabled Task: '%s'"), UE::StateTree::DebugIndentSize, TEXT(""), *Task.Name.ToString());
+				continue;
+			}
+
 			SetNodeDataView(Task, InstanceStructIndex, InstanceObjectIndex);
 
 			const bool bNeedsTick = bShouldTickTasks && (Task.bShouldCallTick || (bHasEvents && Task.bShouldCallTickOnlyOnEvents));
@@ -1119,6 +1150,14 @@ EStateTreeRunStatus FStateTreeExecutionContext::StartEvaluatorsAndGlobalTasks(FS
 		{
 			StateTree.PropertyBindings.CopyTo(DataViews, Task.BindingsBatch, DataViews[Task.DataViewIndex.Get()]);
 		}
+
+		// Ignore disabled task
+		if (Task.bTaskEnabled == false)
+		{
+			STATETREE_LOG(VeryVerbose, TEXT("%*sSkipped 'EnterState' for disabled Task: '%s'"), UE::StateTree::DebugIndentSize, TEXT(""), *Task.Name.ToString());
+			continue;
+		}
+
 		STATETREE_LOG(Verbose, TEXT("  Start: '%s'"), *Task.Name.ToString());
 		{
 			QUICK_SCOPE_CYCLE_COUNTER(StateTree_Task_TreeStart);
@@ -1183,6 +1222,13 @@ void FStateTreeExecutionContext::StopEvaluatorsAndGlobalTasks(const EStateTreeRu
 	{
 		const FStateTreeTaskBase& Task =  StateTree.Nodes[TaskIndex].Get<const FStateTreeTaskBase>();
 
+		// Ignore disabled task
+		if (Task.bTaskEnabled == false)
+		{
+			STATETREE_LOG(VeryVerbose, TEXT("%*sSkipped 'ExitState' for disabled Task: '%s'"), UE::StateTree::DebugIndentSize, TEXT(""), *Task.Name.ToString());
+			continue;
+		}
+
 		// Relying here that invalid value of LastInitializedTaskIndex == MAX_uint16.
 		if (TaskIndex <= LastInitializedTaskIndex.Get())
 		{
@@ -1256,6 +1302,14 @@ EStateTreeRunStatus FStateTreeExecutionContext::TickTasks(const float DeltaTime)
 		for (int32 TaskIndex = State.TasksBegin; TaskIndex < (State.TasksBegin + State.TasksNum); TaskIndex++)
 		{
 			const FStateTreeTaskBase& Task = StateTree.Nodes[TaskIndex].Get<const FStateTreeTaskBase>();
+
+			// Ignore disabled task
+			if (Task.bTaskEnabled == false)
+			{
+				STATETREE_LOG(VeryVerbose, TEXT("%*sSkipped 'Tick' for disabled Task: '%s'"), UE::StateTree::DebugIndentSize, TEXT(""), *Task.Name.ToString());
+				continue;
+			}
+
 			SetNodeDataView(Task, InstanceStructIndex, InstanceObjectIndex);
 
 			const bool bNeedsTick = bShouldTickTasks && (Task.bShouldCallTick || (bHasEvents && Task.bShouldCallTickOnlyOnEvents));
@@ -1342,26 +1396,35 @@ bool FStateTreeExecutionContext::TestAllConditions(const int32 ConditionsOffset,
 			DataView = SharedInstanceData->GetMutableStruct(Cond.InstanceIndex.Get());
 		}
 
-		// Copy bound properties.
-		if (Cond.BindingsBatch.IsValid())
+		bool bValue = false;
+		if (Cond.EvaluationMode == EStateTreeConditionEvaluationMode::Evaluated)
 		{
-			if (!StateTree.PropertyBindings.CopyTo(DataViews, Cond.BindingsBatch, DataView))
+			// Copy bound properties.
+			if (Cond.BindingsBatch.IsValid())
 			{
-				// If the source data cannot be accessed, the whole expression evaluates to false.
-				Values[0] = false;
-				break;
+				if (!StateTree.PropertyBindings.CopyTo(DataViews, Cond.BindingsBatch, DataView))
+				{
+					// If the source data cannot be accessed, the whole expression evaluates to false.
+					Values[0] = false;
+					break;
+				}
+			}
+			
+			bValue = Cond.TestCondition(*this);
+			
+			// Reset copied properties that might contain object references.
+			if (Cond.BindingsBatch.IsValid())
+			{
+				StateTree.PropertyBindings.ResetObjects(Cond.BindingsBatch, DataView);
 			}
 		}
-
-		const bool bValue = Cond.TestCondition(*this);
-		STATETREE_TRACE_CONDITION_EVENT(ConditionsOffset + Index, DataView, bValue ? EStateTreeTraceEventType::Passed : EStateTreeTraceEventType::Failed);
-
-		// Reset copied properties that might contain object references.
-		if (Cond.BindingsBatch.IsValid())
+		else
 		{
-			StateTree.PropertyBindings.ResetObjects(Cond.BindingsBatch, DataView);
+			bValue = Cond.EvaluationMode == EStateTreeConditionEvaluationMode::ForcedTrue ? true : /* EStateTreeConditionEvaluationMode::AlwaysFalse */ false;
 		}
 
+		STATETREE_TRACE_CONDITION_EVENT(ConditionsOffset + Index, DataView, bValue ? EStateTreeTraceEventType::Passed : EStateTreeTraceEventType::Failed);
+
 		const int32 DeltaIndent = Cond.DeltaIndent;
 		const int32 OpenParens = FMath::Max(0, DeltaIndent) + 1;	// +1 for the current value that is stored at the empty slot at the top of the value stack.
 		const int32 ClosedParens = FMath::Max(0, -DeltaIndent) + 1;
@@ -1504,6 +1567,8 @@ bool FStateTreeExecutionContext::TriggerTransitions()
 		const FStateTreeStateHandle StateHandle = Exec.ActiveStates[StateIndex];
 		const FCompactStateTreeState& State = StateTree.States[StateHandle.Index];
 
+		ensureMsgf(State.bEnabled, TEXT("Should never try to transition to a disabled state. This indicates an unhandled case."));
+
 		FCurrentlyProcessedStateScope StateScope(*this, StateHandle);
 
 		if (State.bHasTransitionTasks)
@@ -1513,6 +1578,14 @@ bool FStateTreeExecutionContext::TriggerTransitions()
 			for (int32 TaskIndex = (State.TasksBegin + State.TasksNum) - 1; TaskIndex >= State.TasksBegin; TaskIndex--)
 			{
 				const FStateTreeTaskBase& Task = StateTree.Nodes[TaskIndex].Get<const FStateTreeTaskBase>();
+
+				// Ignore disabled task
+				if (Task.bTaskEnabled == false)
+				{
+					STATETREE_LOG(VeryVerbose, TEXT("%*sSkipped 'TriggerTransitions' for disabled Task: '%s'"), UE::StateTree::DebugIndentSize, TEXT(""), *Task.Name.ToString());
+					continue;
+				}
+
 				if (Task.bShouldAffectTransitions)
 				{
 					STATETREE_LOG(VeryVerbose, TEXT("%*sTriggerTransitions: '%s'"), UE::StateTree::DebugIndentSize, TEXT(""), *Task.Name.ToString());
@@ -1614,6 +1687,14 @@ bool FStateTreeExecutionContext::TriggerTransitions()
 		for (int32 TaskIndex = (StateTree.GlobalTasksBegin + StateTree.GlobalTasksNum) - 1; TaskIndex >= StateTree.GlobalTasksBegin; TaskIndex--)
 		{
 			const FStateTreeTaskBase& Task =  StateTree.Nodes[TaskIndex].Get<const FStateTreeTaskBase>();
+
+			// Ignore disabled task
+			if (Task.bTaskEnabled == false)
+			{
+				STATETREE_LOG(VeryVerbose, TEXT("%*sSkipped 'TriggerTransitions' for disabled Task: '%s'"), UE::StateTree::DebugIndentSize, TEXT(""), *Task.Name.ToString());
+				continue;
+			}
+
 			if (Task.bShouldAffectTransitions)
 			{
 				STATETREE_LOG(VeryVerbose, TEXT("%*sTriggerTransitions: '%s'"), UE::StateTree::DebugIndentSize, TEXT(""), *Task.Name.ToString());
@@ -1822,6 +1903,14 @@ bool FStateTreeExecutionContext::SelectStateInternal(const FStateTreeStateHandle
 
 	const FCompactStateTreeState& State = StateTree.States[NextState.Index];
 
+	if (State.bEnabled == false)
+	{
+		// Do not select disabled state
+		STATETREE_LOG(VeryVerbose, TEXT("%hs: Ignoring disabled state '%s'.  '%s' using StateTree '%s'."),
+			__FUNCTION__, *GetSafeStateName(NextState), *GetNameSafe(&Owner), *GetFullNameSafe(&StateTree));
+		return false;
+	}
+
 	// Push State before the scoped phase so the debugger can create the hierarchy parent using the state description.
 	STATETREE_TRACE_SCOPED_STATE_SELECTION(NextState, State.SelectionBehavior);
 	STATETREE_TRACE_SCOPED_PHASE(EStateTreeUpdatePhase::StateSelection);
@@ -2071,7 +2160,10 @@ FString FStateTreeExecutionContext::GetDebugInfoString() const
 				for (int32 TaskIndex = State.TasksBegin; TaskIndex < (State.TasksBegin + State.TasksNum); TaskIndex++)
 				{
 					const FStateTreeTaskBase& Task = StateTree.Nodes[TaskIndex].Get<const FStateTreeTaskBase>();
-					Task.AppendDebugInfoString(DebugString, *this);
+					if (Task.bTaskEnabled)
+					{
+						Task.AppendDebugInfoString(DebugString, *this);
+					}
 				}
 			}
 		}
diff --git a/Engine/Plugins/Runtime/StateTree/Source/StateTreeModule/Public/StateTreeConditionBase.h b/Engine/Plugins/Runtime/StateTree/Source/StateTreeModule/Public/StateTreeConditionBase.h
--- a/Engine/Plugins/Runtime/StateTree/Source/StateTreeModule/Public/StateTreeConditionBase.h
+++ b/Engine/Plugins/Runtime/StateTree/Source/StateTreeModule/Public/StateTreeConditionBase.h
@@ -29,6 +29,9 @@ struct STATETREEMODULE_API FStateTreeConditionBase : public FStateTreeNodeBase
 
 	UPROPERTY()
 	int8 DeltaIndent = 0;
+
+	UPROPERTY()
+	EStateTreeConditionEvaluationMode EvaluationMode = EStateTreeConditionEvaluationMode::Evaluated;
 };
 
 /**
diff --git a/Engine/Plugins/Runtime/StateTree/Source/StateTreeModule/Public/StateTreeTaskBase.h b/Engine/Plugins/Runtime/StateTree/Source/StateTreeModule/Public/StateTreeTaskBase.h
--- a/Engine/Plugins/Runtime/StateTree/Source/StateTreeModule/Public/StateTreeTaskBase.h
+++ b/Engine/Plugins/Runtime/StateTree/Source/StateTreeModule/Public/StateTreeTaskBase.h
@@ -23,6 +23,7 @@ struct STATETREEMODULE_API FStateTreeTaskBase : public FStateTreeNodeBase
 		, bShouldCopyBoundPropertiesOnTick(true)
 		, bShouldCopyBoundPropertiesOnExitState(true)
 		, bShouldAffectTransitions(false)
+		, bTaskEnabled(true)
 	{
 	}
 	
@@ -99,6 +100,10 @@ struct STATETREEMODULE_API FStateTreeTaskBase : public FStateTreeNodeBase
 
 	/** If set to true, TriggerTransitions() is called during transition handling. Default false. */
 	uint8 bShouldAffectTransitions : 1;
+
+	/** True if the node is Enabled (i.e. not explicitly disabled in the asset). */
+	UPROPERTY()
+	uint8 bTaskEnabled : 1;
 };
 
 /**
diff --git a/Engine/Plugins/Runtime/StateTree/Source/StateTreeModule/Public/StateTreeTypes.h b/Engine/Plugins/Runtime/StateTree/Source/StateTreeModule/Public/StateTreeTypes.h
--- a/Engine/Plugins/Runtime/StateTree/Source/StateTreeModule/Public/StateTreeTypes.h
+++ b/Engine/Plugins/Runtime/StateTree/Source/StateTreeModule/Public/StateTreeTypes.h
@@ -91,6 +91,20 @@ enum class EStateTreeConditionOperand : uint8
 	Or,
 };
 
+/** Defines how to assign the result of a condition to evaluate.  */
+UENUM()
+enum class EStateTreeConditionEvaluationMode : uint8
+{
+	/** Condition is evaluated to set the result. This is the normal behavior. */
+	Evaluated,
+	
+	/** Do not evaluate the condition and force result to 'true'. */
+	ForcedTrue,
+	
+	/** Do not evaluate the condition and force result to 'false'. */
+	ForcedFalse,
+};
+
 UENUM()
 enum class EStateTreeStateType : uint8
 {
@@ -688,6 +702,10 @@ struct STATETREEMODULE_API FCompactStateTreeState
 	/** True if the state contains tasks that should be called during transition handling. */
 	UPROPERTY()
 	uint8 bHasTransitionTasks : 1;
+
+	/** True if the state is Enabled (i.e. not explicitly marked as disabled). */
+	UPROPERTY()
+	uint8 bEnabled : 1;
 };
 
 USTRUCT()
