Fix CIS Issue 314921
#rnx

diff --git a/Engine/Plugins/Runtime/StateTree/Source/StateTreeEditorModule/Private/SStateTreeView.cpp b/Engine/Plugins/Runtime/StateTree/Source/StateTreeEditorModule/Private/SStateTreeView.cpp
--- a/Engine/Plugins/Runtime/StateTree/Source/StateTreeEditorModule/Private/SStateTreeView.cpp
+++ b/Engine/Plugins/Runtime/StateTree/Source/StateTreeEditorModule/Private/SStateTreeView.cpp
@@ -208,8 +208,8 @@ void SStateTreeView::BindCommands()
 
 				// Should not happen since action is not visible in this case
 				return ECheckBoxState::Undetermined;
-			} )),
-		FIsActionButtonVisible::CreateLambda([this] { return CanEnableStates() || CanDisableStates(); });
+			}),
+		FIsActionButtonVisible::CreateLambda([this] { return CanEnableStates() || CanDisableStates(); }));
 }
 
 bool SStateTreeView::HasSelection() const
