[StateTree]
- fixed scrub state not refreshed when selecting another instance
 -added GetInstanceName and using it instead of GetInstanceDescription for the instance tree view to get lighter UI
- used a single external scrollbar for instances and timelines treeviews
#rnx

diff --git a/Engine/Plugins/AI/MassAI/Source/MassAIBehavior/Public/MassStateTreeExecutionContext.h b/Engine/Plugins/AI/MassAI/Source/MassAIBehavior/Public/MassStateTreeExecutionContext.h
--- a/Engine/Plugins/AI/MassAI/Source/MassAIBehavior/Public/MassStateTreeExecutionContext.h
+++ b/Engine/Plugins/AI/MassAI/Source/MassAIBehavior/Public/MassStateTreeExecutionContext.h
@@ -28,7 +28,7 @@ public:
 protected:
 
 	/** Prefix that will be used by STATETREE_LOG and STATETREE_CLOG, using Entity description. */
-	virtual FString GetInstanceDescription() const override { return FString::Printf(TEXT("Entity [%s]: "), *Entity.DebugGetDescription()); }
+	virtual FString GetInstanceDescription() const override { return FString::Printf(TEXT("Entity [%s]"), *Entity.DebugGetDescription()); }
 
 	virtual void BeginDelayedTransition(const FStateTreeTransitionDelayedState& DelayedState) override;
 
diff --git a/Engine/Plugins/Runtime/StateTree/Source/StateTreeEditorModule/Private/Debugger/SStateTreeDebuggerView.cpp b/Engine/Plugins/Runtime/StateTree/Source/StateTreeEditorModule/Private/Debugger/SStateTreeDebuggerView.cpp
--- a/Engine/Plugins/Runtime/StateTree/Source/StateTreeEditorModule/Private/Debugger/SStateTreeDebuggerView.cpp
+++ b/Engine/Plugins/Runtime/StateTree/Source/StateTreeEditorModule/Private/Debugger/SStateTreeDebuggerView.cpp
@@ -319,8 +319,11 @@ void SStateTreeDebuggerView::Construct(const FArguments& InArgs, const UStateTre
 			})
 		];
 
+	TSharedPtr<SScrollBar> ScrollBar = SNew(SScrollBar); 
+	
 	// Instances TreeView
 	InstancesTreeView =	SNew(SStateTreeDebuggerInstanceTree)
+		.ExternalScrollBar(ScrollBar)
 		.OnExpansionChanged_Lambda([this]() { InstanceTimelinesTreeView->RestoreExpansion(); })
 		.OnScrolled_Lambda([this](double ScrollOffset)
 		{
@@ -337,6 +340,7 @@ void SStateTreeDebuggerView::Construct(const FArguments& InArgs, const UStateTre
 	
 	// Timelines TreeView
 	InstanceTimelinesTreeView = SNew(SStateTreeDebuggerTimelines)
+		.ExternalScrollbar(ScrollBar)
 		.OnExpansionChanged_Lambda([this]() { InstancesTreeView->RestoreExpansion(); })
 		.OnScrolled_Lambda([this](double ScrollOffset) { InstancesTreeView->ScrollTo(ScrollOffset); })
 		.InstanceTracks(&InstanceTracks)
@@ -527,12 +531,15 @@ void SStateTreeDebuggerView::Construct(const FArguments& InArgs, const UStateTre
 							HeaderSplitter->SlotAt(1).SetSizeValue(Size);
 						})
 					[
-						SNew(SScrollBox)
-						.Orientation(Orient_Vertical)
-						+ SScrollBox::Slot()
+						SNew(SOverlay)
+						+ SOverlay::Slot()
 						[
 							InstanceTimelinesTreeView.ToSharedRef()
 						]
+						+ SOverlay::Slot().HAlign(EHorizontalAlignment::HAlign_Right)
+						[
+							ScrollBar.ToSharedRef()
+						]
 					]
 				]
 				+ SSplitter::Slot()
@@ -744,7 +751,7 @@ void SStateTreeDebuggerView::OnDebuggerScrubStateChanged(const UE::StateTreeDebu
 	TArray<FParentInfo, TInlineAllocator<8>> PhaseStack;
 	TArray<FStateTreeStateHandle, TInlineAllocator<8>> StateStack;
 
-	const int32 SpanIdx = ScrubState.FrameSpanIndex;
+	const int32 SpanIdx = ScrubState.GetFrameSpanIndex();
 	const int32 FirstEventIdx = Spans[SpanIdx].EventIdx;
 	const TraceServices::FFrame Frame = Spans[SpanIdx].Frame;
 	const int32 MaxEventIdx = Spans.IsValidIndex(SpanIdx+1) ? Spans[SpanIdx+1].EventIdx : Events.Num();
@@ -834,7 +841,7 @@ void SStateTreeDebuggerView::OnNewInstance(FStateTreeInstanceDebugId InstanceId)
 
 	if (ExistingTrack == nullptr)
 	{
-		InstanceTracks.Add(MakeShared<FStateTreeDebuggerTrack>(Debugger, InstanceId, Debugger->GetInstanceDescription(InstanceId), ViewRange));
+		InstanceTracks.Add(MakeShared<FStateTreeDebuggerTrack>(Debugger, InstanceId, Debugger->GetInstanceName(InstanceId), ViewRange));
 	}
 
 	InstancesTreeView->Refresh();
diff --git a/Engine/Plugins/Runtime/StateTree/Source/StateTreeModule/Private/Debugger/StateTreeDebugger.cpp b/Engine/Plugins/Runtime/StateTree/Source/StateTreeModule/Private/Debugger/StateTreeDebugger.cpp
--- a/Engine/Plugins/Runtime/StateTree/Source/StateTreeModule/Private/Debugger/StateTreeDebugger.cpp
+++ b/Engine/Plugins/Runtime/StateTree/Source/StateTreeModule/Private/Debugger/StateTreeDebugger.cpp
@@ -137,6 +137,17 @@ void FStateTreeDebugger::SyncToCurrentSessionDuration()
 	}
 }
 
+FText FStateTreeDebugger::GetInstanceName(FStateTreeInstanceDebugId InstanceId) const
+{
+	using namespace UE::StateTreeDebugger;
+	const FInstanceDescriptor* FoundInstance = InstanceDescs.FindByPredicate([InstanceId](const FInstanceDescriptor& InstanceDesc)
+		{
+			return InstanceDesc.Id == InstanceId;
+		});
+
+	return (FoundInstance != nullptr) ? FText::FromString(FoundInstance->Name) : LOCTEXT("InstanceNotFound","Instance not found");
+}
+
 FText FStateTreeDebugger::GetInstanceDescription(FStateTreeInstanceDebugId InstanceId) const
 {
 	using namespace UE::StateTreeDebugger;
@@ -158,15 +169,15 @@ void FStateTreeDebugger::SelectInstance(const FStateTreeInstanceDebugId Instance
 		OnSelectedInstanceCleared.ExecuteIfBound();
 
 		// Update event collection index for newly debugged instance
-		ScrubState.EventCollectionIndex = InstanceId.IsValid() ? EventCollections.IndexOfByPredicate([InstanceId = InstanceId](const UE::StateTreeDebugger::FInstanceEventCollection& Entry)
+		ScrubState.SetEventCollectionIndex(InstanceId.IsValid() ? EventCollections.IndexOfByPredicate([InstanceId = InstanceId](const UE::StateTreeDebugger::FInstanceEventCollection& Entry)
 			{
 				return Entry.InstanceId == InstanceId;
 			})
-			: INDEX_NONE;
+			: INDEX_NONE);
 
-		// Set scrub time to refresh internal tracking for the new instance
-		// This will notify scrub position changed and active states
-		SetScrubTime(ScrubState.ScrubTime);
+		OnScrubStateChanged.Execute(ScrubState);
+
+		RefreshActiveStates();
 	}
 }
 
@@ -427,8 +438,8 @@ void FStateTreeDebugger::RefreshActiveStates()
 
 	if (ScrubState.IsPointingToValidActiveStates())
 	{
-		const UE::StateTreeDebugger::FInstanceEventCollection& EventCollection = EventCollections[ScrubState.EventCollectionIndex];
-		const int32 EventIndex = EventCollection.ActiveStatesChanges[ScrubState.ActiveStatesIndex].EventIndex;
+		const UE::StateTreeDebugger::FInstanceEventCollection& EventCollection = EventCollections[ScrubState.GetEventCollectionIndex()];
+		const int32 EventIndex = EventCollection.ActiveStatesChanges[ScrubState.GetActiveStatesIndex()].EventIndex;
 		NewActiveStates = EventCollection.Events[EventIndex].Get<FStateTreeTraceActiveStatesEvent>().ActiveStates;
 	}
 
@@ -678,9 +689,9 @@ bool FStateTreeDebugger::ProcessEvent(const FStateTreeInstanceDebugId InstanceId
 
 		// Update the active event collection index when it's newly created for the currently debugged instance.
 		// Otherwise (i.e. EventCollection already exists) it is updated when switching instance (i.e. SelectInstance)
-		if (SelectedInstanceId == InstanceId && ScrubState.EventCollectionIndex == INDEX_NONE)
+		if (SelectedInstanceId == InstanceId && ScrubState.GetEventCollectionIndex() == INDEX_NONE)
 		{
-			ScrubState.EventCollectionIndex = EventCollections.Num();
+			ScrubState.SetEventCollectionIndex(EventCollections.Num());
 		}
 
 		ExistingCollection = &EventCollections.Emplace_GetRef(InstanceId);
diff --git a/Engine/Plugins/Runtime/StateTree/Source/StateTreeModule/Private/Debugger/StateTreeDebuggerTypes.cpp b/Engine/Plugins/Runtime/StateTree/Source/StateTreeModule/Private/Debugger/StateTreeDebuggerTypes.cpp
--- a/Engine/Plugins/Runtime/StateTree/Source/StateTreeModule/Private/Debugger/StateTreeDebuggerTypes.cpp
+++ b/Engine/Plugins/Runtime/StateTree/Source/StateTreeModule/Private/Debugger/StateTreeDebuggerTypes.cpp
@@ -83,6 +83,16 @@ bool FScrubState::SetScrubTime(const double NewScrubTime)
 	return true;
 }
 
+void FScrubState::SetEventCollectionIndex(const int32 InEventCollectionIndex)
+{
+	EventCollectionIndex = InEventCollectionIndex;
+
+	// Force refresh of internal indices with current time applied on the new event collection. 
+	const double PrevScrubTime = ScrubTime;
+	ScrubTime = 0;
+	SetScrubTime(PrevScrubTime);
+}
+
 void FScrubState::SetFrameSpanIndex(const int32 NewFrameSpanIndex)
 {
 	FrameSpanIndex = NewFrameSpanIndex;
diff --git a/Engine/Plugins/Runtime/StateTree/Source/StateTreeModule/Private/StateTreeExecutionContext.cpp b/Engine/Plugins/Runtime/StateTree/Source/StateTreeModule/Private/StateTreeExecutionContext.cpp
--- a/Engine/Plugins/Runtime/StateTree/Source/StateTreeModule/Private/StateTreeExecutionContext.cpp
+++ b/Engine/Plugins/Runtime/StateTree/Source/StateTreeModule/Private/StateTreeExecutionContext.cpp
@@ -12,8 +12,8 @@
 #include "ProfilingDebugging/CsvProfiler.h"
 #include "Logging/LogScopedVerbosityOverride.h"
 
-#define STATETREE_LOG(Verbosity, Format, ...) UE_VLOG_UELOG(GetOwner(), LogStateTree, Verbosity, TEXT("%s") Format, *GetInstanceDescription(), ##__VA_ARGS__)
-#define STATETREE_CLOG(Condition, Verbosity, Format, ...) UE_CVLOG_UELOG((Condition), GetOwner(), LogStateTree, Verbosity, TEXT("%s") Format, *GetInstanceDescription(), ##__VA_ARGS__)
+#define STATETREE_LOG(Verbosity, Format, ...) UE_VLOG_UELOG(GetOwner(), LogStateTree, Verbosity, TEXT("%s: ") Format, *GetInstanceDescription(), ##__VA_ARGS__)
+#define STATETREE_CLOG(Condition, Verbosity, Format, ...) UE_CVLOG_UELOG((Condition), GetOwner(), LogStateTree, Verbosity, TEXT("%s: ") Format, *GetInstanceDescription(), ##__VA_ARGS__)
 
 #if WITH_STATETREE_DEBUGGER
 	#define ID_NAME PREPROCESSOR_JOIN(InstanceId,__LINE__) \
@@ -2100,7 +2100,7 @@ EStateTreeRunStatus FStateTreeExecutionContext::GetLastTickStatus() const
 
 FString FStateTreeExecutionContext::GetInstanceDescription() const
 {
-	return FString::Printf(TEXT("%s: "), *GetNameSafe(&Owner));
+	return FString::Printf(TEXT("%s"), *GetNameSafe(&Owner));
 }
 
 const FStateTreeActiveStates& FStateTreeExecutionContext::GetActiveStates() const
diff --git a/Engine/Plugins/Runtime/StateTree/Source/StateTreeModule/Public/Debugger/StateTreeDebugger.h b/Engine/Plugins/Runtime/StateTree/Source/StateTreeModule/Public/Debugger/StateTreeDebugger.h
--- a/Engine/Plugins/Runtime/StateTree/Source/StateTreeModule/Public/Debugger/StateTreeDebugger.h
+++ b/Engine/Plugins/Runtime/StateTree/Source/StateTreeModule/Public/Debugger/StateTreeDebugger.h
@@ -80,6 +80,7 @@ struct STATETREEMODULE_API FStateTreeDebugger : FTickableGameObject
 	const TraceServices::IAnalysisSession* GetAnalysisSession() const;
 
 	bool IsActiveInstance(double Time, FStateTreeInstanceDebugId InstanceId) const;
+	FText GetInstanceName(FStateTreeInstanceDebugId InstanceId) const;
 	FText GetInstanceDescription(FStateTreeInstanceDebugId InstanceId) const;
 	void SelectInstance(const FStateTreeInstanceDebugId InstanceId);
 	FStateTreeInstanceDebugId GetSelectedInstance() const { return SelectedInstanceId; }
@@ -102,7 +103,7 @@ struct STATETREEMODULE_API FStateTreeDebugger : FTickableGameObject
 	double GetAnalysisDuration() const { return AnalysisDuration; }
 
 	/** Returns the time (based on the recording duration) associated to the selected frame. */
-	double GetScrubTime() const	{ return ScrubState.ScrubTime; }
+	double GetScrubTime() const	{ return ScrubState.GetScrubTime(); }
 	void SetScrubTime(double ScrubTime);
 
 	void GetLiveTraces(TArray<FTraceDescriptor>& OutTraceDescriptors) const;
diff --git a/Engine/Plugins/Runtime/StateTree/Source/StateTreeModule/Public/Debugger/StateTreeDebuggerTypes.h b/Engine/Plugins/Runtime/StateTree/Source/StateTreeModule/Public/Debugger/StateTreeDebuggerTypes.h
--- a/Engine/Plugins/Runtime/StateTree/Source/StateTreeModule/Public/Debugger/StateTreeDebuggerTypes.h
+++ b/Engine/Plugins/Runtime/StateTree/Source/StateTreeModule/Public/Debugger/StateTreeDebuggerTypes.h
@@ -138,33 +138,26 @@ struct STATETREEMODULE_API FInstanceEventCollection
 
 struct STATETREEMODULE_API FScrubState
 {
-	enum class EScrubTimeBoundState : uint8
-	{
-		Unset,
-		/** There are events but current time is before the first frame. */
-		BeforeLowerBound,
-		/** There are events and current time is within the frames received. */
-		InBounds,
-		/** There are events but current time is after the last frame. */
-		AfterHigherBound
-	};
-
 	explicit FScrubState(const TArray<FInstanceEventCollection>& EventCollections)
 		: EventCollections(EventCollections)
 	{
 	}
 
-private:
-	const TArray<FInstanceEventCollection>& EventCollections;
-	
-public:
-	double ScrubTime = 0;
-	int32 EventCollectionIndex = INDEX_NONE;
-	uint64 TraceFrameIndex = INDEX_NONE;
-	int32 FrameSpanIndex = INDEX_NONE;
-	int32 ActiveStatesIndex = INDEX_NONE;
-	EScrubTimeBoundState ScrubTimeBoundState = EScrubTimeBoundState::Unset;
-	
+	/** @return Index of the currently selected event collection; INDEX_NONE if nothing is selected. */
+	int32 GetEventCollectionIndex() const { return EventCollectionIndex; }
+
+	/** Assigns a new collection index and updates internal indices for current scrub time. */
+	void SetEventCollectionIndex(const int32 InEventCollectionIndex);
+
+	/** @return Index of the span for the currently selected frame; INDEX_NONE if there is no span for the current scrub time. */
+	int32 GetFrameSpanIndex() const { return FrameSpanIndex; }
+
+	/** @return Index of the list of active states for the currently selected frame; INDEX_NONE if there is no active states for the current scrub time. */
+	int32 GetActiveStatesIndex() const { return ActiveStatesIndex; }
+
+	/** @return Current scrub time. */
+	double GetScrubTime() const { return ScrubTime; }
+
 	/**
 	 * Updates internal indices based on the new time.
 	 * @return true if values were updated; false otherwise (i.e. no changes)
@@ -231,9 +224,28 @@ public:
 	const FInstanceEventCollection& GetEventCollection() const;
 
 private:
+	enum class EScrubTimeBoundState : uint8
+	{
+		Unset,
+		/** There are events but current time is before the first frame. */
+		BeforeLowerBound,
+		/** There are events and current time is within the frames received. */
+		InBounds,
+		/** There are events but current time is after the last frame. */
+		AfterHigherBound
+	};
+
 	void SetFrameSpanIndex(int32 NewFrameSpanIndex);
 	void SetActiveStatesIndex(int32 NewActiveStatesIndex);
 	void UpdateActiveStatesIndex(int32 SpanIndex);
+
+	const TArray<FInstanceEventCollection>& EventCollections;
+	double ScrubTime = 0;
+	int32 EventCollectionIndex = INDEX_NONE;
+	uint64 TraceFrameIndex = INDEX_NONE;
+	int32 FrameSpanIndex = INDEX_NONE;
+	int32 ActiveStatesIndex = INDEX_NONE;
+	EScrubTimeBoundState ScrubTimeBoundState = EScrubTimeBoundState::Unset;
 };
 
 } // UE::StateTreeDebugger
