StateTree: Fix source binding path validation for context structs.

diff --git a/Engine/Plugins/Runtime/StateTree/Source/StateTreeEditorModule/Private/StateTreeCompiler.cpp b/Engine/Plugins/Runtime/StateTree/Source/StateTreeEditorModule/Private/StateTreeCompiler.cpp
--- a/Engine/Plugins/Runtime/StateTree/Source/StateTreeEditorModule/Private/StateTreeCompiler.cpp
+++ b/Engine/Plugins/Runtime/StateTree/Source/StateTreeEditorModule/Private/StateTreeCompiler.cpp
@@ -1426,12 +1426,17 @@ bool FStateTreeCompiler::GetAndValidateBindings(const FStateTreeBindableStructDe
 			return false;
 		}
 
-		if (!Binding.GetMutableSourcePath().UpdateInstanceStructsFromValue(IDToStructValue[SourceStructID]))
+		// Update the source structs only if we have value for it. For some sources (e.g. context structs) we know only type, and in that case there are no instance structs.
+		const FStateTreeDataView SourceValue = IDToStructValue[SourceStructID];
+		if (SourceValue.IsValid())
 		{
-			Log.Reportf(EMessageSeverity::Error, TargetStruct,
-				TEXT("Malformed target property path for binding source property '%s' for source '%s:%s'."),
-				*Binding.GetSourcePath().ToString(), *TargetStruct.Name.ToString(), *Binding.GetTargetPath().ToString());
-			return false;
+			if (!Binding.GetMutableSourcePath().UpdateInstanceStructsFromValue(SourceValue))
+			{
+				Log.Reportf(EMessageSeverity::Error, TargetStruct,
+					TEXT("Malformed target property path for binding source property '%s' for source '%s:%s'."),
+					*Binding.GetSourcePath().ToString(), *TargetStruct.Name.ToString(), *Binding.GetTargetPath().ToString());
+				return false;
+			}
 		}
 
 		// Special case fo AnyEnum. StateTreeBindingExtension allows AnyEnums to bind to other enum types.
