Fixed unitialized property
#rnx

diff --git a/Engine/Plugins/Runtime/StateTree/Source/StateTreeModule/Public/StateTreeTypes.h b/Engine/Plugins/Runtime/StateTree/Source/StateTreeModule/Public/StateTreeTypes.h
--- a/Engine/Plugins/Runtime/StateTree/Source/StateTreeModule/Public/StateTreeTypes.h
+++ b/Engine/Plugins/Runtime/StateTree/Source/StateTreeModule/Public/StateTreeTypes.h
@@ -630,6 +630,7 @@ struct STATETREEMODULE_API FCompactStateTreeState
 
 	FCompactStateTreeState()
 		: bHasTransitionTasks(false)
+		, bEnabled(true)
 	{
 	}
 	
