[StateTreeDebugger] QOL improvements for PIE debugging
- improved the connection to the next live session in case other live traces are actives
- give priority to the Editor actions (e.g. start/stop recording or auto-record on PIE start) even if an analysis was active (remote process). This will stop current analysis and connect to the new live trace.
#rnx

diff --git a/Engine/Plugins/Runtime/StateTree/Source/StateTreeEditorModule/Private/Debugger/SStateTreeDebuggerView.cpp b/Engine/Plugins/Runtime/StateTree/Source/StateTreeEditorModule/Private/Debugger/SStateTreeDebuggerView.cpp
--- a/Engine/Plugins/Runtime/StateTree/Source/StateTreeEditorModule/Private/Debugger/SStateTreeDebuggerView.cpp
+++ b/Engine/Plugins/Runtime/StateTree/Source/StateTreeEditorModule/Private/Debugger/SStateTreeDebuggerView.cpp
@@ -172,15 +172,17 @@ SStateTreeDebuggerView::~SStateTreeDebuggerView()
 
 void SStateTreeDebuggerView::StartRecording()
 {
-	if (!CanStartRecording())
-	{
-		return;
-	}
+	// We give priority to the Editor actions even if an analysis was active (remote process)
+	// This will stop current analysis and connect to the new live trace.
+	Debugger->RequestAnalysisOfNextLiveSession();
 
-	IStateTreeModule& StateTreeModule = FModuleManager::GetModuleChecked<IStateTreeModule>("StateTreeModule");
-	StateTreeModule.StartTraces();
+	if (CanStartRecording())
+	{
+		IStateTreeModule& StateTreeModule = FModuleManager::GetModuleChecked<IStateTreeModule>("StateTreeModule");
+		StateTreeModule.StartTraces();
 
-	bRecording = true;
+		bRecording = true;
+	}
 }
 
 void SStateTreeDebuggerView::StopRecording()
@@ -208,11 +210,6 @@ void SStateTreeDebuggerView::OnPIEStarted(const bool bIsSimulating)
 	if (UStateTreeEditorSettings::Get().bShouldDebuggerAutoRecordOnPIE)
 	{
 		StartRecording();
-
-		if (!Debugger->IsAnalysisSessionActive())
-		{
-			Debugger->StartLastLiveSessionAnalysis();
-		}
 	}
 }
 
diff --git a/Engine/Plugins/Runtime/StateTree/Source/StateTreeModule/Private/Debugger/StateTreeDebugger.cpp b/Engine/Plugins/Runtime/StateTree/Source/StateTreeModule/Private/Debugger/StateTreeDebugger.cpp
--- a/Engine/Plugins/Runtime/StateTree/Source/StateTreeModule/Private/Debugger/StateTreeDebugger.cpp
+++ b/Engine/Plugins/Runtime/StateTree/Source/StateTreeModule/Private/Debugger/StateTreeDebugger.cpp
@@ -90,14 +90,14 @@ void FStateTreeDebugger::Tick(const float DeltaTime)
 {
 	TRACE_CPUPROFILER_EVENT_SCOPE(FStateTreeDebugger::Tick);
 
-	if (RetryLoadLastLiveSessionTimer > 0.0f)
+	if (RetryLoadNextLiveSessionTimer > 0.0f)
 	{
 		// We are still not connected to the last live session.
 		// Update polling timer and retry with remaining time; 0 or less will stop retries.
-		StartLastLiveSessionAnalysis(RetryLoadLastLiveSessionTimer - DeltaTime);
-		if (ActiveSessionTraceDescriptor.IsValid())
+		if (TryStartNewLiveSessionAnalysis(RetryLoadNextLiveSessionTimer - DeltaTime))
 		{
-			RetryLoadLastLiveSessionTimer = 0.0f;
+			RetryLoadNextLiveSessionTimer = 0.0f;
+			LastLiveSessionId = INDEX_NONE;
 		}
 	}
 	
@@ -207,34 +207,52 @@ void FStateTreeDebugger::UpdateInstances()
 	}
 }
 
-void FStateTreeDebugger::StartLastLiveSessionAnalysis(const float RetryPollingDuration)
+void FStateTreeDebugger::RequestAnalysisOfNextLiveSession()
 {
+	// Invalidate our current active session
+	ActiveSessionTraceDescriptor = FTraceDescriptor();
+
+	// Stop current analysis if any
+	StopAnalysis();
+
 	TArray<FTraceDescriptor> Traces;
 	GetLiveTraces(Traces);
+	
+	LastLiveSessionId = Traces.Num() ? Traces.Last().TraceId : INDEX_NONE;
 
-	if (Traces.Num())
-	{
-		StartSessionAnalysis(Traces.Last());
-	}
-	else
+	// This won't succeed yet but will schedule our next retry
+	TryStartNewLiveSessionAnalysis(1.0f);
+}
+
+bool FStateTreeDebugger::TryStartNewLiveSessionAnalysis(const float RetryPollingDuration)
+{
+	TArray<FTraceDescriptor> Traces;
+	GetLiveTraces(Traces);
+
+	if (Traces.Num() && Traces.Last().TraceId != LastLiveSessionId)
 	{
-		RetryLoadLastLiveSessionTimer = RetryPollingDuration;
-		ensure(RetryLoadLastLiveSessionTimer > 0);
-		UE_CLOG(RetryLoadLastLiveSessionTimer > 0, LogStateTree, Log, TEXT("Unable to start analysis for the most recent live session."));
+		return StartSessionAnalysis(Traces.Last());
 	}
+	
+	RetryLoadNextLiveSessionTimer = RetryPollingDuration;
+	ensure(RetryLoadNextLiveSessionTimer > 0);
+	UE_CLOG(RetryLoadNextLiveSessionTimer > 0, LogStateTree, Log, TEXT("Unable to start analysis for the most recent live session."));
+
+	return false;
 }
 
-void FStateTreeDebugger::StartSessionAnalysis(const FTraceDescriptor& TraceDescriptor)
+bool FStateTreeDebugger::StartSessionAnalysis(const FTraceDescriptor& TraceDescriptor)
 {
 	if (ActiveSessionTraceDescriptor == TraceDescriptor)
 	{
-		return;
+		return ActiveSessionTraceDescriptor.IsValid();
 	}
-	
+
+	ActiveSessionTraceDescriptor = FTraceDescriptor();
 	UE::Trace::FStoreClient* StoreClient = GetStoreClient();
 	if (StoreClient == nullptr)
 	{
-		return;
+		return false;
 	}
 
 	// Make sure any active analysis is stopped
@@ -253,7 +271,7 @@ void FStateTreeDebugger::StartSessionAnalysis(const FTraceDescriptor& TraceDescr
 		UE::Trace::FStoreClient::FTraceData TraceData = StoreClient->ReadTrace(TraceId);
 		if (!TraceData)
 		{
-			return;
+			return false;
 		}
 
 		FString TraceName(StoreClient->GetStatus()->GetStoreDir());
@@ -278,8 +296,13 @@ void FStateTreeDebugger::StartSessionAnalysis(const FTraceDescriptor& TraceDescr
 			SyncToCurrentSessionDuration();
 		}
 
-		ActiveSessionTraceDescriptor = AnalysisSession.IsValid() ? TraceDescriptor : FTraceDescriptor();
+		if (AnalysisSession.IsValid())
+		{
+			ActiveSessionTraceDescriptor = TraceDescriptor;	
+		}
 	}
+
+	return ActiveSessionTraceDescriptor.IsValid();
 }
 
 void FStateTreeDebugger::GetLiveTraces(TArray<FTraceDescriptor>& OutTraceDescriptors) const
diff --git a/Engine/Plugins/Runtime/StateTree/Source/StateTreeModule/Public/Debugger/StateTreeDebugger.h b/Engine/Plugins/Runtime/StateTree/Source/StateTreeModule/Public/Debugger/StateTreeDebugger.h
--- a/Engine/Plugins/Runtime/StateTree/Source/StateTreeModule/Public/Debugger/StateTreeDebugger.h
+++ b/Engine/Plugins/Runtime/StateTree/Source/StateTreeModule/Public/Debugger/StateTreeDebugger.h
@@ -107,15 +107,14 @@ struct STATETREEMODULE_API FStateTreeDebugger : FTickableGameObject
 	void SetScrubTime(double ScrubTime);
 
 	void GetLiveTraces(TArray<FTraceDescriptor>& OutTraceDescriptors) const;
-	
+
 	/**
-	 * Creates a new analysis session instance and loads the latest available trace that is live.
-	 * Replaces the current analysis session.
-	 * On failure, if 'RetryPollingDuration is > 0', will retry connecting every frame for 'RetryPollingDuration' seconds 
-	 * @param RetryPollingDuration - On failure, how many seconds to retry connecting during FStateTreeDebugger::Tick
+	 * Queue a request to auto start an analysis session on the next available live trace.
+	 * This will replace the current analysis session if any.
 	 */
-	void StartLastLiveSessionAnalysis(float RetryPollingDuration = 1.0f);
-	void StartSessionAnalysis(const FTraceDescriptor& TraceDescriptor);
+	void RequestAnalysisOfNextLiveSession();
+
+	bool StartSessionAnalysis(const FTraceDescriptor& TraceDescriptor);
 	void StopAnalysis();
 	FTraceDescriptor GetSelectedTraceDescriptor() const { return ActiveSessionTraceDescriptor; }
 	FText GetSelectedTraceDescription() const;
@@ -149,6 +148,14 @@ private:
 
 	void SetActiveStates(const TConstArrayView<FStateTreeStateHandle> NewActiveStates);
 
+	/**
+	 * Looks for a new live traces to start an analysis session.
+	 * On failure, if 'RetryPollingDuration is > 0', will retry connecting every frame for 'RetryPollingDuration' seconds 
+	 * @param RetryPollingDuration - On failure, how many seconds to retry connecting during FStateTreeDebugger::Tick
+	 * @return True if the analysis was successfully started; false otherwise.
+	 */
+	bool TryStartNewLiveSessionAnalysis(float RetryPollingDuration);
+	
 	/**
 	 * Recompute index of the span that contains the active states change event and update the active states.
 	 * This method handles unselected instances in which case it will reset the active states and set the span index to INDEX_NONE
@@ -191,10 +198,16 @@ private:
 	TArray<FStateTreeStateHandle> ActiveStates;
 
 	/**
-	 * When auto-connecting on session start it is possible that a few frames are required for the tracing session to be accessible and connected to.
+	 * When auto-connecting on next live session it is possible that a few frames are required for the tracing session to be accessible and connected to.
+	 * This is to keep track of the previous last live session id so we can detect when the new one is available.
+	 */
+	int32 LastLiveSessionId = INDEX_NONE;
+	
+	/**
+	 * When auto-connecting on next live session it is possible that a few frames are required for the tracing session to be accessible and connected to.
 	 * This is to keep track of the time window where we will retry.
 	 */
-	float RetryLoadLastLiveSessionTimer = 0.0f;
+	float RetryLoadNextLiveSessionTimer = 0.0f;
 	
 	/** Recording duration of the analysis session in world recorded time. */
 	double RecordingDuration = 0;
