[StateTreeDebugger]
- Added owner track to hold instance tracks as sub tracks. This will regroup all statetree instances running on the same owner during a recorded session.
   - Instance tracks are marked as stale at the end of the session but owner track can be reactivated from a subsequent PIE sessions (i.e. same owners as previous session)
- Added some placeholders icons for the two track types
#rnx

diff --git a/Engine/Plugins/Runtime/StateTree/Source/StateTreeEditorModule/Private/Debugger/SStateTreeDebuggerInstanceTree.cpp b/Engine/Plugins/Runtime/StateTree/Source/StateTreeEditorModule/Private/Debugger/SStateTreeDebuggerInstanceTree.cpp
--- a/Engine/Plugins/Runtime/StateTree/Source/StateTreeEditorModule/Private/Debugger/SStateTreeDebuggerInstanceTree.cpp
+++ b/Engine/Plugins/Runtime/StateTree/Source/StateTreeEditorModule/Private/Debugger/SStateTreeDebuggerInstanceTree.cpp
@@ -49,9 +49,9 @@ TSharedRef<ITableRow> SStateTreeDebuggerInstanceTree::GenerateTreeRow(TSharedPtr
 				.Font(FAppStyle::GetFontStyle(TEXT("PropertyWindow.NormalFont")))
 				.ColorAndOpacity_Lambda([Item]()
 				{
-					if (const FStateTreeDebuggerTrack* DebuggerTrack = static_cast<FStateTreeDebuggerTrack*>(Item.Get()))
+					if (const FStateTreeDebuggerBaseTrack* DebuggerTrack = static_cast<FStateTreeDebuggerBaseTrack*>(Item.Get()))
 					{
-						if (DebuggerTrack->IsStale())
+						if (DebuggerTrack != nullptr && DebuggerTrack->IsStale())
 						{
 							return FSlateColor::UseSubduedForeground();
 						}
diff --git a/Engine/Plugins/Runtime/StateTree/Source/StateTreeEditorModule/Private/Debugger/SStateTreeDebuggerView.cpp b/Engine/Plugins/Runtime/StateTree/Source/StateTreeEditorModule/Private/Debugger/SStateTreeDebuggerView.cpp
--- a/Engine/Plugins/Runtime/StateTree/Source/StateTreeEditorModule/Private/Debugger/SStateTreeDebuggerView.cpp
+++ b/Engine/Plugins/Runtime/StateTree/Source/StateTreeEditorModule/Private/Debugger/SStateTreeDebuggerView.cpp
@@ -195,9 +195,9 @@ void SStateTreeDebuggerView::StopRecording()
 		bRecording = false;
 
 		// Mark all tracks from the stopped session as stale to have different look.
-		for (const TSharedPtr<RewindDebugger::FRewindDebuggerTrack>& DebugTrack : InstanceTracks)
+		for (const TSharedPtr<RewindDebugger::FRewindDebuggerTrack>& DebugTrack : InstanceOwnerTracks)
 		{
-			if (FStateTreeDebuggerTrack* StateTreeTrack = static_cast<FStateTreeDebuggerTrack*>(DebugTrack.Get()))
+			if (FStateTreeDebuggerBaseTrack* StateTreeTrack = static_cast<FStateTreeDebuggerBaseTrack*>(DebugTrack.Get()))
 			{
 				StateTreeTrack->MarkAsStale();
 			}
@@ -316,7 +316,7 @@ void SStateTreeDebuggerView::Construct(const FArguments& InArgs, const UStateTre
 			})
 		];
 
-	TSharedPtr<SScrollBar> ScrollBar = SNew(SScrollBar); 
+	const TSharedPtr<SScrollBar> ScrollBar = SNew(SScrollBar); 
 	
 	// Instances TreeView
 	InstancesTreeView =	SNew(SStateTreeDebuggerInstanceTree)
@@ -326,13 +326,15 @@ void SStateTreeDebuggerView::Construct(const FArguments& InArgs, const UStateTre
 		{
 			InstanceTimelinesTreeView->ScrollTo(ScrollOffset);
 		})
-		.InstanceTracks(&InstanceTracks)
+		.InstanceTracks(&InstanceOwnerTracks)
 		.OnSelectionChanged_Lambda([this](TSharedPtr<RewindDebugger::FRewindDebuggerTrack> SelectedItem, ESelectInfo::Type SelectInfo)
 			{
 				InstanceTimelinesTreeView->SetSelection(SelectedItem);
 
-				const FStateTreeDebuggerTrack* StateTreeTrack = static_cast<FStateTreeDebuggerTrack*>(SelectedItem.Get());
-				Debugger->SelectInstance(StateTreeTrack != nullptr ? StateTreeTrack->GetInstanceId() : FStateTreeInstanceDebugId::Invalid);
+				if (FStateTreeDebuggerBaseTrack* StateTreeTrack = static_cast<FStateTreeDebuggerBaseTrack*>(SelectedItem.Get()))
+				{
+					StateTreeTrack->OnSelected();
+				}
 			});
 	
 	// Timelines TreeView
@@ -340,7 +342,7 @@ void SStateTreeDebuggerView::Construct(const FArguments& InArgs, const UStateTre
 		.ExternalScrollbar(ScrollBar)
 		.OnExpansionChanged_Lambda([this]() { InstancesTreeView->RestoreExpansion(); })
 		.OnScrolled_Lambda([this](double ScrollOffset) { InstancesTreeView->ScrollTo(ScrollOffset); })
-		.InstanceTracks(&InstanceTracks)
+		.InstanceTracks(&InstanceOwnerTracks)
 		.ViewRange_Lambda([this]() { return ViewRange; })
 		.ClampRange_Lambda([this]() { return TRange<double>(0, Debugger->GetRecordingDuration()); })
 		.OnViewRangeChanged_Lambda([this](TRange<double> NewRange) { ViewRange = NewRange; })
@@ -587,7 +589,7 @@ void SStateTreeDebuggerView::Tick(const FGeometry& AllottedGeometry, const doubl
 void SStateTreeDebuggerView::RefreshTracks()
 {
 	bool bChanged = false;
-	for (const TSharedPtr<RewindDebugger::FRewindDebuggerTrack>& DebugTrack : InstanceTracks)
+	for (const TSharedPtr<RewindDebugger::FRewindDebuggerTrack>& DebugTrack : InstanceOwnerTracks)
 	{
 		bChanged = DebugTrack->Update() || bChanged;
 	}
@@ -831,14 +833,27 @@ void SStateTreeDebuggerView::OnBreakpointHit(const FStateTreeInstanceDebugId Ins
 
 void SStateTreeDebuggerView::OnNewInstance(FStateTreeInstanceDebugId InstanceId)
 {
-	const TSharedPtr<RewindDebugger::FRewindDebuggerTrack>* ExistingTrack = InstanceTracks.FindByPredicate([InstanceId](const TSharedPtr<RewindDebugger::FRewindDebuggerTrack>& Track)
+	const UE::StateTreeDebugger::FInstanceDescriptor* FoundDescriptor = Debugger->GetInstanceDescriptor(InstanceId);
+	if (!ensureMsgf(FoundDescriptor != nullptr, TEXT("This callback is from the Debugger so we expect to be able to find matching descriptor.")))
+	{
+		return;
+	}
+
+	const TSharedPtr<RewindDebugger::FRewindDebuggerTrack>* ExistingOwnerTrack = InstanceOwnerTracks.FindByPredicate(
+		[FoundDescriptor](const TSharedPtr<RewindDebugger::FRewindDebuggerTrack>& Track)
+		{
+			return Track.Get()->GetName() == FoundDescriptor->Name;
+		});
+
+	if (ExistingOwnerTrack == nullptr)
 	{
-		return static_cast<FStateTreeDebuggerTrack*>(Track.Get())->GetInstanceId() == InstanceId;
-	});
+		ExistingOwnerTrack = &InstanceOwnerTracks.Add_GetRef(MakeShared<FStateTreeDebuggerOwnerTrack>(FText::FromString(FoundDescriptor->Name)));
+	}
 
-	if (ExistingTrack == nullptr)
+	if (FStateTreeDebuggerOwnerTrack* OwnerTrack = static_cast<FStateTreeDebuggerOwnerTrack*>(ExistingOwnerTrack->Get()))
 	{
-		InstanceTracks.Add(MakeShared<FStateTreeDebuggerTrack>(Debugger, InstanceId, Debugger->GetInstanceName(InstanceId), ViewRange));
+		const FString TrackName = FString::Printf(TEXT("Execution #%d"), OwnerTrack->NumSubTracks()+1);
+		OwnerTrack->AddSubTrack(MakeShared<FStateTreeDebuggerInstanceTrack>(Debugger, InstanceId, FText::FromString(TrackName), ViewRange));
 	}
 
 	InstancesTreeView->Refresh();
@@ -850,7 +865,7 @@ void SStateTreeDebuggerView::OnSelectedInstanceCleared()
 	EventsTreeElements.Reset();
 	if (EventsTreeView)
 	{
-		EventsTreeView->RequestTreeRefresh();	
+		EventsTreeView->RequestTreeRefresh();
 	}
 
 	PropertiesBorder->ClearContent();
diff --git a/Engine/Plugins/Runtime/StateTree/Source/StateTreeEditorModule/Private/Debugger/SStateTreeDebuggerView.h b/Engine/Plugins/Runtime/StateTree/Source/StateTreeEditorModule/Private/Debugger/SStateTreeDebuggerView.h
--- a/Engine/Plugins/Runtime/StateTree/Source/StateTreeEditorModule/Private/Debugger/SStateTreeDebuggerView.h
+++ b/Engine/Plugins/Runtime/StateTree/Source/StateTreeEditorModule/Private/Debugger/SStateTreeDebuggerView.h
@@ -92,8 +92,8 @@ private:
 	TSharedPtr<FStateTreeViewModel> StateTreeViewModel;
 	TWeakObjectPtr<const UStateTree> StateTree;
 
-	/** Tracks for all instances producing trace events for the associated state tree asset. */
-	TArray<TSharedPtr<RewindDebugger::FRewindDebuggerTrack>> InstanceTracks;
+	/** Tracks for all statetree instance owners producing trace events for the associated state tree asset. */
+	TArray<TSharedPtr<RewindDebugger::FRewindDebuggerTrack>> InstanceOwnerTracks;
 
 	/**
 	 * Tree view displaying the instance names and synced with InstanceTimelinesTreeView.
diff --git a/Engine/Plugins/Runtime/StateTree/Source/StateTreeEditorModule/Private/Debugger/StateTreeDebuggerTrack.cpp b/Engine/Plugins/Runtime/StateTree/Source/StateTreeEditorModule/Private/Debugger/StateTreeDebuggerTrack.cpp
--- a/Engine/Plugins/Runtime/StateTree/Source/StateTreeEditorModule/Private/Debugger/StateTreeDebuggerTrack.cpp
+++ b/Engine/Plugins/Runtime/StateTree/Source/StateTreeEditorModule/Private/Debugger/StateTreeDebuggerTrack.cpp
@@ -5,13 +5,17 @@
 #include "Debugger/StateTreeDebuggerTrack.h"
 #include "Debugger/StateTreeDebugger.h"
 #include "SStateTreeDebuggerEventTimelineView.h"
-#include "Styling/AppStyle.h"
 
 //----------------------------------------------------------------------//
-// FStateTreeTrack
+// FStateTreeDebuggerInstanceTrack
 //----------------------------------------------------------------------//
-FStateTreeDebuggerTrack::FStateTreeDebuggerTrack(const TSharedPtr<FStateTreeDebugger>& InDebugger, const FStateTreeInstanceDebugId InInstanceId, const FText InName, const TRange<double>& InViewRange)
-	: TrackName(InName)
+FStateTreeDebuggerInstanceTrack::FStateTreeDebuggerInstanceTrack(
+	const TSharedPtr<FStateTreeDebugger>& InDebugger,
+	const FStateTreeInstanceDebugId InInstanceId,
+	const FText& InName,
+	const TRange<double>& InViewRange
+	)
+	: FStateTreeDebuggerBaseTrack(FSlateIcon("StateTreeEditorStyle", "StateTreeEditor.Debugger.InstanceTrack", "StateTreeEditor.Debugger.InstanceTrack"), InName)
 	, StateTreeDebugger(InDebugger)
 	, InstanceId(InInstanceId)
 	, ViewRange(InViewRange)
@@ -19,12 +23,15 @@ FStateTreeDebuggerTrack::FStateTreeDebuggerTrack(const TSharedPtr<FStateTreeDebu
 	EventData = MakeShared<SStateTreeDebuggerEventTimelineView::FTimelineEventData>();
 }
 
-TSharedPtr<SWidget> FStateTreeDebuggerTrack::GetDetailsViewInternal() 
+void FStateTreeDebuggerInstanceTrack::OnSelected()
 {
-	return SNullWidget::NullWidget;
+	if (FStateTreeDebugger* Debugger = StateTreeDebugger.Get())
+	{
+		Debugger->SelectInstance(InstanceId);
+	}
 }
 
-bool FStateTreeDebuggerTrack::UpdateInternal()
+bool FStateTreeDebuggerInstanceTrack::UpdateInternal()
 {
 	const int32 PrevNumPoints = EventData->Points.Num();
 	const int32 PrevNumWindows = EventData->Windows.Num();
@@ -106,11 +113,70 @@ bool FStateTreeDebuggerTrack::UpdateInternal()
 	return bChanged;
 }
 
-TSharedPtr<SWidget> FStateTreeDebuggerTrack::GetTimelineViewInternal()
+TSharedPtr<SWidget> FStateTreeDebuggerInstanceTrack::GetTimelineViewInternal()
 {
 	return SNew(SStateTreeDebuggerEventTimelineView)
 		.ViewRange_Lambda([this](){ return ViewRange; })
 		.EventData_Lambda([this](){ return EventData; });
 }
 
-#endif // WITH_STATETREE_DEBUGGER
\ No newline at end of file
+
+//----------------------------------------------------------------------//
+// FStateTreeDebuggerOwnerTrack
+//----------------------------------------------------------------------//
+FStateTreeDebuggerOwnerTrack::FStateTreeDebuggerOwnerTrack(const FText& InInstanceName)
+	: FStateTreeDebuggerBaseTrack(FSlateIcon("StateTreeEditorStyle", "StateTreeEditor.Debugger.OwnerTrack", "StateTreeEditor.Debugger.OwnerTrack"), InInstanceName)
+{
+}
+
+bool FStateTreeDebuggerOwnerTrack::UpdateInternal()
+{
+	bool bChanged = false;
+	for (const TSharedPtr<FStateTreeDebuggerInstanceTrack>& Track : SubTracks)
+	{
+		bChanged = Track->Update() || bChanged;
+	}
+
+	return bChanged;
+}
+
+void FStateTreeDebuggerOwnerTrack::IterateSubTracksInternal(TFunction<void(TSharedPtr<FRewindDebuggerTrack> SubTrack)> IteratorFunction)
+{
+	for (TSharedPtr<FStateTreeDebuggerInstanceTrack>& Track : SubTracks)
+	{
+		IteratorFunction(Track);
+	}
+}
+
+void FStateTreeDebuggerOwnerTrack::MarkAsStale()
+{
+	for (TSharedPtr<FStateTreeDebuggerInstanceTrack>& Track : SubTracks)
+	{
+		if (FStateTreeDebuggerBaseTrack* InstanceTrack = Track.Get())
+		{
+			InstanceTrack->MarkAsStale();
+		}
+	}
+}
+
+bool FStateTreeDebuggerOwnerTrack::IsStale() const
+{
+	// Considered stale only if all sub tracks are stale
+	if (SubTracks.IsEmpty())
+	{
+		return false;
+	}
+
+	for (const TSharedPtr<FStateTreeDebuggerInstanceTrack>& Track : SubTracks)
+	{
+		const FStateTreeDebuggerBaseTrack* InstanceTrack = Track.Get();
+		if (InstanceTrack && InstanceTrack->IsStale() == false)
+		{
+			return false;
+		}
+	}
+
+	return true;
+}
+
+#endif // WITH_STATETREE_DEBUGGER
diff --git a/Engine/Plugins/Runtime/StateTree/Source/StateTreeEditorModule/Private/Debugger/StateTreeDebuggerTrack.h b/Engine/Plugins/Runtime/StateTree/Source/StateTreeEditorModule/Private/Debugger/StateTreeDebuggerTrack.h
--- a/Engine/Plugins/Runtime/StateTree/Source/StateTreeEditorModule/Private/Debugger/StateTreeDebuggerTrack.h
+++ b/Engine/Plugins/Runtime/StateTree/Source/StateTreeEditorModule/Private/Debugger/StateTreeDebuggerTrack.h
@@ -10,35 +10,77 @@
 
 struct FStateTreeDebugger;
 
+/** Base struct for Debugger tracks to append some functionalities not available in RewindDebuggerTrack */
+struct FStateTreeDebuggerBaseTrack : RewindDebugger::FRewindDebuggerTrack
+{
+	explicit FStateTreeDebuggerBaseTrack(const FSlateIcon& Icon, const FText& TrackName)
+		: Icon(Icon)
+		, TrackName(TrackName)
+	{
+	}
+
+	virtual bool IsStale() const { return false; }
+	virtual void MarkAsStale() {}
+	virtual void OnSelected() {}
+
+protected:
+	FSlateIcon Icon;
+	FText TrackName;
+};
+
 /**
  * Struct use to store timeline events for a single StateTree instance
  * This is currently not tied to RewindDebugger but we use this base
  * type to facilitate a future integration.
  */
-struct FStateTreeDebuggerTrack : RewindDebugger::FRewindDebuggerTrack
+struct FStateTreeDebuggerInstanceTrack : FStateTreeDebuggerBaseTrack
 {
-	FStateTreeDebuggerTrack(const TSharedPtr<FStateTreeDebugger>& InDebugger, const FStateTreeInstanceDebugId InInstanceId, const FText InName, const TRange<double>& InViewRange);
-	
-	FStateTreeInstanceDebugId GetInstanceId() const { return InstanceId; }
-	bool IsStale() const { return bIsStale; }
-	void MarkAsStale() { bIsStale = true; }
+	explicit FStateTreeDebuggerInstanceTrack(
+		const TSharedPtr<FStateTreeDebugger>& InDebugger,
+		const FStateTreeInstanceDebugId InInstanceId,
+		const FText& InName,
+		const TRange<double>& InViewRange);
 
-private:
+protected:
 	virtual bool UpdateInternal() override;
-	virtual TSharedPtr<SWidget> GetDetailsViewInternal() override;
 	virtual TSharedPtr<SWidget> GetTimelineViewInternal() override;
 	virtual FSlateIcon GetIconInternal() override { return Icon; }
+	virtual FName GetNameInternal() const override { return FName(TrackName.ToString()); }
 	virtual FText GetDisplayNameInternal() const override { return TrackName; }
 
+	virtual void OnSelected() override;
+	virtual void MarkAsStale() override { bIsStale = true; }
+	virtual bool IsStale() const override { return bIsStale; }
+
+private:
 	TSharedPtr<SStateTreeDebuggerEventTimelineView::FTimelineEventData> EventData;
 	
-	FSlateIcon Icon;
-	FText TrackName;
-
 	TSharedPtr<FStateTreeDebugger> StateTreeDebugger;
 	FStateTreeInstanceDebugId InstanceId;
 	const TRange<double>& ViewRange;
 	bool bIsStale = false;
 };
 
+
+/** Parent track of all the statetree instance tracks sharing the same execution context owner */
+struct FStateTreeDebuggerOwnerTrack : FStateTreeDebuggerBaseTrack
+{
+	explicit FStateTreeDebuggerOwnerTrack(const FText& InInstanceName);
+
+	void AddSubTrack(const TSharedPtr<FStateTreeDebuggerInstanceTrack>& SubTrack) { SubTracks.Emplace(SubTrack); }
+	int32 NumSubTracks() const { return SubTracks.Num(); }
+
+protected:
+	virtual bool UpdateInternal() override;	
+	virtual FSlateIcon GetIconInternal() override { return Icon; }
+	virtual FName GetNameInternal() const override { return FName(TrackName.ToString()); }
+	virtual FText GetDisplayNameInternal() const override { return TrackName; }
+	virtual void IterateSubTracksInternal(TFunction<void(TSharedPtr<FRewindDebuggerTrack> SubTrack)> IteratorFunction) override;
+
+	virtual void MarkAsStale() override;
+	virtual bool IsStale() const override;
+private:
+	TArray<TSharedPtr<FStateTreeDebuggerInstanceTrack>> SubTracks;
+};
+
 #endif // WITH_STATETREE_DEBUGGER
\ No newline at end of file
diff --git a/Engine/Plugins/Runtime/StateTree/Source/StateTreeEditorModule/Private/StateTreeEditorStyle.cpp b/Engine/Plugins/Runtime/StateTree/Source/StateTreeEditorModule/Private/StateTreeEditorStyle.cpp
--- a/Engine/Plugins/Runtime/StateTree/Source/StateTreeEditorModule/Private/StateTreeEditorStyle.cpp
+++ b/Engine/Plugins/Runtime/StateTree/Source/StateTreeEditorModule/Private/StateTreeEditorStyle.cpp
@@ -204,6 +204,9 @@ FStateTreeEditorStyle::FStateTreeEditorStyle()
 
 		Set("StateTreeEditor.Debugger.ToggleBreakpoint", new IMAGE_BRUSH_SVG("Starship/Blueprints/Breakpoint_Valid", CoreStyleConstants::Icon16x16));
 		Set("StateTreeEditor.DebugOptions", new IMAGE_BRUSH_SVG("Starship/Common/Bug", CoreStyleConstants::Icon16x16));
+
+		Set("StateTreeEditor.Debugger.OwnerTrack", new IMAGE_BRUSH_SVG("Starship/AssetIcons/AIController_64", CoreStyleConstants::Icon16x16));
+		Set("StateTreeEditor.Debugger.InstanceTrack", new IMAGE_BRUSH_SVG("Starship/AssetIcons/AnimInstance_64", CoreStyleConstants::Icon16x16));		
 	}
 
 	{
diff --git a/Engine/Plugins/Runtime/StateTree/Source/StateTreeModule/Private/Debugger/StateTreeDebugger.cpp b/Engine/Plugins/Runtime/StateTree/Source/StateTreeModule/Private/Debugger/StateTreeDebugger.cpp
--- a/Engine/Plugins/Runtime/StateTree/Source/StateTreeModule/Private/Debugger/StateTreeDebugger.cpp
+++ b/Engine/Plugins/Runtime/StateTree/Source/StateTreeModule/Private/Debugger/StateTreeDebugger.cpp
@@ -137,26 +137,27 @@ void FStateTreeDebugger::SyncToCurrentSessionDuration()
 	}
 }
 
-FText FStateTreeDebugger::GetInstanceName(FStateTreeInstanceDebugId InstanceId) const
+const UE::StateTreeDebugger::FInstanceDescriptor* FStateTreeDebugger::GetInstanceDescriptor(const FStateTreeInstanceDebugId InstanceId) const
 {
-	using namespace UE::StateTreeDebugger;
-	const FInstanceDescriptor* FoundInstance = InstanceDescs.FindByPredicate([InstanceId](const FInstanceDescriptor& InstanceDesc)
+	const UE::StateTreeDebugger::FInstanceDescriptor* FoundDescriptor = InstanceDescs.FindByPredicate(
+		[InstanceId](const UE::StateTreeDebugger::FInstanceDescriptor& Descriptor)
 		{
-			return InstanceDesc.Id == InstanceId;
+			return Descriptor.Id == InstanceId;
 		});
 
-	return (FoundInstance != nullptr) ? FText::FromString(FoundInstance->Name) : LOCTEXT("InstanceNotFound","Instance not found");
+	return FoundDescriptor;
 }
 
-FText FStateTreeDebugger::GetInstanceDescription(FStateTreeInstanceDebugId InstanceId) const
+FText FStateTreeDebugger::GetInstanceName(const FStateTreeInstanceDebugId InstanceId) const
 {
-	using namespace UE::StateTreeDebugger;
-	const FInstanceDescriptor* FoundInstance = InstanceDescs.FindByPredicate([InstanceId](const FInstanceDescriptor& InstanceDesc)
-		{
-			return InstanceDesc.Id == InstanceId;
-		});
+	const UE::StateTreeDebugger::FInstanceDescriptor* FoundDescriptor = GetInstanceDescriptor(InstanceId);
+	return (FoundDescriptor != nullptr) ? FText::FromString(FoundDescriptor->Name) : LOCTEXT("InstanceNotFound","Instance not found");
+}
 
-	return (FoundInstance != nullptr) ? DescribeInstance(*FoundInstance) : LOCTEXT("InstanceNotFound","Instance not found");
+FText FStateTreeDebugger::GetInstanceDescription(const FStateTreeInstanceDebugId InstanceId) const
+{
+	const UE::StateTreeDebugger::FInstanceDescriptor* FoundDescriptor = GetInstanceDescriptor(InstanceId);
+	return (FoundDescriptor != nullptr) ? DescribeInstance(*FoundDescriptor) : LOCTEXT("InstanceNotFound","Instance not found");
 }
 
 void FStateTreeDebugger::SelectInstance(const FStateTreeInstanceDebugId InstanceId)
diff --git a/Engine/Plugins/Runtime/StateTree/Source/StateTreeModule/Public/Debugger/StateTreeDebugger.h b/Engine/Plugins/Runtime/StateTree/Source/StateTreeModule/Public/Debugger/StateTreeDebugger.h
--- a/Engine/Plugins/Runtime/StateTree/Source/StateTreeModule/Public/Debugger/StateTreeDebugger.h
+++ b/Engine/Plugins/Runtime/StateTree/Source/StateTreeModule/Public/Debugger/StateTreeDebugger.h
@@ -83,7 +83,9 @@ struct STATETREEMODULE_API FStateTreeDebugger : FTickableGameObject
 	FText GetInstanceName(FStateTreeInstanceDebugId InstanceId) const;
 	FText GetInstanceDescription(FStateTreeInstanceDebugId InstanceId) const;
 	void SelectInstance(const FStateTreeInstanceDebugId InstanceId);
-	FStateTreeInstanceDebugId GetSelectedInstance() const { return SelectedInstanceId; }
+	FStateTreeInstanceDebugId GetSelectedInstanceId() const { return SelectedInstanceId; }
+	const UE::StateTreeDebugger::FInstanceDescriptor* GetInstanceDescriptor(const FStateTreeInstanceDebugId InstanceId) const ;
+	const UE::StateTreeDebugger::FInstanceDescriptor* GetSelectedInstanceDescriptor() const { return GetInstanceDescriptor(SelectedInstanceId); }
 
 	static FText DescribeTrace(const FTraceDescriptor& TraceDescriptor);
 	static FText DescribeInstance(const UE::StateTreeDebugger::FInstanceDescriptor& StateTreeInstanceDesc);
