[StateTreeDebugger]
- Added buffered instance lifetime event in case recording is started of channel enabled after simulation is running with ticking instances.
- Fixed display issue for inactive instances
- Fixed edge case where we receive StateTree events before the first game frame of the FrameProvider for late recording sessions.
#rnx

diff --git a/Engine/Plugins/Runtime/StateTree/Source/StateTreeEditorModule/Private/Debugger/StateTreeDebuggerTrack.cpp b/Engine/Plugins/Runtime/StateTree/Source/StateTreeEditorModule/Private/Debugger/StateTreeDebuggerTrack.cpp
--- a/Engine/Plugins/Runtime/StateTree/Source/StateTreeEditorModule/Private/Debugger/StateTreeDebuggerTrack.cpp
+++ b/Engine/Plugins/Runtime/StateTree/Source/StateTreeEditorModule/Private/Debugger/StateTreeDebuggerTrack.cpp
@@ -86,7 +86,7 @@ bool FStateTreeDebuggerInstanceTrack::UpdateInternal()
 			}
 			else
 			{
-				Window.TimeEnd = Debugger->IsActiveInstance(RecordingDuration, InstanceId) ? RecordingDuration : Span.GetWorldTimeEnd();
+				Window.TimeEnd = Debugger->IsActiveInstance(RecordingDuration, InstanceId) ? RecordingDuration : EventCollection.FrameSpans.Last().GetWorldTimeEnd();
 			}
 		}
 
diff --git a/Engine/Plugins/Runtime/StateTree/Source/StateTreeModule/Private/Debugger/StateTreeDebugger.cpp b/Engine/Plugins/Runtime/StateTree/Source/StateTreeModule/Private/Debugger/StateTreeDebugger.cpp
--- a/Engine/Plugins/Runtime/StateTree/Source/StateTreeModule/Private/Debugger/StateTreeDebugger.cpp
+++ b/Engine/Plugins/Runtime/StateTree/Source/StateTreeModule/Private/Debugger/StateTreeDebugger.cpp
@@ -788,15 +788,37 @@ void FStateTreeDebugger::AddEvents(const double StartTime, const double EndTime,
 			TimelineData.EnumerateEvents(StartTime,	EndTime,
 				[this, InstanceId, &FrameProvider, &Frame](const double EventStartTime, const double EventEndTime, uint32 InDepth, const FStateTreeTraceEventVariantType& Event)
 				{
+					bool bValidFrame = true;
+
 					// Fetch frame when not set yet or if events no longer part of the current one
 					if (Frame.Index == INDEX_NONE ||
 						(EventEndTime < Frame.StartTime || Frame.EndTime < EventStartTime))
 					{
-						FrameProvider.GetFrameFromTime(TraceFrameType_Game, EventStartTime, Frame);
+						bValidFrame = FrameProvider.GetFrameFromTime(TraceFrameType_Game, EventStartTime, Frame);
+
+						if (bValidFrame == false)
+						{
+							// Edge case for events from a missing first complete frame.
+							// (i.e. FrameProvider didn't get BeginFrame event but StateTreeEvent were sent in that frame)
+							// Doing this will merge our two first frames of state tree events using the same recording world time
+							// but this should happen only for late start recording.
+							const TraceServices::FFrame* FirstFrame = FrameProvider.GetFrame(TraceFrameType_Game, 0);
+							if (FirstFrame != nullptr && EventEndTime < FirstFrame->StartTime)
+							{
+								Frame = *FirstFrame;
+								bValidFrame = true;
+							}
+						}
+					}
+
+					if (bValidFrame)
+					{
+						const bool bKeepProcessing = ProcessEvent(InstanceId, Frame, Event);
+						return bKeepProcessing ? TraceServices::EEventEnumerate::Continue : TraceServices::EEventEnumerate::Stop;
 					}
 
-					const bool bKeepProcessing = ProcessEvent(InstanceId, Frame, Event);
-					return bKeepProcessing ? TraceServices::EEventEnumerate::Continue : TraceServices::EEventEnumerate::Stop;
+					// Skip events outside of game frames
+					return TraceServices::EEventEnumerate::Continue;
 				});
 		});
 }
diff --git a/Engine/Plugins/Runtime/StateTree/Source/StateTreeModule/Private/Debugger/StateTreeTrace.cpp b/Engine/Plugins/Runtime/StateTree/Source/StateTreeModule/Private/Debugger/StateTreeTrace.cpp
--- a/Engine/Plugins/Runtime/StateTree/Source/StateTreeModule/Private/Debugger/StateTreeTrace.cpp
+++ b/Engine/Plugins/Runtime/StateTree/Source/StateTreeModule/Private/Debugger/StateTreeTrace.cpp
@@ -3,6 +3,7 @@
 #if WITH_STATETREE_DEBUGGER
 
 #include "Debugger/StateTreeTrace.h"
+#include "Algo/Removeif.h"
 #include "Debugger/StateTreeDebugger.h"
 #include "Exporters/Exporter.h"
 #include "ObjectTrace.h"
@@ -101,7 +102,6 @@ FDelegateHandle GOnPIEStartDelegateHandle;
 #endif // WITH_EDITOR
 
 double GRecordingWorldTime = -1;
-double GTracedRecordingWorldTime = -1;
 
 /** Struct to keep track if a given phase was traced or not. */
 struct FPhaseTraceStatusPair
@@ -120,40 +120,151 @@ struct FPhaseStack
 };
 
 /**
- * Stacks to keep track of all received phase events so other events will control when and if a given phase trace will be sent.
- * This is per thread since it is possible to update execution contexts on multiple threads.
+ * Struct to hold data for active states events until we are ready to trace the events (i.e. traces are active and channel is enabled).
  */
-thread_local TArray<FPhaseStack> GPhaseStacks;
+struct FInstanceEventBufferedData
+{
+	FInstanceEventBufferedData() = default;
+	explicit FInstanceEventBufferedData(
+		const double RecordingWorldTime,
+		const UStateTree& StateTree,
+		const FStateTreeInstanceDebugId InstanceId,
+		const FString& InstanceName,
+		const EStateTreeTraceEventType EventType)	
+		: InstanceName(InstanceName)
+		, WeakStateTree(&StateTree)
+		, InstanceId(InstanceId)
+		, LifetimeRecordingWorldTime(RecordingWorldTime)
+		, EventType(EventType)
+	{
+	}
 
-/**
- * Called by TraceBufferedEvents from the OutputXYZ methods to flush pending phase events.
- * Phases popped before TraceStackedPhases gets called will never produce any trace since
- * they will not be required for the analysis.
- */
-void TraceStackedPhases(const FStateTreeInstanceDebugId InstanceId)
+	void Trace() const
+	{
+		if (!ensureMsgf(UE_TRACE_CHANNELEXPR_IS_ENABLED(StateTreeDebugChannel), TEXT("Tracing a buffered data is expected only if channel is enabled.")))
+		{
+			return;
+		}
+
+		if (const UStateTree* StateTree = WeakStateTree.Get())
+		{
+			// Force a world time update since we are tracing an event from the past
+			UE_TRACE_LOG(StateTreeDebugger, WorldTimestampEvent, StateTreeDebugChannel)
+				<< WorldTimestampEvent.WorldTime(LifetimeRecordingWorldTime);
+
+			OutputInstanceLifetimeEvent(InstanceId, StateTree, *InstanceName, EventType);
+
+			if (ActiveStates.Num() > 0)
+			{
+				// Force a world time update since we are tracing an event from the past
+				UE_TRACE_LOG(StateTreeDebugger, WorldTimestampEvent, StateTreeDebugChannel)
+					<< WorldTimestampEvent.WorldTime(ActiveStatesRecordingWorldTime);
+
+				OutputActiveStatesEventTrace(InstanceId, ActiveStates);
+			}
+		}
+	}
+
+	FStateTreeActiveStates ActiveStates;
+	FString InstanceName;
+	TWeakObjectPtr<const UStateTree> WeakStateTree;
+	FStateTreeInstanceDebugId InstanceId;
+	double LifetimeRecordingWorldTime = 0;
+	double ActiveStatesRecordingWorldTime = 0;
+	EStateTreeTraceEventType EventType = EStateTreeTraceEventType::Unset;
+};
+
+/** Struct to keep track of the buffered event data and flush them. */
+struct FBufferedDataList
 {
-	for (FPhaseStack& PhaseStack : GPhaseStacks)
+	double TracedRecordingWorldTime = -1;
+	
+	/**
+	 * Stacks to keep track of all received phase events so other events will control when and if a given phase trace will be sent.
+	 * This is per thread since it is possible to update execution contexts on multiple threads.
+	 */
+	TArray<FPhaseStack> PhaseStacks;
+
+	/** List of lifetime events that will be output if channel gets enabled in the Push - Pop lifetime window of an instance. */
+	TArray<FInstanceEventBufferedData> Events;
+
+	/** Flag use to prevent reentrant calls */
+	bool bFlushing = false;
+
+	void Flush(const FStateTreeInstanceDebugId InstanceId)
 	{
-		if (PhaseStack.InstanceId == InstanceId)
+		if (bFlushing)
+		{
+			return;
+		}
+
+		TGuardValue<bool> GuardReentry(bFlushing, true);
+
+		// Trace instance lifetime events first since they are required for other event types.
+		// It is also associated to an older world time.
+		for (const FInstanceEventBufferedData& InstanceEventData : Events)
 		{
-			for (FPhaseTraceStatusPair& StackEntry : PhaseStack.Stack)
+			InstanceEventData.Trace();
+		}
+		Events.Empty();
+
+		TraceWorldTime();
+		TraceStackedPhases(InstanceId);
+	}
+
+	/**
+	 * Called by TraceBufferedEvents from the OutputXYZ methods to make sure we have the current world time was sent.
+	 */
+	void TraceWorldTime()
+	{
+		if (TracedRecordingWorldTime != GRecordingWorldTime)
+		{
+			TracedRecordingWorldTime = GRecordingWorldTime;
+			UE_TRACE_LOG(StateTreeDebugger, WorldTimestampEvent, StateTreeDebugChannel)
+				<< WorldTimestampEvent.WorldTime(GRecordingWorldTime);
+		}
+	}
+
+	/**
+	 * Called by TraceBufferedEvents from the OutputXYZ methods to flush pending phase events.
+	 * Phases popped before TraceStackedPhases gets called will never produce any trace since
+	 * they will not be required for the analysis.
+	 */
+	void TraceStackedPhases(const FStateTreeInstanceDebugId InstanceId)
+	{
+		for (FPhaseStack& PhaseStack : PhaseStacks)
+		{
+			if (PhaseStack.InstanceId == InstanceId)
 			{
-				if (StackEntry.bTraced == false)
+				for (FPhaseTraceStatusPair& StackEntry : PhaseStack.Stack)
 				{
-					UE_TRACE_LOG(StateTreeDebugger, PhaseEvent, StateTreeDebugChannel)
-						<< PhaseEvent.Cycle(FPlatformTime::Cycles64())
-						<< PhaseEvent.InstanceId(InstanceId.Id)
-						<< PhaseEvent.InstanceSerial(InstanceId.SerialNumber)
-						<< PhaseEvent.Phase(static_cast<std::underlying_type_t<EStateTreeUpdatePhase>>(StackEntry.Phase))
-						<< PhaseEvent.EventType(static_cast<std::underlying_type_t<EStateTreeTraceEventType>>(EStateTreeTraceEventType::Push));
+					// Trace push phase event and marked as traced only if not already traced and our channel is enabled.
+					// We need the pop phase event to be sent only in this case to enforce complementary events in case of
+					// late recording (e.g. recording started, or channel enabled, after simulation is running and instances are ticked)
+					if (StackEntry.bTraced == false && UE_TRACE_CHANNELEXPR_IS_ENABLED(StateTreeDebugChannel))
+					{
+						UE_TRACE_LOG(StateTreeDebugger, PhaseEvent, StateTreeDebugChannel)
+							<< PhaseEvent.Cycle(FPlatformTime::Cycles64())
+							<< PhaseEvent.InstanceId(InstanceId.Id)
+							<< PhaseEvent.InstanceSerial(InstanceId.SerialNumber)
+							<< PhaseEvent.Phase(static_cast<std::underlying_type_t<EStateTreeUpdatePhase>>(StackEntry.Phase))
+							<< PhaseEvent.EventType(static_cast<std::underlying_type_t<EStateTreeTraceEventType>>(EStateTreeTraceEventType::Push));
 			
-					StackEntry.bTraced = true;
+						StackEntry.bTraced = true;
+					}
 				}
+				break;
 			}
-			break;
 		}
 	}
-}
+};
+
+/**
+ * Buffered events (e.g. lifetime, active state, scoped phase) in case channel is not active yet or phase are empty and don't need to be traced.
+ * This is per thread since it is possible to update execution contexts on multiple threads.
+ * @note The current implementation of the lifetime events doesn't properly support same instance getting ticked in different threads.
+ */
+thread_local FBufferedDataList GBufferedEvents;
 
 /**
  * Pushed or pops an entry on the Phase stack for a given Instance.
@@ -161,21 +272,22 @@ void TraceStackedPhases(const FStateTreeInstanceDebugId InstanceId)
  */
 void ProcessPhaseScopeEvent(const FStateTreeInstanceDebugId InstanceId, const EStateTreeUpdatePhase Phase, const EStateTreeTraceEventType EventType)
 {
-	int32 ExistingStackIndex = GPhaseStacks.IndexOfByPredicate([InstanceId](const FPhaseStack& PhaseStack){ return PhaseStack.InstanceId == InstanceId; });
+	TArray<FPhaseStack>& PhaseStacks = GBufferedEvents.PhaseStacks;
+	int32 ExistingStackIndex = PhaseStacks.IndexOfByPredicate([InstanceId](const FPhaseStack& PhaseStack){ return PhaseStack.InstanceId == InstanceId; });
 
 	if (EventType == EStateTreeTraceEventType::Push)
 	{
 		if (ExistingStackIndex == INDEX_NONE)
 		{
-			ExistingStackIndex = GPhaseStacks.AddDefaulted();
+			ExistingStackIndex = PhaseStacks.AddDefaulted();
 		}
-		FPhaseStack& PhaseStack = GPhaseStacks[ExistingStackIndex];
+		FPhaseStack& PhaseStack = PhaseStacks[ExistingStackIndex];
 		PhaseStack.InstanceId = InstanceId;
 		PhaseStack.Stack.Push(FPhaseTraceStatusPair(Phase));
 	}
 	else if (ensureMsgf(ExistingStackIndex != INDEX_NONE, TEXT("Not expected to pop phases for an instance that never pushed a phase.")))
 	{
-		FPhaseStack& PhaseStack = GPhaseStacks[ExistingStackIndex];
+		FPhaseStack& PhaseStack = PhaseStacks[ExistingStackIndex];
 
 		if (ensureMsgf(PhaseStack.Stack.IsEmpty() == false, TEXT("Not expected to pop phases that never got pushed.")) &&
 			ensureMsgf(PhaseStack.InstanceId == InstanceId, TEXT("Not expected to pop phases for an instance that is not the one currently assigned to the stack.")))
@@ -186,7 +298,7 @@ void ProcessPhaseScopeEvent(const FStateTreeInstanceDebugId InstanceId, const ES
 			// Clear associated InstanceId when removing last entry from the stack.
 			if (PhaseStack.Stack.IsEmpty())
 			{
-				GPhaseStacks.RemoveAt(ExistingStackIndex);
+				PhaseStacks.RemoveAt(ExistingStackIndex);
 			}
 
 			// Phase was previously traced (i.e. other events were traced in that scope so we need to trace the closing (i.e. Pop) event.
@@ -203,26 +315,13 @@ void ProcessPhaseScopeEvent(const FStateTreeInstanceDebugId InstanceId, const ES
 	}
 }
 
-/**
- * Called by TraceBufferedEvents from the OutputXYZ methods to make sure we have the current world time was sent.
- */
-void TraceTimeMapping()
-{
-	if (GTracedRecordingWorldTime != GRecordingWorldTime)
-	{
-		GTracedRecordingWorldTime = GRecordingWorldTime;
-		UE_TRACE_LOG(StateTreeDebugger, WorldTimestampEvent, StateTreeDebugChannel)
-			<< WorldTimestampEvent.WorldTime(GRecordingWorldTime);
-	}
-}
 
 /**
  * Called by the OutputXYZ methods to flush pending events (e.g. Push or WorldTime).
  */
 void TraceBufferedEvents(const FStateTreeInstanceDebugId InstanceId)
 {
-	TraceTimeMapping();
-	TraceStackedPhases(InstanceId);
+	GBufferedEvents.Flush(InstanceId);
 }
 
 void SerializeDataViewToArchive(FBufferArchive Ar, const FStateTreeDataView DataView)
@@ -257,7 +356,7 @@ void SerializeDataViewToArchive(FBufferArchive Ar, const FStateTreeDataView Data
 void RegisterGlobalDelegates()
 {
 #if WITH_EDITOR
-	GOnPIEStartDelegateHandle = FEditorDelegates::BeginPIE.AddLambda([&LastRecordingWorldTime=GTracedRecordingWorldTime](const bool bIsSimulating)
+	GOnPIEStartDelegateHandle = FEditorDelegates::BeginPIE.AddLambda([&LastRecordingWorldTime=GBufferedEvents.TracedRecordingWorldTime](const bool bIsSimulating)
 		{
 			LastRecordingWorldTime = -1;
 		});
@@ -292,20 +391,44 @@ void OutputInstanceLifetimeEvent(
 	{
 		return;
 	}
-	const FString ObjectName = StateTree->GetName();
-	const FString ObjectPackageName = StateTree->GetPackage()->GetName();
 
-	TraceBufferedEvents(InstanceId);
-
-	UE_TRACE_LOG(StateTreeDebugger, InstanceEvent, StateTreeDebugChannel)
-		<< InstanceEvent.Cycle(FPlatformTime::Cycles64())
-		<< InstanceEvent.TreeName(*ObjectName, ObjectName.Len())
-		<< InstanceEvent.TreePath(*ObjectPackageName, ObjectPackageName.Len())
-		<< InstanceEvent.CompiledDataHash(StateTree->LastCompiledEditorDataHash)
-		<< InstanceEvent.InstanceId(InstanceId.Id)
-		<< InstanceEvent.InstanceSerial(InstanceId.SerialNumber)
-		<< InstanceEvent.InstanceName(InstanceName)
-		<< InstanceEvent.EventType(static_cast<std::underlying_type_t<EStateTreeTraceEventType>>(EventType));
+	if (UE_TRACE_CHANNELEXPR_IS_ENABLED(StateTreeDebugChannel))
+	{
+		TraceBufferedEvents(InstanceId);
+
+		const FString TreeName = StateTree->GetName();
+		const FString TreePath = StateTree->GetPackage()->GetName();
+
+		UE_TRACE_LOG(StateTreeDebugger, InstanceEvent, StateTreeDebugChannel)
+			<< InstanceEvent.Cycle(FPlatformTime::Cycles64())
+			<< InstanceEvent.TreeName(*TreeName, TreeName.Len())
+			<< InstanceEvent.TreePath(*TreePath, TreePath.Len())
+			<< InstanceEvent.CompiledDataHash(StateTree->LastCompiledEditorDataHash)
+			<< InstanceEvent.InstanceId(InstanceId.Id)
+			<< InstanceEvent.InstanceSerial(InstanceId.SerialNumber)
+			<< InstanceEvent.InstanceName(InstanceName)
+			<< InstanceEvent.EventType(static_cast<std::underlying_type_t<EStateTreeTraceEventType>>(EventType));
+	}
+	else
+	{
+		if (EventType == EStateTreeTraceEventType::Push)
+		{
+			GBufferedEvents.Events.Emplace(GRecordingWorldTime, *StateTree, InstanceId, InstanceName, EventType);
+		}
+		else if (EventType == EStateTreeTraceEventType::Pop)
+		{
+			// Remove matching instance events since if it was not sent then no other events were sent between, hence not needed in the trace.
+			GBufferedEvents.Events.SetNum(Algo::StableRemoveIf(GBufferedEvents.Events, 
+				[InstanceId](const FInstanceEventBufferedData& BufferedData)
+				{
+					return BufferedData.InstanceId == InstanceId;
+				}));
+		}
+		else
+		{
+			ensureMsgf(false, TEXT("Unexpected EventType '%s' for instance lifetime event."), *UEnum::GetDisplayValueAsText(EventType).ToString());
+		}
+	}
 }
 
 void OutputLogEventTrace(
@@ -411,19 +534,37 @@ void OutputActiveStatesEventTrace(
 	const FStateTreeActiveStates& ActiveStates
 	)
 {
-	TArray<uint16, TInlineAllocator<FStateTreeActiveStates::MaxStates>> StatesIndices;
-	for (int32 i = 0; i < ActiveStates.Num(); i++)
+	if (UE_TRACE_CHANNELEXPR_IS_ENABLED(StateTreeDebugChannel))
 	{
-		StatesIndices.Add(ActiveStates[i].Index);
-	}
+		TraceBufferedEvents(InstanceId);
 
-	TraceBufferedEvents(InstanceId);
+		TArray<uint16, TInlineAllocator<FStateTreeActiveStates::MaxStates>> StatesIndices;
+		for (int32 i = 0; i < ActiveStates.Num(); i++)
+		{
+			StatesIndices.Add(ActiveStates[i].Index);
+		}
 
-	UE_TRACE_LOG(StateTreeDebugger, ActiveStatesEvent, StateTreeDebugChannel)
-		<< ActiveStatesEvent.Cycle(FPlatformTime::Cycles64())
-		<< ActiveStatesEvent.InstanceId(InstanceId.Id)
-		<< ActiveStatesEvent.InstanceSerial(InstanceId.SerialNumber)
-		<< ActiveStatesEvent.ActiveStates(StatesIndices.GetData(), StatesIndices.Num());
+		UE_TRACE_LOG(StateTreeDebugger, ActiveStatesEvent, StateTreeDebugChannel)
+			<< ActiveStatesEvent.Cycle(FPlatformTime::Cycles64())
+			<< ActiveStatesEvent.InstanceId(InstanceId.Id)
+			<< ActiveStatesEvent.InstanceSerial(InstanceId.SerialNumber)
+			<< ActiveStatesEvent.ActiveStates(StatesIndices.GetData(), StatesIndices.Num());
+	}
+	else
+	{
+		FInstanceEventBufferedData* ExisingBufferedData = GBufferedEvents.Events.FindByPredicate([InstanceId](const FInstanceEventBufferedData& BufferedData)
+			{
+				return BufferedData.InstanceId == InstanceId;
+			});
+
+		// We keep only the most recent active states since this is all we need to know in which state was the instance
+		// when we start receiving the events once the channel is enabled. 
+		if (ExisingBufferedData != nullptr)
+		{
+			ExisingBufferedData->ActiveStates = ActiveStates;
+			ExisingBufferedData->ActiveStatesRecordingWorldTime = GRecordingWorldTime;
+		}
+	}
 }
 
 } // UE::StateTreeTrace
diff --git a/Engine/Plugins/Runtime/StateTree/Source/StateTreeModule/Public/Debugger/StateTreeTrace.h b/Engine/Plugins/Runtime/StateTree/Source/StateTreeModule/Public/Debugger/StateTreeTrace.h
--- a/Engine/Plugins/Runtime/StateTree/Source/StateTreeModule/Public/Debugger/StateTreeTrace.h
+++ b/Engine/Plugins/Runtime/StateTree/Source/StateTreeModule/Public/Debugger/StateTreeTrace.h
@@ -34,10 +34,7 @@ namespace UE::StateTreeTrace
 }
 
 #define TRACE_STATETREE_INSTANCE_EVENT(InstanceID, StateTree, InstanceName, EventType) \
-	if (UE_TRACE_CHANNELEXPR_IS_ENABLED(StateTreeDebugChannel)) \
-	{ \
-		UE::StateTreeTrace::OutputInstanceLifetimeEvent(InstanceID, StateTree, InstanceName, EventType); \
-	}
+	UE::StateTreeTrace::OutputInstanceLifetimeEvent(InstanceID, StateTree, InstanceName, EventType);
 
 #define TRACE_STATETREE_PHASE_EVENT(InstanceID, Phase, EventType) \
 	UE::StateTreeTrace::ProcessPhaseScopeEvent(InstanceID, Phase, EventType); \
@@ -73,10 +70,7 @@ namespace UE::StateTreeTrace
 	}
 
 #define TRACE_STATETREE_ACTIVE_STATES_EVENT(InstanceId, ActivateStates) \
-	if (UE_TRACE_CHANNELEXPR_IS_ENABLED(StateTreeDebugChannel)) \
-	{ \
-		UE::StateTreeTrace::OutputActiveStatesEventTrace(InstanceId, ActivateStates); \
-	}
+		UE::StateTreeTrace::OutputActiveStatesEventTrace(InstanceId, ActivateStates);
 
 #else //STATETREE_DEBUG_TRACE_ENABLED
 
