[StateTreeDebugger] updated some traces call sites in execution context
- State OnEnter no longer traced if a state is sustained (similar to OnExit)
- Added Task OnEnter for global and state tasks
- Added OnTicking and OnTicked for global tasks (same as state tasks)
#rnx

diff --git a/Engine/Plugins/Runtime/StateTree/Source/StateTreeModule/Private/StateTreeExecutionContext.cpp b/Engine/Plugins/Runtime/StateTree/Source/StateTreeModule/Private/StateTreeExecutionContext.cpp
--- a/Engine/Plugins/Runtime/StateTree/Source/StateTreeModule/Private/StateTreeExecutionContext.cpp
+++ b/Engine/Plugins/Runtime/StateTree/Source/StateTreeModule/Private/StateTreeExecutionContext.cpp
@@ -759,8 +759,6 @@ EStateTreeRunStatus FStateTreeExecutionContext::EnterState(const FStateTreeTrans
 			break;
 		}
 		
-		STATETREE_TRACE_STATE_EVENT(CurrentHandle, EStateTreeTraceEventType::OnEnter);
-
 		if (State.Type == EStateTreeStateType::Linked)
 		{
 			UpdateLinkedStateParameters(State, InstanceStructIndex);
@@ -778,8 +776,14 @@ EStateTreeRunStatus FStateTreeExecutionContext::EnterState(const FStateTreeTrans
 		CurrentTransition.CurrentState = CurrentHandle;
 		CurrentTransition.ChangeType = bWasActive ? EStateTreeStateChangeType::Sustained : EStateTreeStateChangeType::Changed;
 
-		STATETREE_CLOG(bIsEnteringState, Log, TEXT("%*sState '%s' %s"), Index*UE::StateTree::DebugIndentSize, TEXT(""), *DebugGetStatePath(Transition.NextActiveStates, Index), *UEnum::GetDisplayValueAsText(CurrentTransition.ChangeType).ToString());
-		
+		if (bIsEnteringState)
+		{
+			STATETREE_TRACE_STATE_EVENT(CurrentHandle, EStateTreeTraceEventType::OnEnter);
+			STATETREE_LOG(Log, TEXT("%*sState '%s' %s"), Index*UE::StateTree::DebugIndentSize, TEXT(""),
+				*DebugGetStatePath(Transition.NextActiveStates, Index),
+				*UEnum::GetDisplayValueAsText(CurrentTransition.ChangeType).ToString());
+		}
+
 		// Activate tasks on current state.
 		for (int32 TaskIndex = State.TasksBegin; TaskIndex < (State.TasksBegin + State.TasksNum); TaskIndex++)
 		{
@@ -814,9 +818,10 @@ EStateTreeRunStatus FStateTreeExecutionContext::EnterState(const FStateTreeTrans
 					Status = Task.EnterState(*this, CurrentTransition);
 				}
 
+				STATETREE_TRACE_TASK_EVENT(TaskIndex, DataViews[Task.DataViewIndex.Get()], EStateTreeTraceEventType::OnEnter, Status);
+
 				if (Status != EStateTreeRunStatus::Running)
 				{
-					STATETREE_TRACE_TASK_EVENT(TaskIndex, DataViews[Task.DataViewIndex.Get()], EStateTreeTraceEventType::OnTaskCompleted, Status);
 					// Store the first state that completed, will be used to decide where to trigger transitions.
 					if (!Exec.CompletedStateHandle.IsValid())
 					{
@@ -1072,13 +1077,16 @@ EStateTreeRunStatus FStateTreeExecutionContext::TickEvaluatorsAndGlobalTasks(con
 				continue;
 			}
 
+			const FStateTreeDataView TaskDataView = DataViews[Task.DataViewIndex.Get()];
+
 			// Copy bound properties.
 			// Only copy properties when the task is actually ticked, and copy properties at tick is requested.
 			if (Task.BindingsBatch.IsValid() && Task.bShouldCopyBoundPropertiesOnTick)
 			{
-				StateTree.PropertyBindings.CopyTo(DataViews, Task.BindingsBatch, DataViews[Task.DataViewIndex.Get()]);
+				StateTree.PropertyBindings.CopyTo(DataViews, Task.BindingsBatch, TaskDataView);
 			}
 
+			STATETREE_TRACE_TASK_EVENT(TaskIndex, TaskDataView, EStateTreeTraceEventType::OnTickingTask, EStateTreeRunStatus::Running);
 			EStateTreeRunStatus TaskResult = EStateTreeRunStatus::Unset;
 			{
 				QUICK_SCOPE_CYCLE_COUNTER(StateTree_Task_Tick);
@@ -1087,10 +1095,13 @@ EStateTreeRunStatus FStateTreeExecutionContext::TickEvaluatorsAndGlobalTasks(con
 				TaskResult = Task.Tick(*this, DeltaTime);
 			}
 
+			STATETREE_TRACE_TASK_EVENT(TaskIndex, TaskDataView,
+				TaskResult != EStateTreeRunStatus::Running ? EStateTreeTraceEventType::OnTaskCompleted : EStateTreeTraceEventType::OnTaskTicked,
+				TaskResult);
+
 			// If a global task succeeds or fails, it will stop the whole tree.
 			if (TaskResult != EStateTreeRunStatus::Running)
 			{
-				STATETREE_TRACE_TASK_EVENT(TaskIndex, DataViews[Task.DataViewIndex.Get()], EStateTreeTraceEventType::OnTaskCompleted, TaskResult);
 				Result = TaskResult;
 			}
 				
@@ -1160,7 +1171,10 @@ EStateTreeRunStatus FStateTreeExecutionContext::StartEvaluatorsAndGlobalTasks(FS
 		STATETREE_LOG(Verbose, TEXT("  Start: '%s'"), *Task.Name.ToString());
 		{
 			QUICK_SCOPE_CYCLE_COUNTER(StateTree_Task_TreeStart);
-			EStateTreeRunStatus TaskStatus = Task.EnterState(*this, Transition); 
+			const EStateTreeRunStatus TaskStatus = Task.EnterState(*this, Transition); 
+
+			STATETREE_TRACE_TASK_EVENT(TaskIndex, DataViews[Task.DataViewIndex.Get()], EStateTreeTraceEventType::OnEnter, TaskStatus);
+
 			if (TaskStatus != EStateTreeRunStatus::Running)
 			{
 				OutLastInitializedTaskIndex = FStateTreeIndex16(TaskIndex);
