State Tree: Fix handling of subtree completion transitions.

diff --git a/Engine/Plugins/Runtime/StateTree/Source/StateTreeModule/Private/StateTreeExecutionContext.cpp b/Engine/Plugins/Runtime/StateTree/Source/StateTreeModule/Private/StateTreeExecutionContext.cpp
--- a/Engine/Plugins/Runtime/StateTree/Source/StateTreeModule/Private/StateTreeExecutionContext.cpp
+++ b/Engine/Plugins/Runtime/StateTree/Source/StateTreeModule/Private/StateTreeExecutionContext.cpp
@@ -1786,10 +1786,16 @@ bool FStateTreeExecutionContext::TriggerTransitions()
 			const EStateTreeRunStatus RunStatus = NextTransition.TargetState.ToCompletionStatus(); 
 			STATETREE_LOG(Verbose, TEXT("Completed subtree '%s' from state '%s' (%s): %s"),
 				*GetSafeStateName(ParentLinkedState), *GetSafeStateName(Exec.ActiveStates.Last()), *GetSafeStateName(NextTransition.SourceState), *UEnum::GetDisplayValueAsText(RunStatus).ToString());
+
+			// Set the parent linked state as last completed state, and update tick status to the status from the transition.
 			Exec.CompletedStateHandle = ParentLinkedState;
 			Exec.LastTickStatus = RunStatus;
 
-			// Not a transition
+			// Clear the transition and return that no transition took place.
+			// Since the LastTickStatus != running, the transition loop will try another transition
+			// now starting from the linked parent state. If we run out of retires in the selection loop (e.g. very deep hierarchy)
+			// we will continue on next tick.
+			NextTransition.Reset();
 			return false;
 		}
 	}
@@ -1825,6 +1831,10 @@ FStateTreeStateHandle FStateTreeExecutionContext::GetParentLinkedStateHandle(con
 		StateIndex--;
 	}
 
+	// The function result is used to iteratively traverse to the root-most parent linked state.
+	// Skip the start state, as we want to always find a parent state to the start state, or else the iteration will hit infinite loop. 
+	StateIndex--;
+
 	// Find parent linked state.
 	while (StateIndex >= 0)
 	{
diff --git a/Engine/Plugins/Runtime/StateTree/Source/StateTreeTestSuite/Private/StateTreeTest.cpp b/Engine/Plugins/Runtime/StateTree/Source/StateTreeTestSuite/Private/StateTreeTest.cpp
--- a/Engine/Plugins/Runtime/StateTree/Source/StateTreeTestSuite/Private/StateTreeTest.cpp
+++ b/Engine/Plugins/Runtime/StateTree/Source/StateTreeTestSuite/Private/StateTreeTest.cpp
@@ -398,7 +398,7 @@ struct FStateTreeTest_SubTree : FAITestBase
 		Status = Exec.Start();
 
 		AITEST_TRUE("StateTree Active States should be in Root/State1/State3/State3A", Exec.ExpectInActiveStates(Root.Name, State1.Name, State3.Name, State3A.Name));
-		AITEST_TRUE("StateTree Task2 should enter state", !Exec.Expect(Task2.GetName(), EnterStateStr));
+		AITEST_FALSE("StateTree Task2 should not enter state", Exec.Expect(Task2.GetName(), EnterStateStr));
 		AITEST_TRUE("StateTree Task3A should enter state", Exec.Expect(Task3A.GetName(), EnterStateStr));
 		AITEST_TRUE("StateTree should be running", Status == EStateTreeRunStatus::Running);
 		Exec.LogClear();
@@ -427,6 +427,78 @@ struct FStateTreeTest_SubTree : FAITestBase
 };
 IMPLEMENT_AI_INSTANT_TEST(FStateTreeTest_SubTree, "System.StateTree.SubTree");
 
+struct FStateTreeTest_SubTree_CascadedSucceeded : FAITestBase
+{
+	virtual bool InstantTest() override
+	{
+		UStateTree& StateTree = UE::StateTree::Tests::NewStateTree(&GetWorld());
+		UStateTreeEditorData& EditorData = *Cast<UStateTreeEditorData>(StateTree.EditorData);
+
+		//	- Root [TaskA]
+		//		- LinkedState>SubTreeState -> (F)Failed
+		//		- SubTreeState [TaskB]
+		//			- SubLinkedState>SubSubTreeState -> (S)Failed
+		//		- SubSubTreeState
+		//			- SubSubLeaf [TaskC] -> (S)Succeeded
+		
+		UStateTreeState& Root = EditorData.AddSubTree(FName(TEXT("Root")));
+		UStateTreeState& LinkedState = Root.AddChildState(FName(TEXT("Linked")), EStateTreeStateType::Linked);
+		
+		UStateTreeState& SubTreeState = Root.AddChildState(FName(TEXT("SubTreeState")), EStateTreeStateType::Subtree);
+		UStateTreeState& SubLinkedState = SubTreeState.AddChildState(FName(TEXT("SubLinkedState")), EStateTreeStateType::Linked);
+		
+		UStateTreeState& SubSubTreeState = Root.AddChildState(FName(TEXT("SubSubTreeState")), EStateTreeStateType::Subtree);
+		UStateTreeState& SubSubLeaf = SubSubTreeState.AddChildState(FName(TEXT("SubSubLeaf")));
+
+		LinkedState.LinkedSubtree = SubTreeState.GetLinkToState();
+		SubLinkedState.LinkedSubtree = SubSubTreeState.GetLinkToState();
+
+		LinkedState.AddTransition(EStateTreeTransitionTrigger::OnStateFailed, EStateTreeTransitionType::Failed);
+		SubLinkedState.AddTransition(EStateTreeTransitionTrigger::OnStateSucceeded, EStateTreeTransitionType::Failed);
+		SubSubLeaf.AddTransition(EStateTreeTransitionTrigger::OnStateSucceeded, EStateTreeTransitionType::Succeeded);
+
+		TStateTreeEditorNode<FTestTask_Stand>& TaskA = Root.AddTask<FTestTask_Stand>(FName(TEXT("TaskA")));
+		TStateTreeEditorNode<FTestTask_Stand>& TaskB = SubTreeState.AddTask<FTestTask_Stand>(FName(TEXT("TaskB")));
+		TStateTreeEditorNode<FTestTask_Stand>& TaskC = SubSubLeaf.AddTask<FTestTask_Stand>(FName(TEXT("TaskC")));
+
+		TaskA.GetNode().TicksToCompletion = 2;
+		TaskB.GetNode().TicksToCompletion = 2;
+		TaskC.GetNode().TicksToCompletion = 1; // The deepest task completes first.
+		
+		FStateTreeCompilerLog Log;
+		FStateTreeCompiler Compiler(Log);
+		const bool bResult = Compiler.Compile(StateTree);
+		AITEST_TRUE("StateTree should get compiled", bResult);
+
+		EStateTreeRunStatus Status = EStateTreeRunStatus::Unset;
+		FStateTreeInstanceData InstanceData;
+		FTestStateTreeExecutionContext Exec(StateTree, StateTree, InstanceData);
+		const bool bInitSucceeded = Exec.IsValid();
+		AITEST_TRUE("StateTree should init", bInitSucceeded);
+
+		const FString TickStr(TEXT("Tick"));
+		const FString EnterStateStr(TEXT("EnterState"));
+		const FString ExitStateStr(TEXT("ExitState"));
+		const FString StateCompletedStr(TEXT("StateCompleted"));
+
+		// Start and enter state
+		Status = Exec.Start();
+		AITEST_TRUE("StateTree Active States should be in Root/Linked/SubTreeState", Exec.ExpectInActiveStates(Root.Name, LinkedState.Name, SubTreeState.Name, SubLinkedState.Name, SubSubTreeState.Name, SubSubLeaf.Name));
+		AITEST_TRUE("TaskA,B,C should enter state", Exec.Expect(TaskA.GetName(), EnterStateStr).Then(TaskB.GetName(), EnterStateStr).Then(TaskC.GetName(), EnterStateStr));
+		AITEST_TRUE("StateTree should be running", Status == EStateTreeRunStatus::Running);
+		Exec.LogClear();
+
+		// Subtrees completes, and it completes the whole tree too.
+		// There's no good way to observe this externally. We switch the return along the way to make sure the transition does not happen directly from the leaf to failed.
+		Status = Exec.Tick(0.1f);
+		AITEST_TRUE("StateTree should be Failed", Status == EStateTreeRunStatus::Failed);
+		Exec.LogClear();
+		
+		return true;
+	}
+};
+IMPLEMENT_AI_INSTANT_TEST(FStateTreeTest_SubTree_CascadedSucceeded, "System.StateTree.SubTree.CascadedSucceeded");
+
 
 struct FStateTreeTest_SharedInstanceData : FAITestBase
 {
