StateTree: Fix TriggerTransitions() call on tasks which have not been ticked this frame.

diff --git a/Engine/Plugins/Runtime/StateTree/Source/StateTreeModule/Private/StateTree.cpp b/Engine/Plugins/Runtime/StateTree/Source/StateTreeModule/Private/StateTree.cpp
--- a/Engine/Plugins/Runtime/StateTree/Source/StateTreeModule/Private/StateTree.cpp
+++ b/Engine/Plugins/Runtime/StateTree/Source/StateTreeModule/Private/StateTree.cpp
@@ -433,7 +433,7 @@ bool UStateTree::Link()
 			SourceStructs[ParametersDataViewIndex.Get()].Struct = Parameters.GetPropertyBagStruct();
 		}
 		
-		for (const FCompactStateTreeState& State : States)
+		for (FCompactStateTreeState& State : States)
 		{
 			if (State.Type == EStateTreeStateType::Subtree)
 			{
@@ -479,6 +479,39 @@ bool UStateTree::Link()
 					Batch.TargetStruct.Struct = Params.Parameters.GetPropertyBagStruct();
 				}
 			}
+
+			// Update instance data num for each state.
+			State.TaskInstanceStructNum = 0;
+			State.TaskInstanceObjectNum = 0;
+
+			if (State.Type == EStateTreeStateType::Linked)
+			{
+				// Linked state parameters.
+				State.TaskInstanceStructNum++;
+			}
+			
+			for (int32 TaskIndex = (State.TasksBegin + State.TasksNum) - 1; TaskIndex >= State.TasksBegin; TaskIndex--)
+			{
+				const FStateTreeTaskBase& Task = Nodes[TaskIndex].Get<const FStateTreeTaskBase>();
+				if (Task.bInstanceIsObject)
+				{
+					if (State.TaskInstanceObjectNum == MAX_uint8)
+					{
+						UE_LOG(LogStateTree, Error, TEXT("%s: State '%s' has more than %d struct based task instances."), *GetFullName(), *State.Name.ToString(), (int32)MAX_uint8);
+						return false;
+					}
+					State.TaskInstanceObjectNum++;
+				}
+				else
+				{
+					if (State.TaskInstanceStructNum == MAX_uint8)
+					{
+						UE_LOG(LogStateTree, Error, TEXT("%s: State '%s' has more than %d object based task instances."), *GetFullName(), *State.Name.ToString(), (int32)MAX_uint8);
+						return false;
+					}
+					State.TaskInstanceStructNum++;
+				}
+			}
 		}
 		
 		// Resolves property paths used by bindings a store property pointers
diff --git a/Engine/Plugins/Runtime/StateTree/Source/StateTreeModule/Private/StateTreeExecutionContext.cpp b/Engine/Plugins/Runtime/StateTree/Source/StateTreeModule/Private/StateTreeExecutionContext.cpp
--- a/Engine/Plugins/Runtime/StateTree/Source/StateTreeModule/Private/StateTreeExecutionContext.cpp
+++ b/Engine/Plugins/Runtime/StateTree/Source/StateTreeModule/Private/StateTreeExecutionContext.cpp
@@ -1574,121 +1574,172 @@ bool FStateTreeExecutionContext::TriggerTransitions()
 	//
 	// Check tick, event, and task based transitions first.
 	//
-	for (int32 StateIndex = Exec.ActiveStates.Num() - 1; StateIndex >= 0; StateIndex--)
+	if (Exec.ActiveStates.Num() > 0)
 	{
-		const FStateTreeStateHandle StateHandle = Exec.ActiveStates[StateIndex];
-		const FCompactStateTreeState& State = StateTree.States[StateHandle.Index];
+		// Setup data views for the tasks that will get called.
+		// It is possible that not all tasks views are set up at this stage (e.g. failed tick, pending transition handling).
+		check(Exec.FirstTaskStructIndex.IsValid() && Exec.FirstTaskObjectIndex.IsValid());
+		int32 InstanceStructIndex = Exec.FirstTaskStructIndex.Get();
+		int32 InstanceObjectIndex = Exec.FirstTaskObjectIndex.Get();
 
-		ensureMsgf(State.bEnabled, TEXT("Should never try to transition to a disabled state. This indicates an unhandled case."));
-
-		FCurrentlyProcessedStateScope StateScope(*this, StateHandle);
-
-		if (State.bHasTransitionTasks)
+		for (int32 Index = 0; Index < Exec.ActiveStates.Num(); Index++)
 		{
-			STATETREE_CLOG(State.TasksNum > 0, VeryVerbose, TEXT("%*sTrigger task transitions in state '%s'"), StateIndex*UE::StateTree::DebugIndentSize, TEXT(""), *DebugGetStatePath(Exec.ActiveStates, StateIndex));
+			const FStateTreeStateHandle CurrentHandle = Exec.ActiveStates[Index];
+			const FCompactStateTreeState& State = StateTree.States[CurrentHandle.Index];
 
-			for (int32 TaskIndex = (State.TasksBegin + State.TasksNum) - 1; TaskIndex >= State.TasksBegin; TaskIndex--)
+			if (State.bHasTransitionTasks)
 			{
-				const FStateTreeTaskBase& Task = StateTree.Nodes[TaskIndex].Get<const FStateTreeTaskBase>();
-
-				// Ignore disabled task
-				if (Task.bTaskEnabled == false)
+				// Update index to skip over the linked state params
+				if (State.Type == EStateTreeStateType::Linked)
 				{
-					STATETREE_LOG(VeryVerbose, TEXT("%*sSkipped 'TriggerTransitions' for disabled Task: '%s'"), UE::StateTree::DebugIndentSize, TEXT(""), *Task.Name.ToString());
-					continue;
+					InstanceStructIndex++;
 				}
 
-				if (Task.bShouldAffectTransitions)
+				for (int32 TaskIndex = State.TasksBegin; TaskIndex < (State.TasksBegin + State.TasksNum); TaskIndex++)
 				{
-					STATETREE_LOG(VeryVerbose, TEXT("%*sTriggerTransitions: '%s'"), UE::StateTree::DebugIndentSize, TEXT(""), *Task.Name.ToString());
-					check(DataViews[Task.DataViewIndex.Get()].IsValid());
-					Task.TriggerTransitions(*this);
+					const FStateTreeTaskBase& Task = StateTree.Nodes[TaskIndex].Get<const FStateTreeTaskBase>();
+					if (Task.bShouldAffectTransitions)
+					{
+						SetNodeDataView(Task, InstanceStructIndex, InstanceObjectIndex);
+					}
+					else
+					{
+						if (Task.bInstanceIsObject)
+						{
+							InstanceObjectIndex++;
+						}
+						else
+						{
+							InstanceStructIndex++;
+						}
+					}
 				}
 			}
+			else
+			{
+				// Skip over all instances in the state.
+				InstanceStructIndex += (int32)State.TaskInstanceStructNum;
+				InstanceObjectIndex += (int32)State.TaskInstanceObjectNum;
+			}
 		}
 		
-		for (uint8 i = 0; i < State.TransitionsNum; i++)
+		for (int32 StateIndex = Exec.ActiveStates.Num() - 1; StateIndex >= 0; StateIndex--)
 		{
-			// All transition conditions must pass
-			const int16 TransitionIndex = State.TransitionsBegin + i;
-			const FCompactStateTransition& Transition = StateTree.Transitions[TransitionIndex];
+			const FStateTreeStateHandle StateHandle = Exec.ActiveStates[StateIndex];
+			const FCompactStateTreeState& State = StateTree.States[StateHandle.Index];
 
-			// No need to test the transition if same or higher priority transition has already been processed.
-			if (Transition.Priority <= NextTransition.Priority)
-			{
-				continue;
-			}
+			ensureMsgf(State.bEnabled, TEXT("Should never try to transition to a disabled state. This indicates an unhandled case."));
 
-			// Skip completion transitions
-			if (EnumHasAnyFlags(Transition.Trigger, EStateTreeTransitionTrigger::OnStateCompleted))
-			{
-				continue;
-			}
+			FCurrentlyProcessedStateScope StateScope(*this, StateHandle);
 
-			// If a delayed transition has passed the delay, and remove it from the queue, and try trigger it.
-			FStateTreeTransitionDelayedState* DelayedState = nullptr;
-			if (Transition.HasDelay())
+			if (State.bHasTransitionTasks)
 			{
-				DelayedState = Exec.FindDelayedTransition(FStateTreeIndex16(TransitionIndex));
-				if (DelayedState != nullptr && DelayedState->TimeLeft <= 0.0f)
+				STATETREE_CLOG(State.TasksNum > 0, VeryVerbose, TEXT("%*sTrigger task transitions in state '%s'"), StateIndex*UE::StateTree::DebugIndentSize, TEXT(""), *DebugGetStatePath(Exec.ActiveStates, StateIndex));
+
+				for (int32 TaskIndex = (State.TasksBegin + State.TasksNum) - 1; TaskIndex >= State.TasksBegin; TaskIndex--)
 				{
-					STATETREE_LOG(Verbose, TEXT("Passed delayed transition from '%s' (%s) -> '%s'"),
-						*GetSafeStateName(Exec.ActiveStates.Last()), *State.Name.ToString(), *GetSafeStateName(Transition.State));
+					const FStateTreeTaskBase& Task = StateTree.Nodes[TaskIndex].Get<const FStateTreeTaskBase>();
 
-					Exec.DelayedTransitions.RemoveAllSwap([TransitionIndex](const FStateTreeTransitionDelayedState& DelayedState)
-						{
-							return DelayedState.TransitionIndex.Get() == TransitionIndex;
-						});
+					// Ignore disabled task
+					if (Task.bTaskEnabled == false)
+					{
+						STATETREE_LOG(VeryVerbose, TEXT("%*sSkipped 'TriggerTransitions' for disabled Task: '%s'"), UE::StateTree::DebugIndentSize, TEXT(""), *Task.Name.ToString());
+						continue;
+					}
 
-					// Trigger Delayed Transition when the delay has passed.
-					RequestTransition(Transition.State, Transition.Priority);
-					continue;
+					if (Task.bShouldAffectTransitions)
+					{
+						STATETREE_LOG(VeryVerbose, TEXT("%*sTriggerTransitions: '%s'"), UE::StateTree::DebugIndentSize, TEXT(""), *Task.Name.ToString());
+						check(DataViews[Task.DataViewIndex.Get()].IsValid());
+						Task.TriggerTransitions(*this);
+					}
 				}
 			}
+			
+			
+			for (uint8 i = 0; i < State.TransitionsNum; i++)
+			{
+				// All transition conditions must pass
+				const int16 TransitionIndex = State.TransitionsBegin + i;
+				const FCompactStateTransition& Transition = StateTree.Transitions[TransitionIndex];
 
-			const bool bShouldTrigger = Transition.Trigger == EStateTreeTransitionTrigger::OnTick
-										|| (Transition.Trigger == EStateTreeTransitionTrigger::OnEvent
-											&& HasEventToProcess(Transition.EventTag));
+				// No need to test the transition if same or higher priority transition has already been processed.
+				if (Transition.Priority <= NextTransition.Priority)
+				{
+					continue;
+				}
 
-			bool bPassed = false; 
-			if (bShouldTrigger)
-			{
-				STATETREE_TRACE_TRANSITION_EVENT(TransitionIndex, EStateTreeTraceEventType::OnEvaluating);
-				STATETREE_TRACE_SCOPED_PHASE(EStateTreeUpdatePhase::TransitionConditions);
-				bPassed = TestAllConditions(Transition.ConditionsBegin, Transition.ConditionsNum);
-			}
+				// Skip completion transitions
+				if (EnumHasAnyFlags(Transition.Trigger, EStateTreeTransitionTrigger::OnStateCompleted))
+				{
+					continue;
+				}
 
-			if (bPassed)
-			{
-				// If the transitions is delayed, set up the delay. 
+				// If a delayed transition has passed the delay, and remove it from the queue, and try trigger it.
+				FStateTreeTransitionDelayedState* DelayedState = nullptr;
 				if (Transition.HasDelay())
 				{
-					if (DelayedState == nullptr)
+					DelayedState = Exec.FindDelayedTransition(FStateTreeIndex16(TransitionIndex));
+					if (DelayedState != nullptr && DelayedState->TimeLeft <= 0.0f)
 					{
-						// Initialize new delayed transition.
-						const float DelayDuration = Transition.Delay.GetRandomDuration();
-						if (DelayDuration > 0.0f)
+						STATETREE_LOG(Verbose, TEXT("Passed delayed transition from '%s' (%s) -> '%s'"),
+							*GetSafeStateName(Exec.ActiveStates.Last()), *State.Name.ToString(), *GetSafeStateName(Transition.State));
+
+						Exec.DelayedTransitions.RemoveAllSwap([TransitionIndex](const FStateTreeTransitionDelayedState& DelayedState)
+							{
+								return DelayedState.TransitionIndex.Get() == TransitionIndex;
+							});
+
+						// Trigger Delayed Transition when the delay has passed.
+						RequestTransition(Transition.State, Transition.Priority);
+						continue;
+					}
+				}
+
+				const bool bShouldTrigger = Transition.Trigger == EStateTreeTransitionTrigger::OnTick
+											|| (Transition.Trigger == EStateTreeTransitionTrigger::OnEvent
+												&& HasEventToProcess(Transition.EventTag));
+
+				bool bPassed = false; 
+				if (bShouldTrigger)
+				{
+					STATETREE_TRACE_TRANSITION_EVENT(TransitionIndex, EStateTreeTraceEventType::OnEvaluating);
+					STATETREE_TRACE_SCOPED_PHASE(EStateTreeUpdatePhase::TransitionConditions);
+					bPassed = TestAllConditions(Transition.ConditionsBegin, Transition.ConditionsNum);
+				}
+
+				if (bPassed)
+				{
+					// If the transitions is delayed, set up the delay. 
+					if (Transition.HasDelay())
+					{
+						if (DelayedState == nullptr)
+						{
+							// Initialize new delayed transition.
+							const float DelayDuration = Transition.Delay.GetRandomDuration();
+							if (DelayDuration > 0.0f)
+							{
+								DelayedState = &Exec.DelayedTransitions.AddDefaulted_GetRef();
+								DelayedState->TransitionIndex = FStateTreeIndex16(TransitionIndex);
+								DelayedState->TimeLeft = DelayDuration;
+								BeginDelayedTransition(*DelayedState);
+								STATETREE_LOG(Verbose, TEXT("Delayed transition triggered from '%s' (%s) -> '%s' %.1fs"),
+									*GetSafeStateName(Exec.ActiveStates.Last()), *State.Name.ToString(), *GetSafeStateName(Transition.State), DelayedState->TimeLeft);
+								
+								// Delay state added, skip requesting the transition.
+								continue;
+							}
+							// Fallthrough to request transition if duration was zero. 
+						}
+						else
 						{
-							DelayedState = &Exec.DelayedTransitions.AddDefaulted_GetRef();
-							DelayedState->TransitionIndex = FStateTreeIndex16(TransitionIndex);
-							DelayedState->TimeLeft = DelayDuration;
-							BeginDelayedTransition(*DelayedState);
-							STATETREE_LOG(Verbose, TEXT("Delayed transition triggered from '%s' (%s) -> '%s' %.1fs"),
-								*GetSafeStateName(Exec.ActiveStates.Last()), *State.Name.ToString(), *GetSafeStateName(Transition.State), DelayedState->TimeLeft);
-							
-							// Delay state added, skip requesting the transition.
+							// We get here if the transitions re-triggers during the delay, on which case we'll just ignore it.
 							continue;
 						}
-						// Fallthrough to request transition if duration was zero. 
-					}
-					else
-					{
-						// We get here if the transitions re-triggers during the delay, on which case we'll just ignore it.
-						continue;
 					}
-				}
 
-				RequestTransition(Transition.State, Transition.Priority);
+					RequestTransition(Transition.State, Transition.Priority);
+				}
 			}
 		}
 	}
diff --git a/Engine/Plugins/Runtime/StateTree/Source/StateTreeModule/Public/StateTreeTypes.h b/Engine/Plugins/Runtime/StateTree/Source/StateTreeModule/Public/StateTreeTypes.h
--- a/Engine/Plugins/Runtime/StateTree/Source/StateTreeModule/Public/StateTreeTypes.h
+++ b/Engine/Plugins/Runtime/StateTree/Source/StateTreeModule/Public/StateTreeTypes.h
@@ -692,6 +692,14 @@ struct STATETREEMODULE_API FCompactStateTreeState
 	UPROPERTY()
 	uint8 TasksNum = 0;
 
+	/** Number of tasks with struct instance data. */
+	UPROPERTY()
+	uint8 TaskInstanceStructNum = 0;
+
+	/** Number of tasks with object instance data. */
+	UPROPERTY()
+	uint8 TaskInstanceObjectNum = 0;
+
 	/** Type of the state */
 	UPROPERTY()
 	EStateTreeStateType Type = EStateTreeStateType::State;
