StateTree: Changed EnterState and TickState to events to allow latent tasks
- Added new EnterState and TickState and deprecated olds ones
- Added function to finish the state

diff --git a/Engine/Plugins/Runtime/StateTree/Source/StateTreeModule/Private/Blueprint/StateTreeTaskBlueprintBase.cpp b/Engine/Plugins/Runtime/StateTree/Source/StateTreeModule/Private/Blueprint/StateTreeTaskBlueprintBase.cpp
--- a/Engine/Plugins/Runtime/StateTree/Source/StateTreeModule/Private/Blueprint/StateTreeTaskBlueprintBase.cpp
+++ b/Engine/Plugins/Runtime/StateTree/Source/StateTreeModule/Private/Blueprint/StateTreeTaskBlueprintBase.cpp
@@ -4,6 +4,8 @@
 #include "Blueprint/StateTreeNodeBlueprintBase.h"
 #include "StateTreeExecutionContext.h"
 #include "BlueprintNodeHelpers.h"
+#include "Engine/World.h"
+#include "TimerManager.h"
 
 #include UE_INLINE_GENERATED_CPP_BY_NAME(StateTreeTaskBlueprintBase)
 
@@ -19,22 +21,37 @@ UStateTreeTaskBlueprintBase::UStateTreeTaskBlueprintBase(const FObjectInitialize
 	, bShouldCopyBoundPropertiesOnTick(true)
 	, bShouldCopyBoundPropertiesOnExitState(true)
 {
-	bHasEnterState = BlueprintNodeHelpers::HasBlueprintFunction(TEXT("ReceiveEnterState"), *this, *StaticClass());
 	bHasExitState = BlueprintNodeHelpers::HasBlueprintFunction(TEXT("ReceiveExitState"), *this, *StaticClass());
 	bHasStateCompleted = BlueprintNodeHelpers::HasBlueprintFunction(TEXT("ReceiveStateCompleted"), *this, *StaticClass());
-	bHasTick = BlueprintNodeHelpers::HasBlueprintFunction(TEXT("ReceiveTick"), *this, *StaticClass());
+	bHasLatentEnterState = BlueprintNodeHelpers::HasBlueprintFunction(TEXT("ReceiveLatentEnterState"), *this, *StaticClass());
+	bHasLatentTick = BlueprintNodeHelpers::HasBlueprintFunction(TEXT("ReceiveLatentTick"), *this, *StaticClass());
+PRAGMA_DISABLE_DEPRECATION_WARNINGS		
+	bHasEnterState_DEPRECATED = BlueprintNodeHelpers::HasBlueprintFunction(TEXT("ReceiveEnterState"), *this, *StaticClass());
+	bHasTick_DEPRECATED = BlueprintNodeHelpers::HasBlueprintFunction(TEXT("ReceiveTick"), *this, *StaticClass());
+PRAGMA_ENABLE_DEPRECATION_WARNINGS	
 }
 
 EStateTreeRunStatus UStateTreeTaskBlueprintBase::EnterState(FStateTreeExecutionContext& Context, const FStateTreeTransitionResult& Transition)
 {
 	// Task became active, cache event queue and owner.
 	SetCachedInstanceDataFromContext(Context);
-	
-	if (bHasEnterState)
+
+	// Reset status to running since the same task may be restarted.
+	RunStatus = EStateTreeRunStatus::Running;
+
+	if (bHasLatentEnterState)
+	{
+		// Note: the name contains latent just to differentiate it from the deprecated version (the old version did not allow latent actions to be started).
+		ReceiveLatentEnterState(Transition);
+	}
+PRAGMA_DISABLE_DEPRECATION_WARNINGS		
+	else if (bHasEnterState_DEPRECATED)
 	{
 		return ReceiveEnterState(Transition);
 	}
-	return EStateTreeRunStatus::Running;
+PRAGMA_ENABLE_DEPRECATION_WARNINGS
+	
+	return RunStatus;
 }
 
 void UStateTreeTaskBlueprintBase::ExitState(FStateTreeExecutionContext& Context, const FStateTreeTransitionResult& Transition)
@@ -44,6 +61,12 @@ void UStateTreeTaskBlueprintBase::ExitState(FStateTreeExecutionContext& Context,
 		ReceiveExitState(Transition);
 	}
 
+	if (UWorld* CurrentWorld = GetWorld())
+	{
+		CurrentWorld->GetLatentActionManager().RemoveActionsForObject(this);
+		CurrentWorld->GetTimerManager().ClearAllTimersForObject(this);
+	}
+
 	// Task became inactive, clear cached event queue and owner.
 	ClearCachedInstanceData();
 }
@@ -58,11 +81,19 @@ void UStateTreeTaskBlueprintBase::StateCompleted(FStateTreeExecutionContext& Con
 
 EStateTreeRunStatus UStateTreeTaskBlueprintBase::Tick(FStateTreeExecutionContext& Context, const float DeltaTime)
 {
-	if (bHasTick)
+	if (bHasLatentTick)
+	{
+		// Note: the name contains latent just to differentiate it from the deprecated version (the old version did not allow latent actions to be started).
+		ReceiveLatentTick(DeltaTime);
+	}
+PRAGMA_DISABLE_DEPRECATION_WARNINGS
+	else if (bHasTick_DEPRECATED)
 	{
 		return ReceiveTick(DeltaTime);
 	}
-	return EStateTreeRunStatus::Running;
+PRAGMA_ENABLE_DEPRECATION_WARNINGS
+	
+	return RunStatus;
 }
 
 //----------------------------------------------------------------------//
@@ -75,7 +106,11 @@ EDataValidationResult FStateTreeBlueprintTaskWrapper::Compile(FStateTreeDataView
 	
 	// Copy over ticking related options.
 	bShouldStateChangeOnReselect = InstanceData.bShouldStateChangeOnReselect;
-	bShouldCallTick = InstanceData.bShouldCallTick || InstanceData.bHasTick;
+	
+	bShouldCallTick = InstanceData.bShouldCallTick || InstanceData.bHasLatentTick;
+PRAGMA_DISABLE_DEPRECATION_WARNINGS
+	bShouldCallTick |= InstanceData.bHasTick_DEPRECATED;
+PRAGMA_ENABLE_DEPRECATION_WARNINGS
 	bShouldCallTickOnlyOnEvents = InstanceData.bShouldCallTickOnlyOnEvents;
 	bShouldCopyBoundPropertiesOnTick = InstanceData.bShouldCopyBoundPropertiesOnTick;
 	bShouldCopyBoundPropertiesOnExitState = InstanceData.bShouldCopyBoundPropertiesOnExitState;
diff --git a/Engine/Plugins/Runtime/StateTree/Source/StateTreeModule/Public/Blueprint/StateTreeTaskBlueprintBase.h b/Engine/Plugins/Runtime/StateTree/Source/StateTreeModule/Public/Blueprint/StateTreeTaskBlueprintBase.h
--- a/Engine/Plugins/Runtime/StateTree/Source/StateTreeModule/Public/Blueprint/StateTreeTaskBlueprintBase.h
+++ b/Engine/Plugins/Runtime/StateTree/Source/StateTreeModule/Public/Blueprint/StateTreeTaskBlueprintBase.h
@@ -10,7 +10,7 @@
 struct FStateTreeExecutionContext;
 
 /*
- * Base class for Blueprint based Tasks. 
+ * Base class for Blueprint based Tasks.
  */
 UCLASS(Abstract, Blueprintable)
 class STATETREEMODULE_API UStateTreeTaskBlueprintBase : public UStateTreeNodeBlueprintBase
@@ -19,16 +19,54 @@ class STATETREEMODULE_API UStateTreeTaskBlueprintBase : public UStateTreeNodeBlu
 public:
 	UStateTreeTaskBlueprintBase(const FObjectInitializer& ObjectInitializer);
 
+	/**
+	 * Called when a new state is entered and task is part of active states.
+	 * Use FinishTask() to set the task execution completed. State completion is controlled by completed tasks.
+	 *
+	 * GameplayTasks and other latent actions should be generally triggered on EnterState. When using a GameplayTasks it's required
+	 * to manually cancel active tasks on ExitState if the GameplayTask's lifetime is tied to the State Tree task.
+	 *
+	 * @param Transition Describes the states involved in the transition
+	 */
 	UFUNCTION(BlueprintImplementableEvent, meta = (DisplayName = "EnterState"))
-	EStateTreeRunStatus ReceiveEnterState(const FStateTreeTransitionResult& Transition);
+	void ReceiveLatentEnterState(const FStateTreeTransitionResult& Transition);
 
+	/**
+	 * Called when a current state is exited and task is part of active states.
+	 * @param Transition Describes the states involved in the transition
+	 */
 	UFUNCTION(BlueprintImplementableEvent, meta = (DisplayName = "ExitState"))
 	void ReceiveExitState(const FStateTreeTransitionResult& Transition);
 
+	/**
+	 * Called right after a state has been completed, but before new state has been selected. StateCompleted is called in reverse order to allow to propagate state to other Tasks that
+	 * are executed earlier in the tree. Note that StateCompleted is not called if conditional transition changes the state.
+	 * @param CompletionStatus Describes the running status of the completed state (Succeeded/Failed).
+	 * @param CompletedActiveStates Active states at the time of completion.
+	 */
 	UFUNCTION(BlueprintImplementableEvent, meta = (DisplayName = "StateCompleted"))
 	void ReceiveStateCompleted(const EStateTreeRunStatus CompletionStatus, const FStateTreeActiveStates CompletedActiveStates);
 
+	/**
+	 * Called during state tree tick when the task is on active state.
+	 * Use FinishTask() to set the task execution completed. State completion is controlled by completed tasks.
+	 *
+	 * Triggering GameplayTasks and other latent action should generally be done on EnterState. Tick is called on each update (or event)
+	 * and can cause huge amount of task added if the logic is not handled carefully.
+	 * Tick should be generally be used for monitoring that require polling, or actions that require constant ticking.  
+	 *
+	 * Note: The method is called only if bShouldCallTick or bShouldCallTickOnlyOnEvents is set.
+	 * @param DeltaTime Time since last StateTree tick.
+	 */
 	UFUNCTION(BlueprintImplementableEvent, meta = (DisplayName = "Tick"))
+	void ReceiveLatentTick(const float DeltaTime);
+
+	UE_DEPRECATED(5.3, "Use the new EnterState event without without return value instead. Task status is now controlled via FinishTask node, instead of a return value. Default status is running.")
+	UFUNCTION(BlueprintImplementableEvent, meta = (DisplayName = "EnterState (Deprecated)", DeprecatedFunction, DeprecationMessage="Use the new EnterState event without without return value instead. Task status is now controlled via FinishTask node, instead of a return value. Default status is running."))
+	EStateTreeRunStatus ReceiveEnterState(const FStateTreeTransitionResult& Transition);
+
+	UE_DEPRECATED(5.3, "Use the new Tick event without without return value instead. Task status is now controlled via FinishTask node, instead of a return value. Default status is running.")
+	UFUNCTION(BlueprintImplementableEvent, meta = (DisplayName = "Tick (Deprecated)", DeprecatedFunction, DeprecationMessage="Use the new Tick event without without return value instead. Task status is now controlled via FinishTask node, instead of a return value. Default status is running."))
 	EStateTreeRunStatus ReceiveTick(const float DeltaTime);
 
 protected:
@@ -47,6 +85,16 @@ protected:
 	virtual void StateCompleted(FStateTreeExecutionContext& Context, const EStateTreeRunStatus CompletionStatus, const FStateTreeActiveStates& CompletedActiveStates);
 	virtual EStateTreeRunStatus Tick(FStateTreeExecutionContext& Context, const float DeltaTime);
 
+	/** Finish the task and sets it's status. */
+	UFUNCTION(BlueprintCallable, Category = "StateTree", meta = (HideSelfPin = "true", DisplayName = "Finish Task"))
+	void FinishTask(const bool bSucceeded = true)
+	{
+		RunStatus = bSucceeded ? EStateTreeRunStatus::Succeeded : EStateTreeRunStatus::Failed;
+	}
+
+	/** Run status when using latent EnterState and Tick */
+	EStateTreeRunStatus RunStatus = EStateTreeRunStatus::Running;
+	
 	/**
 	 * If set to true, the task will receive EnterState/ExitState even if the state was previously active.
 	 * Generally this should be true for action type tasks, like playing animation,
@@ -72,10 +120,14 @@ protected:
 	UPROPERTY(EditDefaultsOnly, Category="Default")
 	uint8 bShouldCopyBoundPropertiesOnExitState : 1;
 	
-	uint8 bHasEnterState : 1;
 	uint8 bHasExitState : 1;
 	uint8 bHasStateCompleted : 1;
-	uint8 bHasTick : 1;
+	uint8 bHasLatentEnterState : 1;
+	uint8 bHasLatentTick : 1;
+	UE_DEPRECATED(5.3, "Use bHasLatentEnterState instead.")
+	uint8 bHasEnterState_DEPRECATED : 1;
+	UE_DEPRECATED(5.3, "Use bHasLatentTick instead.")
+	uint8 bHasTick_DEPRECATED : 1;
 
 	friend struct FStateTreeBlueprintTaskWrapper;
 };
diff --git a/Engine/Plugins/Runtime/StateTree/Source/StateTreeModule/Public/StateTreeTaskBase.h b/Engine/Plugins/Runtime/StateTree/Source/StateTreeModule/Public/StateTreeTaskBase.h
--- a/Engine/Plugins/Runtime/StateTree/Source/StateTreeModule/Public/StateTreeTaskBase.h
+++ b/Engine/Plugins/Runtime/StateTree/Source/StateTreeModule/Public/StateTreeTaskBase.h
@@ -35,8 +35,7 @@ struct STATETREEMODULE_API FStateTreeTaskBase : public FStateTreeNodeBase
 	virtual void ExitState(FStateTreeExecutionContext& Context, const EStateTreeStateChangeType ChangeType, const FStateTreeTransitionResult& Transition) const final {}
 
 	/**
-	 * Called when a new state is entered and task is part of active states. The change type parameter describes if the task's state
-	 * was previously part of the list of active states (Sustained), or if it just became active (Changed).
+	 * Called when a new state is entered and task is part of active states.
 	 * @param Context Reference to current execution context.
 	 * @param Transition Describes the states involved in the transition
 	 * @return Succeed/Failed will end the state immediately and trigger to select new state, Running will carry on to tick the state.
@@ -44,15 +43,14 @@ struct STATETREEMODULE_API FStateTreeTaskBase : public FStateTreeNodeBase
 	virtual EStateTreeRunStatus EnterState(FStateTreeExecutionContext& Context, const FStateTreeTransitionResult& Transition) const { return EStateTreeRunStatus::Running; }
 
 	/**
-	* Called when a current state is exited and task is part of active states. The change type parameter describes if the task's state
-	* will be active after the transition (Sustained), or if it will became inactive (Changed).
+	* Called when a current state is exited and task is part of active states.
 	* @param Context Reference to current execution context.
 	* @param Transition Describes the states involved in the transition
 	*/
 	virtual void ExitState(FStateTreeExecutionContext& Context, const FStateTreeTransitionResult& Transition) const {}
 
 	/**
-	 * Called Right after a state has been completed. StateCompleted is called in reverse order to allow to propagate state to Evaluators and Tasks that
+	 * Called right after a state has been completed, but before new state has been selected. StateCompleted is called in reverse order to allow to propagate state to other Tasks that
 	 * are executed earlier in the tree. Note that StateCompleted is not called if conditional transition changes the state.
 	 * @param Context Reference to current execution context.
 	 * @param CompletionStatus Describes the running status of the completed state (Succeeded/Failed).
