StateTree: Fixed undo on paste/duplicate state

UE-187155

diff --git a/Engine/Plugins/Runtime/StateTree/Source/StateTreeEditorModule/Private/SStateTreeView.cpp b/Engine/Plugins/Runtime/StateTree/Source/StateTreeEditorModule/Private/SStateTreeView.cpp
--- a/Engine/Plugins/Runtime/StateTree/Source/StateTreeEditorModule/Private/SStateTreeView.cpp
+++ b/Engine/Plugins/Runtime/StateTree/Source/StateTreeEditorModule/Private/SStateTreeView.cpp
@@ -586,8 +586,7 @@ void SStateTreeView::HandleDuplicateSelectedStates()
 {
 	if (StateTreeViewModel)
 	{
-		StateTreeViewModel->CopySelectedStates();
-		StateTreeViewModel->PasteStatesFromClipboard(GetFirstSelectedState());
+		StateTreeViewModel->DuplicateSelectedStates();
 	}
 }
 
diff --git a/Engine/Plugins/Runtime/StateTree/Source/StateTreeEditorModule/Private/StateTreeViewModel.cpp b/Engine/Plugins/Runtime/StateTree/Source/StateTreeEditorModule/Private/StateTreeViewModel.cpp
--- a/Engine/Plugins/Runtime/StateTree/Source/StateTreeEditorModule/Private/StateTreeViewModel.cpp
+++ b/Engine/Plugins/Runtime/StateTree/Source/StateTreeEditorModule/Private/StateTreeViewModel.cpp
@@ -654,22 +654,38 @@ void FStateTreeViewModel::PasteStatesFromClipboard(UStateTreeState* AfterState)
 		const int32 Index = AfterState->Parent ? AfterState->Parent->Children.Find(AfterState) : TreeData->SubTrees.Find(AfterState);
 		if (Index != INDEX_NONE)
 		{
-			PasteStatesAsChildrenFromClipboard(AfterState->Parent, Index + 1);
+			FString TextToImport;
+			FPlatformApplicationMisc::ClipboardPaste(TextToImport);
+			
+			const FScopedTransaction Transaction(LOCTEXT("PasteStatesTransaction", "Paste State(s)"));
+			PasteStatesAsChildrenFromText(TextToImport, AfterState->Parent, Index + 1);
 		}
 	}
 }
 
-void FStateTreeViewModel::PasteStatesAsChildrenFromClipboard(UStateTreeState* ParentState, const int32 Index)
+void FStateTreeViewModel::PasteStatesAsChildrenFromClipboard(UStateTreeState* ParentState)
 {
 	UStateTreeEditorData* TreeData = TreeDataWeak.Get();
 	if (TreeData == nullptr)
 	{
 		return;
 	}
-
+	
 	FString TextToImport;
 	FPlatformApplicationMisc::ClipboardPaste(TextToImport);
 
+	const FScopedTransaction Transaction(LOCTEXT("PasteStatesTransaction", "Paste State(s)"));
+	PasteStatesAsChildrenFromText(TextToImport, ParentState, INDEX_NONE);
+}
+
+void FStateTreeViewModel::PasteStatesAsChildrenFromText(const FString& TextToImport, UStateTreeState* ParentState, const int32 IndexToInsertAt)
+{
+	UStateTreeEditorData* TreeData = TreeDataWeak.Get();
+	if (TreeData == nullptr)
+	{
+		return;
+	}
+
 	UObject* Outer = ParentState ? static_cast<UObject*>(ParentState) : static_cast<UObject*>(TreeData);
 	Outer->Modify();
 
@@ -677,7 +693,7 @@ void FStateTreeViewModel::PasteStatesAsChildrenFromClipboard(UStateTreeState* Pa
 	Factory.ProcessBuffer(Outer, RF_Transactional, TextToImport);
 
 	TArray<TObjectPtr<UStateTreeState>>& ParentArray = ParentState ? ParentState->Children : TreeData->SubTrees;
-	const int32 TargetIndex = (Index == INDEX_NONE) ? ParentArray.Num() : Index;
+	const int32 TargetIndex = (IndexToInsertAt == INDEX_NONE) ? ParentArray.Num() : IndexToInsertAt;
 	ParentArray.Insert(Factory.States, TargetIndex);
 
 	TArray<FStateTreeStateLink*> Links;
@@ -731,6 +747,40 @@ void FStateTreeViewModel::PasteStatesAsChildrenFromClipboard(UStateTreeState* Pa
 	}
 }
 
+void FStateTreeViewModel::DuplicateSelectedStates()
+{
+	UStateTreeEditorData* TreeData = TreeDataWeak.Get();
+	if (TreeData == nullptr)
+	{
+		return;
+	}
+
+	TArray<UStateTreeState*> States;
+	GetSelectedStates(States);
+	UE::StateTree::Editor::RemoveContainedChildren(States);
+
+	if (States.IsEmpty())
+	{
+		return;
+	}
+	
+	FString ExportedText = UE::StateTree::Editor::ExportStatesToText(TreeData, States);
+
+	// Place duplicates after first selected state.
+	UStateTreeState* AfterState = States[0];
+	
+	const int32 Index = AfterState->Parent ? AfterState->Parent->Children.Find(AfterState) : TreeData->SubTrees.Find(AfterState);
+	if (Index != INDEX_NONE)
+	{
+		FString TextToImport;
+		FPlatformApplicationMisc::ClipboardPaste(TextToImport);
+			
+		const FScopedTransaction Transaction(LOCTEXT("DuplicateStatesTransaction", "Duplicate State(s)"));
+		PasteStatesAsChildrenFromText(TextToImport, AfterState->Parent, Index + 1);
+	}
+}
+
+
 void FStateTreeViewModel::MoveSelectedStatesBefore(UStateTreeState* TargetState)
 {
 	MoveSelectedStates(TargetState, FStateTreeViewModelInsert::Before);
diff --git a/Engine/Plugins/Runtime/StateTree/Source/StateTreeEditorModule/Private/StateTreeViewModel.h b/Engine/Plugins/Runtime/StateTree/Source/StateTreeEditorModule/Private/StateTreeViewModel.h
--- a/Engine/Plugins/Runtime/StateTree/Source/StateTreeEditorModule/Private/StateTreeViewModel.h
+++ b/Engine/Plugins/Runtime/StateTree/Source/StateTreeEditorModule/Private/StateTreeViewModel.h
@@ -75,7 +75,8 @@ public:
 	void CopySelectedStates();
 	bool CanPasteStatesFromClipboard() const;
 	void PasteStatesFromClipboard(UStateTreeState* AfterState);
-	void PasteStatesAsChildrenFromClipboard(UStateTreeState* ParentState, const int32 Index = INDEX_NONE);
+	void PasteStatesAsChildrenFromClipboard(UStateTreeState* ParentState);
+	void DuplicateSelectedStates();
 	void MoveSelectedStatesBefore(UStateTreeState* TargetState);
 	void MoveSelectedStatesAfter(UStateTreeState* TargetState);
 	void MoveSelectedStatesInto(UStateTreeState* TargetState);
@@ -117,6 +118,8 @@ protected:
 	void GetExpandedStatesRecursive(UStateTreeState* State, TSet<TWeakObjectPtr<UStateTreeState>>& ExpandedStates);
 	void MoveSelectedStates(UStateTreeState* TargetState, const FStateTreeViewModelInsert RelativeLocation);
 
+	void PasteStatesAsChildrenFromText(const FString& TextToImport, UStateTreeState* ParentState, const int32 IndexToInsertAt);
+
 	void HandleIdentifierChanged(const UStateTree& StateTree) const;
 
 	void BindToDebuggerDelegates();
