StateTree: Display state links from tasks as transitions.

diff --git a/Engine/Plugins/Runtime/StateTree/Source/StateTreeEditorModule/Private/SStateTreeViewRow.cpp b/Engine/Plugins/Runtime/StateTree/Source/StateTreeEditorModule/Private/SStateTreeViewRow.cpp
--- a/Engine/Plugins/Runtime/StateTree/Source/StateTreeEditorModule/Private/SStateTreeViewRow.cpp
+++ b/Engine/Plugins/Runtime/StateTree/Source/StateTreeEditorModule/Private/SStateTreeViewRow.cpp
@@ -675,36 +675,78 @@ bool SStateTreeViewRow::HasParentTransitionForTrigger(const UStateTreeState& Sta
 	return EnumHasAllFlags(CombinedTrigger, Trigger);
 }
 
+
+FText SStateTreeViewRow::GetLinkDescription(const FStateTreeStateLink& Link)
+{
+	switch (Link.LinkType)
+	{
+	case EStateTreeTransitionType::None:
+		return LOCTEXT("TransitionNoneStyled", "[None]");
+		break;
+	case EStateTreeTransitionType::Succeeded:
+		return LOCTEXT("TransitionTreeSucceededStyled", "[Succeeded]");
+		break;
+	case EStateTreeTransitionType::Failed:
+		return LOCTEXT("TransitionTreeFailedStyled", "[Failed]");
+		break;
+	case EStateTreeTransitionType::NextState:
+		return LOCTEXT("TransitionNextStateStyled", "[Next]");
+		break;
+	case EStateTreeTransitionType::GotoState:
+		return FText::FromName(Link.Name);
+		break;
+	default:
+		ensureMsgf(false, TEXT("Unhandled transition type."));
+		break;
+	}
+
+	return FText::GetEmpty();
+};
+
 FText SStateTreeViewRow::GetTransitionsDesc(const UStateTreeState& State, const EStateTreeTransitionTrigger Trigger, const bool bUseMask) const
 {
 	TArray<FText> DescItems;
+	
 	for (const FStateTreeTransition& Transition : State.Transitions)
 	{
 		const bool bMatch = bUseMask ? EnumHasAnyFlags(Transition.Trigger, Trigger) : Transition.Trigger == Trigger;
 		if (bMatch)
 		{
-			switch (Transition.State.LinkType)
+			DescItems.Add(GetLinkDescription(Transition.State));
+		}
+	}
+
+	// Find states from transition tasks
+	if (EnumHasAnyFlags(Trigger, EStateTreeTransitionTrigger::OnTick | EStateTreeTransitionTrigger::OnEvent))
+	{
+		auto AddLinksFromStruct = [&DescItems](FStateTreeDataView Struct)
+		{
+			if (!Struct.IsValid())
 			{
-			case EStateTreeTransitionType::None:
-				DescItems.Add(LOCTEXT("TransitionNoneStyled", "[None]"));
-				break;
-			case EStateTreeTransitionType::Succeeded:
-				DescItems.Add(LOCTEXT("TransitionTreeSucceededStyled", "[Succeeded]"));
-				break;
-			case EStateTreeTransitionType::Failed:
-				DescItems.Add(LOCTEXT("TransitionTreeFailedStyled", "[Failed]"));
-				break;
-			case EStateTreeTransitionType::NextState:
-				DescItems.Add(LOCTEXT("TransitionNextStateStyled", "[Next]"));
-				break;
-			case EStateTreeTransitionType::GotoState:
-				DescItems.Add(FText::FromName(Transition.State.Name));
-				break;
-			default:
-				ensureMsgf(false, TEXT("Unhandled transition type."));
-				break;
+				return;
 			}
+			for (TPropertyValueIterator<FStructProperty> It(Struct.GetStruct(), Struct.GetMemory()); It; ++It)
+			{
+				const UScriptStruct* StructType = It.Key()->Struct;
+				if (StructType == TBaseStructure<FStateTreeStateLink>::Get())
+				{
+					const FStateTreeStateLink& Link = *static_cast<const FStateTreeStateLink*>(It.Value());
+					if (Link.LinkType != EStateTreeTransitionType::None)
+					{
+						DescItems.Add(GetLinkDescription(Link));
+					}
+				}
+			}
+		};
+		
+		for (const FStateTreeEditorNode& Task : State.Tasks)
+		{
+			AddLinksFromStruct(FStateTreeDataView(Task.Node.GetScriptStruct(), const_cast<uint8*>(Task.Node.GetMemory())));
+			AddLinksFromStruct(Task.GetInstance());
 		}
+
+		AddLinksFromStruct(FStateTreeDataView(State.SingleTask.Node.GetScriptStruct(), const_cast<uint8*>(State.SingleTask.Node.GetMemory())));
+		AddLinksFromStruct(State.SingleTask.GetInstance());
 	}
 
 	if (State.Children.Num() == 0
@@ -831,6 +873,46 @@ EVisibility SStateTreeViewRow::GetTransitionsVisibility(const UStateTreeState& S
 		return bExactMatch ? EVisibility::Visible : EVisibility::Collapsed;
 	}
 
+	// Find states from transition tasks
+	if (EnumHasAnyFlags(Trigger, EStateTreeTransitionTrigger::OnTick | EStateTreeTransitionTrigger::OnEvent))
+	{
+		auto HasAnyLinksInStruct = [](FStateTreeDataView Struct) -> bool
+		{
+			if (!Struct.IsValid())
+			{
+				return false;
+			}
+			for (TPropertyValueIterator<FStructProperty> It(Struct.GetStruct(), Struct.GetMemory()); It; ++It)
+			{
+				const UScriptStruct* StructType = It.Key()->Struct;
+				if (StructType == TBaseStructure<FStateTreeStateLink>::Get())
+				{
+					const FStateTreeStateLink& Link = *static_cast<const FStateTreeStateLink*>(It.Value());
+					if (Link.LinkType != EStateTreeTransitionType::None)
+					{
+						return true;
+					}
+				}
+			}
+			return false;
+		};
+		
+		for (const FStateTreeEditorNode& Task : State.Tasks)
+		{
+			if (HasAnyLinksInStruct(FStateTreeDataView(Task.Node.GetScriptStruct(), const_cast<uint8*>(Task.Node.GetMemory())))
+				|| HasAnyLinksInStruct(Task.GetInstance()))
+			{
+				return EVisibility::Visible;
+			}
+		}
+
+		if (HasAnyLinksInStruct(FStateTreeDataView(State.SingleTask.Node.GetScriptStruct(), const_cast<uint8*>(State.SingleTask.Node.GetMemory())))
+			|| HasAnyLinksInStruct(State.SingleTask.GetInstance()))
+		{
+			return EVisibility::Visible;
+		}
+	}
+	
 	// Handle the test
 	for (const FStateTreeTransition& Transition : State.Transitions)
 	{
diff --git a/Engine/Plugins/Runtime/StateTree/Source/StateTreeEditorModule/Private/SStateTreeViewRow.h b/Engine/Plugins/Runtime/StateTree/Source/StateTreeEditorModule/Private/SStateTreeViewRow.h
--- a/Engine/Plugins/Runtime/StateTree/Source/StateTreeEditorModule/Private/SStateTreeViewRow.h
+++ b/Engine/Plugins/Runtime/StateTree/Source/StateTreeEditorModule/Private/SStateTreeViewRow.h
@@ -8,6 +8,7 @@
 class STableViewBase;
 class UStateTreeState;
 enum class EStateTreeTransitionTrigger : uint8;
+struct FStateTreeStateLink;
 
 class UStateTreeEditorData;
 class SStateTreeView;
@@ -63,6 +64,7 @@ private:
 	EVisibility GetConditionalTransitionsVisibility() const;
 	FText GetConditionalTransitionsDesc() const;
 
+	static FText GetLinkDescription(const FStateTreeStateLink& Link);
 	FText GetTransitionsDesc(const UStateTreeState& State, const EStateTreeTransitionTrigger Trigger, const bool bUseMask = false) const;
 	FText GetTransitionsIcon(const UStateTreeState& State, const EStateTreeTransitionTrigger Trigger, const bool bUseMask = false) const;
 	EVisibility GetTransitionsVisibility(const UStateTreeState& State, const EStateTreeTransitionTrigger Trigger) const;
