StateTree: Updating bindings instanced structs while linking State Tree.

diff --git a/Engine/Plugins/Runtime/StateTree/Source/StateTreeEditorModule/Private/StateTreeEditorData.cpp b/Engine/Plugins/Runtime/StateTree/Source/StateTreeEditorModule/Private/StateTreeEditorData.cpp
--- a/Engine/Plugins/Runtime/StateTree/Source/StateTreeEditorModule/Private/StateTreeEditorData.cpp
+++ b/Engine/Plugins/Runtime/StateTree/Source/StateTreeEditorModule/Private/StateTreeEditorData.cpp
@@ -24,6 +24,7 @@ void UStateTreeEditorData::PostLoad()
 	Super::PostLoad();
 	ReparentStates();
 	FixObjectNodes();
+	UpdateBindingsInstanceStructs();
 }
 
 void UStateTreeEditorData::PostEditChangeChainProperty(FPropertyChangedChainEvent& PropertyChangedEvent)
@@ -429,6 +430,24 @@ void UStateTreeEditorData::FixObjectNodes()
 	}
 }
 
+void UStateTreeEditorData::UpdateBindingsInstanceStructs()
+{
+	TMap<FGuid, const FStateTreeDataView> AllValues;
+	GetAllStructValues(AllValues);
+	for (FStateTreePropertyPathBinding& Binding : EditorBindings.GetMutableBindings())
+	{
+		if (AllValues.Contains(Binding.GetSourcePath().GetStructID()))
+		{
+			Binding.GetMutableSourcePath().UpdateInstanceStructsFromValue(AllValues[Binding.GetSourcePath().GetStructID()]);
+		}
+
+		if (AllValues.Contains(Binding.GetTargetPath().GetStructID()))
+		{
+			Binding.GetMutableTargetPath().UpdateInstanceStructsFromValue(AllValues[Binding.GetTargetPath().GetStructID()]);
+		}
+	}
+}
+
 EStateTreeVisitor UStateTreeEditorData::VisitStateNodes(const UStateTreeState& State, TFunctionRef<EStateTreeVisitor(const UStateTreeState* State, const FGuid& ID, const FName& Name, const EStateTreeNodeType NodeType, const UScriptStruct* NodeStruct, const UStruct* InstanceStruct)> InFunc) const
 {
 	bool bContinue = true;
@@ -701,8 +720,9 @@ EStateTreeVisitor UStateTreeEditorData::VisitGlobalNodes(TFunctionRef<EStateTree
 			Desc.Name = ContextDesc.Name;
 			Desc.ID = ContextDesc.ID;
 			Desc.DataSource = EStateTreeBindableStructSource::Context;
-			
-			if (InFunc(nullptr, Desc, FStateTreeDataView()) == EStateTreeVisitor::Break)
+
+			// We don't have value for the external objects, but return the type and null value so that users of GetAllStructValues() can use the type. 
+			if (InFunc(nullptr, Desc, FStateTreeDataView(Desc.Struct, nullptr)) == EStateTreeVisitor::Break)
 			{
 				return EStateTreeVisitor::Break;
 			}
diff --git a/Engine/Plugins/Runtime/StateTree/Source/StateTreeEditorModule/Public/StateTreeEditorData.h b/Engine/Plugins/Runtime/StateTree/Source/StateTreeEditorModule/Public/StateTreeEditorData.h
--- a/Engine/Plugins/Runtime/StateTree/Source/StateTreeEditorModule/Public/StateTreeEditorData.h
+++ b/Engine/Plugins/Runtime/StateTree/Source/StateTreeEditorModule/Public/StateTreeEditorData.h
@@ -201,6 +201,7 @@ PRAGMA_ENABLE_DEPRECATION_WARNINGS
 private:
 	void FixObjectInstance(TSet<UObject*>& SeenObjects, UObject& Outer, FStateTreeEditorNode& Node);
 	void FixObjectNodes();
+	void UpdateBindingsInstanceStructs();
 
 public:
 	/** Schema describing which inputs, evaluators, and tasks a StateTree can contain */	
diff --git a/Engine/Plugins/Runtime/StateTree/Source/StateTreeModule/Private/StateTree.cpp b/Engine/Plugins/Runtime/StateTree/Source/StateTreeModule/Private/StateTree.cpp
--- a/Engine/Plugins/Runtime/StateTree/Source/StateTreeModule/Private/StateTree.cpp
+++ b/Engine/Plugins/Runtime/StateTree/Source/StateTreeModule/Private/StateTree.cpp
@@ -5,6 +5,7 @@
 #include "StateTreeLinker.h"
 #include "StateTreeTaskBase.h"
 #include "StateTreeEvaluatorBase.h"
+#include "StateTreeConditionBase.h"
 #include "AssetRegistry/AssetData.h"
 #include "Misc/ScopeRWLock.h"
 #include "StateTreeDelegates.h"
@@ -350,6 +351,29 @@ bool UStateTree::Link()
 	// This data will be used to allocate runtime instance on all StateTree users.
 	ResetLinked();
 
+	// Resolves nodes references to other StateTree data
+	FStateTreeLinker Linker(Schema);
+	Linker.SetExternalDataBaseIndex(PropertyBindings.GetSourceStructNum());
+
+	for (int32 Index = 0; Index < Nodes.Num(); Index++)
+	{
+		FStructView Node = Nodes[Index];
+		if (FStateTreeNodeBase* NodePtr = Node.GetPtr<FStateTreeNodeBase>())
+		{
+			Linker.SetCurrentInstanceDataType(NodePtr->GetInstanceDataType(), NodePtr->DataViewIndex.Get());
+			const bool bLinkSucceeded = NodePtr->Link(Linker);
+			if (!bLinkSucceeded || Linker.GetStatus() == EStateTreeLinkerStatus::Failed)
+			{
+				UE_LOG(LogStateTree, Error, TEXT("%s: node '%s' failed to resolve its references."), *GetFullName(), *NodePtr->StaticStruct()->GetName());
+				return false;
+			}
+		}
+	}
+
+	ExternalDataBaseIndex = PropertyBindings.GetSourceStructNum();
+	ExternalDataDescs = Linker.GetExternalDataDescs();
+	NumDataViews = ExternalDataBaseIndex + ExternalDataDescs.Num();
+
 	if (States.Num() > 0 && Nodes.Num() > 0)
 	{
 		if (!DefaultInstanceData.IsValid())
@@ -380,173 +404,235 @@ bool UStateTree::Link()
 			return false;
 		}
 
-		
-		const TArrayView<FStateTreeBindableStructDesc> SourceStructs = PropertyBindings.SourceStructs;
-		const TArrayView<FStateTreePropertyCopyBatch> CopyBatches = PropertyBindings.CopyBatches;
-		const TArrayView<FStateTreePropertyPathBinding> PropertyPathBindings = PropertyBindings.PropertyPathBindings;
+		if (!PatchBindings())
+		{
+			return false;
+		}
+
+		// Resolves property paths used by bindings a store property pointers
+		if (!PropertyBindings.ResolvePaths())
+		{
+			return false;
+		}
+	}
+
+	// Link succeeded, setup tree to be ready to run
+	bIsLinked = true;
+	
+	return true;
+}
+
+bool UStateTree::PatchBindings()
+{
+	const TArrayView<FStateTreeBindableStructDesc> SourceStructs = PropertyBindings.SourceStructs;
+	const TArrayView<FStateTreePropertyCopyBatch> CopyBatches = PropertyBindings.CopyBatches;
+	const TArrayView<FStateTreePropertyPathBinding> PropertyPathBindings = PropertyBindings.PropertyPathBindings;
 
-		// Reconcile out of date classes.
-		for (FStateTreeBindableStructDesc& SourceStruct : SourceStructs)
+	// Reconcile out of date classes.
+	for (FStateTreeBindableStructDesc& SourceStruct : SourceStructs)
+	{
+		if (const UClass* SourceClass = Cast<UClass>(SourceStruct.Struct))
 		{
-			if (const UClass* SourceClass = Cast<UClass>(SourceStruct.Struct))
+			if (SourceClass->HasAnyClassFlags(CLASS_NewerVersionExists))
 			{
-				if (SourceClass->HasAnyClassFlags(CLASS_NewerVersionExists))
-				{
-					SourceStruct.Struct = SourceClass->GetAuthoritativeClass();
-				}
+				SourceStruct.Struct = SourceClass->GetAuthoritativeClass();
 			}
 		}
-		for (FStateTreePropertyCopyBatch& CopyBatch : CopyBatches)
+	}
+	for (FStateTreePropertyCopyBatch& CopyBatch : CopyBatches)
+	{
+		if (const UClass* TargetClass = Cast<UClass>(CopyBatch.TargetStruct.Struct))
 		{
-			if (const UClass* TargetClass = Cast<UClass>(CopyBatch.TargetStruct.Struct))
+			if (TargetClass->HasAnyClassFlags(CLASS_NewerVersionExists))
 			{
-				if (TargetClass->HasAnyClassFlags(CLASS_NewerVersionExists))
-				{
-					CopyBatch.TargetStruct.Struct = TargetClass->GetAuthoritativeClass();
-				}
+				CopyBatch.TargetStruct.Struct = TargetClass->GetAuthoritativeClass();
 			}
 		}
+	}
 
-		auto PatchPropertyPath = [](FStateTreePropertyPath& PropertyPath)
+	auto PatchPropertyPath = [](FStateTreePropertyPath& PropertyPath)
+	{
+		for (FStateTreePropertyPathSegment& Segment : PropertyPath.GetMutableSegments())
 		{
-			for (FStateTreePropertyPathSegment& Segment : PropertyPath.GetMutableSegments())
+			if (const UClass* InstanceStruct = Cast<UClass>(Segment.GetInstanceStruct()))
 			{
-				if (const UClass* InstanceStruct = Cast<UClass>(Segment.GetInstanceStruct()))
+				if (InstanceStruct->HasAnyClassFlags(CLASS_NewerVersionExists))
 				{
-					if (InstanceStruct->HasAnyClassFlags(CLASS_NewerVersionExists))
-					{
-						Segment.SetInstanceStruct(InstanceStruct->GetAuthoritativeClass());
-					}
+					Segment.SetInstanceStruct(InstanceStruct->GetAuthoritativeClass());
 				}
 			}
-		};
-		
-		for (FStateTreePropertyPathBinding& PropertyPathBinding : PropertyPathBindings)
-		{
-			PatchPropertyPath(PropertyPathBinding.GetMutableSourcePath());
-			PatchPropertyPath(PropertyPathBinding.GetMutableTargetPath());
 		}
+	};
 
-		// Update property bag structs before resolving binding.
-		if (ParametersDataViewIndex.IsValid() && SourceStructs.IsValidIndex(ParametersDataViewIndex.Get()))
-		{
-			SourceStructs[ParametersDataViewIndex.Get()].Struct = Parameters.GetPropertyBagStruct();
-		}
-		
-		for (FCompactStateTreeState& State : States)
+	for (FStateTreePropertyPathBinding& PropertyPathBinding : PropertyPathBindings)
+	{
+		PatchPropertyPath(PropertyPathBinding.GetMutableSourcePath());
+		PatchPropertyPath(PropertyPathBinding.GetMutableTargetPath());
+	}
+
+	// Update property bag structs before resolving binding.
+	if (ParametersDataViewIndex.IsValid() && SourceStructs.IsValidIndex(ParametersDataViewIndex.Get()))
+	{
+		SourceStructs[ParametersDataViewIndex.Get()].Struct = Parameters.GetPropertyBagStruct();
+	}
+
+	for (const FCompactStateTreeState& State : States)
+	{
+		if (State.Type == EStateTreeStateType::Subtree)
 		{
-			if (State.Type == EStateTreeStateType::Subtree)
+			if (State.ParameterInstanceIndex.IsValid() == false)
 			{
-				if (State.ParameterInstanceIndex.IsValid() == false)
-				{
-					UE_LOG(LogStateTree, Error, TEXT("%s: Data for state '%s' is malformed. Please recompile the StateTree asset."), *GetFullName(), *State.Name.ToString());
-					return false;
-				}
-
-				// Subtree is a bind source, update bag struct.
-				if (State.ParameterDataViewIndex.IsValid())
-				{
-					const FCompactStateTreeParameters& Params = DefaultInstanceData.GetMutableStruct(State.ParameterInstanceIndex.Get()).Get<FCompactStateTreeParameters>();
-					FStateTreeBindableStructDesc& Desc = SourceStructs[State.ParameterDataViewIndex.Get()];
-					Desc.Struct = Params.Parameters.GetPropertyBagStruct();
-				}
+				UE_LOG(LogStateTree, Error, TEXT("%s: Data for state '%s' is malformed. Please recompile the StateTree asset."), *GetFullName(), *State.Name.ToString());
+				return false;
 			}
-			else if (State.Type == EStateTreeStateType::Linked && State.LinkedState.IsValid())
+
+			// Subtree is a bind source, update bag struct.
+			if (State.ParameterDataViewIndex.IsValid())
 			{
-				const FCompactStateTreeState& LinkedState = States[State.LinkedState.Index];
+				const FCompactStateTreeParameters& Params = DefaultInstanceData.GetMutableStruct(State.ParameterInstanceIndex.Get()).Get<FCompactStateTreeParameters>();
+				FStateTreeBindableStructDesc& Desc = SourceStructs[State.ParameterDataViewIndex.Get()];
+				Desc.Struct = Params.Parameters.GetPropertyBagStruct();
+			}
+		}
+		else if (State.Type == EStateTreeStateType::Linked && State.LinkedState.IsValid())
+		{
+			const FCompactStateTreeState& LinkedState = States[State.LinkedState.Index];
 
-				if (State.ParameterInstanceIndex.IsValid() == false
-					|| LinkedState.ParameterInstanceIndex.IsValid() == false)
-				{
-					UE_LOG(LogStateTree, Error, TEXT("%s: Data for state '%s' is malformed. Please recompile the StateTree asset."), *GetFullName(), *State.Name.ToString());
-					return false;
-				}
+			if (State.ParameterInstanceIndex.IsValid() == false
+				|| LinkedState.ParameterInstanceIndex.IsValid() == false)
+			{
+				UE_LOG(LogStateTree, Error, TEXT("%s: Data for state '%s' is malformed. Please recompile the StateTree asset."), *GetFullName(), *State.Name.ToString());
+				return false;
+			}
 
-				const FCompactStateTreeParameters& Params = DefaultInstanceData.GetMutableStruct(State.ParameterInstanceIndex.Get()).Get<FCompactStateTreeParameters>();
+			const FCompactStateTreeParameters& Params = DefaultInstanceData.GetMutableStruct(State.ParameterInstanceIndex.Get()).Get<FCompactStateTreeParameters>();
 
-				// Check that the bag in linked state matches.
-				const FCompactStateTreeParameters& LinkedStateParams = DefaultInstanceData.GetMutableStruct(LinkedState.ParameterInstanceIndex.Get()).Get<FCompactStateTreeParameters>();
+			// Check that the bag in linked state matches.
+			const FCompactStateTreeParameters& LinkedStateParams = DefaultInstanceData.GetMutableStruct(LinkedState.ParameterInstanceIndex.Get()).Get<FCompactStateTreeParameters>();
 
-				if (LinkedStateParams.Parameters.GetPropertyBagStruct() != Params.Parameters.GetPropertyBagStruct())
-				{
-					UE_LOG(LogStateTree, Log, TEXT("%s: The parameters on state '%s' does not match the linked state parameters in state '%s'. Please recompile the StateTree asset."), *GetFullName(), *State.Name.ToString(), *LinkedState.Name.ToString());
-					return false;
-				}
+			if (LinkedStateParams.Parameters.GetPropertyBagStruct() != Params.Parameters.GetPropertyBagStruct())
+			{
+				UE_LOG(LogStateTree, Error, TEXT("%s: The parameters on state '%s' does not match the linked state parameters in state '%s'. Please recompile the StateTree asset."), *GetFullName(), *State.Name.ToString(), *LinkedState.Name.ToString());
+				return false;
+			}
 
-				if (Params.BindingsBatch.IsValid())
-				{
-					FStateTreePropertyCopyBatch& Batch = CopyBatches[Params.BindingsBatch.Get()];
-					Batch.TargetStruct.Struct = Params.Parameters.GetPropertyBagStruct();
-				}
+			if (Params.BindingsBatch.IsValid())
+			{
+				FStateTreePropertyCopyBatch& Batch = CopyBatches[Params.BindingsBatch.Get()];
+				Batch.TargetStruct.Struct = Params.Parameters.GetPropertyBagStruct();
 			}
+		}
+	}
+
+	TArray<FStateTreeDataView> DataViews;
+	DataViews.SetNum(NumDataViews);
 
-			// Update instance data num for each state.
-			State.TaskInstanceStructNum = 0;
-			State.TaskInstanceObjectNum = 0;
+	// Tree parameters
+	DataViews[ParametersDataViewIndex.Get()] = Parameters.GetMutableValue();
 
-			if (State.Type == EStateTreeStateType::Linked)
+	// Setup data views for context data. Since the external data is passed at runtime, we can only provide the type.
+	for (const FStateTreeExternalDataDesc& DataDesc : ContextDataDescs)
+	{
+		DataViews[DataDesc.Handle.DataViewIndex.Get()] = FStateTreeDataView(DataDesc.Struct, nullptr);
+	}
+	
+	// Setup data views for state parameters.
+	for (FCompactStateTreeState& State : States)
+	{
+		if (State.Type == EStateTreeStateType::Linked && State.LinkedState.IsValid())
+		{
+			FCompactStateTreeParameters& Params = DefaultInstanceData.GetMutableStruct(State.ParameterInstanceIndex.Get()).Get<FCompactStateTreeParameters>();
+			if (Params.Parameters.IsValid())
 			{
-				// Linked state parameters.
-				State.TaskInstanceStructNum++;
+				const FCompactStateTreeState& LinkedState = States[State.LinkedState.Index];
+				DataViews[LinkedState.ParameterDataViewIndex.Get()] = Params.Parameters.GetMutableValue();
 			}
-			
-			for (int32 TaskIndex = (State.TasksBegin + State.TasksNum) - 1; TaskIndex >= State.TasksBegin; TaskIndex--)
+		}
+	}
+
+	// Setup data views for all nodes.
+	for (FConstStructView NodeView : Nodes)
+	{
+		if (const FStateTreeConditionBase* Condition = NodeView.GetPtr<const FStateTreeConditionBase>())
+		{
+			if (Condition->bInstanceIsObject)
 			{
-				const FStateTreeTaskBase& Task = Nodes[TaskIndex].Get<const FStateTreeTaskBase>();
-				if (Task.bInstanceIsObject)
-				{
-					if (State.TaskInstanceObjectNum == MAX_uint8)
-					{
-						UE_LOG(LogStateTree, Error, TEXT("%s: State '%s' has more than %d struct based task instances."), *GetFullName(), *State.Name.ToString(), (int32)MAX_uint8);
-						return false;
-					}
-					State.TaskInstanceObjectNum++;
-				}
-				else
-				{
-					if (State.TaskInstanceStructNum == MAX_uint8)
-					{
-						UE_LOG(LogStateTree, Error, TEXT("%s: State '%s' has more than %d object based task instances."), *GetFullName(), *State.Name.ToString(), (int32)MAX_uint8);
-						return false;
-					}
-					State.TaskInstanceStructNum++;
-				}
+				DataViews[Condition->DataViewIndex.Get()] = SharedInstanceData.GetMutableObject(Condition->InstanceIndex.Get());
+			}
+			else
+			{
+				DataViews[Condition->DataViewIndex.Get()] = SharedInstanceData.GetMutableStruct(Condition->InstanceIndex.Get());
 			}
 		}
-		
-		// Resolves property paths used by bindings a store property pointers
-		if (!PropertyBindings.ResolvePaths())
+		else
 		{
-			return false;
+			const FStateTreeNodeBase& Node = NodeView.Get<const FStateTreeNodeBase>();
+			if (Node.bInstanceIsObject)
+			{
+				DataViews[Node.DataViewIndex.Get()] = DefaultInstanceData.GetMutableObject(Node.InstanceIndex.Get());
+			}
+			else
+			{
+				DataViews[Node.DataViewIndex.Get()] = DefaultInstanceData.GetMutableStruct(Node.InstanceIndex.Get());
+			}
 		}
 	}
 
-	// Resolves nodes references to other StateTree data
-	FStateTreeLinker Linker(Schema);
-	Linker.SetExternalDataBaseIndex(PropertyBindings.GetSourceStructNum());
-	
-	for (int32 Index = 0; Index < Nodes.Num(); Index++)
+	for (int32 BatchIndex = 0; BatchIndex < CopyBatches.Num(); ++BatchIndex)
 	{
-		FStructView Node = Nodes[Index];
-		if (FStateTreeNodeBase* NodePtr = Node.GetPtr<FStateTreeNodeBase>())
+		const FStateTreePropertyCopyBatch& Batch = CopyBatches[BatchIndex];
+		for (int32 i = Batch.BindingsBegin; i != Batch.BindingsEnd; i++)
 		{
-			Linker.SetCurrentInstanceDataType(NodePtr->GetInstanceDataType(), NodePtr->DataViewIndex.Get());
-			const bool bLinkSucceeded = NodePtr->Link(Linker);
-			if (!bLinkSucceeded || Linker.GetStatus() == EStateTreeLinkerStatus::Failed)
+			FStateTreePropertyPathBinding& Binding = PropertyPathBindings[i];
+			FStateTreeIndex16 TargetStructIndex;
+			for (FConstStructView NodeView : GetNodes())
 			{
-				UE_LOG(LogStateTree, Error, TEXT("%s: node '%s' failed to resolve its references."), *GetFullName(), *NodePtr->StaticStruct()->GetName());
+				const FStateTreeNodeBase& Node = NodeView.Get<const FStateTreeNodeBase>();
+				if (Node.BindingsBatch.AsInt32() == BatchIndex)
+				{
+					TargetStructIndex = Node.DataViewIndex;
+					break;
+				}
+			}
+
+			if (!TargetStructIndex.IsValid())
+			{
+				for (const FCompactStateTreeState& State : GetStates())
+				{
+					if (State.Type == EStateTreeStateType::Linked)
+					{
+						const FCompactStateTreeParameters& Params = DefaultInstanceData.GetStruct(State.ParameterInstanceIndex.Get()).Get<const FCompactStateTreeParameters>();
+						if (Params.BindingsBatch.Get() == BatchIndex)
+						{
+							if (const FCompactStateTreeState* LinkedState = GetStateFromHandle(State.LinkedState))
+							{
+								TargetStructIndex = LinkedState->ParameterDataViewIndex;
+							}
+						}
+					}
+				}
+			}
+
+			if (!TargetStructIndex.IsValid())
+			{
+				UE_LOG(LogStateTree, Error, TEXT("%hs: Invalid target struct for property binding %s."), __FUNCTION__, *Binding.GetTargetPath().ToString());
+				return false;
+			}
+
+			if (!Binding.GetMutableSourcePath().UpdateInstanceStructsFromValue(DataViews[Binding.GetCompiledSourceStructIndex().Get()]))
+			{
+				UE_LOG(LogStateTree, Error, TEXT("%hs: Failed to update source instance structs for property binding %s."), __FUNCTION__, *Binding.GetTargetPath().ToString());
+				return false;
+			}
+
+			if (!Binding.GetMutableTargetPath().UpdateInstanceStructsFromValue(DataViews[TargetStructIndex.Get()]))
+			{
+				UE_LOG(LogStateTree, Error, TEXT("%hs: Failed to update target instance structs for property binding %s."), __FUNCTION__, *Binding.GetTargetPath().ToString());
 				return false;
 			}
 		}
 	}
 
-	// Link succeeded, setup tree to be ready to run
-	ExternalDataBaseIndex = PropertyBindings.GetSourceStructNum();
-	ExternalDataDescs = Linker.GetExternalDataDescs();
-	NumDataViews = ExternalDataBaseIndex + ExternalDataDescs.Num();
-
-	bIsLinked = true;
-	
 	return true;
 }
 
diff --git a/Engine/Plugins/Runtime/StateTree/Source/StateTreeModule/Public/StateTree.h b/Engine/Plugins/Runtime/StateTree/Source/StateTreeModule/Public/StateTree.h
--- a/Engine/Plugins/Runtime/StateTree/Source/StateTreeModule/Public/StateTree.h
+++ b/Engine/Plugins/Runtime/StateTree/Source/StateTreeModule/Public/StateTree.h
@@ -82,6 +82,9 @@ class STATETREEMODULE_API UStateTree : public UDataAsset
 	GENERATED_BODY()
 
 public:
+	/** @return Default instance data. */
+	const FStateTreeInstanceData& GetDefaultInstanceData() const { return DefaultInstanceData; }
+
 	/** @return Shared instance data. */
 	TSharedPtr<FStateTreeInstanceData> GetSharedInstanceData() const;
 
@@ -115,6 +118,9 @@ public:
 	/** @return Struct view of the node matching a given node index; invalid view if state not found. */
 	FConstStructView GetNode(const int32 NodeIndex) const;
 
+	/** @return Struct views of all nodes */
+	const FInstancedStructContainer& GetNodes() const { return Nodes; }
+
 	/** @return View of all states. */
 	TConstArrayView<FCompactStateTreeState> GetStates() const { return States; }
 
@@ -177,6 +183,8 @@ private:
 	 */
 	void ResetLinked();
 
+	bool PatchBindings();
+
 	// Data created during compilation, source data in EditorData.
 	
 	/** Schema used to compile the StateTree. */
