[StateTreeDebugger] breakpoint improvements
- added support for breakpoint when exiting state (was only when entering)
- added support for breakpoint when entering or exiting a task
- state breakpoints can be added from contextual menu in the TreeView or from the Debug options in the details panel
- task breakpoints can be added from the Debug options in the details panel
- added feedback labels in the details panel for states and tasks with breakpoints

diff --git a/Engine/Plugins/Runtime/StateTree/Source/StateTreeEditorModule/Private/Customizations/StateTreeEditorNodeDetails.cpp b/Engine/Plugins/Runtime/StateTree/Source/StateTreeEditorModule/Private/Customizations/StateTreeEditorNodeDetails.cpp
--- a/Engine/Plugins/Runtime/StateTree/Source/StateTreeEditorModule/Private/Customizations/StateTreeEditorNodeDetails.cpp
+++ b/Engine/Plugins/Runtime/StateTree/Source/StateTreeEditorModule/Private/Customizations/StateTreeEditorNodeDetails.cpp
@@ -454,7 +454,7 @@ void FStateTreeEditorNodeDetails::CustomizeHeader(TSharedRef<class IPropertyHand
 				.HAlign(HAlign_Right)
 				.FillWidth(1.0f)
 				[
-					UE::StateTreeEditor::DebuggerExtensions::CreateEditorNodeWidget(StructPropertyHandle)
+					UE::StateTreeEditor::DebuggerExtensions::CreateEditorNodeWidget(StructPropertyHandle, EditorData)
 				]
 			]
 		]
diff --git a/Engine/Plugins/Runtime/StateTree/Source/StateTreeEditorModule/Private/Customizations/StateTreeStateDetails.cpp b/Engine/Plugins/Runtime/StateTree/Source/StateTreeEditorModule/Private/Customizations/StateTreeStateDetails.cpp
--- a/Engine/Plugins/Runtime/StateTree/Source/StateTreeEditorModule/Private/Customizations/StateTreeStateDetails.cpp
+++ b/Engine/Plugins/Runtime/StateTree/Source/StateTreeEditorModule/Private/Customizations/StateTreeStateDetails.cpp
@@ -20,14 +20,14 @@ TSharedRef<IDetailCustomization> FStateTreeStateDetails::MakeInstance()
 void FStateTreeStateDetails::CustomizeDetails(IDetailLayoutBuilder& DetailBuilder)
 {
 	// Find StateTreeEditorData associated with this panel.
-	const UStateTreeEditorData* EditorData = nullptr;
+	UStateTreeEditorData* EditorData = nullptr;
 	TArray<TWeakObjectPtr<UObject>> Objects;
 	DetailBuilder.GetObjectsBeingCustomized(Objects);
 	for (TWeakObjectPtr<UObject>& WeakObject : Objects)
 	{
 		if (const UObject* Object = WeakObject.Get())
 		{
-			const UStateTreeEditorData* OuterEditorData = Object->GetTypedOuter<UStateTreeEditorData>();
+			if (UStateTreeEditorData* OuterEditorData = Object->GetTypedOuter<UStateTreeEditorData>())
 			{
 				EditorData = OuterEditorData;
 				break;
@@ -55,7 +55,7 @@ void FStateTreeStateDetails::CustomizeDetails(IDetailLayoutBuilder& DetailBuilde
 	IDetailCategoryBuilder& StateCategory = DetailBuilder.EditCategory(TEXT("State"), LOCTEXT("StateDetailsState", "State"));
 	StateCategory.SetSortOrder(0);
 
-	StateCategory.HeaderContent(UE::StateTreeEditor::DebuggerExtensions::CreateStateWidget(EnabledProperty));
+	StateCategory.HeaderContent(UE::StateTreeEditor::DebuggerExtensions::CreateStateWidget(DetailBuilder, EditorData));
 
 	if (StateType != EStateTreeStateType::Linked)
 	{
diff --git a/Engine/Plugins/Runtime/StateTree/Source/StateTreeEditorModule/Private/Debugger/SStateTreeDebuggerView.cpp b/Engine/Plugins/Runtime/StateTree/Source/StateTreeEditorModule/Private/Debugger/SStateTreeDebuggerView.cpp
--- a/Engine/Plugins/Runtime/StateTree/Source/StateTreeEditorModule/Private/Debugger/SStateTreeDebuggerView.cpp
+++ b/Engine/Plugins/Runtime/StateTree/Source/StateTreeEditorModule/Private/Debugger/SStateTreeDebuggerView.cpp
@@ -13,11 +13,13 @@
 #include "Modules/ModuleManager.h"
 #include "ProfilingDebugging/TraceAuxiliary.h"
 #include "PropertyEditorModule.h"
+#include "ScopedTransaction.h"
 #include "SStateTreeDebuggerInstanceTree.h"
 #include "SStateTreeDebuggerTimelines.h"
 #include "StateTree.h"
 #include "StateTreeDebuggerCommands.h"
 #include "StateTreeDebuggerTrack.h"
+#include "StateTreeEditorData.h"
 #include "StateTreeEditorSettings.h"
 #include "StateTreeExecutionContext.h"
 #include "StateTreeModule.h"
@@ -175,6 +177,8 @@ void SStateTreeDebuggerView::StartRecording()
 	// We give priority to the Editor actions even if an analysis was active (remote process)
 	// This will stop current analysis and connect to the new live trace.
 	Debugger->RequestAnalysisOfNextLiveSession();
+	
+	Debugger->ClearSelection();
 
 	if (CanStartRecording())
 	{
@@ -235,10 +239,11 @@ void SStateTreeDebuggerView::OnPIESingleStepped(bool bSimulating) const
 	Debugger->SyncToCurrentSessionDuration();
 }
 
-void SStateTreeDebuggerView::Construct(const FArguments& InArgs, const UStateTree* InStateTree, const TSharedRef<FStateTreeViewModel>& InStateTreeViewModel, const TSharedRef<FUICommandList>& InCommandList)
+void SStateTreeDebuggerView::Construct(const FArguments& InArgs, const UStateTree& InStateTree, const TSharedRef<FStateTreeViewModel>& InStateTreeViewModel, const TSharedRef<FUICommandList>& InCommandList)
 {
 	StateTreeViewModel = InStateTreeViewModel;
-	StateTree = InStateTree;
+	StateTree = &InStateTree;
+	StateTreeEditorData = Cast<UStateTreeEditorData>(InStateTree.EditorData.Get()); 
 
 	Debugger = InStateTreeViewModel->GetDebugger();
 	bRecording = FTraceAuxiliary::IsConnected();
@@ -264,11 +269,20 @@ void SStateTreeDebuggerView::Construct(const FArguments& InArgs, const UStateTre
 	// Register the play world commands
 	InCommandList->Append(FPlayWorldCommands::GlobalPlayWorldActions.ToSharedRef());
 
+	// Debug commands
+	InCommandList->MapAction(
+		FStateTreeDebuggerCommands::Get().EnableOnEnterStateBreakpoint,
+		FExecuteAction::CreateLambda([this] { HandleEnableStateBreakpoint(EStateTreeBreakpointType::OnEnter); }),
+		FCanExecuteAction(),
+		FGetActionCheckState::CreateLambda([this]{ return GetStateBreakpointCheckState(EStateTreeBreakpointType::OnEnter); }),
+		FIsActionButtonVisible::CreateLambda([this] { return CanAddStateBreakpoint(EStateTreeBreakpointType::OnEnter) || CanRemoveStateBreakpoint(EStateTreeBreakpointType::OnEnter); }));
+
 	InCommandList->MapAction(
-		FStateTreeDebuggerCommands::Get().ToggleBreakpoint,
-		FExecuteAction::CreateSP(this, &SStateTreeDebuggerView::ToggleBreakpoint),
-		FCanExecuteAction::CreateSP(this, &SStateTreeDebuggerView::CanToggleBreakpoint)
-	);
+		FStateTreeDebuggerCommands::Get().EnableOnExitStateBreakpoint,
+		FExecuteAction::CreateLambda([this] { HandleEnableStateBreakpoint(EStateTreeBreakpointType::OnExit); }),
+		FCanExecuteAction(),
+		FGetActionCheckState::CreateLambda([this] { return GetStateBreakpointCheckState(EStateTreeBreakpointType::OnExit); }),
+		FIsActionButtonVisible::CreateLambda([this] { return CanAddStateBreakpoint(EStateTreeBreakpointType::OnExit) || CanRemoveStateBreakpoint(EStateTreeBreakpointType::OnExit); }));
 	
 	// Toolbar
 	FSlimHorizontalToolBarBuilder ToolbarBuilder(InCommandList, FMultiBoxCustomization::None, /*InExtender*/ nullptr, /*InForceSmallIcons*/ true);
@@ -681,36 +695,145 @@ void SStateTreeDebuggerView::StepForwardToNextStateChange()
 	Debugger->StepForwardToNextStateChange();
 }
 
-bool SStateTreeDebuggerView::CanToggleBreakpoint() const
+bool SStateTreeDebuggerView::CanAddStateBreakpoint(const EStateTreeBreakpointType Type) const
 {
-	return (Debugger.IsValid() && StateTreeViewModel.IsValid() && StateTreeViewModel->HasSelection());
+	check(StateTreeViewModel);
+
+	TArray<UStateTreeState*> SelectedStates;
+	StateTreeViewModel->GetSelectedStates(SelectedStates);
+	if (SelectedStates.IsEmpty())
+	{
+		return false;
+	}
+
+	const UStateTreeEditorData* EditorData = StateTreeEditorData.Get();
+	if (!ensure(EditorData != nullptr))
+	{
+		return false;
+	}
+
+	for (const UStateTreeState* SelectedState : SelectedStates)
+	{
+		if (SelectedState != nullptr)
+		{
+			if (EditorData->HasBreakpoint(SelectedState->ID, Type) == false)
+			{
+				return true;
+			}
+		}
+	}
+
+	return false;
 }
 
-void SStateTreeDebuggerView::ToggleBreakpoint() const
+bool SStateTreeDebuggerView::CanRemoveStateBreakpoint(const EStateTreeBreakpointType Type) const
 {
 	check(StateTreeViewModel);
-	check(Debugger);
-	check(StateTree.IsValid());
 
-	TArray<UStateTreeState*> States;
-	StateTreeViewModel->GetSelectedStates(States);
+	TArray<UStateTreeState*> SelectedStates;
+	StateTreeViewModel->GetSelectedStates(SelectedStates);
+	if (SelectedStates.IsEmpty())
+	{
+		return false;
+	}
+
+	const UStateTreeEditorData* EditorData = StateTreeEditorData.Get();
+	if (!ensure(EditorData != nullptr))
+	{
+		return false;
+	}
 
-	TArray<FStateTreeStateHandle> StateHandles;
-	StateHandles.Reserve(States.Num());
-	for (const UStateTreeState* SelectedState : States)
+	for (const UStateTreeState* SelectedState : SelectedStates)
 	{
-		if (SelectedState->Type != EStateTreeStateType::State || SelectedState->Parent == nullptr)
+		if (SelectedState != nullptr)
 		{
-			continue;
+			if (EditorData->HasBreakpoint(SelectedState->ID, Type))
+			{
+				return true;
+			}
 		}
-		
-		FStateTreeStateHandle Handle = StateTree->GetStateHandleFromId(SelectedState->ID);
-		if (Handle.IsValid())
+	}
+
+	return false;
+}
+
+
+ECheckBoxState SStateTreeDebuggerView::GetStateBreakpointCheckState(const EStateTreeBreakpointType Type) const
+{
+	const bool bCanAdd = CanAddStateBreakpoint(Type);
+	const bool bCanRemove = CanRemoveStateBreakpoint(Type);
+	if (bCanAdd && bCanRemove)
+	{
+		return ECheckBoxState::Undetermined;
+	}
+
+	if (bCanRemove)
+	{
+		return ECheckBoxState::Checked;
+	}
+
+	if (bCanAdd)
+	{
+		return ECheckBoxState::Unchecked;
+	}
+
+	// Should not happen since action is not visible in this case
+	return ECheckBoxState::Undetermined;
+}
+
+void SStateTreeDebuggerView::HandleEnableStateBreakpoint(EStateTreeBreakpointType Type)
+{
+	check(StateTreeViewModel);
+
+	TArray<UStateTreeState*> SelectedStates;
+	StateTreeViewModel->GetSelectedStates(SelectedStates);
+	if (SelectedStates.IsEmpty())
+	{
+		return;
+	}
+
+	UStateTreeEditorData* EditorData = StateTreeEditorData.Get();
+	if (!ensure(EditorData != nullptr))
+	{
+		return;
+	}
+
+	TBitArray<> HasBreakpoint;
+	HasBreakpoint.Reserve(SelectedStates.Num());
+	for (const UStateTreeState* SelectedState : SelectedStates)
+	{
+		HasBreakpoint.Add(SelectedState != nullptr && EditorData->HasBreakpoint(SelectedState->ID, Type));
+	}
+
+	check(HasBreakpoint.Num() == SelectedStates.Num());
+
+	// Process CanAdd first so in case of undetermined state (mixed selection) we add by default. 
+	if (CanAddStateBreakpoint(Type))
+	{
+		const FScopedTransaction Transaction(LOCTEXT("AddStateBreakpoint", "Add State Breakpoint(s)"));
+		EditorData->Modify();
+		for (int Index = 0; Index < SelectedStates.Num(); ++Index)
 		{
-			StateHandles.Add(Handle);
+			const UStateTreeState* SelectedState = SelectedStates[Index];
+			if (HasBreakpoint[Index] == false && SelectedState != nullptr)
+			{
+				EditorData->AddBreakpoint(SelectedState->ID, Type);	
+			}
+		}
+	}
+	else if (CanRemoveStateBreakpoint(Type))
+	{
+		const FScopedTransaction Transaction(LOCTEXT("RemoveStateBreakpoint", "Remove State Breakpoint(s)"));
+		EditorData->Modify();
+		for (int Index = 0; Index < SelectedStates.Num(); ++Index)
+		{
+			const UStateTreeState* SelectedState = SelectedStates[Index];
+			if (HasBreakpoint[Index] && SelectedState != nullptr)
+			{
+				EditorData->RemoveBreakpoint(SelectedState->ID, Type);	
+			}
 		}
 	}
-	Debugger->ToggleBreakpoints(StateHandles);
 }
 
 void SStateTreeDebuggerView::OnTimeLineScrubPositionChanged(double Time, bool bIsScrubbing)
@@ -820,7 +943,7 @@ void SStateTreeDebuggerView::ExpandAll(const TArray<TSharedPtr<UE::StateTreeDebu
 	}
 }
 
-void SStateTreeDebuggerView::OnBreakpointHit(const FStateTreeInstanceDebugId InstanceId, const FStateTreeStateHandle StateHandle, const TSharedRef<FUICommandList> ActionList) const
+void SStateTreeDebuggerView::OnBreakpointHit(const FStateTreeInstanceDebugId InstanceId, const FStateTreeDebuggerBreakpoint Breakpoint, const TSharedRef<FUICommandList> ActionList) const
 {
 	if (FPlayWorldCommands::Get().PausePlaySession.IsValid())
 	{
@@ -853,7 +976,22 @@ void SStateTreeDebuggerView::OnNewInstance(FStateTreeInstanceDebugId InstanceId)
 	if (FStateTreeDebuggerOwnerTrack* OwnerTrack = static_cast<FStateTreeDebuggerOwnerTrack*>(ExistingOwnerTrack->Get()))
 	{
 		const FString TrackName = FString::Printf(TEXT("Execution #%d"), OwnerTrack->NumSubTracks()+1);
-		OwnerTrack->AddSubTrack(MakeShared<FStateTreeDebuggerInstanceTrack>(Debugger, InstanceId, FText::FromString(TrackName), ViewRange));
+		const TSharedPtr<FStateTreeDebuggerInstanceTrack> SubTrack = MakeShared<FStateTreeDebuggerInstanceTrack>(Debugger, InstanceId, FText::FromString(TrackName), ViewRange);
+		OwnerTrack->AddSubTrack(SubTrack);
+
+		// Look at current selection; if nothing selected or stale track then select new track
+		const TSharedPtr<RewindDebugger::FRewindDebuggerTrack> Selection = InstancesTreeView->GetSelection();
+		if (const FStateTreeDebuggerBaseTrack* DebuggerBaseTrack = static_cast<FStateTreeDebuggerBaseTrack*>(ExistingOwnerTrack->Get()))
+		{
+			if (DebuggerBaseTrack->IsStale())
+			{
+				InstancesTreeView->SetSelection(SubTrack);
+			}
+		}
+		else
+		{
+			InstancesTreeView->SetSelection(SubTrack);
+		}
 	}
 
 	InstancesTreeView->Refresh();
diff --git a/Engine/Plugins/Runtime/StateTree/Source/StateTreeEditorModule/Private/Debugger/SStateTreeDebuggerView.h b/Engine/Plugins/Runtime/StateTree/Source/StateTreeEditorModule/Private/Debugger/SStateTreeDebuggerView.h
--- a/Engine/Plugins/Runtime/StateTree/Source/StateTreeEditorModule/Private/Debugger/SStateTreeDebuggerView.h
+++ b/Engine/Plugins/Runtime/StateTree/Source/StateTreeEditorModule/Private/Debugger/SStateTreeDebuggerView.h
@@ -5,7 +5,6 @@
 
 #if WITH_STATETREE_DEBUGGER
 
-#include "Debugger/StateTreeTraceTypes.h"
 #include "InstancedStruct.h"
 #include "Widgets/Views/STreeView.h"
 #include "Widgets/SCompoundWidget.h"
@@ -19,14 +18,19 @@ namespace UE::StateTreeDebugger
 	struct FFrameSpan;
 	struct FEventTreeElement;
 }
+
+enum class EStateTreeBreakpointType : uint8;
+
 struct FStateTreeDebugger;
 struct FStateTreeInstanceDebugId;
+struct FStateTreeDebuggerBreakpoint;
 class FStateTreeEditor;
 class FStateTreeViewModel;
 class FUICommandList;
 class UStateTree;
 class SStateTreeDebuggerTimelines;
 class SStateTreeDebuggerInstanceTree;
+class UStateTreeEditorData;
 
 
 /**
@@ -42,7 +46,7 @@ public:
 	SStateTreeDebuggerView();
 	~SStateTreeDebuggerView();
 
-	void Construct(const FArguments& InArgs, const UStateTree* InStateTree, const TSharedRef<FStateTreeViewModel>& InStateTreeViewModel, const TSharedRef<FUICommandList>& InCommandList);
+	void Construct(const FArguments& InArgs, const UStateTree& InStateTree, const TSharedRef<FStateTreeViewModel>& InStateTreeViewModel, const TSharedRef<FUICommandList>& InCommandList);
 
 private:
 	virtual void Tick(const FGeometry& AllottedGeometry, const double InCurrentTime, const float InDeltaTime) override;
@@ -59,7 +63,7 @@ private:
 	void OnPIEResumed(bool bIsSimulating) const;
 	void OnPIESingleStepped(bool bIsSimulating) const;
 	
-	void OnBreakpointHit(const FStateTreeInstanceDebugId InstanceId, const FStateTreeStateHandle StateHandle, const TSharedRef<FUICommandList> ActionList) const;
+	void OnBreakpointHit(const FStateTreeInstanceDebugId InstanceId, const FStateTreeDebuggerBreakpoint Breakpoint, const TSharedRef<FUICommandList> ActionList) const;
 	void OnNewInstance(FStateTreeInstanceDebugId InstanceId);
 	void OnSelectedInstanceCleared();
 
@@ -84,13 +88,16 @@ private:
 
 	bool CanStepForwardToNextStateChange() const;
 	void StepForwardToNextStateChange();
-	
-	bool CanToggleBreakpoint() const;
-	void ToggleBreakpoint() const;
 
+	bool CanAddStateBreakpoint(EStateTreeBreakpointType Type) const;
+	bool CanRemoveStateBreakpoint(EStateTreeBreakpointType Type) const;
+	ECheckBoxState GetStateBreakpointCheckState(EStateTreeBreakpointType Type) const;
+	void HandleEnableStateBreakpoint(EStateTreeBreakpointType Type);
+	
 	TSharedPtr<FStateTreeDebugger> Debugger;
 	TSharedPtr<FStateTreeViewModel> StateTreeViewModel;
 	TWeakObjectPtr<const UStateTree> StateTree;
+	TWeakObjectPtr<UStateTreeEditorData> StateTreeEditorData;
 
 	/** Tracks for all statetree instance owners producing trace events for the associated state tree asset. */
 	TArray<TSharedPtr<RewindDebugger::FRewindDebuggerTrack>> InstanceOwnerTracks;
diff --git a/Engine/Plugins/Runtime/StateTree/Source/StateTreeEditorModule/Private/Debugger/StateTreeDebuggerCommands.cpp b/Engine/Plugins/Runtime/StateTree/Source/StateTreeEditorModule/Private/Debugger/StateTreeDebuggerCommands.cpp
--- a/Engine/Plugins/Runtime/StateTree/Source/StateTreeEditorModule/Private/Debugger/StateTreeDebuggerCommands.cpp
+++ b/Engine/Plugins/Runtime/StateTree/Source/StateTreeEditorModule/Private/Debugger/StateTreeDebuggerCommands.cpp
@@ -22,7 +22,8 @@ FStateTreeDebuggerCommands::FStateTreeDebuggerCommands()
 
 void FStateTreeDebuggerCommands::RegisterCommands()
 {
-	UI_COMMAND(ToggleBreakpoint, "Toggle Breakpoint", "Adds or removes a breakpoint on each selected node", EUserInterfaceActionType::Button, FInputChord(EKeys::F9))
+	UI_COMMAND(EnableOnEnterStateBreakpoint, "Break on Enter", "Adds or removes a breakpoint when entering the selected state(s)", EUserInterfaceActionType::Check, FInputChord(EKeys::F9));
+	UI_COMMAND(EnableOnExitStateBreakpoint, "Break on Exit", "Adds or removes a breakpoint when exiting the selected state(s)", EUserInterfaceActionType::Check, FInputChord(EModifierKey::Shift, EKeys::F9));
 
 	UI_COMMAND(StartRecording, "Start Recording", "Start a new trace session.", EUserInterfaceActionType::Button, FInputChord(EModifierKey::Control | EModifierKey::Shift, EKeys::R));
 	UI_COMMAND(StopRecording, "Stop Recording", "Stop the current trace session.", EUserInterfaceActionType::Button, FInputChord(EModifierKey::Control | EModifierKey::Shift, EKeys::S));
diff --git a/Engine/Plugins/Runtime/StateTree/Source/StateTreeEditorModule/Private/Debugger/StateTreeDebuggerCommands.h b/Engine/Plugins/Runtime/StateTree/Source/StateTreeEditorModule/Private/Debugger/StateTreeDebuggerCommands.h
--- a/Engine/Plugins/Runtime/StateTree/Source/StateTreeEditorModule/Private/Debugger/StateTreeDebuggerCommands.h
+++ b/Engine/Plugins/Runtime/StateTree/Source/StateTreeEditorModule/Private/Debugger/StateTreeDebuggerCommands.h
@@ -17,7 +17,8 @@ public:
 	// TCommands<> overrides
 	virtual void RegisterCommands() override;
 
-	TSharedPtr<FUICommandInfo> ToggleBreakpoint;
+	TSharedPtr<FUICommandInfo> EnableOnEnterStateBreakpoint;
+	TSharedPtr<FUICommandInfo> EnableOnExitStateBreakpoint;
 	TSharedPtr<FUICommandInfo> StartRecording;
 	TSharedPtr<FUICommandInfo> StopRecording;
 	TSharedPtr<FUICommandInfo> PreviousFrameWithStateChange;
diff --git a/Engine/Plugins/Runtime/StateTree/Source/StateTreeEditorModule/Private/Debugger/StateTreeDebuggerUIExtensions.cpp b/Engine/Plugins/Runtime/StateTree/Source/StateTreeEditorModule/Private/Debugger/StateTreeDebuggerUIExtensions.cpp
--- a/Engine/Plugins/Runtime/StateTree/Source/StateTreeEditorModule/Private/Debugger/StateTreeDebuggerUIExtensions.cpp
+++ b/Engine/Plugins/Runtime/StateTree/Source/StateTreeEditorModule/Private/Debugger/StateTreeDebuggerUIExtensions.cpp
@@ -2,10 +2,12 @@
 
 #include "StateTreeDebuggerUIExtensions.h"
 #include "Customizations/StateTreeEditorNodeUtils.h"
+#include "DetailLayoutBuilder.h"
 #include "Framework/MultiBox/MultiBoxBuilder.h"
 #include "PropertyHandle.h"
 #include "ScopedTransaction.h"
 #include "StateTreeConditionBase.h"
+#include "StateTreeEditorData.h"
 #include "StateTreeEditorNode.h"
 #include "StateTreeEditorStyle.h"
 #include "StateTreeTaskBase.h"
@@ -19,8 +21,138 @@
 
 namespace UE::StateTreeEditor::DebuggerExtensions
 {
-TSharedRef<SWidget> CreateStateWidget(TSharedPtr<IPropertyHandle> EnabledProperty)
+
+void OnConditionEvaluationModeChanged(TSharedPtr<IPropertyHandle> StructProperty, const EStateTreeConditionEvaluationMode Mode)
+{
+	EditorNodeUtils::ModifyNodeInTransaction(LOCTEXT("OnEvaluationModeChanged", "Condition Evaluation Mode Changed"),
+		StructProperty,
+		[Mode](IPropertyHandle& StructProperty)
+		{
+			if (FStateTreeEditorNode* Node = EditorNodeUtils::GetMutableCommonNode(StructProperty))
+			{
+				if (FStateTreeConditionBase* ConditionBase = Node->Node.GetMutablePtr<FStateTreeConditionBase>())
+				{
+					ConditionBase->EvaluationMode = Mode;
+				}
+			}
+		});
+}
+
+void OnTaskEnableToggled(TSharedPtr<IPropertyHandle> StructProperty)
+{
+	EditorNodeUtils::ModifyNodeInTransaction(LOCTEXT("OnTaskEnableToggled", "Toggled Task Enabled"),
+		StructProperty,
+		[](IPropertyHandle& StructProperty)
+		{
+			if (FStateTreeEditorNode* Node = EditorNodeUtils::GetMutableCommonNode(StructProperty))
+			{
+				if (FStateTreeTaskBase* TaskBase = Node->Node.GetMutablePtr<FStateTreeTaskBase>())
+				{
+					TaskBase->bTaskEnabled = !TaskBase->bTaskEnabled;
+				}
+			}
+		});
+}
+
+void OnStateEnableToggled(TSharedPtr<IPropertyHandle> EnabledProperty)
+{
+	FScopedTransaction Transaction(LOCTEXT("OnStateEnableToggled", "Toggled State Enabled"));
+
+	bool bEnabled = false;
+	const FPropertyAccess::Result Result = EnabledProperty->GetValue(bEnabled);
+
+	if (Result == FPropertyAccess::MultipleValues || bEnabled == false)
+	{
+		EnabledProperty->SetValue(true);
+	}
+	else
+	{
+		EnabledProperty->SetValue(false);
+	}
+}
+
+bool HasStateBreakpoint(const TArray<TWeakObjectPtr<>>& StatesBeingCustomized, const UStateTreeEditorData* EditorData, const EStateTreeBreakpointType Type)
+{
+#if WITH_STATETREE_DEBUGGER
+	if (EditorData == nullptr)
+	{
+		return false;
+	}
+	for (const TWeakObjectPtr<>& WeakStateObject : StatesBeingCustomized)
+	{
+		if (const UStateTreeState* State = static_cast<const UStateTreeState*>(WeakStateObject.Get()))
+		{
+			if (EditorData->HasBreakpoint(State->ID, Type))
+			{
+				return true;
+			}
+		}
+	}
+#endif // WITH_STATETREE_DEBUGGER
+	return false;
+}
+
+void OnStateBreakpointToggled(const TArray<TWeakObjectPtr<>>& States, UStateTreeEditorData* EditorData, const EStateTreeBreakpointType Type)
+{
+#if WITH_STATETREE_DEBUGGER
+	if (EditorData == nullptr || States.IsEmpty())
+	{
+		return;
+	}
+
+	const FScopedTransaction Transaction(LOCTEXT("ToggleStateBreakpoint", "Toggle State Breakpoint"));
+	EditorData->Modify();
+
+	for (const TWeakObjectPtr<>& WeakStateObject : States)
+	{
+		if (const UStateTreeState* State = static_cast<const UStateTreeState*>(WeakStateObject.Get()))
+		{
+			const bool bRemoved = EditorData->RemoveBreakpoint(State->ID, Type);
+			if (bRemoved == false)
+			{
+				EditorData->AddBreakpoint(State->ID, Type);
+			}
+		}
+	}
+#endif // WITH_STATETREE_DEBUGGER
+}
+
+bool HasTaskBreakpoint(TSharedPtr<IPropertyHandle> StructProperty, const UStateTreeEditorData* EditorData, const EStateTreeBreakpointType Type)
 {
+#if WITH_STATETREE_DEBUGGER
+	const FStateTreeEditorNode* Node = EditorNodeUtils::GetCommonNode(StructProperty);
+	if (EditorData != nullptr && Node != nullptr)
+	{
+		return EditorData->HasBreakpoint(Node->ID, Type);
+	}
+#endif // WITH_STATETREE_DEBUGGER
+	return false;
+}
+
+void OnTaskBreakpointToggled(TSharedPtr<IPropertyHandle> StructProperty, UStateTreeEditorData* EditorData, const EStateTreeBreakpointType Type)
+{
+#if WITH_STATETREE_DEBUGGER
+	const FStateTreeEditorNode* Node = EditorNodeUtils::GetMutableCommonNode(StructProperty);
+	if (EditorData != nullptr && Node != nullptr)
+	{
+		const FScopedTransaction Transaction(LOCTEXT("ToggleTaskBreakpoint", "Toggle Task Breakpoint"));
+		EditorData->Modify();
+		
+		const bool bRemoved = EditorData->RemoveBreakpoint(Node->ID, Type);
+		if (bRemoved == false)
+		{
+			EditorData->AddBreakpoint(Node->ID, Type);
+		}
+	}
+#endif // WITH_STATETREE_DEBUGGER
+}
+
+TSharedRef<SWidget> CreateStateWidget(IDetailLayoutBuilder& DetailBuilder, UStateTreeEditorData* TreeData)
+{
+	TWeakObjectPtr<UStateTreeEditorData> WeakTreeData = TreeData;
+	const TSharedPtr<IPropertyHandle> EnabledProperty = DetailBuilder.GetProperty(GET_MEMBER_NAME_CHECKED(UStateTreeState, bEnabled));
+	TArray<TWeakObjectPtr<UObject>> StatesBeingCustomized = DetailBuilder.GetSelectedObjects();
+	
 	FMenuBuilder DebuggingActions(/*CloseWindowAfterMenuSelection*/ true, /*CommandList*/nullptr);
 
 	DebuggingActions.BeginSection(FName("Options"), LOCTEXT("StateOptions", "State Debug Options"));
@@ -48,6 +180,41 @@ TSharedRef<SWidget> CreateStateWidget(TSharedPtr<IPropertyHandle> EnabledPropert
 		NAME_None,
 		EUserInterfaceActionType::Check
 	);
+	DebuggingActions.AddSeparator();
+	DebuggingActions.AddMenuEntry
+	(
+		LOCTEXT("ToggleOnEnterStateBreakpoint", "Break on Enter"),
+		LOCTEXT("ToggleOnEnterStateBreakpoint_ToolTip", "Enables or disables breakpoint when entering selected state(s)."),
+		FSlateIcon(),
+		FUIAction
+		(
+			FExecuteAction::CreateLambda([StatesBeingCustomized, WeakTreeData] { OnStateBreakpointToggled(StatesBeingCustomized, WeakTreeData.Get(), EStateTreeBreakpointType::OnEnter); }),
+			FCanExecuteAction(),
+			FGetActionCheckState::CreateLambda([StatesBeingCustomized, WeakTreeData]
+				{
+					return HasStateBreakpoint(StatesBeingCustomized, WeakTreeData.Get(), EStateTreeBreakpointType::OnEnter) ? ECheckBoxState::Checked : ECheckBoxState::Unchecked;
+				})
+		),
+		NAME_None,
+		EUserInterfaceActionType::Check
+	);
+	DebuggingActions.AddMenuEntry
+	(
+		LOCTEXT("ToggleOnExitStateBreakpoint", "Break on Exit"),
+		LOCTEXT("ToggleOnExitStateBreakpoint_ToolTip", "Enables or disables breakpoint when exiting selected state(s)."),
+		FSlateIcon(),
+		FUIAction
+		(
+			FExecuteAction::CreateLambda([StatesBeingCustomized, WeakTreeData] { OnStateBreakpointToggled(StatesBeingCustomized, WeakTreeData.Get(), EStateTreeBreakpointType::OnExit); }),
+			FCanExecuteAction(),
+			FGetActionCheckState::CreateLambda([StatesBeingCustomized, WeakTreeData]
+				{
+					return HasStateBreakpoint(StatesBeingCustomized, WeakTreeData.Get(), EStateTreeBreakpointType::OnExit) ? ECheckBoxState::Checked : ECheckBoxState::Unchecked;
+				})
+		),
+		NAME_None,
+		EUserInterfaceActionType::Check
+	);
 	DebuggingActions.EndSection();
 
 	const TSharedRef<SHorizontalBox> HeaderContentWidget = SNew(SHorizontalBox)
@@ -81,6 +248,53 @@ TSharedRef<SWidget> CreateStateWidget(TSharedPtr<IPropertyHandle> EnabledPropert
 					.ToolTipText(LOCTEXT("DisabledStateTooltip", "This state has been disabled."))
 				]
 			]
+			
+			// Break on enter Label
+			+ SHorizontalBox::Slot()
+			.AutoWidth()
+			.HAlign(HAlign_Left)
+			.VAlign(VAlign_Center)
+			.Padding(4.0f, 0.0f)
+			[
+				SNew(SBorder)
+				.Padding(FMargin(6, 2))
+				.BorderImage(FStateTreeEditorStyle::Get().GetBrush("StateTree.Node.Label"))
+				.Visibility_Lambda([StatesBeingCustomized, WeakTreeData]
+					{
+						return HasStateBreakpoint(StatesBeingCustomized, WeakTreeData.Get(), EStateTreeBreakpointType::OnEnter) ? EVisibility::Visible : EVisibility::Collapsed;
+					})
+				[
+					SNew(STextBlock)
+					.TextStyle(FStateTreeEditorStyle::Get(), "StateTree.Param.Label")
+					.ColorAndOpacity(FStyleColors::Foreground)
+					.Text(LOCTEXT("LabelOnEnterBreakpoint", "BREAK ON ENTER"))
+					.ToolTipText(LOCTEXT("BreakOnEnterTaskTooltip", "Break when entering this task."))
+				]
+			]
+
+			// Break on exit Label
+			+ SHorizontalBox::Slot()
+			.AutoWidth()
+			.HAlign(HAlign_Left)
+			.VAlign(VAlign_Center)
+			.Padding(4.0f, 0.0f)
+			[
+				SNew(SBorder)
+				.Padding(FMargin(6, 2))
+				.BorderImage(FStateTreeEditorStyle::Get().GetBrush("StateTree.Node.Label"))
+				.Visibility_Lambda([StatesBeingCustomized, WeakTreeData]
+					{
+						return HasStateBreakpoint(StatesBeingCustomized, WeakTreeData.Get(), EStateTreeBreakpointType::OnExit) ? EVisibility::Visible : EVisibility::Collapsed;
+					})
+				[
+					SNew(STextBlock)
+					.TextStyle(FStateTreeEditorStyle::Get(), "StateTree.Param.Label")
+					.ColorAndOpacity(FStyleColors::Foreground)
+					.Text(LOCTEXT("LabelOnExitBreakpoint", "BREAK ON EXIT"))
+					.ToolTipText(LOCTEXT("BreakOnExitTaskTooltip", "Break when exiting this task."))
+				]
+			]
+			
 		]
 		// Debug Options
 		+ SHorizontalBox::Slot()
@@ -93,7 +307,6 @@ TSharedRef<SWidget> CreateStateWidget(TSharedPtr<IPropertyHandle> EnabledPropert
 			.ToolTipText(LOCTEXT("DebugOptions", "Debug Options"))
 			.HasDownArrow(false)
 			.ContentPadding(0)
-			.ForegroundColor(FSlateColor::UseForeground())
 			.ButtonStyle(FAppStyle::Get(), "SimpleButton")
 			.MenuContent()
 			[
@@ -103,24 +316,27 @@ TSharedRef<SWidget> CreateStateWidget(TSharedPtr<IPropertyHandle> EnabledPropert
 			[
 				SNew(SImage)
 				.Image(FStateTreeEditorStyle::Get().GetBrush("StateTreeEditor.DebugOptions"))
+				.ColorAndOpacity(FSlateColor::UseForeground())
 			]
 		];
 
 	return HeaderContentWidget;
 }
 
-TSharedRef<SWidget> CreateEditorNodeWidget(TSharedPtr<IPropertyHandle> InPropertyHandle)
+TSharedRef<SWidget> CreateEditorNodeWidget(const TSharedPtr<IPropertyHandle>& StructPropertyHandle, UStateTreeEditorData* TreeData)
 {
 	using namespace EditorNodeUtils;
+
+	TWeakObjectPtr<UStateTreeEditorData> WeakTreeData = TreeData;
 	
 	bool bIsTask = false;
 	bool bIsCondition = false;
-	if (const FStructProperty* TmpStructProperty = CastField<FStructProperty>(InPropertyHandle->GetProperty()))
+	if (const FStructProperty* TmpStructProperty = CastField<FStructProperty>(StructPropertyHandle->GetProperty()))
 	{
 		if (TmpStructProperty->Struct->IsChildOf(TBaseStructure<FStateTreeEditorNode>::Get()))
 		{
 			const UScriptStruct* ScriptStruct = nullptr;
-			if (const FStateTreeEditorNode* Node = GetCommonNode(InPropertyHandle))
+			if (const FStateTreeEditorNode* Node = GetCommonNode(StructPropertyHandle))
 			{
 				ScriptStruct = Node->Node.GetScriptStruct();
 			}
@@ -151,11 +367,46 @@ TSharedRef<SWidget> CreateEditorNodeWidget(TSharedPtr<IPropertyHandle> InPropert
 			FSlateIcon(),
 			FUIAction
 			(
-				FExecuteAction::CreateLambda([InPropertyHandle] { OnTaskEnableToggled(InPropertyHandle); }),
+				FExecuteAction::CreateLambda([StructPropertyHandle] { OnTaskEnableToggled(StructPropertyHandle); }),
+				FCanExecuteAction(),
+				FGetActionCheckState::CreateLambda([StructPropertyHandle]
+					{
+						return IsTaskDisabled(StructPropertyHandle) ? ECheckBoxState::Unchecked : ECheckBoxState::Checked;
+					})
+			),
+			NAME_None,
+			EUserInterfaceActionType::Check
+		);
+		DebuggingActions.AddSeparator();
+		DebuggingActions.AddMenuEntry
+		(
+			LOCTEXT("ToggleOnEnterTaskBreakpoint", "Break on Enter"),
+			LOCTEXT("ToggleOnEnterTaskBreakpoint_ToolTip", "Enables or disables breakpoint when entering selected task(s)."),
+			FSlateIcon(),
+			FUIAction
+			(
+				FExecuteAction::CreateLambda([StructPropertyHandle, WeakTreeData] { OnTaskBreakpointToggled(StructPropertyHandle, WeakTreeData.Get(), EStateTreeBreakpointType::OnEnter); }),
 				FCanExecuteAction(),
-				FGetActionCheckState::CreateLambda([InPropertyHandle]
+				FGetActionCheckState::CreateLambda([StructPropertyHandle, WeakTreeData]
 					{
-						return IsTaskDisabled(InPropertyHandle) ? ECheckBoxState::Unchecked : ECheckBoxState::Checked;
+						return HasTaskBreakpoint(StructPropertyHandle, WeakTreeData.Get(), EStateTreeBreakpointType::OnEnter) ? ECheckBoxState::Checked : ECheckBoxState::Unchecked;
+					})
+			),
+			NAME_None,
+			EUserInterfaceActionType::Check
+		);
+		DebuggingActions.AddMenuEntry
+		(
+			LOCTEXT("ToggleOnExitTaskBreakpoint", "Break on Exit"),
+			LOCTEXT("ToggleOnExitTaskBreakpoint_ToolTip", "Enables or disables breakpoint when exiting selected task(s)."),
+			FSlateIcon(),
+			FUIAction
+			(
+				FExecuteAction::CreateLambda([StructPropertyHandle, WeakTreeData] { OnTaskBreakpointToggled(StructPropertyHandle, WeakTreeData.Get(), EStateTreeBreakpointType::OnExit); }),
+				FCanExecuteAction(),
+				FGetActionCheckState::CreateLambda([StructPropertyHandle, WeakTreeData]
+					{
+						return HasTaskBreakpoint(StructPropertyHandle, WeakTreeData.Get(), EStateTreeBreakpointType::OnExit) ? ECheckBoxState::Checked : ECheckBoxState::Unchecked;
 					})
 			),
 			NAME_None,
@@ -174,14 +425,14 @@ TSharedRef<SWidget> CreateEditorNodeWidget(TSharedPtr<IPropertyHandle> InPropert
 			FSlateIcon(),
 			FUIAction
 			(
-				FExecuteAction::CreateLambda([InPropertyHandle]
+				FExecuteAction::CreateLambda([StructPropertyHandle]
 					{
-						OnConditionEvaluationModeChanged(InPropertyHandle, EStateTreeConditionEvaluationMode::Evaluated);
+						OnConditionEvaluationModeChanged(StructPropertyHandle, EStateTreeConditionEvaluationMode::Evaluated);
 					}),
 				FCanExecuteAction(),
-				FIsActionChecked::CreateLambda([InPropertyHandle]
+				FIsActionChecked::CreateLambda([StructPropertyHandle]
 					{
-						return GetConditionEvaluationMode(InPropertyHandle) == EStateTreeConditionEvaluationMode::Evaluated;
+						return GetConditionEvaluationMode(StructPropertyHandle) == EStateTreeConditionEvaluationMode::Evaluated;
 					})
 			),
 			NAME_None,
@@ -194,14 +445,14 @@ TSharedRef<SWidget> CreateEditorNodeWidget(TSharedPtr<IPropertyHandle> InPropert
 			FSlateIcon(),
 			FUIAction
 			(
-				FExecuteAction::CreateLambda([InPropertyHandle]
+				FExecuteAction::CreateLambda([StructPropertyHandle]
 					{
-						OnConditionEvaluationModeChanged(InPropertyHandle, EStateTreeConditionEvaluationMode::ForcedTrue);
+						OnConditionEvaluationModeChanged(StructPropertyHandle, EStateTreeConditionEvaluationMode::ForcedTrue);
 					}),
 				FCanExecuteAction(),
-				FIsActionChecked::CreateLambda([InPropertyHandle]
+				FIsActionChecked::CreateLambda([StructPropertyHandle]
 					{
-						return GetConditionEvaluationMode(InPropertyHandle) == EStateTreeConditionEvaluationMode::ForcedTrue;
+						return GetConditionEvaluationMode(StructPropertyHandle) == EStateTreeConditionEvaluationMode::ForcedTrue;
 					})
 			),
 			NAME_None,
@@ -214,14 +465,14 @@ TSharedRef<SWidget> CreateEditorNodeWidget(TSharedPtr<IPropertyHandle> InPropert
 			FSlateIcon(),
 			FUIAction
 			(
-				FExecuteAction::CreateLambda([InPropertyHandle]
+				FExecuteAction::CreateLambda([StructPropertyHandle]
 					{
-						OnConditionEvaluationModeChanged(InPropertyHandle, EStateTreeConditionEvaluationMode::ForcedFalse);
+						OnConditionEvaluationModeChanged(StructPropertyHandle, EStateTreeConditionEvaluationMode::ForcedFalse);
 					}),
 				FCanExecuteAction(),
-				FIsActionChecked::CreateLambda([InPropertyHandle]
+				FIsActionChecked::CreateLambda([StructPropertyHandle]
 					{
-						return GetConditionEvaluationMode(InPropertyHandle) == EStateTreeConditionEvaluationMode::ForcedFalse;
+						return GetConditionEvaluationMode(StructPropertyHandle) == EStateTreeConditionEvaluationMode::ForcedFalse;
 					})
 			),
 			NAME_None,
@@ -250,9 +501,9 @@ TSharedRef<SWidget> CreateEditorNodeWidget(TSharedPtr<IPropertyHandle> InPropert
 				SNew(SBorder)
 				.Padding(FMargin(6, 2))
 				.BorderImage(FStateTreeEditorStyle::Get().GetBrush("StateTree.Node.Label"))
-				.Visibility_Lambda([InPropertyHandle]
+				.Visibility_Lambda([StructPropertyHandle]
 					{
-						return IsTaskDisabled(InPropertyHandle) ? EVisibility::Visible : EVisibility::Collapsed;
+						return IsTaskDisabled(StructPropertyHandle) ? EVisibility::Visible : EVisibility::Collapsed;
 					})
 				[
 					SNew(STextBlock)
@@ -263,6 +514,52 @@ TSharedRef<SWidget> CreateEditorNodeWidget(TSharedPtr<IPropertyHandle> InPropert
 				]
 			]
 
+			// Break on enter Label
+			+ SHorizontalBox::Slot()
+			.AutoWidth()
+			.HAlign(HAlign_Left)
+			.VAlign(VAlign_Center)
+			.Padding(4.0f, 0.0f)
+			[
+				SNew(SBorder)
+				.Padding(FMargin(6, 2))
+				.BorderImage(FStateTreeEditorStyle::Get().GetBrush("StateTree.Node.Label"))
+				.Visibility_Lambda([StructPropertyHandle, WeakTreeData]
+					{
+						return HasTaskBreakpoint(StructPropertyHandle, WeakTreeData.Get(), EStateTreeBreakpointType::OnEnter) ? EVisibility::Visible : EVisibility::Collapsed;
+					})
+				[
+					SNew(STextBlock)
+					.TextStyle(FStateTreeEditorStyle::Get(), "StateTree.Param.Label")
+					.ColorAndOpacity(FStyleColors::Foreground)
+					.Text(LOCTEXT("LabelOnEnterBreakpoint", "BREAK ON ENTER"))
+					.ToolTipText(LOCTEXT("BreakOnEnterTaskTooltip", "Break when entering this task."))
+				]
+			]
+
+			// Break on exit Label
+			+ SHorizontalBox::Slot()
+			.AutoWidth()
+			.HAlign(HAlign_Left)
+			.VAlign(VAlign_Center)
+			.Padding(4.0f, 0.0f)
+			[
+				SNew(SBorder)
+				.Padding(FMargin(6, 2))
+				.BorderImage(FStateTreeEditorStyle::Get().GetBrush("StateTree.Node.Label"))
+				.Visibility_Lambda([StructPropertyHandle, WeakTreeData]
+					{
+						return HasTaskBreakpoint(StructPropertyHandle, WeakTreeData.Get(), EStateTreeBreakpointType::OnExit) ? EVisibility::Visible : EVisibility::Collapsed;
+					})
+				[
+					SNew(STextBlock)
+					.TextStyle(FStateTreeEditorStyle::Get(), "StateTree.Param.Label")
+					.ColorAndOpacity(FStyleColors::Foreground)
+					.Text(LOCTEXT("LabelOnExitBreakpoint", "BREAK ON EXIT"))
+					.ToolTipText(LOCTEXT("BreakOnExitTaskTooltip", "Break when exiting this task."))
+				]
+			]
+
 			// Force True / Force False Label
 			+ SHorizontalBox::Slot()
 			.AutoWidth()
@@ -273,15 +570,15 @@ TSharedRef<SWidget> CreateEditorNodeWidget(TSharedPtr<IPropertyHandle> InPropert
 				SNew(SBorder)
 				.Padding(FMargin(6, 2))
 				.BorderImage(FStateTreeEditorStyle::Get().GetBrush("StateTree.Node.Label"))
-				.BorderBackgroundColor_Lambda([InPropertyHandle]
+				.BorderBackgroundColor_Lambda([StructPropertyHandle]
 					{
-						return GetConditionEvaluationMode(InPropertyHandle) == EStateTreeConditionEvaluationMode::ForcedTrue
+						return GetConditionEvaluationMode(StructPropertyHandle) == EStateTreeConditionEvaluationMode::ForcedTrue
 								   ? FStyleColors::AccentGreen.GetSpecifiedColor()
 								   : FStyleColors::AccentRed.GetSpecifiedColor();
 					})
-				.Visibility_Lambda([InPropertyHandle]
+				.Visibility_Lambda([StructPropertyHandle]
 					{
-						return GetConditionEvaluationMode(InPropertyHandle) != EStateTreeConditionEvaluationMode::Evaluated ? EVisibility::Visible : EVisibility::Collapsed;
+						return GetConditionEvaluationMode(StructPropertyHandle) != EStateTreeConditionEvaluationMode::Evaluated ? EVisibility::Visible : EVisibility::Collapsed;
 					})
 				[
 					SNew(SOverlay)
@@ -291,9 +588,9 @@ TSharedRef<SWidget> CreateEditorNodeWidget(TSharedPtr<IPropertyHandle> InPropert
 						.TextStyle(FStateTreeEditorStyle::Get(), "StateTree.Param.Label")
 						.Text(LOCTEXT("ConditionForcedTrue", "TRUE"))
 						.ToolTipText(LOCTEXT("ForcedTrueConditionTooltip", "This condition is not evaluated and result forced to 'true'."))
-						.Visibility_Lambda([InPropertyHandle]
+						.Visibility_Lambda([StructPropertyHandle]
 							{
-								return GetConditionEvaluationMode(InPropertyHandle) == EStateTreeConditionEvaluationMode::ForcedTrue ? EVisibility::Visible : EVisibility::Collapsed;
+								return GetConditionEvaluationMode(StructPropertyHandle) == EStateTreeConditionEvaluationMode::ForcedTrue ? EVisibility::Visible : EVisibility::Collapsed;
 							})
 					]
 					+ SOverlay::Slot()
@@ -302,16 +599,16 @@ TSharedRef<SWidget> CreateEditorNodeWidget(TSharedPtr<IPropertyHandle> InPropert
 						.TextStyle(FStateTreeEditorStyle::Get(), "StateTree.Param.Label")
 						.Text(LOCTEXT("ConditionForcedFalse", "FALSE"))
 						.ToolTipText(LOCTEXT("ForcedFalseConditionTooltip", "This condition is not evaluated and result forced to 'false'."))
-						.Visibility_Lambda([InPropertyHandle]
+						.Visibility_Lambda([StructPropertyHandle]
 							{
-								return GetConditionEvaluationMode(InPropertyHandle) == EStateTreeConditionEvaluationMode::ForcedFalse ? EVisibility::Visible : EVisibility::Collapsed;
+								return GetConditionEvaluationMode(StructPropertyHandle) == EStateTreeConditionEvaluationMode::ForcedFalse ? EVisibility::Visible : EVisibility::Collapsed;
 							})
 					]
 				]
 			]
 		]
 		+ SHorizontalBox::Slot()
-		.Padding(0)
+		.Padding(0,0,6,0)
 		.HAlign(HAlign_Right)
 		.VAlign(VAlign_Center)
 		.AutoWidth()
@@ -320,7 +617,6 @@ TSharedRef<SWidget> CreateEditorNodeWidget(TSharedPtr<IPropertyHandle> InPropert
 			.ToolTipText(LOCTEXT("DebugOptions", "Debug Options"))
 			.HasDownArrow(false)
 			.ContentPadding(0)
-			.ForegroundColor(FSlateColor::UseForeground())
 			.ButtonStyle(FAppStyle::Get(), "SimpleButton")
 			.MenuContent()
 			[
@@ -330,59 +626,11 @@ TSharedRef<SWidget> CreateEditorNodeWidget(TSharedPtr<IPropertyHandle> InPropert
 			[
 				SNew(SImage)
 				.Image(FStateTreeEditorStyle::Get().GetBrush("StateTreeEditor.DebugOptions"))
+				.ColorAndOpacity(FSlateColor::UseForeground())
 			]
 		];
 }
 
-void OnConditionEvaluationModeChanged(TSharedPtr<IPropertyHandle> StructProperty, const EStateTreeConditionEvaluationMode Mode)
-{
-	EditorNodeUtils::ModifyNodeInTransaction(LOCTEXT("OnEvaluationModeChanged", "Condition Evaluation Mode Changed"),
-		StructProperty,
-		[Mode](IPropertyHandle& StructProperty)
-		{
-			if (FStateTreeEditorNode* Node = EditorNodeUtils::GetMutableCommonNode(StructProperty))
-			{
-				if (FStateTreeConditionBase* ConditionBase = Node->Node.GetMutablePtr<FStateTreeConditionBase>())
-				{
-					ConditionBase->EvaluationMode = Mode;
-				}
-			}
-		});
-}
-
-void OnTaskEnableToggled(TSharedPtr<IPropertyHandle> StructProperty)
-{
-	EditorNodeUtils::ModifyNodeInTransaction(LOCTEXT("OnTaskEnableToggled", "Toggled Task Enabled"),
-		StructProperty,
-		[](IPropertyHandle& StructProperty)
-		{
-			if (FStateTreeEditorNode* Node = EditorNodeUtils::GetMutableCommonNode(StructProperty))
-			{
-				if (FStateTreeTaskBase* TaskBase = Node->Node.GetMutablePtr<FStateTreeTaskBase>())
-				{
-					TaskBase->bTaskEnabled = !TaskBase->bTaskEnabled;
-				}
-			}
-		});
-}
-
-void OnStateEnableToggled(TSharedPtr<IPropertyHandle> EnabledProperty)
-{
-	FScopedTransaction Transaction(LOCTEXT("OnStateEnableToggled", "Toggled State Enabled"));
-
-	bool bEnabled = false;
-	const FPropertyAccess::Result Result = EnabledProperty->GetValue(bEnabled);
-
-	if (Result == FPropertyAccess::MultipleValues || bEnabled == false)
-	{
-		EnabledProperty->SetValue(true);
-	}
-	else
-	{
-		EnabledProperty->SetValue(false);
-	}
-}
-
 } // namespace UE::StateTreeEditor::DebuggerExtensions
 
 #undef LOCTEXT_NAMESPACE
diff --git a/Engine/Plugins/Runtime/StateTree/Source/StateTreeEditorModule/Private/Debugger/StateTreeDebuggerUIExtensions.h b/Engine/Plugins/Runtime/StateTree/Source/StateTreeEditorModule/Private/Debugger/StateTreeDebuggerUIExtensions.h
--- a/Engine/Plugins/Runtime/StateTree/Source/StateTreeEditorModule/Private/Debugger/StateTreeDebuggerUIExtensions.h
+++ b/Engine/Plugins/Runtime/StateTree/Source/StateTreeEditorModule/Private/Debugger/StateTreeDebuggerUIExtensions.h
@@ -12,15 +12,12 @@ class IPropertyHandle;
 class IDetailLayoutBuilder;
 class FDetailWidgetRow;
 class FStateTreeViewModel;
+class UStateTreeEditorData;
 
 namespace UE::StateTreeEditor::DebuggerExtensions
 {
 
-TSharedRef<SWidget> CreateStateWidget(TSharedPtr<IPropertyHandle> EnabledProperty);
-TSharedRef<SWidget> CreateEditorNodeWidget(TSharedPtr<IPropertyHandle> Shared);
-
-void OnConditionEvaluationModeChanged(TSharedPtr<IPropertyHandle> StructProperty, EStateTreeConditionEvaluationMode Mode);
-void OnStateEnableToggled(TSharedPtr<IPropertyHandle> StructProperty);
-void OnTaskEnableToggled(TSharedPtr<IPropertyHandle> StructProperty);
+TSharedRef<SWidget> CreateStateWidget(IDetailLayoutBuilder& DetailBuilder, UStateTreeEditorData* TreeData);
+TSharedRef<SWidget> CreateEditorNodeWidget(const TSharedPtr<IPropertyHandle>& StructPropertyHandle, UStateTreeEditorData* TreeData);
 
 }; // UE::StateTreeEditor::DebuggerExtensions
diff --git a/Engine/Plugins/Runtime/StateTree/Source/StateTreeEditorModule/Private/SStateTreeView.cpp b/Engine/Plugins/Runtime/StateTree/Source/StateTreeEditorModule/Private/SStateTreeView.cpp
--- a/Engine/Plugins/Runtime/StateTree/Source/StateTreeEditorModule/Private/SStateTreeView.cpp
+++ b/Engine/Plugins/Runtime/StateTree/Source/StateTreeEditorModule/Private/SStateTreeView.cpp
@@ -477,7 +477,9 @@ TSharedPtr<SWidget> SStateTreeView::HandleContextMenuOpening()
 #if WITH_STATETREE_DEBUGGER
 	if (UStateTreeSettings::Get().bUseDebugger)
 	{
-		MenuBuilder.AddMenuEntry(FStateTreeDebuggerCommands::Get().ToggleBreakpoint);
+		MenuBuilder.AddSeparator();
+		MenuBuilder.AddMenuEntry(FStateTreeDebuggerCommands::Get().EnableOnEnterStateBreakpoint);
+		MenuBuilder.AddMenuEntry(FStateTreeDebuggerCommands::Get().EnableOnExitStateBreakpoint);
 	}
 #endif // WITH_STATETREE_DEBUGGER
 	
diff --git a/Engine/Plugins/Runtime/StateTree/Source/StateTreeEditorModule/Private/SStateTreeView.h b/Engine/Plugins/Runtime/StateTree/Source/StateTreeEditorModule/Private/SStateTreeView.h
--- a/Engine/Plugins/Runtime/StateTree/Source/StateTreeEditorModule/Private/SStateTreeView.h
+++ b/Engine/Plugins/Runtime/StateTree/Source/StateTreeEditorModule/Private/SStateTreeView.h
@@ -11,7 +11,6 @@ class SScrollBar;
 class STableViewBase;
 namespace ESelectInfo { enum Type : int; }
 struct FPropertyChangedEvent;
-class UStateTreeEditorData;
 class UStateTreeState;
 class SScrollBox;
 class FUICommandList;
diff --git a/Engine/Plugins/Runtime/StateTree/Source/StateTreeEditorModule/Private/SStateTreeViewRow.cpp b/Engine/Plugins/Runtime/StateTree/Source/StateTreeEditorModule/Private/SStateTreeViewRow.cpp
--- a/Engine/Plugins/Runtime/StateTree/Source/StateTreeEditorModule/Private/SStateTreeViewRow.cpp
+++ b/Engine/Plugins/Runtime/StateTree/Source/StateTreeEditorModule/Private/SStateTreeViewRow.cpp
@@ -5,10 +5,12 @@
 
 
 #include "EditorFontGlyphs.h"
+#include "StateTreeEditorData.h"
 #include "StateTreeEditorStyle.h"
 
 #include "Widgets/Text/SInlineEditableTextBlock.h"
 #include "Widgets/Layout/SScrollBox.h"
+#include "Widgets/Layout/SSpacer.h"
 
 #include "StateTreeState.h"
 #include "StateTreeTaskBase.h"
@@ -21,6 +23,8 @@ void SStateTreeViewRow::Construct(const FArguments& InArgs, const TSharedRef<STa
 {
 	StateTreeViewModel = InStateTreeViewModel;
 	WeakState = InState;
+	const UStateTreeState* State = InState.Get();
+	WeakTreeData = State != nullptr ? State->GetTypedOuter<UStateTreeEditorData>() : nullptr;
 
 	ConstructInternal(STableRow::FArguments()
 		.Padding(5.0f)
@@ -111,71 +115,79 @@ void SStateTreeViewRow::Construct(const FArguments& InArgs, const TSharedRef<STa
 										return State != nullptr && State->bEnabled;
 									})
 								[
-									SNew(SHorizontalBox)
-
-									// Breakpoint box
-									+ SHorizontalBox::Slot()
-									.VAlign(VAlign_Top)
-									.HAlign(HAlign_Left)
-									.AutoWidth()
+									SNew(SOverlay)
+									+ SOverlay::Slot()
 									[
-										SNew(SBox)
-										.Padding(FMargin(0.0f, 0.0f, 4.0f, 0.0f))
+										SNew(SHorizontalBox)
+										// Conditions icon
+										+SHorizontalBox::Slot()
+										.VAlign(VAlign_Center)
+										.AutoWidth()
 										[
-											SNew(SImage)
-											.DesiredSizeOverride(FVector2D(12.f, 12.f))
-											.Image(FAppStyle::GetBrush(TEXT("Kismet.DebuggerOverlay.Breakpoint.EnabledAndValid")))
-											.Visibility(this, &SStateTreeViewRow::GetBreakpointVisibility)
+											SNew(SBox)
+											.Padding(FMargin(0.0f, 0.0f, 4.0f, 0.0f))
+											.Visibility(this, &SStateTreeViewRow::GetConditionVisibility)
+											[
+												SNew(SImage)
+												.ColorAndOpacity(IconTint)
+												.Image(FStateTreeEditorStyle::Get().GetBrush("StateTreeEditor.Conditions"))
+												.ToolTipText(LOCTEXT("StateHasEnterConditions", "State selection is guarded with enter conditions."))
+											]
 										]
-									]
 
-									// Conditions icon
-									+ SHorizontalBox::Slot()
-									.VAlign(VAlign_Center)
-									.AutoWidth()
-									[
-										SNew(SBox)
-										.Padding(FMargin(0.0f, 0.0f, 4.0f, 0.0f))
-										.Visibility(this, &SStateTreeViewRow::GetConditionVisibility)
+										// Selector icon
+										+ SHorizontalBox::Slot()
+										.VAlign(VAlign_Center)
+										.AutoWidth()
 										[
-											SNew(SImage)
-											.ColorAndOpacity(IconTint)
-											.Image(FStateTreeEditorStyle::Get().GetBrush("StateTreeEditor.Conditions"))
-											.ToolTipText(LOCTEXT("StateHasEnterConditions", "State selection is guarded with enter conditions."))
+											SNew(SBox)
+											.Padding(FMargin(0.0f, 0.0f, 4.0f, 0.0f))
+											[
+												SNew(SImage)
+												.Image(this, &SStateTreeViewRow::GetSelectorIcon)
+												.ColorAndOpacity(IconTint)
+												.ToolTipText(this, &SStateTreeViewRow::GetSelectorTooltip)
+											]
 										]
-									]
 
-									// Selector icon
-									+ SHorizontalBox::Slot()
-									.VAlign(VAlign_Center)
-									.AutoWidth()
-									[
-										SNew(SBox)
-										.Padding(FMargin(0.0f, 0.0f, 4.0f, 0.0f))
+										// State Name
+										+ SHorizontalBox::Slot()
+										.VAlign(VAlign_Center)
+										.AutoWidth()
 										[
-											SNew(SImage)
-											.Image(this, &SStateTreeViewRow::GetSelectorIcon)
-											.ColorAndOpacity(IconTint)
-											.ToolTipText( this, &SStateTreeViewRow::GetSelectorTooltip)
+											SAssignNew(NameTextBlock, SInlineEditableTextBlock)
+											.Style(FStateTreeEditorStyle::Get(), "StateTree.State.TitleInlineEditableText")
+											.OnVerifyTextChanged_Lambda([](const FText& NewLabel, FText& OutErrorMessage)
+												{
+													return !NewLabel.IsEmptyOrWhitespace();
+												})
+											.OnTextCommitted(this, &SStateTreeViewRow::HandleNodeLabelTextCommitted)
+											.Text(this, &SStateTreeViewRow::GetStateDesc)
+											.ToolTipText(this, &SStateTreeViewRow::GetStateTypeTooltip)
+											.Clipping(EWidgetClipping::ClipToBounds)
+											.IsSelected(this, &SStateTreeViewRow::IsStateSelected)
 										]
 									]
-
-									// State Name
-									+ SHorizontalBox::Slot()
-									.VAlign(VAlign_Center)
-									.AutoWidth()
+									+ SOverlay::Slot()
 									[
-										SAssignNew(NameTextBlock, SInlineEditableTextBlock)
-										.Style(FStateTreeEditorStyle::Get(), "StateTree.State.TitleInlineEditableText")
-										.OnVerifyTextChanged_Lambda([](const FText& NewLabel, FText& OutErrorMessage)
-										{
-											return !NewLabel.IsEmptyOrWhitespace();
-										})
-										.OnTextCommitted(this, &SStateTreeViewRow::HandleNodeLabelTextCommitted)
-										.Text(this, &SStateTreeViewRow::GetStateDesc)
-										.ToolTipText(this, &SStateTreeViewRow::GetStateTypeTooltip)
-										.Clipping(EWidgetClipping::ClipToBounds)
-										.IsSelected(this, &SStateTreeViewRow::IsStateSelected)
+										SNew(SHorizontalBox)
+
+										// State breakpoint box
+										+ SHorizontalBox::Slot()
+										.VAlign(VAlign_Top)
+										.HAlign(HAlign_Left)
+										.AutoWidth()
+										[
+											SNew(SBox)
+											.Padding(FMargin(-12.0f, -6.0f, 0.0f, 0.0f))
+											[
+												SNew(SImage)
+												.DesiredSizeOverride(FVector2D(12.f, 12.f))
+												.Image(FStateTreeEditorStyle::Get().GetBrush(TEXT("StateTreeEditor.Debugger.Breakpoint.EnabledAndValid")))
+												.Visibility(this, &SStateTreeViewRow::GetStateBreakpointVisibility)
+												.ToolTipText(this, &SStateTreeViewRow::GetStateBreakpointTooltipText)
+											]
+										]
 									]
 								]
 							]
@@ -392,10 +404,13 @@ TSharedRef<SHorizontalBox> SStateTreeViewRow::CreateTasksWidget()
 		return TasksBox;
 	}
 
+	TWeakObjectPtr<UStateTreeEditorData> WeakEditorData = State->GetTypedOuter<UStateTreeEditorData>();
+
 	for (int32 TaskIndex = 0; TaskIndex < State->Tasks.Num(); TaskIndex++)
 	{
 		if (const FStateTreeTaskBase* Task = State->Tasks[TaskIndex].Node.GetPtr<FStateTreeTaskBase>())
 		{
+			FGuid TaskId = State->Tasks[TaskIndex].ID;
 			auto IsTaskEnabledFunc = [WeakState=WeakState, TaskIndex]
 				{
 					const UStateTreeState* State = WeakState.Get();
@@ -405,10 +420,53 @@ TSharedRef<SHorizontalBox> SStateTreeViewRow::CreateTasksWidget()
 						{
 							return (State->bEnabled && Task->bTaskEnabled);
 						}
-					}					
+					}
 					return true;
 				};
+
+			auto IsTaskBreakpointEnabledFunc = [WeakEditorData, TaskId]
+				{
+#if WITH_STATETREE_DEBUGGER
+					const UStateTreeEditorData* EditorData = WeakEditorData.Get();
+					if (EditorData != nullptr && EditorData->HasAnyBreakpoint(TaskId))
+					{
+						return EVisibility::Visible;
+					}
+#endif // WITH_STATETREE_DEBUGGER
+					return EVisibility::Hidden;
+				};
 			
+			auto GetTaskBreakpointTooltipFunc = [WeakEditorData, TaskId]
+				{
+#if WITH_STATETREE_DEBUGGER
+					if (const UStateTreeEditorData* EditorData = WeakEditorData.Get())
+					{
+						const bool bHasBreakpointOnEnter = EditorData->HasBreakpoint(TaskId, EStateTreeBreakpointType::OnEnter);
+						const bool bHasBreakpointOnExit = EditorData->HasBreakpoint(TaskId, EStateTreeBreakpointType::OnExit);
+						if (bHasBreakpointOnEnter && bHasBreakpointOnExit)
+						{
+							return LOCTEXT("StateTreeTaskBreakpointOnEnterAndOnExitTooltip","Break when entering or exiting task");
+						}
+
+						if (bHasBreakpointOnEnter)
+						{
+							return LOCTEXT("StateTreeTaskBreakpointOnEnterTooltip","Break when entering task");
+						}
+
+						if (bHasBreakpointOnExit)
+						{
+							return LOCTEXT("StateTreeTaskBreakpointOnExitTooltip","Break when exiting task");
+						}
+					}
+#endif // WITH_STATETREE_DEBUGGER
+					return FText::GetEmpty();
+				};
+
+			FText TaskName = FText::FromName(Task->Name);
+
+			// Uncomment to visualize Task Ids
+			// TaskName = FText::FromString(FString::Printf(TEXT("%s (%s)"), *Task->Name.ToString(), *LexToString(TaskId)));
+
 			TasksBox->AddSlot()
 				.AutoWidth()
 				.VAlign(VAlign_Fill)
@@ -416,14 +474,39 @@ TSharedRef<SHorizontalBox> SStateTreeViewRow::CreateTasksWidget()
 					SNew(SBorder)
 					.VAlign(VAlign_Center)
 					.BorderImage(FStateTreeEditorStyle::Get().GetBrush("StateTree.Task.Rect"))
-					.IsEnabled_Lambda(IsTaskEnabledFunc)					
+					.Padding(0)
+					.IsEnabled_Lambda(IsTaskEnabledFunc)
 					[
-						SNew(STextBlock)
-						.Margin(FMargin(4.f, 0.f))
-						.Text(FText::FromName(Task->Name))
-						.TextStyle(FStateTreeEditorStyle::Get(), "StateTree.Task.Title")
-						.IsEnabled_Lambda(IsTaskEnabledFunc)						
-						.ToolTipText(FText::FromName(Task->Name))
+						SNew(SOverlay)
+						+ SOverlay::Slot()
+						[
+							SNew(STextBlock)
+							.Margin(FMargin(4.f, 0.f))
+							.Text(TaskName)
+							.TextStyle(FStateTreeEditorStyle::Get(), "StateTree.Task.Title")
+							.IsEnabled_Lambda(IsTaskEnabledFunc)
+							.ToolTipText(FText::FromName(Task->Name))
+						]
+						+ SOverlay::Slot()
+						[
+							// Task Breakpoint box
+							SNew(SHorizontalBox)
+							+ SHorizontalBox::Slot()
+							.VAlign(VAlign_Top)
+							.HAlign(HAlign_Left)
+							.AutoWidth()
+							[
+								SNew(SBox)
+								.Padding(FMargin(0.0f, -10.0f, 0.0f, 0.0f))
+								[
+									SNew(SImage)
+									.DesiredSizeOverride(FVector2D(10.f, 10.f))
+									.Image(FStateTreeEditorStyle::Get().GetBrush(TEXT("StateTreeEditor.Debugger.Breakpoint.EnabledAndValid")))
+									.Visibility_Lambda(IsTaskBreakpointEnabledFunc)
+									.ToolTipText_Lambda(GetTaskBreakpointTooltipFunc)
+								]
+							]
+						]
 					]
 				];
 		}
@@ -502,17 +585,46 @@ EVisibility SStateTreeViewRow::GetConditionVisibility() const
 	return EVisibility::Collapsed;
 }
 
-EVisibility SStateTreeViewRow::GetBreakpointVisibility() const
+EVisibility SStateTreeViewRow::GetStateBreakpointVisibility() const
 {
-	if (const UStateTreeState* State = WeakState.Get())
+#if WITH_STATETREE_DEBUGGER
+	const UStateTreeState* State = WeakState.Get();
+	const UStateTreeEditorData* EditorData = WeakTreeData.Get();
+	if (State != nullptr && EditorData != nullptr)
 	{
-		if (StateTreeViewModel && StateTreeViewModel->DoesStateHaveBreakpoint(*State))
+		return (EditorData != nullptr && EditorData->HasAnyBreakpoint(State->ID)) ? EVisibility::Visible : EVisibility::Hidden;
+	}
+#endif // WITH_STATETREE_DEBUGGER
+	return EVisibility::Hidden;
+}
+
+FText SStateTreeViewRow::GetStateBreakpointTooltipText() const
+{
+#if WITH_STATETREE_DEBUGGER
+	const UStateTreeState* State = WeakState.Get();
+	const UStateTreeEditorData* EditorData = WeakTreeData.Get();
+	if (State != nullptr && EditorData != nullptr)
+	{
+		const bool bHasBreakpointOnEnter = EditorData->HasBreakpoint(State->ID, EStateTreeBreakpointType::OnEnter);
+		const bool bHasBreakpointOnExit = EditorData->HasBreakpoint(State->ID, EStateTreeBreakpointType::OnExit);
+
+		if (bHasBreakpointOnEnter && bHasBreakpointOnExit)
 		{
-			return EVisibility::Visible;
+			return LOCTEXT("StateTreeStateBreakpointOnEnterAndOnExitTooltip","Break when entering or exiting state");
 		}
-	}
 
-	return EVisibility::Collapsed;
+		if (bHasBreakpointOnEnter)
+		{
+			return LOCTEXT("StateTreeStateBreakpointOnEnterTooltip","Break when entering state");
+		}
+
+		if (bHasBreakpointOnExit)
+		{
+			return LOCTEXT("StateTreeStateBreakpointOnExitTooltip","Break when exiting state");
+		}
+	}
+#endif // WITH_STATETREE_DEBUGGER
+	return FText::GetEmpty();
 }
 
 const FSlateBrush* SStateTreeViewRow::GetSelectorIcon() const
@@ -617,26 +729,6 @@ EVisibility SStateTreeViewRow::GetTasksVisibility() const
 	return EVisibility::Collapsed;
 }
 
-FText SStateTreeViewRow::GetTasksDesc() const
-{
-	const UStateTreeState* State = WeakState.Get();
-	if (!State)
-	{
-		return FText::GetEmpty();
-	}
-
-	TArray<FText> Names;
-	for (int32 i = 0; i < State->Tasks.Num(); i++)
-	{
-		if (const FStateTreeTaskBase* Task = State->Tasks[i].Node.GetPtr<FStateTreeTaskBase>())
-		{
-			Names.Add(FText::FromName(Task->Name));
-		}
-	}
-
-	return FText::Join((FText::FromString(TEXT(" & "))), Names);
-}
-
 EVisibility SStateTreeViewRow::GetLinkedStateVisibility() const
 {
 	if (const UStateTreeState* State = WeakState.Get())
diff --git a/Engine/Plugins/Runtime/StateTree/Source/StateTreeEditorModule/Private/SStateTreeViewRow.h b/Engine/Plugins/Runtime/StateTree/Source/StateTreeEditorModule/Private/SStateTreeViewRow.h
--- a/Engine/Plugins/Runtime/StateTree/Source/StateTreeEditorModule/Private/SStateTreeViewRow.h
+++ b/Engine/Plugins/Runtime/StateTree/Source/StateTreeEditorModule/Private/SStateTreeViewRow.h
@@ -37,14 +37,14 @@ private:
 	FText GetStateDesc() const;
 
 	EVisibility GetConditionVisibility() const;
-	EVisibility GetBreakpointVisibility() const;
+	EVisibility GetStateBreakpointVisibility() const;
+	FText GetStateBreakpointTooltipText() const;
 	
 	const FSlateBrush* GetSelectorIcon() const;
 	FText GetSelectorTooltip() const;
 	FText GetStateTypeTooltip() const;
 
 	EVisibility GetTasksVisibility() const;
-	FText GetTasksDesc() const;
 
 	EVisibility GetLinkedStateVisibility() const;
 	FText GetLinkedStateDesc() const;
@@ -82,5 +82,6 @@ private:
 	
 	TSharedPtr<FStateTreeViewModel> StateTreeViewModel;
 	TWeakObjectPtr<UStateTreeState> WeakState;
+	TWeakObjectPtr<UStateTreeEditorData> WeakTreeData;
 	TSharedPtr<SInlineEditableTextBlock> NameTextBlock;
 };
diff --git a/Engine/Plugins/Runtime/StateTree/Source/StateTreeEditorModule/Private/StateTreeCompiler.cpp b/Engine/Plugins/Runtime/StateTree/Source/StateTreeEditorModule/Private/StateTreeCompiler.cpp
--- a/Engine/Plugins/Runtime/StateTree/Source/StateTreeEditorModule/Private/StateTreeCompiler.cpp
+++ b/Engine/Plugins/Runtime/StateTree/Source/StateTreeEditorModule/Private/StateTreeCompiler.cpp
@@ -239,8 +239,14 @@ bool FStateTreeCompiler::Compile(UStateTree& InStateTree)
 		return false;
 	}
 
+	// Store mapping between node unique ID and their compiled index. Used for debugging purposes.
+	for (const TPair<FGuid, int32>& ToNode : IDToNode)
+	{
+		StateTree->IDToNodeMappings.Emplace(ToNode.Key, FStateTreeIndex16(ToNode.Value));
+	}
+
 	// Store mapping between state unique ID and state handle. Used for debugging purposes.
-	for (const TPair<FGuid, int>& ToState : IDToState)
+	for (const TPair<FGuid, int32>& ToState : IDToState)
 	{
 		StateTree->IDToStateMappings.Emplace(ToState.Key, FStateTreeStateHandle(ToState.Value));
 	}
@@ -856,6 +862,7 @@ bool FStateTreeCompiler::CreateCondition(UStateTreeState& State, const FStateTre
 	}
 
 	// Copy the condition
+	IDToNode.Add(CondNode.ID, Nodes.Num());
 	FInstancedStruct& Node = Nodes.Add_GetRef(CondNode.Node);
 	InstantiateStructSubobjects(Node);
 
@@ -1026,6 +1033,7 @@ bool FStateTreeCompiler::CreateTask(UStateTreeState* State, const FStateTreeEdit
 	}
 
 	// Copy the task
+	IDToNode.Add(TaskNode.ID, Nodes.Num());
 	FInstancedStruct& Node = Nodes.Add_GetRef(TaskNode.Node);
 	InstantiateStructSubobjects(Node);
 	
@@ -1140,6 +1148,7 @@ bool FStateTreeCompiler::CreateEvaluator(const FStateTreeEditorNode& EvalNode)
     }
 
 	// Copy the evaluator
+	IDToNode.Add(EvalNode.ID, Nodes.Num());
 	FInstancedStruct& Node = Nodes.Add_GetRef(EvalNode.Node);
 	InstantiateStructSubobjects(Node);
 	
diff --git a/Engine/Plugins/Runtime/StateTree/Source/StateTreeEditorModule/Private/StateTreeEditor.cpp b/Engine/Plugins/Runtime/StateTree/Source/StateTreeEditorModule/Private/StateTreeEditor.cpp
--- a/Engine/Plugins/Runtime/StateTree/Source/StateTreeEditorModule/Private/StateTreeEditor.cpp
+++ b/Engine/Plugins/Runtime/StateTree/Source/StateTreeEditorModule/Private/StateTreeEditor.cpp
@@ -381,13 +381,20 @@ TSharedRef<SDockTab> FStateTreeEditor::SpawnTab_CompilerResults(const FSpawnTabA
 #if WITH_STATETREE_DEBUGGER
 TSharedRef<SDockTab> FStateTreeEditor::SpawnTab_Debugger(const FSpawnTabArgs& Args)
 {
+	TSharedPtr<SWidget> Widget = SNullWidget::NullWidget;
+	if (StateTree != nullptr)
+	{
+		DebuggerView = SNew(SStateTreeDebuggerView, *StateTree, StateTreeViewModel.ToSharedRef(), TreeViewCommandList);
+		Widget = DebuggerView;
+	}
+	
 	check(Args.GetTabId() == DebuggerTabId);
 
 	return SNew(SDockTab)
 		.Label(LOCTEXT("DebuggerTitle", "Debugger"))
 		.TabColorScale(GetTabColorScale())
 		[
-			SAssignNew(DebuggerView, SStateTreeDebuggerView, StateTree, StateTreeViewModel.ToSharedRef(), TreeViewCommandList)
+			Widget.ToSharedRef()
 		];
 }
 #endif // WITH_STATETREE_DEBUGGER
@@ -835,7 +842,7 @@ void FStateTreeEditor::RegisterToolbar()
 	const FToolMenuInsert InsertAfterAssetSection("Asset", EToolMenuInsertType::After);
 
 	FToolMenuSection& Section = ToolBar->AddSection("Compile", TAttribute<FText>(), InsertAfterAssetSection);
-
+	
 	Section.AddDynamicEntry("CompileCommands", FNewToolMenuSectionDelegate::CreateLambda([](FToolMenuSection& InSection)
 	{
 		const UStateTreeToolMenuContext* Context = InSection.FindContext<UStateTreeToolMenuContext>();
diff --git a/Engine/Plugins/Runtime/StateTree/Source/StateTreeEditorModule/Private/StateTreeEditorCommands.cpp b/Engine/Plugins/Runtime/StateTree/Source/StateTreeEditorModule/Private/StateTreeEditorCommands.cpp
--- a/Engine/Plugins/Runtime/StateTree/Source/StateTreeEditorModule/Private/StateTreeEditorCommands.cpp
+++ b/Engine/Plugins/Runtime/StateTree/Source/StateTreeEditorModule/Private/StateTreeEditorCommands.cpp
@@ -31,7 +31,7 @@ void FStateTreeEditorCommands::RegisterCommands()
 	UI_COMMAND(DuplicateStates, "Duplicate", "Duplicate Selected States", EUserInterfaceActionType::Button, FInputChord(EModifierKey::Control, EKeys::D));
 	UI_COMMAND(DeleteStates, "Delete", "Delete Selected States", EUserInterfaceActionType::Button, FInputChord(EKeys::Delete));
 	UI_COMMAND(RenameState, "Rename", "Rename Selected State", EUserInterfaceActionType::Button, FInputChord(EKeys::F2));
-	UI_COMMAND(EnableStates, "Enable", "Enables Selected States", EUserInterfaceActionType::Check, FInputChord());
+	UI_COMMAND(EnableStates, "State Enabled", "Enables Selected States", EUserInterfaceActionType::Check, FInputChord());
 }
 
 #undef LOCTEXT_NAMESPACE
diff --git a/Engine/Plugins/Runtime/StateTree/Source/StateTreeEditorModule/Private/StateTreeEditorData.cpp b/Engine/Plugins/Runtime/StateTree/Source/StateTreeEditorModule/Private/StateTreeEditorData.cpp
--- a/Engine/Plugins/Runtime/StateTree/Source/StateTreeEditorModule/Private/StateTreeEditorData.cpp
+++ b/Engine/Plugins/Runtime/StateTree/Source/StateTreeEditorModule/Private/StateTreeEditorData.cpp
@@ -356,7 +356,6 @@ void UStateTreeEditorData::ReparentStates()
 	});
 }
 
-
 void UStateTreeEditorData::FixObjectInstance(TSet<UObject*>& SeenObjects, UObject& Outer, FStateTreeEditorNode& Node)
 {
 	if (Node.InstanceObject)
@@ -786,9 +785,58 @@ EStateTreeVisitor UStateTreeEditorData::VisitAllNodes(TFunctionRef<EStateTreeVis
 	return VisitHierarchyNodes(InFunc);
 }
 
+#if WITH_STATETREE_DEBUGGER
+bool UStateTreeEditorData::HasAnyBreakpoint(const FGuid ID) const
+{
+	return Breakpoints.ContainsByPredicate([ID](const FStateTreeEditorBreakpoint& Breakpoint) { return Breakpoint.ID == ID; });
+}
+
+bool UStateTreeEditorData::HasBreakpoint(const FGuid ID, const EStateTreeBreakpointType BreakpointType) const
+{
+	return GetBreakpoint(ID, BreakpointType) != nullptr;
+}
+
+const FStateTreeEditorBreakpoint* UStateTreeEditorData::GetBreakpoint(const FGuid ID, const EStateTreeBreakpointType BreakpointType) const
+{
+	return Breakpoints.FindByPredicate([ID, BreakpointType](const FStateTreeEditorBreakpoint& Breakpoint)
+		{
+			return Breakpoint.ID == ID && Breakpoint.BreakpointType == BreakpointType;
+		});
+}
+
+void UStateTreeEditorData::AddBreakpoint(const FGuid ID, const EStateTreeBreakpointType BreakpointType)
+{
+	Breakpoints.Emplace(ID, BreakpointType);
+
+	const UStateTree* StateTree = GetTypedOuter<UStateTree>();
+	checkf(StateTree, TEXT("UStateTreeEditorData should only be allocated within a UStateTree"));
+	UE::StateTree::Delegates::OnBreakpointsChanged.Broadcast(*StateTree);
+}
+
+bool UStateTreeEditorData::RemoveBreakpoint(const FGuid ID, const EStateTreeBreakpointType BreakpointType)
+{
+	const int32 Index = Breakpoints.IndexOfByPredicate([ID, BreakpointType](const FStateTreeEditorBreakpoint& Breakpoint)
+		{
+			return Breakpoint.ID == ID && Breakpoint.BreakpointType == BreakpointType;
+		});
+		
+	if (Index != INDEX_NONE)
+	{
+		Breakpoints.RemoveAtSwap(Index);
+		
+		const UStateTree* StateTree = GetTypedOuter<UStateTree>();
+		checkf(StateTree, TEXT("UStateTreeEditorData should only be allocated within a UStateTree"));
+		UE::StateTree::Delegates::OnBreakpointsChanged.Broadcast(*StateTree);
+	}
+
+	return Index != INDEX_NONE;
+}
+
+#endif // WITH_STATETREE_DEBUGGER
+
 PRAGMA_DISABLE_DEPRECATION_WARNINGS
 void UStateTreeEditorData::AddPropertyBinding(const FStateTreeEditorPropertyPath& SourcePath, const FStateTreeEditorPropertyPath& TargetPath)
 {
 	EditorBindings.AddPropertyBinding(UE::StateTree::Private::ConvertEditorPath(SourcePath), UE::StateTree::Private::ConvertEditorPath(TargetPath));
 }
-PRAGMA_ENABLE_DEPRECATION_WARNINGS
\ No newline at end of file
+PRAGMA_ENABLE_DEPRECATION_WARNINGS
diff --git a/Engine/Plugins/Runtime/StateTree/Source/StateTreeEditorModule/Private/StateTreeEditorStyle.cpp b/Engine/Plugins/Runtime/StateTree/Source/StateTreeEditorModule/Private/StateTreeEditorStyle.cpp
--- a/Engine/Plugins/Runtime/StateTree/Source/StateTreeEditorModule/Private/StateTreeEditorStyle.cpp
+++ b/Engine/Plugins/Runtime/StateTree/Source/StateTreeEditorModule/Private/StateTreeEditorStyle.cpp
@@ -107,7 +107,7 @@ FStateTreeEditorStyle::FStateTreeEditorStyle()
 
 		Set("StateTree.Task.Rect", new FSlateRoundedBoxBrush(FLinearColor(FColor(21, 104, 116)), 0.0f, FLinearColor(FColor(27, 138, 154)), 1.0f));
 	}
-
+	
 	// Details rich text
 	{
 		Set("Details.Normal", FTextBlockStyle(NormalText)
@@ -187,7 +187,6 @@ FStateTreeEditorStyle::FStateTreeEditorStyle()
 		Set("StateTreeEditor.DuplicateStates", new IMAGE_BRUSH_SVG("Starship/Common/Duplicate", CoreStyleConstants::Icon16x16));
 		Set("StateTreeEditor.DeleteStates", new IMAGE_BRUSH_SVG("Starship/Common/Delete", CoreStyleConstants::Icon16x16));
 		Set("StateTreeEditor.RenameState", new IMAGE_BRUSH_SVG("Starship/Common/Rename", CoreStyleConstants::Icon16x16));
-		Set("StateTreeEditor.EnableStates", new IMAGE_BRUSH_SVG("Starship/Common/check-circle-large", CoreStyleConstants::Icon16x16));
 	}
 
 	{
@@ -202,11 +201,16 @@ FStateTreeEditorStyle::FStateTreeEditorStyle()
 		Set("StateTreeEditor.Debugger.NextFrameWithEvents", new IMAGE_BRUSH("Sequencer/Transport_Bar/Step_Forward_24x", CoreStyleConstants::Icon16x16));
 		Set("StateTreeEditor.Debugger.NextFrameWithStateChange", new IMAGE_BRUSH("Sequencer/Transport_Bar/Go_To_End_24x", CoreStyleConstants::Icon16x16));
 
-		Set("StateTreeEditor.Debugger.ToggleBreakpoint", new IMAGE_BRUSH_SVG("Starship/Blueprints/Breakpoint_Valid", CoreStyleConstants::Icon16x16));
+		Set("StateTreeEditor.Debugger.ToggleOnEnterStateBreakpoint", new IMAGE_BRUSH_SVG("Starship/Blueprints/Breakpoint_Valid", CoreStyleConstants::Icon16x16));
+		Set("StateTreeEditor.Debugger.EnableOnEnterStateBreakpoint", new IMAGE_BRUSH_SVG("Starship/Blueprints/Breakpoint_Valid", CoreStyleConstants::Icon16x16));
+		Set("StateTreeEditor.Debugger.EnableOnExitStateBreakpoint", new IMAGE_BRUSH_SVG("Starship/Blueprints/Breakpoint_Valid", CoreStyleConstants::Icon16x16));
 		Set("StateTreeEditor.DebugOptions", new IMAGE_BRUSH_SVG("Starship/Common/Bug", CoreStyleConstants::Icon16x16));
 
 		Set("StateTreeEditor.Debugger.OwnerTrack", new IMAGE_BRUSH_SVG("Starship/AssetIcons/AIController_64", CoreStyleConstants::Icon16x16));
-		Set("StateTreeEditor.Debugger.InstanceTrack", new IMAGE_BRUSH_SVG("Starship/AssetIcons/AnimInstance_64", CoreStyleConstants::Icon16x16));		
+		Set("StateTreeEditor.Debugger.InstanceTrack", new IMAGE_BRUSH_SVG("Starship/AssetIcons/AnimInstance_64", CoreStyleConstants::Icon16x16));
+
+		Set("StateTreeEditor.EnableStates", new IMAGE_BRUSH("Icons/Empty_16x", CoreStyleConstants::Icon16x16));
+		Set( "StateTreeEditor.Debugger.Breakpoint.EnabledAndValid", new IMAGE_BRUSH_SVG( "Starship/Blueprints/Breakpoint_Valid", CoreStyleConstants::Icon16x16, FStyleColors::AccentRed));
 	}
 
 	{
diff --git a/Engine/Plugins/Runtime/StateTree/Source/StateTreeEditorModule/Private/StateTreeViewModel.cpp b/Engine/Plugins/Runtime/StateTree/Source/StateTreeEditorModule/Private/StateTreeViewModel.cpp
--- a/Engine/Plugins/Runtime/StateTree/Source/StateTreeEditorModule/Private/StateTreeViewModel.cpp
+++ b/Engine/Plugins/Runtime/StateTree/Source/StateTreeEditorModule/Private/StateTreeViewModel.cpp
@@ -250,10 +250,14 @@ void FStateTreeViewModel::Init(UStateTreeEditorData* InTreeData)
 	GEditor->RegisterForUndo(this);
 
 	UE::StateTree::Delegates::OnIdentifierChanged.AddSP(this, &FStateTreeViewModel::HandleIdentifierChanged);
-
+	
 #if WITH_STATETREE_DEBUGGER
+	UE::StateTree::Delegates::OnBreakpointsChanged.AddSP(this, &FStateTreeViewModel::HandleBreakpointsChanged);
+	UE::StateTree::Delegates::OnPostCompile.AddSP(this, &FStateTreeViewModel::HandlePostCompile);
+
 	Debugger->SetAsset(GetStateTree());
 	BindToDebuggerDelegates();
+	RefreshDebuggerBreakpoints();
 #endif // WITH_STATETREE_DEBUGGER	
 }
 
@@ -275,6 +279,52 @@ void FStateTreeViewModel::HandleIdentifierChanged(const UStateTree& StateTree) c
 	}
 }
 
+#if WITH_STATETREE_DEBUGGER
+void FStateTreeViewModel::HandleBreakpointsChanged(const UStateTree& StateTree)
+{
+	if (GetStateTree() == &StateTree)
+	{
+		RefreshDebuggerBreakpoints();
+	}
+}
+
+void FStateTreeViewModel::HandlePostCompile(const UStateTree& StateTree)
+{
+	if (GetStateTree() == &StateTree)
+	{
+		RefreshDebuggerBreakpoints();
+	}
+}
+
+void FStateTreeViewModel::RefreshDebuggerBreakpoints()
+{
+	const UStateTree* StateTree = GetStateTree();
+	const UStateTreeEditorData* TreeData = TreeDataWeak.Get();
+	if (StateTree != nullptr && TreeData != nullptr)
+	{
+		Debugger->ClearAllBreakpoints();
+
+		for (const FStateTreeEditorBreakpoint& Breakpoint : TreeData->Breakpoints)
+		{
+			const FStateTreeIndex16 Index = StateTree->GetNodeIndexFromId(Breakpoint.ID);
+			if (Index.IsValid())
+			{
+				Debugger->SetBreakpoint(Index, Breakpoint.BreakpointType);	
+			}
+			else
+			{
+				FStateTreeStateHandle StateHandle = StateTree->GetStateHandleFromId(Breakpoint.ID);
+				if (StateHandle.IsValid())
+				{
+					Debugger->SetBreakpoint(StateHandle, Breakpoint.BreakpointType);
+				}
+			}
+		}
+	}
+}
+
+#endif // WITH_STATETREE_DEBUGGER
+
 void FStateTreeViewModel::NotifyAssetChangedExternally() const
 {
 	OnAssetChanged.Broadcast();
@@ -997,35 +1047,6 @@ void FStateTreeViewModel::BindToDebuggerDelegates()
 			}
 		}
 	});
-
-	Debugger->OnBreakpointsChanged.BindLambda([this](TConstArrayView<FStateTreeStateHandle> NewActiveBreakpoints)
-	{
-		StatesWithBreakpoints.Reset(NewActiveBreakpoints.Num());
-
-		const UStateTreeEditorData* TreeData = TreeDataWeak.Get();
-		if (const UStateTree* OuterStateTree = TreeData ? TreeData->GetTypedOuter<UStateTree>() : nullptr)
-		{
-			TreeData->VisitHierarchy([this, &NewActiveBreakpoints, &OuterStateTree](const UStateTreeState& State, UStateTreeState* /*ParentState*/)
-			{
-				const FStateTreeStateHandle Handle = OuterStateTree->GetStateHandleFromId(State.ID);
-				if (NewActiveBreakpoints.Contains(Handle))
-				{
-					StatesWithBreakpoints.Emplace(State.ID);
-				}
-
-				return EStateTreeVisitor::Continue;
-			});
-		}
-	});
-#endif // WITH_STATETREE_DEBUGGER
-}
-
-bool FStateTreeViewModel::DoesStateHaveBreakpoint(const UStateTreeState& State) const
-{
-#if WITH_STATETREE_DEBUGGER
-	return StatesWithBreakpoints.Contains(State.ID);
-#else
-	return false;
 #endif // WITH_STATETREE_DEBUGGER
 }
 
diff --git a/Engine/Plugins/Runtime/StateTree/Source/StateTreeEditorModule/Private/StateTreeViewModel.h b/Engine/Plugins/Runtime/StateTree/Source/StateTreeEditorModule/Private/StateTreeViewModel.h
--- a/Engine/Plugins/Runtime/StateTree/Source/StateTreeEditorModule/Private/StateTreeViewModel.h
+++ b/Engine/Plugins/Runtime/StateTree/Source/StateTreeEditorModule/Private/StateTreeViewModel.h
@@ -6,6 +6,7 @@
 #include "EditorUndoClient.h"
 #include "StateTreeViewModel.generated.h"
 
+struct FStateTreeEditorBreakpoint;
 struct FStateTreeDebugger;
 struct FPropertyChangedEvent;
 
@@ -83,7 +84,7 @@ public:
 	bool CanEnableStates() const;
 	bool CanDisableStates() const;
 	void SetSelectedStatesEnabled(bool bEnable);
-
+	
 	// Force to update the view externally.
 	void NotifyAssetChangedExternally() const;
 	void NotifyStatesChangedExternally(const TSet<UStateTreeState*>& ChangedStates, const FPropertyChangedEvent& PropertyChangedEvent) const;
@@ -91,10 +92,10 @@ public:
 	// Debugging
 #if WITH_STATETREE_DEBUGGER
 	TSharedRef<FStateTreeDebugger> GetDebugger() const { return Debugger; }
+	void RefreshDebuggerBreakpoints();
 #endif // WITH_STATETREE_DEBUGGER
 
 	bool IsStateActiveInDebugger(const UStateTreeState& State) const;
-	bool DoesStateHaveBreakpoint(const UStateTreeState& State) const;
 
 	// Called when the whole asset is updated (i.e. undo/redo).
 	FOnAssetChanged& GetOnAssetChanged() { return OnAssetChanged; }
@@ -121,15 +122,17 @@ protected:
 	void PasteStatesAsChildrenFromText(const FString& TextToImport, UStateTreeState* ParentState, const int32 IndexToInsertAt);
 
 	void HandleIdentifierChanged(const UStateTree& StateTree) const;
-
+	
 	void BindToDebuggerDelegates();
 	
 	TWeakObjectPtr<UStateTreeEditorData> TreeDataWeak;
 	TSet<TWeakObjectPtr<UStateTreeState>> SelectedStates;
 
 #if WITH_STATETREE_DEBUGGER
+	void HandleBreakpointsChanged(const UStateTree& StateTree);
+	void HandlePostCompile(const UStateTree& StateTree);
+
 	TSharedRef<FStateTreeDebugger> Debugger;
-	TArray<FGuid> StatesWithBreakpoints;
 	TArray<FGuid> ActiveStates;
 #endif // WITH_STATETREE_DEBUGGER
 	
diff --git a/Engine/Plugins/Runtime/StateTree/Source/StateTreeEditorModule/Public/StateTreeCompiler.h b/Engine/Plugins/Runtime/StateTree/Source/StateTreeEditorModule/Public/StateTreeCompiler.h
--- a/Engine/Plugins/Runtime/StateTree/Source/StateTreeEditorModule/Public/StateTreeCompiler.h
+++ b/Engine/Plugins/Runtime/StateTree/Source/StateTreeEditorModule/Public/StateTreeCompiler.h
@@ -65,6 +65,7 @@ private:
 	FStateTreeCompilerLog& Log;
 	UStateTree* StateTree = nullptr;
 	UStateTreeEditorData* EditorData = nullptr;
+	TMap<FGuid, int32> IDToNode;
 	TMap<FGuid, int32> IDToState;
 	TMap<FGuid, const FStateTreeDataView > IDToStructValue;
 	TArray<UStateTreeState*> SourceStates;
diff --git a/Engine/Plugins/Runtime/StateTree/Source/StateTreeEditorModule/Public/StateTreeEditorData.h b/Engine/Plugins/Runtime/StateTree/Source/StateTreeEditorModule/Public/StateTreeEditorData.h
--- a/Engine/Plugins/Runtime/StateTree/Source/StateTreeEditorModule/Public/StateTreeEditorData.h
+++ b/Engine/Plugins/Runtime/StateTree/Source/StateTreeEditorModule/Public/StateTreeEditorData.h
@@ -4,6 +4,7 @@
 
 #include "StateTreeState.h"
 #include "StateTreeEditorPropertyBindings.h"
+#include "Debugger/StateTreeDebuggerTypes.h"
 #include "StateTreeEditorData.generated.h"
 
 struct FStateTreeBindableStructDesc;
@@ -11,6 +12,25 @@ struct FStateTreeEditorPropertyPath;
 
 class UStateTreeSchema;
 
+USTRUCT()
+struct FStateTreeEditorBreakpoint
+{
+	GENERATED_BODY()
+
+	FStateTreeEditorBreakpoint() = default;
+	explicit FStateTreeEditorBreakpoint(const FGuid& ID, const EStateTreeBreakpointType BreakpointType)
+		: ID(ID),
+		BreakpointType(BreakpointType)
+	{
+	}
+
+	UPROPERTY()
+	FGuid ID;
+
+	UPROPERTY()
+	EStateTreeBreakpointType BreakpointType = EStateTreeBreakpointType::Unset;
+};
+
 UENUM()
 enum class EStateTreeVisitor : uint8
 {
@@ -196,6 +216,14 @@ PRAGMA_ENABLE_DEPRECATION_WARNINGS
 		return false;
 	}
 
+#if WITH_STATETREE_DEBUGGER
+	bool HasAnyBreakpoint(FGuid ID) const;
+	bool HasBreakpoint(FGuid ID, EStateTreeBreakpointType Type) const;
+	const FStateTreeEditorBreakpoint* GetBreakpoint(FGuid ID, EStateTreeBreakpointType Type) const;
+	void AddBreakpoint(FGuid ID, EStateTreeBreakpointType Type);
+	bool RemoveBreakpoint(FGuid ID, EStateTreeBreakpointType Type);
+#endif // WITH_STATETREE_DEBUGGER
+
 	// ~StateTree Builder API
 
 private:
@@ -224,4 +252,13 @@ public:
 	/** Top level States. */
 	UPROPERTY()
 	TArray<TObjectPtr<UStateTreeState>> SubTrees;
+
+	/**
+	 * Transient list of breakpoints added in the debugging session.
+	 * These will be lost if the asset gets reloaded.
+	 * If there is eventually a change to make those persist with the asset
+	 * we need to prune all dangling breakpoints after states/tasks got removed. (see RemoveUnusedBindings)
+	 */
+	UPROPERTY(Transient)
+	TArray<FStateTreeEditorBreakpoint> Breakpoints;
 };
diff --git a/Engine/Plugins/Runtime/StateTree/Source/StateTreeModule/Private/Debugger/StateTreeDebugger.cpp b/Engine/Plugins/Runtime/StateTree/Source/StateTreeModule/Private/Debugger/StateTreeDebugger.cpp
--- a/Engine/Plugins/Runtime/StateTree/Source/StateTreeModule/Private/Debugger/StateTreeDebugger.cpp
+++ b/Engine/Plugins/Runtime/StateTree/Source/StateTreeModule/Private/Debugger/StateTreeDebugger.cpp
@@ -3,13 +3,11 @@
 #if WITH_STATETREE_DEBUGGER
 
 #include "Debugger/StateTreeDebugger.h"
-#include "Algo/RemoveIf.h"
 #include "Debugger/IStateTreeTraceProvider.h"
 #include "Debugger/StateTreeTraceProvider.h"
 #include "Debugger/StateTreeTraceTypes.h"
 #include "Misc/Paths.h"
 #include "Modules/ModuleManager.h"
-#include "StateTree.h"
 #include "StateTreeModule.h"
 #include "Trace/StoreClient.h"
 #include "TraceServices/AnalysisService.h"
@@ -522,32 +520,72 @@ void FStateTreeDebugger::StepForwardToNextStateChange()
 	RefreshActiveStates();
 }
 
-void FStateTreeDebugger::ToggleBreakpoints(const TConstArrayView<FStateTreeStateHandle> SelectedStates)
+bool FStateTreeDebugger::HasStateBreakpoint(const FStateTreeStateHandle StateHandle, const EStateTreeBreakpointType BreakpointType) const
 {
-	TArray<FStateTreeStateHandle> CurrentBreakpoints = StatesWithBreakpoint;
-	TArray<FStateTreeStateHandle> NewBreakpoints(SelectedStates);
-
-	// remove from the selected states any state that already has a breakpoint
-	const int32 ExistingBreakpointsStartIndex = Algo::RemoveIf
-	(NewBreakpoints,
-		[&CurrentBreakpoints](const FStateTreeStateHandle BreakpointToToggle)
+	return Breakpoints.ContainsByPredicate([StateHandle, BreakpointType](const FStateTreeDebuggerBreakpoint Breakpoint)
 		{
-			return CurrentBreakpoints.Contains(BreakpointToToggle);
+			if (Breakpoint.BreakpointType == BreakpointType)
+			{
+				const FStateTreeStateHandle* BreakpointStateHandle = Breakpoint.ElementIdentifier.TryGet<FStateTreeStateHandle>();
+				return (BreakpointStateHandle != nullptr && *BreakpointStateHandle == StateHandle);
+			}
+			return false; 
 		});
+}
 
-	for (int32 i = ExistingBreakpointsStartIndex; i < NewBreakpoints.Num(); i++)
+bool FStateTreeDebugger::HasTaskBreakpoint(const FStateTreeIndex16 Index, const EStateTreeBreakpointType BreakpointType) const
+{
+	return Breakpoints.ContainsByPredicate([Index, BreakpointType](const FStateTreeDebuggerBreakpoint Breakpoint)
 	{
-		CurrentBreakpoints.RemoveSingleSwap(NewBreakpoints[i], /*bAllowShrinking*/false);
-	}
+		if (Breakpoint.BreakpointType == BreakpointType)
+		{
+			const FStateTreeDebuggerBreakpoint::FStateTreeTaskIndex* BreakpointTaskIndex = Breakpoint.ElementIdentifier.TryGet<FStateTreeDebuggerBreakpoint::FStateTreeTaskIndex>();
+			return (BreakpointTaskIndex != nullptr && BreakpointTaskIndex->Index == Index);
+		}
+		return false; 
+	});
+}
+
+bool FStateTreeDebugger::HasTransitionBreakpoint(const FStateTreeIndex16 Index, const EStateTreeBreakpointType BreakpointType) const
+{
+	return Breakpoints.ContainsByPredicate([Index, BreakpointType](const FStateTreeDebuggerBreakpoint Breakpoint)
+	{
+		if (Breakpoint.BreakpointType == BreakpointType)
+		{
+			const FStateTreeDebuggerBreakpoint::FStateTreeTransitionIndex* BreakpointTransitionIndex = Breakpoint.ElementIdentifier.TryGet<FStateTreeDebuggerBreakpoint::FStateTreeTransitionIndex>();
+			return (BreakpointTransitionIndex != nullptr && BreakpointTransitionIndex->Index == Index);
+		}
+		return false; 
+	});
+}
 
-	NewBreakpoints.SetNum(ExistingBreakpointsStartIndex, /*bAllowShrinking*/false);
+void FStateTreeDebugger::SetBreakpoint(FStateTreeStateHandle StateHandle, EStateTreeBreakpointType BreakpointType)
+{
+	Breakpoints.Emplace(StateHandle, BreakpointType);
+}
 
-	// Both lists were reduced and can be merged as the new complete list of breakpoints.
-	CurrentBreakpoints.Append(NewBreakpoints);
+void FStateTreeDebugger::SetBreakpoint(const FStateTreeIndex16 NodeIndex, EStateTreeBreakpointType BreakpointType)
+{
+	Breakpoints.Emplace(NodeIndex, BreakpointType);
+}
 
-	StatesWithBreakpoint = CurrentBreakpoints;
+void FStateTreeDebugger::ClearBreakpoint(const FStateTreeIndex16 NodeIndex, EStateTreeBreakpointType BreakpointType)
+{
+	const int32 Index = Breakpoints.IndexOfByPredicate([NodeIndex, BreakpointType](const FStateTreeDebuggerBreakpoint& Breakpoint)
+		{
+			const FStateTreeDebuggerBreakpoint::FStateTreeTaskIndex* IndexPtr = Breakpoint.ElementIdentifier.TryGet<FStateTreeDebuggerBreakpoint::FStateTreeTaskIndex>();
+			return (IndexPtr != nullptr && IndexPtr->Index == NodeIndex && Breakpoint.BreakpointType == BreakpointType);
+		});
 
-	OnBreakpointsChanged.ExecuteIfBound(StatesWithBreakpoint);
+	if (Index != INDEX_NONE)
+	{
+		Breakpoints.RemoveAtSwap(Index);
+	}
+}
+
+void FStateTreeDebugger::ClearAllBreakpoints()
+{
+	Breakpoints.Empty();
 }
 
 const TraceServices::IAnalysisSession* FStateTreeDebugger::GetAnalysisSession() const
@@ -619,10 +657,10 @@ void FStateTreeDebugger::SendNotifications()
 		NewInstances.Reset();
 	}
 
-	if (HitBreakpointStateIndex != INDEX_NONE)
+	if (HitBreakpointIndex != INDEX_NONE)
 	{
 		check(HitBreakpointInstanceId.IsValid());
-		check(StatesWithBreakpoint.IsValidIndex(HitBreakpointStateIndex));
+		check(Breakpoints.IsValidIndex(HitBreakpointIndex));
 
 		// Force scrub time to latest read time to reflect most recent events.
 		// This will notify scrub position changed and active states
@@ -634,9 +672,9 @@ void FStateTreeDebugger::SendNotifications()
 			SelectInstance(HitBreakpointInstanceId);
 		}
 
-		OnBreakpointHit.ExecuteIfBound(HitBreakpointInstanceId, StatesWithBreakpoint[HitBreakpointStateIndex]);
+		OnBreakpointHit.ExecuteIfBound(HitBreakpointInstanceId, Breakpoints[HitBreakpointIndex]);
 
-		HitBreakpointStateIndex = INDEX_NONE;
+		HitBreakpointIndex = INDEX_NONE;
 		HitBreakpointInstanceId.Reset();
 	}
 }
@@ -659,48 +697,69 @@ void FStateTreeDebugger::ReadTrace(
 	}
 }
 
-bool FStateTreeDebugger::ProcessEvent(const FStateTreeInstanceDebugId InstanceId, const TraceServices::FFrame& Frame, const FStateTreeTraceEventVariantType& Event)
+void FStateTreeDebugger::EvaluateBreakpoints(const FStateTreeInstanceDebugId InstanceId, const FStateTreeTraceEventVariantType& Event)
 {
-	using namespace UE::StateTreeDebugger;
+	if (bPaused // ignored when scrubbing a paused session
+		|| StateTreeAsset == nullptr // asset is required to properly match state handles
+		|| HitBreakpointIndex != INDEX_NONE // stop on first hit breakpoint
+		|| Breakpoints.IsEmpty()
+		|| (SelectedInstanceId.IsValid() && InstanceId != SelectedInstanceId)) // ignore events not for the selected instances		
+	{
+		return;
+	}
+	
+	EStateTreeTraceEventType EventType = EStateTreeTraceEventType::Unset;
+	
+	Visit([&EventType](auto& TypedEvent) { EventType = TypedEvent.EventType; }, Event); 
 
-	// Breakpoints 
-	if (bPaused == false // ignored when scrubbing a paused session
-		&& StateTreeAsset != nullptr // asset is required to properly match state handles
-		&& HitBreakpointStateIndex == INDEX_NONE // stop on first hit breakpoint
-		&& StatesWithBreakpoint.Num()
-		&& (SelectedInstanceId == InstanceId || SelectedInstanceId.IsInvalid())) // allow breakpoints on any instances if not specified
+	for (int BreakpointIndex = 0; BreakpointIndex < Breakpoints.Num(); ++BreakpointIndex)
 	{
-		const FStateTreeTraceStateEvent* StateEvent = Event.TryGet<FStateTreeTraceStateEvent>();
-		if (StateEvent != nullptr && StateEvent->EventType == EStateTreeTraceEventType::OnEnter)
+		const FStateTreeDebuggerBreakpoint Breakpoint = Breakpoints[BreakpointIndex];
+		if (Breakpoint.EventType == EventType)
 		{
-			const UStateTree* InstanceStateTree = nullptr;
-			if (SelectedInstanceId.IsInvalid())
+			if (const FStateTreeStateHandle* StateHandle = Breakpoint.ElementIdentifier.TryGet<FStateTreeStateHandle>())
 			{
-				// No specific instance selected yet, find matching descriptor from id to extract the associated StateTree asset
-				const FInstanceDescriptor* FoundInstance = InstanceDescs.FindByPredicate(
-					[InstanceId](const FInstanceDescriptor& InstanceDesc)
+				if (const FStateTreeTraceStateEvent* StateEvent = Event.TryGet<FStateTreeTraceStateEvent>())
+				{
+					if (StateEvent->GetStateHandle() == *StateHandle)
 					{
-						return InstanceDesc.Id == InstanceId;
-					});
-
-				if (FoundInstance != nullptr)
+						HitBreakpointIndex = BreakpointIndex;
+						HitBreakpointInstanceId = InstanceId;
+					}
+				}
+			}
+			else if (const FStateTreeDebuggerBreakpoint::FStateTreeTaskIndex* TaskIndex = Breakpoint.ElementIdentifier.TryGet<FStateTreeDebuggerBreakpoint::FStateTreeTaskIndex>())
+			{
+				if (const FStateTreeTraceTaskEvent* TaskEvent = Event.TryGet<FStateTreeTraceTaskEvent>())
 				{
-					InstanceStateTree = FoundInstance->StateTree.Get();
+					if (TaskEvent->Index == TaskIndex->Index)
+					{
+						HitBreakpointIndex = BreakpointIndex;
+						HitBreakpointInstanceId = InstanceId;
+					}
 				}
 			}
-
-			if (SelectedInstanceId.IsValid() || InstanceStateTree == StateTreeAsset)
+			else if (const FStateTreeDebuggerBreakpoint::FStateTreeTransitionIndex* TransitionIndex = Breakpoint.ElementIdentifier.TryGet<FStateTreeDebuggerBreakpoint::FStateTreeTransitionIndex>())
 			{
-				HitBreakpointStateIndex = StatesWithBreakpoint.Find(StateEvent->GetStateHandle());
-				if (HitBreakpointStateIndex != INDEX_NONE)
+				if (const FStateTreeTraceTransitionEvent* TransitionEvent = Event.TryGet<FStateTreeTraceTransitionEvent>())
 				{
-					HitBreakpointInstanceId = InstanceId;
+					if (TransitionEvent->TransitionIndex == TransitionIndex->Index)
+					{
+						HitBreakpointIndex = BreakpointIndex;
+						HitBreakpointInstanceId = InstanceId;
+					}
 				}
 			}
 		}
 	}
+}
+
+bool FStateTreeDebugger::ProcessEvent(const FStateTreeInstanceDebugId InstanceId, const TraceServices::FFrame& Frame, const FStateTreeTraceEventVariantType& Event)
+{
+	EvaluateBreakpoints(InstanceId, Event);
 
-	FInstanceEventCollection* ExistingCollection = EventCollections.FindByPredicate([InstanceId](const FInstanceEventCollection& Entry)
+	UE::StateTreeDebugger::FInstanceEventCollection* ExistingCollection = EventCollections.FindByPredicate(
+		[InstanceId](const UE::StateTreeDebugger::FInstanceEventCollection& Entry)
 		{
 			return Entry.InstanceId == InstanceId;
 		});
diff --git a/Engine/Plugins/Runtime/StateTree/Source/StateTreeModule/Private/Debugger/StateTreeDebuggerTypes.cpp b/Engine/Plugins/Runtime/StateTree/Source/StateTreeModule/Private/Debugger/StateTreeDebuggerTypes.cpp
--- a/Engine/Plugins/Runtime/StateTree/Source/StateTreeModule/Private/Debugger/StateTreeDebuggerTypes.cpp
+++ b/Engine/Plugins/Runtime/StateTree/Source/StateTreeModule/Private/Debugger/StateTreeDebuggerTypes.cpp
@@ -229,4 +229,39 @@ void FScrubState::UpdateActiveStatesIndex(const int32 SpanIndex)
 		});
 }
 } // UE::StateTreeDebugger
+
+FStateTreeDebuggerBreakpoint::FStateTreeDebuggerBreakpoint()
+	: BreakpointType(EStateTreeBreakpointType::Unset)
+	, EventType(EStateTreeTraceEventType::Unset)
+{
+}
+
+FStateTreeDebuggerBreakpoint::FStateTreeDebuggerBreakpoint(const FStateTreeStateHandle StateHandle, const EStateTreeBreakpointType BreakpointType)
+	: ElementIdentifier(TInPlaceType<FStateTreeStateHandle>(), StateHandle)
+	, BreakpointType(BreakpointType)
+{
+	EventType = GetMatchingEventType(BreakpointType);
+}
+
+FStateTreeDebuggerBreakpoint::FStateTreeDebuggerBreakpoint(const FStateTreeIndex16 NodeIndex, const EStateTreeBreakpointType BreakpointType)
+	: ElementIdentifier(TInPlaceType<FStateTreeTaskIndex>(), NodeIndex)
+	, BreakpointType(BreakpointType)
+{
+	EventType = GetMatchingEventType(BreakpointType);
+}
+
+EStateTreeTraceEventType FStateTreeDebuggerBreakpoint::GetMatchingEventType(const EStateTreeBreakpointType BreakpointType)
+{
+	switch (BreakpointType)
+	{
+	case EStateTreeBreakpointType::OnEnter:
+		return EStateTreeTraceEventType::OnEnter;
+	case EStateTreeBreakpointType::OnExit:
+		return EStateTreeTraceEventType::OnExit;
+	default:
+		return EStateTreeTraceEventType::Unset;
+	}
+}
+
 #endif // WITH_STATETREE_DEBUGGER
+
diff --git a/Engine/Plugins/Runtime/StateTree/Source/StateTreeModule/Private/Debugger/StateTreeTraceTypes.cpp b/Engine/Plugins/Runtime/StateTree/Source/StateTreeModule/Private/Debugger/StateTreeTraceTypes.cpp
--- a/Engine/Plugins/Runtime/StateTree/Source/StateTreeModule/Private/Debugger/StateTreeTraceTypes.cpp
+++ b/Engine/Plugins/Runtime/StateTree/Source/StateTreeModule/Private/Debugger/StateTreeTraceTypes.cpp
@@ -123,7 +123,7 @@ FString FStateTreeTraceConditionEvent::ToString(const UStateTree& StateTree) con
 // FStateTreeTraceActiveStatesEvent
 //----------------------------------------------------------------------//
 FStateTreeTraceActiveStatesEvent::FStateTreeTraceActiveStatesEvent(const double RecordingWorldTime)
-	: FStateTreeTraceBaseEvent(RecordingWorldTime)
+	: FStateTreeTraceBaseEvent(RecordingWorldTime, EStateTreeTraceEventType::Unset)
 {
 }
 
diff --git a/Engine/Plugins/Runtime/StateTree/Source/StateTreeModule/Private/StateTree.cpp b/Engine/Plugins/Runtime/StateTree/Source/StateTreeModule/Private/StateTree.cpp
--- a/Engine/Plugins/Runtime/StateTree/Source/StateTreeModule/Private/StateTree.cpp
+++ b/Engine/Plugins/Runtime/StateTree/Source/StateTreeModule/Private/StateTree.cpp
@@ -31,6 +31,12 @@ FConstStructView UStateTree::GetNode(const int32 NodeIndex) const
 	return Nodes.IsValidIndex(NodeIndex) ? Nodes[NodeIndex] : FConstStructView();	
 }
 
+FStateTreeIndex16 UStateTree::GetNodeIndexFromId(const FGuid Id) const
+{
+	const FStateTreeNodeIdToIndex* Entry = IDToNodeMappings.FindByPredicate([Id](const FStateTreeNodeIdToIndex& Entry){ return Entry.Id == Id; });
+	return Entry != nullptr ? Entry->Index : FStateTreeIndex16::Invalid;
+}
+
 const FCompactStateTreeState* UStateTree::GetStateFromHandle(const FStateTreeStateHandle StateHandle) const
 {
 	return States.IsValidIndex(StateHandle.Index) ? &States[StateHandle.Index] : nullptr;
@@ -106,6 +112,7 @@ void UStateTree::ResetCompiled()
 	PropertyBindings.Reset();
 	Parameters.Reset();
 	IDToStateMappings.Reset();
+	IDToNodeMappings.Reset();
 
 	ParametersDataViewIndex = FStateTreeIndex8::Invalid;
 	
diff --git a/Engine/Plugins/Runtime/StateTree/Source/StateTreeModule/Private/StateTreeDelegates.cpp b/Engine/Plugins/Runtime/StateTree/Source/StateTreeModule/Private/StateTreeDelegates.cpp
--- a/Engine/Plugins/Runtime/StateTree/Source/StateTreeModule/Private/StateTreeDelegates.cpp
+++ b/Engine/Plugins/Runtime/StateTree/Source/StateTreeModule/Private/StateTreeDelegates.cpp
@@ -10,6 +10,7 @@ FOnIdentifierChanged OnIdentifierChanged;
 FOnSchemaChanged OnSchemaChanged;
 FOnParametersChanged OnParametersChanged;
 FOnStateParametersChanged OnStateParametersChanged;
+FOnBreakpointsChanged OnBreakpointsChanged;
 FOnPostCompile OnPostCompile;
 FOnRequestCompile OnRequestCompile;
 #endif
diff --git a/Engine/Plugins/Runtime/StateTree/Source/StateTreeModule/Public/Debugger/StateTreeDebugger.h b/Engine/Plugins/Runtime/StateTree/Source/StateTreeModule/Public/Debugger/StateTreeDebugger.h
--- a/Engine/Plugins/Runtime/StateTree/Source/StateTreeModule/Public/Debugger/StateTreeDebugger.h
+++ b/Engine/Plugins/Runtime/StateTree/Source/StateTreeModule/Public/Debugger/StateTreeDebugger.h
@@ -13,6 +13,8 @@
 #include "TraceServices/Model/Diagnostics.h"
 #include "TraceServices/Model/Frames.h"
 
+enum class EStateTreeBreakpointType : uint8;
+
 namespace UE::Trace
 {
 	class FStoreClient;
@@ -23,8 +25,7 @@ class UStateTreeState;
 class UStateTree;
 
 DECLARE_DELEGATE_OneParam(FOnStateTreeDebuggerScrubStateChanged, const UE::StateTreeDebugger::FScrubState& ScrubState);
-DECLARE_DELEGATE_TwoParams(FOnStateTreeDebuggerBreakpointHit, FStateTreeInstanceDebugId InstanceId, FStateTreeStateHandle StateHandle);
-DECLARE_DELEGATE_OneParam(FOnStateTreeDebuggerBreakpointsChanged, TConstArrayView<FStateTreeStateHandle> Breakpoints);
+DECLARE_DELEGATE_TwoParams(FOnStateTreeDebuggerBreakpointHit, FStateTreeInstanceDebugId InstanceId, const FStateTreeDebuggerBreakpoint Breakpoint);
 DECLARE_DELEGATE_OneParam(FOnStateTreeDebuggerActiveStatesChanges, TConstArrayView<FStateTreeStateHandle> ActiveStates);
 DECLARE_DELEGATE_OneParam(FOnStateTreeDebuggerNewInstance, FStateTreeInstanceDebugId InstanceId);
 DECLARE_DELEGATE(FOnStateTreeDebuggerDebuggedInstanceSet);
@@ -83,10 +84,20 @@ struct STATETREEMODULE_API FStateTreeDebugger : FTickableGameObject
 	FText GetInstanceName(FStateTreeInstanceDebugId InstanceId) const;
 	FText GetInstanceDescription(FStateTreeInstanceDebugId InstanceId) const;
 	void SelectInstance(const FStateTreeInstanceDebugId InstanceId);
+	void ClearSelection() { SelectInstance({}); }
 	FStateTreeInstanceDebugId GetSelectedInstanceId() const { return SelectedInstanceId; }
 	const UE::StateTreeDebugger::FInstanceDescriptor* GetInstanceDescriptor(const FStateTreeInstanceDebugId InstanceId) const ;
 	const UE::StateTreeDebugger::FInstanceDescriptor* GetSelectedInstanceDescriptor() const { return GetInstanceDescriptor(SelectedInstanceId); }
 
+	bool HasStateBreakpoint(FStateTreeStateHandle StateHandle, EStateTreeBreakpointType BreakpointType) const;
+	bool HasTaskBreakpoint(FStateTreeIndex16 Index, EStateTreeBreakpointType BreakpointType) const;
+	bool HasTransitionBreakpoint(FStateTreeIndex16 Index, EStateTreeBreakpointType BreakpointType) const;
+	void SetBreakpoint(FStateTreeStateHandle StateHandle, EStateTreeBreakpointType BreakpointType);
+	void SetBreakpoint(FStateTreeIndex16 NodeIndex, EStateTreeBreakpointType BreakpointType);
+	void ClearBreakpoint(FStateTreeIndex16 NodeIndex, EStateTreeBreakpointType BreakpointType);
+	void ClearAllBreakpoints();
+	int32 NumBreakpoints() const { return Breakpoints.Num(); }
+
 	static FText DescribeTrace(const FTraceDescriptor& TraceDescriptor);
 	static FText DescribeInstance(const UE::StateTreeDebugger::FInstanceDescriptor& StateTreeInstanceDesc);
 
@@ -120,14 +131,11 @@ struct STATETREEMODULE_API FStateTreeDebugger : FTickableGameObject
 	void StopAnalysis();
 	FTraceDescriptor GetSelectedTraceDescriptor() const { return ActiveSessionTraceDescriptor; }
 	FText GetSelectedTraceDescription() const;
-	
-	void ToggleBreakpoints(const TConstArrayView<FStateTreeStateHandle> SelectedStates);
 
 	FOnStateTreeDebuggerNewInstance OnNewInstance;
 	FOnStateTreeDebuggerDebuggedInstanceSet OnSelectedInstanceCleared;
 	FOnStateTreeDebuggerScrubStateChanged OnScrubStateChanged;
 	FOnStateTreeDebuggerBreakpointHit OnBreakpointHit;
-	FOnStateTreeDebuggerBreakpointsChanged OnBreakpointsChanged;
 	FOnStateTreeDebuggerActiveStatesChanges OnActiveStatesChanged;
 
 protected:
@@ -146,6 +154,8 @@ private:
 		const TraceServices::FFrame& Frame
 		);
 
+	void EvaluateBreakpoints(FStateTreeInstanceDebugId InstanceId, const FStateTreeTraceEventVariantType& Event);
+
 	void SendNotifications();
 
 	void SetActiveStates(const TConstArrayView<FStateTreeStateHandle> NewActiveStates);
@@ -193,8 +203,8 @@ private:
 	/** Specific instance selected for more details */
 	FStateTreeInstanceDebugId SelectedInstanceId;
 
-	/** Handles of states on which a breakpoint has been set. This is per asset and not specific to an instance. */
-	TArray<FStateTreeStateHandle> StatesWithBreakpoint;
+	/** List of breakpoints set. This is per asset and not specific to an instance. */
+	TArray<FStateTreeDebuggerBreakpoint> Breakpoints;
 
 	/** List of currently active states in the selected instance */
 	TArray<FStateTreeStateHandle> ActiveStates;
@@ -225,9 +235,9 @@ private:
 
 	/** Indicates the instance for which a breakpoint has been hit */
 	FStateTreeInstanceDebugId HitBreakpointInstanceId = FStateTreeInstanceDebugId::Invalid;
-
-	/** Indicates the state for which a breakpoint has been hit */
-	int32 HitBreakpointStateIndex = INDEX_NONE;
+	
+	/** Indicates the index of the breakpoint that has been hit */
+	int32 HitBreakpointIndex = INDEX_NONE;
 
 	/** List of new instances discovered by processing event in the analysis session. */
 	TArray<FStateTreeInstanceDebugId> NewInstances;
diff --git a/Engine/Plugins/Runtime/StateTree/Source/StateTreeModule/Public/Debugger/StateTreeDebuggerTypes.h b/Engine/Plugins/Runtime/StateTree/Source/StateTreeModule/Public/Debugger/StateTreeDebuggerTypes.h
--- a/Engine/Plugins/Runtime/StateTree/Source/StateTreeModule/Public/Debugger/StateTreeDebuggerTypes.h
+++ b/Engine/Plugins/Runtime/StateTree/Source/StateTreeModule/Public/Debugger/StateTreeDebuggerTypes.h
@@ -5,11 +5,25 @@
 
 #if WITH_STATETREE_DEBUGGER
 
-#include "StateTreeTypes.h"
 #include "StateTree.h"
+#include "StateTreeTypes.h"
 #include "TraceServices/Model/Frames.h"
 
+#endif // WITH_STATETREE_DEBUGGER
+
+UENUM()
+enum class EStateTreeBreakpointType : uint8
+{
+	Unset,
+	OnEnter,
+	OnExit,
+	OnTransition,
+};
+
+#if WITH_STATETREE_DEBUGGER
+
 class UStateTree;
+enum class EStateTreeTraceEventType : uint8;
 
 namespace UE::StateTreeDebugger
 {
@@ -250,4 +264,24 @@ private:
 
 } // UE::StateTreeDebugger
 
+struct FStateTreeDebuggerBreakpoint
+{
+	// Wrapper structs to be able to use TVariant with more than one type based on FStateTreeIndex16 (can't use 'using')
+	struct FStateTreeTaskIndex { FStateTreeIndex16 Index; };
+	struct FStateTreeTransitionIndex { FStateTreeIndex16 Index; };
+
+	using FIdentifierVariantType = TVariant<FStateTreeStateHandle, FStateTreeTaskIndex, FStateTreeTransitionIndex>;
+	
+	FStateTreeDebuggerBreakpoint();
+	explicit FStateTreeDebuggerBreakpoint(const FStateTreeStateHandle StateHandle, const EStateTreeBreakpointType BreakpointType);
+	explicit FStateTreeDebuggerBreakpoint(const FStateTreeIndex16 NodeIndex, const EStateTreeBreakpointType BreakpointType);
+
+	FIdentifierVariantType ElementIdentifier;
+	EStateTreeBreakpointType BreakpointType;
+	EStateTreeTraceEventType EventType;
+
+private:
+	static EStateTreeTraceEventType GetMatchingEventType(EStateTreeBreakpointType BreakpointType);
+};
+
 #endif // WITH_STATETREE_DEBUGGER
diff --git a/Engine/Plugins/Runtime/StateTree/Source/StateTreeModule/Public/Debugger/StateTreeTraceTypes.h b/Engine/Plugins/Runtime/StateTree/Source/StateTreeModule/Public/Debugger/StateTreeTraceTypes.h
--- a/Engine/Plugins/Runtime/StateTree/Source/StateTreeModule/Public/Debugger/StateTreeTraceTypes.h
+++ b/Engine/Plugins/Runtime/StateTree/Source/StateTreeModule/Public/Debugger/StateTreeTraceTypes.h
@@ -35,8 +35,9 @@ enum class EStateTreeTraceEventType : uint8
 
 struct FStateTreeTraceBaseEvent
 {
-	explicit FStateTreeTraceBaseEvent(const double RecordingWorldTime)
+	explicit FStateTreeTraceBaseEvent(const double RecordingWorldTime, const EStateTreeTraceEventType EventType)
 		: RecordingWorldTime(RecordingWorldTime)
+		, EventType(EventType)
 	{
 	}
 
@@ -44,27 +45,26 @@ struct FStateTreeTraceBaseEvent
 	static FString GetDataAsText() { return TEXT(""); }
 
 	double RecordingWorldTime = 0;
+	EStateTreeTraceEventType EventType;
 };
 
 struct FStateTreeTracePhaseEvent : FStateTreeTraceBaseEvent
 {
 	explicit FStateTreeTracePhaseEvent(const double RecordingWorldTime, const EStateTreeUpdatePhase Phase, const EStateTreeTraceEventType EventType)
-		: FStateTreeTraceBaseEvent(RecordingWorldTime)
+		: FStateTreeTraceBaseEvent(RecordingWorldTime, EventType)
 		, Phase(Phase)
-		, EventType(EventType)
 	{
 	}
 
 	STATETREEMODULE_API FString ToString(const UStateTree& StateTree) const;
 
 	EStateTreeUpdatePhase Phase;
-	EStateTreeTraceEventType EventType;
 };
 
 struct FStateTreeTraceLogEvent : FStateTreeTraceBaseEvent
 {
 	explicit FStateTreeTraceLogEvent(const double RecordingWorldTime, const FString& Message)
-		: FStateTreeTraceBaseEvent(RecordingWorldTime)
+		: FStateTreeTraceBaseEvent(RecordingWorldTime, EStateTreeTraceEventType::Unset)
 		, Message(Message)
 	{
 	}
@@ -77,31 +77,27 @@ struct FStateTreeTraceLogEvent : FStateTreeTraceBaseEvent
 struct FStateTreeTraceTransitionEvent : FStateTreeTraceBaseEvent
 {
 	explicit FStateTreeTraceTransitionEvent(const double RecordingWorldTime, const FStateTreeIndex16 Index, const EStateTreeTraceEventType EventType)
-		: FStateTreeTraceBaseEvent(RecordingWorldTime)
+		: FStateTreeTraceBaseEvent(RecordingWorldTime, EventType)
 		, TransitionIndex(Index)
-		, EventType(EventType)
 	{
 	}
 
 	STATETREEMODULE_API FString ToString(const UStateTree& StateTree) const;
 
 	FStateTreeIndex16 TransitionIndex;
-	EStateTreeTraceEventType EventType = EStateTreeTraceEventType::Unset;
 };
 
 struct FStateTreeTraceNodeEvent : FStateTreeTraceBaseEvent
 {
 	explicit FStateTreeTraceNodeEvent(const double RecordingWorldTime, const FStateTreeIndex16 Index, const EStateTreeTraceEventType EventType)
-		: FStateTreeTraceBaseEvent(RecordingWorldTime)
+		: FStateTreeTraceBaseEvent(RecordingWorldTime, EventType)
 		, Index(Index)
-		, EventType(EventType)
 	{
 	}
 
 	STATETREEMODULE_API FString ToString(const UStateTree& StateTree) const;
 	
 	FStateTreeIndex16 Index;
-	EStateTreeTraceEventType EventType = EStateTreeTraceEventType::Unset;
 };
 
 struct FStateTreeTraceStateEvent : FStateTreeTraceNodeEvent
diff --git a/Engine/Plugins/Runtime/StateTree/Source/StateTreeModule/Public/StateTree.h b/Engine/Plugins/Runtime/StateTree/Source/StateTreeModule/Public/StateTree.h
--- a/Engine/Plugins/Runtime/StateTree/Source/StateTreeModule/Public/StateTree.h
+++ b/Engine/Plugins/Runtime/StateTree/Source/StateTreeModule/Public/StateTree.h
@@ -121,6 +121,9 @@ public:
 	/** @return Struct views of all nodes */
 	const FInstancedStructContainer& GetNodes() const { return Nodes; }
 
+	/** @return Node index matching a given Id; invalid handle if node not found. */
+	FStateTreeIndex16 GetNodeIndexFromId(const FGuid Id) const;
+
 	/** @return View of all states. */
 	TConstArrayView<FCompactStateTreeState> GetStates() const { return States; }
 
@@ -224,7 +227,11 @@ private:
 	/** Mapping of state guid for the Editor and state handles, created at compilation. */
 	UPROPERTY()
 	TArray<FStateTreeStateIdToHandle> IDToStateMappings;
-	
+
+	/** Mapping of node guid for the Editor and node index, created at compilation. */
+	UPROPERTY()
+	TArray<FStateTreeNodeIdToIndex> IDToNodeMappings;
+
 	/**
 	 * Parameters that could be used for bindings within the Tree.
 	 * Default values are stored within the asset but StateTreeReference can be used to parameterized the tree.
diff --git a/Engine/Plugins/Runtime/StateTree/Source/StateTreeModule/Public/StateTreeDelegates.h b/Engine/Plugins/Runtime/StateTree/Source/StateTreeModule/Public/StateTreeDelegates.h
--- a/Engine/Plugins/Runtime/StateTree/Source/StateTreeModule/Public/StateTreeDelegates.h
+++ b/Engine/Plugins/Runtime/StateTree/Source/StateTreeModule/Public/StateTreeDelegates.h
@@ -37,6 +37,13 @@ namespace UE::StateTree::Delegates
 	DECLARE_MULTICAST_DELEGATE_TwoParams(FOnStateParametersChanged, const UStateTree& /*StateTree*/, const FGuid /*StateID*/);
 	extern STATETREEMODULE_API FOnStateParametersChanged OnStateParametersChanged;
 
+	/**
+	 * Called when breakpoints of the StateTree EditorData changed.
+	 * This should mainly used by the asset editor to update the debugger.
+	 */
+	DECLARE_MULTICAST_DELEGATE_OneParam(FOnBreakpointsChanged, const UStateTree& /*StateTree*/);
+	extern STATETREEMODULE_API FOnBreakpointsChanged OnBreakpointsChanged;
+
 	/** Called when compilation succeeds */
 	DECLARE_MULTICAST_DELEGATE_OneParam(FOnPostCompile, const UStateTree& /*StateTree*/);
 	extern STATETREEMODULE_API FOnPostCompile OnPostCompile;
diff --git a/Engine/Plugins/Runtime/StateTree/Source/StateTreeModule/Public/StateTreeTypes.h b/Engine/Plugins/Runtime/StateTree/Source/StateTreeModule/Public/StateTreeTypes.h
--- a/Engine/Plugins/Runtime/StateTree/Source/StateTreeModule/Public/StateTreeTypes.h
+++ b/Engine/Plugins/Runtime/StateTree/Source/StateTreeModule/Public/StateTreeTypes.h
@@ -818,6 +818,28 @@ struct STATETREEMODULE_API FStateTreeStateIdToHandle
 	FStateTreeStateHandle Handle;
 };
 
+/**
+ * Pair of node id and its associated node index created at compilation.
+ */
+USTRUCT()
+struct STATETREEMODULE_API FStateTreeNodeIdToIndex
+{
+	GENERATED_BODY()
+
+	FStateTreeNodeIdToIndex() = default;
+	explicit FStateTreeNodeIdToIndex(const FGuid& Id, const FStateTreeIndex16 Index)
+		: Id(Id)
+		, Index(Index)
+	{
+	}
+
+	UPROPERTY();
+	FGuid Id;
+
+	UPROPERTY();
+	FStateTreeIndex16 Index;
+};
+
 /**
  * Describes an external data. The data can point to a struct or object.
  * The code that handles StateTree ticking is responsible for passing in the actually data, see FStateTreeExecutionContext.
