Fix for CIS Issue 319943

diff --git a/Engine/Plugins/Runtime/StateTree/Source/StateTreeEditorModule/Private/StateTreeViewModel.cpp b/Engine/Plugins/Runtime/StateTree/Source/StateTreeEditorModule/Private/StateTreeViewModel.cpp
--- a/Engine/Plugins/Runtime/StateTree/Source/StateTreeEditorModule/Private/StateTreeViewModel.cpp
+++ b/Engine/Plugins/Runtime/StateTree/Source/StateTreeEditorModule/Private/StateTreeViewModel.cpp
@@ -309,7 +309,7 @@ void FStateTreeViewModel::RefreshDebuggerBreakpoints()
 			const FStateTreeIndex16 Index = StateTree->GetNodeIndexFromId(Breakpoint.ID);
 			if (Index.IsValid())
 			{
-				Debugger->SetBreakpoint(Index, Breakpoint.BreakpointType);	
+				Debugger->SetTaskBreakpoint(Index, Breakpoint.BreakpointType);	
 			}
 			else
 			{
diff --git a/Engine/Plugins/Runtime/StateTree/Source/StateTreeModule/Private/Debugger/StateTreeDebugger.cpp b/Engine/Plugins/Runtime/StateTree/Source/StateTreeModule/Private/Debugger/StateTreeDebugger.cpp
--- a/Engine/Plugins/Runtime/StateTree/Source/StateTreeModule/Private/Debugger/StateTreeDebugger.cpp
+++ b/Engine/Plugins/Runtime/StateTree/Source/StateTreeModule/Private/Debugger/StateTreeDebugger.cpp
@@ -564,9 +564,9 @@ void FStateTreeDebugger::SetBreakpoint(FStateTreeStateHandle StateHandle, EState
 	Breakpoints.Emplace(StateHandle, BreakpointType);
 }
 
-void FStateTreeDebugger::SetBreakpoint(const FStateTreeIndex16 NodeIndex, EStateTreeBreakpointType BreakpointType)
+void FStateTreeDebugger::SetTaskBreakpoint(const FStateTreeIndex16 NodeIndex, EStateTreeBreakpointType BreakpointType)
 {
-	Breakpoints.Emplace(NodeIndex, BreakpointType);
+	Breakpoints.Emplace(FStateTreeDebuggerBreakpoint::FStateTreeTaskIndex(NodeIndex), BreakpointType);
 }
 
 void FStateTreeDebugger::ClearBreakpoint(const FStateTreeIndex16 NodeIndex, EStateTreeBreakpointType BreakpointType)
diff --git a/Engine/Plugins/Runtime/StateTree/Source/StateTreeModule/Private/Debugger/StateTreeDebuggerTypes.cpp b/Engine/Plugins/Runtime/StateTree/Source/StateTreeModule/Private/Debugger/StateTreeDebuggerTypes.cpp
--- a/Engine/Plugins/Runtime/StateTree/Source/StateTreeModule/Private/Debugger/StateTreeDebuggerTypes.cpp
+++ b/Engine/Plugins/Runtime/StateTree/Source/StateTreeModule/Private/Debugger/StateTreeDebuggerTypes.cpp
@@ -243,8 +243,15 @@ FStateTreeDebuggerBreakpoint::FStateTreeDebuggerBreakpoint(const FStateTreeState
 	EventType = GetMatchingEventType(BreakpointType);
 }
 
-FStateTreeDebuggerBreakpoint::FStateTreeDebuggerBreakpoint(const FStateTreeIndex16 NodeIndex, const EStateTreeBreakpointType BreakpointType)
-	: ElementIdentifier(TInPlaceType<FStateTreeTaskIndex>(), NodeIndex)
+FStateTreeDebuggerBreakpoint::FStateTreeDebuggerBreakpoint(const FStateTreeTaskIndex Index, const EStateTreeBreakpointType BreakpointType)
+	: ElementIdentifier(TInPlaceType<FStateTreeTaskIndex>(), Index)
+	, BreakpointType(BreakpointType)
+{
+	EventType = GetMatchingEventType(BreakpointType);
+}
+
+FStateTreeDebuggerBreakpoint::FStateTreeDebuggerBreakpoint(const FStateTreeTransitionIndex Index, const EStateTreeBreakpointType BreakpointType)
+	: ElementIdentifier(TInPlaceType<FStateTreeTransitionIndex>(), Index)
 	, BreakpointType(BreakpointType)
 {
 	EventType = GetMatchingEventType(BreakpointType);
diff --git a/Engine/Plugins/Runtime/StateTree/Source/StateTreeModule/Public/Debugger/StateTreeDebugger.h b/Engine/Plugins/Runtime/StateTree/Source/StateTreeModule/Public/Debugger/StateTreeDebugger.h
--- a/Engine/Plugins/Runtime/StateTree/Source/StateTreeModule/Public/Debugger/StateTreeDebugger.h
+++ b/Engine/Plugins/Runtime/StateTree/Source/StateTreeModule/Public/Debugger/StateTreeDebugger.h
@@ -93,7 +93,7 @@ struct STATETREEMODULE_API FStateTreeDebugger : FTickableGameObject
 	bool HasTaskBreakpoint(FStateTreeIndex16 Index, EStateTreeBreakpointType BreakpointType) const;
 	bool HasTransitionBreakpoint(FStateTreeIndex16 Index, EStateTreeBreakpointType BreakpointType) const;
 	void SetBreakpoint(FStateTreeStateHandle StateHandle, EStateTreeBreakpointType BreakpointType);
-	void SetBreakpoint(FStateTreeIndex16 NodeIndex, EStateTreeBreakpointType BreakpointType);
+	void SetTaskBreakpoint(FStateTreeIndex16 NodeIndex, EStateTreeBreakpointType BreakpointType);
 	void ClearBreakpoint(FStateTreeIndex16 NodeIndex, EStateTreeBreakpointType BreakpointType);
 	void ClearAllBreakpoints();
 	int32 NumBreakpoints() const { return Breakpoints.Num(); }
diff --git a/Engine/Plugins/Runtime/StateTree/Source/StateTreeModule/Public/Debugger/StateTreeDebuggerTypes.h b/Engine/Plugins/Runtime/StateTree/Source/StateTreeModule/Public/Debugger/StateTreeDebuggerTypes.h
--- a/Engine/Plugins/Runtime/StateTree/Source/StateTreeModule/Public/Debugger/StateTreeDebuggerTypes.h
+++ b/Engine/Plugins/Runtime/StateTree/Source/StateTreeModule/Public/Debugger/StateTreeDebuggerTypes.h
@@ -267,14 +267,33 @@ private:
 struct FStateTreeDebuggerBreakpoint
 {
 	// Wrapper structs to be able to use TVariant with more than one type based on FStateTreeIndex16 (can't use 'using')
-	struct FStateTreeTaskIndex { FStateTreeIndex16 Index; };
-	struct FStateTreeTransitionIndex { FStateTreeIndex16 Index; };
+	struct FStateTreeTaskIndex
+	{
+		FStateTreeTaskIndex() = default;
+		explicit FStateTreeTaskIndex(const FStateTreeIndex16& Index)
+			: Index(Index)
+		{
+		}
+
+		FStateTreeIndex16 Index;
+	};
+	
+	struct FStateTreeTransitionIndex
+	{
+		FStateTreeTransitionIndex() = default;
+		explicit FStateTreeTransitionIndex(const FStateTreeIndex16& Index)
+			: Index(Index)
+		{
+		}
+		FStateTreeIndex16 Index;
+	};
 
 	using FIdentifierVariantType = TVariant<FStateTreeStateHandle, FStateTreeTaskIndex, FStateTreeTransitionIndex>;
 	
 	FStateTreeDebuggerBreakpoint();
 	explicit FStateTreeDebuggerBreakpoint(const FStateTreeStateHandle StateHandle, const EStateTreeBreakpointType BreakpointType);
-	explicit FStateTreeDebuggerBreakpoint(const FStateTreeIndex16 NodeIndex, const EStateTreeBreakpointType BreakpointType);
+	explicit FStateTreeDebuggerBreakpoint(const FStateTreeTaskIndex Index, const EStateTreeBreakpointType BreakpointType);
+	explicit FStateTreeDebuggerBreakpoint(const FStateTreeTransitionIndex Index, const EStateTreeBreakpointType BreakpointType);
 
 	FIdentifierVariantType ElementIdentifier;
 	EStateTreeBreakpointType BreakpointType;
