StateTree: Prevent level Actor references in State Tree
- Added check in the compiler to error when level Actors are found in the node, instance or parameter properties
- Made debug text task's ReferenceActor optional input to hide the level actor picker

diff --git a/Engine/Plugins/Runtime/StateTree/Source/StateTreeEditorModule/Private/StateTreeCompiler.cpp b/Engine/Plugins/Runtime/StateTree/Source/StateTreeEditorModule/Private/StateTreeCompiler.cpp
--- a/Engine/Plugins/Runtime/StateTree/Source/StateTreeEditorModule/Private/StateTreeCompiler.cpp
+++ b/Engine/Plugins/Runtime/StateTree/Source/StateTreeEditorModule/Private/StateTreeCompiler.cpp
@@ -9,6 +9,7 @@
 #include "StateTreeTaskBase.h"
 #include "StateTreeConditionBase.h"
 #include "Serialization/ArchiveUObject.h"
+#include "GameFramework/Actor.h"
 
 
 namespace UE::StateTree::Compiler
@@ -70,6 +71,88 @@ namespace UE::StateTree::Compiler
 		}
 	};
 
+	/** Scans Data for actors that are tied to some level and returns them. */
+	void ScanLevelActorReferences(FStateTreeDataView Data, TSet<const UObject*>& Visited, TArray<const AActor*>& OutActors)
+	{
+		if (!Data.IsValid())
+		{
+			return;
+		}
+		
+		for (TPropertyValueIterator<FProperty> It(Data.GetStruct(), Data.GetMemory()); It; ++It)
+		{
+			const FProperty* Property = It->Key;
+			const void* ValuePtr = It->Value;
+
+			if (!ValuePtr)
+			{
+				continue;
+			}
+			
+			if (const FStructProperty* StructProperty = CastField<FStructProperty>(Property))
+			{
+				if (StructProperty->Struct == TBaseStructure<FInstancedStruct>::Get())
+				{
+					const FInstancedStruct& InstancedStruct = *static_cast<const FInstancedStruct*>(ValuePtr);
+					if (InstancedStruct.IsValid())
+					{
+						ScanLevelActorReferences(FStateTreeDataView(const_cast<FInstancedStruct&>(InstancedStruct)), Visited, OutActors);
+					}
+				}
+			}
+			else if (const FObjectPropertyBase* ObjectProperty = CastField<FObjectPropertyBase>(Property))
+			{
+				if (const UObject* Object = ObjectProperty->GetObjectPropertyValue(ValuePtr))
+				{
+					if (const AActor* Actor = Cast<AActor>(Object))
+					{
+						const ULevel* Level = Actor->GetLevel();
+						if (Level != nullptr)
+						{
+							OutActors.Add(Actor);
+						}
+					}
+					// Recurse into instanced object
+					if (Property->HasAnyPropertyFlags(CPF_InstancedReference))
+					{
+						if (!Visited.Contains(Object))
+						{
+							Visited.Add(Object);
+							ScanLevelActorReferences(FStateTreeDataView(const_cast<UObject*>(Object)), Visited, OutActors);
+						}
+					}
+				}
+			}
+		}
+	}
+
+	bool ValidateNoLevelActorReferences(FStateTreeCompilerLog& Log, const FStateTreeBindableStructDesc& NodeDesc, FStateTreeDataView NodeView, FStateTreeDataView InstanceView)
+	{
+		TSet<const UObject*> Visited;
+		TArray<const AActor*> LevelActors;
+		UE::StateTree::Compiler::ScanLevelActorReferences(NodeView, Visited, LevelActors);
+		UE::StateTree::Compiler::ScanLevelActorReferences(InstanceView, Visited, LevelActors);
+		if (!LevelActors.IsEmpty())
+		{
+			FStringBuilderBase AllActorsString;
+			for (const AActor* Actor : LevelActors)
+			{
+				if (AllActorsString.Len() > 0)
+				{
+					AllActorsString += TEXT(", ");
+				}
+				AllActorsString += *GetNameSafe(Actor);
+			}
+			Log.Reportf(EMessageSeverity::Error, NodeDesc,
+				TEXT("Level Actor references were found: %s. Direct Actor references are not allowed."),
+					*AllActorsString);
+			return false;
+		}
+		
+		return true;
+	}
+
+
 	void FValidationResult::Log(FStateTreeCompilerLog& Log, const TCHAR* ContextText, const FStateTreeBindableStructDesc& ContextStruct) const
 	{
 		Log.Reportf(EMessageSeverity::Error, ContextStruct, TEXT("The StateTree is too complex. Compact index %s out of range %d/%d."), ContextText, Value, MaxValue);
@@ -156,7 +239,7 @@ bool FStateTreeCompiler::Compile(UStateTree& InStateTree)
 
 	// Copy parameters from EditorData	
 	StateTree->Parameters = EditorData->RootParameters.Parameters;
-	
+
 	// Mark parameters as binding source
 	const FStateTreeBindableStructDesc ParametersDesc = {
 			TEXT("Parameters"),
@@ -172,7 +255,13 @@ bool FStateTreeCompiler::Compile(UStateTree& InStateTree)
 		return false;
 	}
 	StateTree->ParametersDataViewIndex = FStateTreeIndex8(ParametersDataViewIndex); 
-	
+
+	if (!UE::StateTree::Compiler::ValidateNoLevelActorReferences(Log, ParametersDesc, FStateTreeDataView(), FStateTreeDataView(EditorData->RootParameters.Parameters.GetMutableValue())))
+	{
+		StateTree->ResetCompiled();
+		return false;
+	}
+
 	// Mark all named external values as binding source
 	if (StateTree->Schema)
 	{
@@ -505,6 +594,12 @@ bool FStateTreeCompiler::CreateStateTasksAndParameters()
 						EStateTreeBindableStructSource::State,
 						SourceState->Parameters.ID
 					};
+
+					if (!UE::StateTree::Compiler::ValidateNoLevelActorReferences(Log, SubtreeParamsDesc, FStateTreeDataView(), FStateTreeDataView(CompactParams.Parameters.GetMutableValue())))
+					{
+						return false;
+					}
+					
 					SourceStructIndex = BindingsCompiler.AddSourceStruct(SubtreeParamsDesc);
 					
 					if (const auto Validation = UE::StateTree::Compiler::IsValidIndex16(SourceStructIndex); Validation.DidFail())
@@ -530,6 +625,11 @@ bool FStateTreeCompiler::CreateStateTasksAndParameters()
 						SourceState->Parameters.ID
 					};
 
+					if (!UE::StateTree::Compiler::ValidateNoLevelActorReferences(Log, LinkedParamsDesc, FStateTreeDataView(), FStateTreeDataView(CompactParams.Parameters.GetMutableValue())))
+					{
+						return false;
+					}
+
 					// Check that the bindings for this struct are still all valid.
 					TArray<FStateTreePropertyPathBinding> Bindings;
 					if (!GetAndValidateBindings(LinkedParamsDesc, FStateTreeDataView(SourceState->Parameters.Parameters.GetMutableValue()), Bindings))
@@ -1007,6 +1107,12 @@ bool FStateTreeCompiler::CompileAndValidateNode(const UStateTreeState* SourceSta
 		}
 	}
 
+	// Make sure there's no level actor references in the data.
+	if (!UE::StateTree::Compiler::ValidateNoLevelActorReferences(Log, NodeDesc, NodeView, InstanceData))
+	{
+		return false;
+	}
+	
 	return Result != EDataValidationResult::Invalid;
 }
 
diff --git a/Engine/Plugins/Runtime/StateTree/Source/StateTreeModule/Private/Tasks/StateTreeDebugTextTask.h b/Engine/Plugins/Runtime/StateTree/Source/StateTreeModule/Private/Tasks/StateTreeDebugTextTask.h
--- a/Engine/Plugins/Runtime/StateTree/Source/StateTreeModule/Private/Tasks/StateTreeDebugTextTask.h
+++ b/Engine/Plugins/Runtime/StateTree/Source/StateTreeModule/Private/Tasks/StateTreeDebugTextTask.h
@@ -14,7 +14,7 @@ struct STATETREEMODULE_API FStateTreeDebugTextTaskInstanceData
 	GENERATED_BODY()
 
 	/** Optional actor where to draw the text at. */
-	UPROPERTY(EditAnywhere, Category = "Parameter")
+	UPROPERTY(EditAnywhere, Category = "Input", meta=(Optional))
 	TObjectPtr<AActor> ReferenceActor = nullptr;
 };
 
