StateTree: Fix PatchBindings() for subtrees.

diff --git a/Engine/Plugins/Runtime/StateTree/Source/StateTreeModule/Private/StateTree.cpp b/Engine/Plugins/Runtime/StateTree/Source/StateTreeModule/Private/StateTree.cpp
--- a/Engine/Plugins/Runtime/StateTree/Source/StateTreeModule/Private/StateTree.cpp
+++ b/Engine/Plugins/Runtime/StateTree/Source/StateTreeModule/Private/StateTree.cpp
@@ -546,13 +546,15 @@ bool UStateTree::PatchBindings()
 	// Setup data views for state parameters.
 	for (FCompactStateTreeState& State : States)
 	{
-		if (State.Type == EStateTreeStateType::Linked && State.LinkedState.IsValid())
+		// Subtree state parameters can be a binding source, but not a target.
+		// Linked state parameters can only by a binding target, so not handled here.
+		if (State.Type == EStateTreeStateType::Subtree)
 		{
-			FCompactStateTreeParameters& Params = DefaultInstanceData.GetMutableStruct(State.ParameterInstanceIndex.Get()).Get<FCompactStateTreeParameters>();
-			if (Params.Parameters.IsValid())
+			if (State.ParameterDataViewIndex.IsValid())
 			{
-				const FCompactStateTreeState& LinkedState = States[State.LinkedState.Index];
-				DataViews[LinkedState.ParameterDataViewIndex.Get()] = Params.Parameters.GetMutableValue();
+				FCompactStateTreeParameters& Params = DefaultInstanceData.GetMutableStruct(State.ParameterInstanceIndex.Get()).Get<FCompactStateTreeParameters>();
+				ensureMsgf(Params.Parameters.IsValid(), TEXT("Expecting valid state parameter value when parameter view is set."));
+				DataViews[State.ParameterDataViewIndex.Get()] = Params.Parameters.GetMutableValue();
 			}
 		}
 	}
@@ -588,53 +590,57 @@ bool UStateTree::PatchBindings()
 	for (int32 BatchIndex = 0; BatchIndex < CopyBatches.Num(); ++BatchIndex)
 	{
 		const FStateTreePropertyCopyBatch& Batch = CopyBatches[BatchIndex];
-		for (int32 i = Batch.BindingsBegin; i != Batch.BindingsEnd; i++)
+
+		// Find data view for the binding target.
+		FStateTreeDataView TargetView;
+
+		// Target can be one of the nodes.
+		for (FConstStructView NodeView : GetNodes())
 		{
-			FStateTreePropertyPathBinding& Binding = PropertyPathBindings[i];
-			FStateTreeIndex16 TargetStructIndex;
-			for (FConstStructView NodeView : GetNodes())
+			const FStateTreeNodeBase& Node = NodeView.Get<const FStateTreeNodeBase>();
+			if (Node.BindingsBatch.AsInt32() == BatchIndex)
 			{
-				const FStateTreeNodeBase& Node = NodeView.Get<const FStateTreeNodeBase>();
-				if (Node.BindingsBatch.AsInt32() == BatchIndex)
-				{
-					TargetStructIndex = Node.DataViewIndex;
-					break;
-				}
+				TargetView = DataViews[Node.DataViewIndex.Get()];
+				break;
 			}
+		}
 
-			if (!TargetStructIndex.IsValid())
+		// Target can be a linked state parameters (subtree state params can only be a binding source).
+		if (!TargetView.IsValid())
+		{
+			for (const FCompactStateTreeState& State : GetStates())
 			{
-				for (const FCompactStateTreeState& State : GetStates())
+				if (State.Type == EStateTreeStateType::Linked)
 				{
-					if (State.Type == EStateTreeStateType::Linked)
+					const FCompactStateTreeParameters& Params = DefaultInstanceData.GetStruct(State.ParameterInstanceIndex.Get()).Get<const FCompactStateTreeParameters>();
+					if (Params.BindingsBatch.Get() == BatchIndex)
 					{
-						const FCompactStateTreeParameters& Params = DefaultInstanceData.GetStruct(State.ParameterInstanceIndex.Get()).Get<const FCompactStateTreeParameters>();
-						if (Params.BindingsBatch.Get() == BatchIndex)
-						{
-							if (const FCompactStateTreeState* LinkedState = GetStateFromHandle(State.LinkedState))
-							{
-								TargetStructIndex = LinkedState->ParameterDataViewIndex;
-							}
-						}
+						const FCompactStateTreeState& LinkedState = States[State.LinkedState.Index];
+						TargetView = DataViews[LinkedState.ParameterDataViewIndex.Get()];
 					}
 				}
 			}
+		}
 
-			if (!TargetStructIndex.IsValid())
-			{
-				UE_LOG(LogStateTree, Error, TEXT("%hs: Invalid target struct for property binding %s."), __FUNCTION__, *Binding.GetTargetPath().ToString());
-				return false;
-			}
+		if (!TargetView.IsValid())
+		{
+			UE_LOG(LogStateTree, Error, TEXT("%hs: Invalid target struct when trying to bind to '%s'."), __FUNCTION__, *Batch.TargetStruct.Name.ToString());
+			return false;
+		}
+
+		for (int32 Index = Batch.BindingsBegin; Index != Batch.BindingsEnd; Index++)
+		{
+			FStateTreePropertyPathBinding& Binding = PropertyPathBindings[Index];
 
 			if (!Binding.GetMutableSourcePath().UpdateInstanceStructsFromValue(DataViews[Binding.GetCompiledSourceStructIndex().Get()]))
 			{
-				UE_LOG(LogStateTree, Error, TEXT("%hs: Failed to update source instance structs for property binding %s."), __FUNCTION__, *Binding.GetTargetPath().ToString());
+				UE_LOG(LogStateTree, Error, TEXT("%hs: Failed to update source instance structs for property binding '%s'."), __FUNCTION__, *Binding.GetTargetPath().ToString());
 				return false;
 			}
 
-			if (!Binding.GetMutableTargetPath().UpdateInstanceStructsFromValue(DataViews[TargetStructIndex.Get()]))
+			if (!Binding.GetMutableTargetPath().UpdateInstanceStructsFromValue(TargetView))
 			{
-				UE_LOG(LogStateTree, Error, TEXT("%hs: Failed to update target instance structs for property binding %s."), __FUNCTION__, *Binding.GetTargetPath().ToString());
+				UE_LOG(LogStateTree, Error, TEXT("%hs: Failed to update target instance structs for property binding '%s'."), __FUNCTION__, *Binding.GetTargetPath().ToString());
 				return false;
 			}
 		}
