StateTree: Improved state tree error logging.
- fixed object tokens not to open asset browser
- made the node text token easier to distinguish and added type to it too
- fixed typo in error

diff --git a/Engine/Plugins/Runtime/StateTree/Source/StateTreeEditorModule/Private/StateTreeCompiler.cpp b/Engine/Plugins/Runtime/StateTree/Source/StateTreeEditorModule/Private/StateTreeCompiler.cpp
--- a/Engine/Plugins/Runtime/StateTree/Source/StateTreeEditorModule/Private/StateTreeCompiler.cpp
+++ b/Engine/Plugins/Runtime/StateTree/Source/StateTreeEditorModule/Private/StateTreeCompiler.cpp
@@ -1653,7 +1653,7 @@ bool FStateTreeCompiler::GetAndValidateBindings(const FStateTreeBindableStructDe
 				else
 				{
 					Log.Reportf(EMessageSeverity::Error, TargetStruct,
-						TEXT("Cound not find matching Context object for Context property '%s:%s'. Property must have munual binding."),
+						TEXT("Cound not find matching Context object for Context property '%s:%s'. Property must have manual binding."),
 						*TargetStruct.Name.ToString(), *PropertyName.ToString());
 					bResult = false;
 				}
diff --git a/Engine/Plugins/Runtime/StateTree/Source/StateTreeEditorModule/Private/StateTreeCompilerLog.cpp b/Engine/Plugins/Runtime/StateTree/Source/StateTreeEditorModule/Private/StateTreeCompilerLog.cpp
--- a/Engine/Plugins/Runtime/StateTree/Source/StateTreeEditorModule/Private/StateTreeCompilerLog.cpp
+++ b/Engine/Plugins/Runtime/StateTree/Source/StateTreeEditorModule/Private/StateTreeCompilerLog.cpp
@@ -17,12 +17,19 @@ void FStateTreeCompilerLog::AppendToLog(IMessageLogListing* LogListing) const
 
 		if (StateTreeMessage.State != nullptr)
 		{
-			Message->AddToken(FUObjectToken::Create(StateTreeMessage.State, FText::FromName(StateTreeMessage.State->Name)));
+			TSharedRef<FUObjectToken> ObjectToken = FUObjectToken::Create(StateTreeMessage.State, FText::FromName(StateTreeMessage.State->Name));
+
+			// Add a dummy activation since default behavior opens content browser.
+			static auto DummyActivation = [](const TSharedRef<class IMessageToken>&) {};
+			ObjectToken->OnMessageTokenActivated(FOnMessageTokenActivated::CreateLambda(DummyActivation));
+			
+			Message->AddToken(ObjectToken);
 		}
 
 		if (StateTreeMessage.Item.ID.IsValid())
 		{
-			Message->AddToken(FTextToken::Create(FText::Format(LOCTEXT("LogMessageItem", " {0}"), FText::FromName(StateTreeMessage.Item.Name))));
+			Message->AddToken(FTextToken::Create(FText::Format(LOCTEXT("LogMessageItem", " {0} '{1}': "),
+				UEnum::GetDisplayValueAsText(StateTreeMessage.Item.DataSource), FText::FromName(StateTreeMessage.Item.Name))));
 		}
 
 		if (!StateTreeMessage.Message.IsEmpty())
