[StateTreeDebugger]
- added option to disable state transitions
- added state transition breakpoints
- fixed scrubbing not using the right time (analysis time instead of world simulation time) when hitting a breakpoint
- fixed new instance auto selection on record when previous selection is a stale subtrack
- added console variable `statetree.displayitemids` to display states, tasks and transitions Ids in the StateTreeEditor (details view and state treeview)

diff --git a/Engine/Plugins/Runtime/StateTree/Source/StateTreeEditorModule/Private/Customizations/StateTreeEditorNodeDetails.cpp b/Engine/Plugins/Runtime/StateTree/Source/StateTreeEditorModule/Private/Customizations/StateTreeEditorNodeDetails.cpp
--- a/Engine/Plugins/Runtime/StateTree/Source/StateTreeEditorModule/Private/Customizations/StateTreeEditorNodeDetails.cpp
+++ b/Engine/Plugins/Runtime/StateTree/Source/StateTreeEditorModule/Private/Customizations/StateTreeEditorNodeDetails.cpp
@@ -6,6 +6,7 @@
 #include "IPropertyUtilities.h"
 #include "IDetailChildrenBuilder.h"
 #include "StateTree.h"
+#include "StateTreeEditor.h"
 #include "StateTreeEditorData.h"
 #include "Widgets/Input/SComboButton.h"
 #include "Framework/MultiBox/MultiBoxBuilder.h"
@@ -646,6 +647,13 @@ void FStateTreeEditorNodeDetails::CustomizeChildren(TSharedRef<class IPropertyHa
 	FGuid ID;
 	UE::StateTree::PropertyHelpers::GetStructValue<FGuid>(IDProperty, ID);
 
+	// ID
+	if (UE::StateTree::Editor::GbDisplayItemIds)
+	{
+		// ID
+		StructBuilder.AddProperty(IDProperty.ToSharedRef());
+	}
+	
 	// Node
 	TSharedRef<FBindableNodeInstanceDetails> NodeDetails = MakeShareable(new FBindableNodeInstanceDetails(NodeProperty, FGuid(), EditorData));
 	StructBuilder.AddCustomBuilder(NodeDetails);
diff --git a/Engine/Plugins/Runtime/StateTree/Source/StateTreeEditorModule/Private/Customizations/StateTreeStateDetails.cpp b/Engine/Plugins/Runtime/StateTree/Source/StateTreeEditorModule/Private/Customizations/StateTreeStateDetails.cpp
--- a/Engine/Plugins/Runtime/StateTree/Source/StateTreeEditorModule/Private/Customizations/StateTreeStateDetails.cpp
+++ b/Engine/Plugins/Runtime/StateTree/Source/StateTreeEditorModule/Private/Customizations/StateTreeStateDetails.cpp
@@ -5,6 +5,7 @@
 #include "DetailLayoutBuilder.h"
 #include "PropertyCustomizationHelpers.h"
 #include "IPropertyUtilities.h"
+#include "StateTreeEditor.h"
 #include "StateTreeEditorData.h"
 #include "StateTreeSchema.h"
 #include "Debugger/StateTreeDebuggerUIExtensions.h"
@@ -36,6 +37,7 @@ void FStateTreeStateDetails::CustomizeDetails(IDetailLayoutBuilder& DetailBuilde
 	}
 	const UStateTreeSchema* Schema = EditorData ? EditorData->Schema : nullptr;
 
+	const TSharedPtr<IPropertyHandle> IDProperty = DetailBuilder.GetProperty(GET_MEMBER_NAME_CHECKED(UStateTreeState, ID));
 	const TSharedPtr<IPropertyHandle> EnabledProperty = DetailBuilder.GetProperty(GET_MEMBER_NAME_CHECKED(UStateTreeState, bEnabled));
 	const TSharedPtr<IPropertyHandle> TasksProperty = DetailBuilder.GetProperty(GET_MEMBER_NAME_CHECKED(UStateTreeState, Tasks));
 	const TSharedPtr<IPropertyHandle> SingleTaskProperty = DetailBuilder.GetProperty(GET_MEMBER_NAME_CHECKED(UStateTreeState, SingleTask));
@@ -47,6 +49,10 @@ void FStateTreeStateDetails::CustomizeDetails(IDetailLayoutBuilder& DetailBuilde
 	const TSharedPtr<IPropertyHandle> SelectionBehaviorProperty = DetailBuilder.GetProperty(GET_MEMBER_NAME_CHECKED(UStateTreeState, SelectionBehavior));
 
 	EnabledProperty->MarkHiddenByCustomization();
+	if (UE::StateTree::Editor::GbDisplayItemIds == false)
+	{
+		IDProperty->MarkHiddenByCustomization();
+	}
 
 	uint8 StateTypeValue = 0;
 	TypeProperty->GetValue(StateTypeValue);
diff --git a/Engine/Plugins/Runtime/StateTree/Source/StateTreeEditorModule/Private/Customizations/StateTreeTransitionDetails.cpp b/Engine/Plugins/Runtime/StateTree/Source/StateTreeEditorModule/Private/Customizations/StateTreeTransitionDetails.cpp
--- a/Engine/Plugins/Runtime/StateTree/Source/StateTreeEditorModule/Private/Customizations/StateTreeTransitionDetails.cpp
+++ b/Engine/Plugins/Runtime/StateTree/Source/StateTreeEditorModule/Private/Customizations/StateTreeTransitionDetails.cpp
@@ -1,16 +1,20 @@
 // Copyright Epic Games, Inc. All Rights Reserved.
 
 #include "StateTreeTransitionDetails.h"
+#include "Debugger/StateTreeDebuggerUIExtensions.h"
+#include "DetailLayoutBuilder.h"
+#include "DetailWidgetRow.h"
+#include "HAL/PlatformApplicationMisc.h"
+#include "IDetailChildrenBuilder.h"
+#include "IPropertyUtilities.h"
+#include "ScopedTransaction.h"
+#include "StateTreeEditor.h"
+#include "StateTreeEditorData.h"
+#include "StateTreePropertyHelpers.h"
 #include "StateTreeTypes.h"
 #include "Widgets/Layout/SBox.h"
 #include "Widgets/SBoxPanel.h"
 #include "Widgets/Text/STextBlock.h"
-#include "DetailWidgetRow.h"
-#include "DetailLayoutBuilder.h"
-#include "IDetailChildrenBuilder.h"
-#include "Widgets/Text/STextBlock.h"
-#include "Widgets/Layout/SBox.h"
-#include "StateTreePropertyHelpers.h"
 
 #define LOCTEXT_NAMESPACE "StateTreeEditor"
 
@@ -19,19 +23,35 @@ TSharedRef<IPropertyTypeCustomization> FStateTreeTransitionDetails::MakeInstance
 	return MakeShareable(new FStateTreeTransitionDetails);
 }
 
-void FStateTreeTransitionDetails::CustomizeHeader(TSharedRef<class IPropertyHandle> StructPropertyHandle, class FDetailWidgetRow& HeaderRow, IPropertyTypeCustomizationUtils& StructCustomizationUtils)
+void FStateTreeTransitionDetails::CustomizeHeader(TSharedRef<IPropertyHandle> StructPropertyHandle, FDetailWidgetRow& HeaderRow, IPropertyTypeCustomizationUtils& StructCustomizationUtils)
 {
 	StructProperty = StructPropertyHandle;
-	PropUtils = StructCustomizationUtils.GetPropertyUtilities().Get();
+	PropUtils = StructCustomizationUtils.GetPropertyUtilities();
 
-	TriggerProperty = StructProperty->GetChildHandle(TEXT("Trigger"));
-	PriorityProperty = StructProperty->GetChildHandle(TEXT("Priority"));
-	EventTagProperty = StructProperty->GetChildHandle(TEXT("EventTag"));
-	StateProperty = StructProperty->GetChildHandle(TEXT("State"));
-	DelayTransitionProperty = StructProperty->GetChildHandle(TEXT("bDelayTransition"));
-	DelayDurationProperty = StructProperty->GetChildHandle(TEXT("DelayDuration"));
-	DelayRandomVarianceProperty = StructProperty->GetChildHandle(TEXT("DelayRandomVariance"));
-	ConditionsProperty = StructProperty->GetChildHandle(TEXT("Conditions"));
+	// Find StateTreeEditorData associated with this panel.
+	UStateTreeEditorData* EditorData = nullptr;
+	const TArray<TWeakObjectPtr<>>& Objects = PropUtils->GetSelectedObjects();
+	for (const TWeakObjectPtr<>& WeakObject : Objects)
+	{
+		if (const UObject* Object = WeakObject.Get())
+		{
+			if (UStateTreeEditorData* OuterEditorData = Object->GetTypedOuter<UStateTreeEditorData>())
+			{
+				EditorData = OuterEditorData;
+				break;
+			}
+		}
+	}
+
+	TriggerProperty = StructProperty->GetChildHandle(GET_MEMBER_NAME_CHECKED(FStateTreeTransition, Trigger));
+	PriorityProperty = StructProperty->GetChildHandle(GET_MEMBER_NAME_CHECKED(FStateTreeTransition, Priority));
+	EventTagProperty = StructProperty->GetChildHandle(GET_MEMBER_NAME_CHECKED(FStateTreeTransition, EventTag));
+	StateProperty = StructProperty->GetChildHandle(GET_MEMBER_NAME_CHECKED(FStateTreeTransition, State));
+	DelayTransitionProperty = StructProperty->GetChildHandle(GET_MEMBER_NAME_CHECKED(FStateTreeTransition, bDelayTransition));
+	DelayDurationProperty = StructProperty->GetChildHandle(GET_MEMBER_NAME_CHECKED(FStateTreeTransition, DelayDuration));
+	DelayRandomVarianceProperty = StructProperty->GetChildHandle(GET_MEMBER_NAME_CHECKED(FStateTreeTransition, DelayRandomVariance));
+	ConditionsProperty = StructProperty->GetChildHandle(GET_MEMBER_NAME_CHECKED(FStateTreeTransition, Conditions));
+	IDProperty = StructProperty->GetChildHandle(GET_MEMBER_NAME_CHECKED(FStateTreeTransition, ID));
 
 	HeaderRow
 		.RowTag(StructProperty->GetProperty()->GetFName())
@@ -53,7 +73,15 @@ void FStateTreeTransitionDetails::CustomizeHeader(TSharedRef<class IPropertyHand
 			[
 				StructPropertyHandle->CreateDefaultPropertyButtonWidgets()
 			]
-		];
+			+ SHorizontalBox::Slot()
+			.HAlign(HAlign_Right)
+			.FillWidth(1.0f)
+			[
+				UE::StateTreeEditor::DebuggerExtensions::CreateTransitionWidget(StructPropertyHandle, EditorData)
+			]
+		]
+		.CopyAction(FUIAction(FExecuteAction::CreateSP(this, &FStateTreeTransitionDetails::OnCopyTransition)))
+		.PasteAction(FUIAction(FExecuteAction::CreateSP(this, &FStateTreeTransitionDetails::OnPasteTransition)));
 }
  
 void FStateTreeTransitionDetails::CustomizeChildren(TSharedRef<class IPropertyHandle> StructPropertyHandle, class IDetailChildrenBuilder& StructBuilder, IPropertyTypeCustomizationUtils& StructCustomizationUtils)
@@ -65,12 +93,18 @@ void FStateTreeTransitionDetails::CustomizeChildren(TSharedRef<class IPropertyHa
 	check(DelayRandomVarianceProperty);
 	check(StateProperty);
 	check(ConditionsProperty);
+	check(IDProperty);
 
 	auto IsTickOrEventTransition = [this]()
 	{
 		return !EnumHasAnyFlags(GetTrigger(), EStateTreeTransitionTrigger::OnStateCompleted) ? EVisibility::Visible : EVisibility::Collapsed;
 	};
 
+	if (UE::StateTree::Editor::GbDisplayItemIds)
+	{
+		StructBuilder.AddProperty(IDProperty.ToSharedRef());
+	}
+	
 	// Trigger
 	StructBuilder.AddProperty(TriggerProperty.ToSharedRef());
 
@@ -198,4 +232,51 @@ FText FStateTreeTransitionDetails::GetDescription() const
 	return FText::Format(LOCTEXT("TransitionDesc", "{0} {1}"), TriggerText, TargetText);
 }
 
+void FStateTreeTransitionDetails::OnCopyTransition() const
+{
+	FString Value;
+	// Use PPF_Copy so that all properties get copied.
+	if (StructProperty->GetValueAsFormattedString(Value, PPF_Copy) == FPropertyAccess::Success)
+	{
+		FPlatformApplicationMisc::ClipboardCopy(*Value);
+	} 
+}
+
+void FStateTreeTransitionDetails::OnPasteTransition() const
+{
+	FString PastedText;
+	FPlatformApplicationMisc::ClipboardPaste(PastedText);
+
+	if (PastedText.IsEmpty())
+	{
+		return;
+	}
+
+	FScopedTransaction Transaction(LOCTEXT("PasteTransition", "Paste Transition"));
+
+	StructProperty->NotifyPreChange();
+
+	// Make sure we instantiate new objects when setting the value.
+	StructProperty->SetValueFromFormattedString(PastedText, EPropertyValueSetFlags::InstanceObjects);
+
+	// Reset GUIDs on paste
+	TArray<void*> RawData;
+	StructProperty->AccessRawData(RawData);
+	for (int32 Index = 0; Index < RawData.Num(); Index++)
+	{
+		if (FStateTreeTransition* Transition = static_cast<FStateTreeTransition*>(RawData[Index]))
+		{
+			Transition->ID = FGuid::NewGuid();
+		}
+	}
+
+	StructProperty->NotifyPostChange(EPropertyChangeType::ValueSet);
+	StructProperty->NotifyFinishedChangingProperties();
+
+	if (PropUtils)
+	{
+		PropUtils->ForceRefresh();
+	}
+}
+
 #undef LOCTEXT_NAMESPACE
diff --git a/Engine/Plugins/Runtime/StateTree/Source/StateTreeEditorModule/Private/Customizations/StateTreeTransitionDetails.h b/Engine/Plugins/Runtime/StateTree/Source/StateTreeEditorModule/Private/Customizations/StateTreeTransitionDetails.h
--- a/Engine/Plugins/Runtime/StateTree/Source/StateTreeEditorModule/Private/Customizations/StateTreeTransitionDetails.h
+++ b/Engine/Plugins/Runtime/StateTree/Source/StateTreeEditorModule/Private/Customizations/StateTreeTransitionDetails.h
@@ -5,6 +5,9 @@
 #include "IPropertyTypeCustomization.h"
 
 class IPropertyHandle;
+class IPropertyUtilities;
+class IDetailChildrenBuilder;
+class FDetailWidgetRow;
 class UStateTree;
 class UStateTreeState;
 enum class EStateTreeTransitionTrigger : uint8;
@@ -20,8 +23,8 @@ public:
 	static TSharedRef<IPropertyTypeCustomization> MakeInstance();
 
 	/** IPropertyTypeCustomization interface */
-	virtual void CustomizeHeader(TSharedRef<IPropertyHandle> StructPropertyHandle, class FDetailWidgetRow& HeaderRow, IPropertyTypeCustomizationUtils& StructCustomizationUtils) override;
-	virtual void CustomizeChildren(TSharedRef<IPropertyHandle> StructPropertyHandle, class IDetailChildrenBuilder& StructBuilder, IPropertyTypeCustomizationUtils& StructCustomizationUtils) override;
+	virtual void CustomizeHeader(TSharedRef<IPropertyHandle> StructPropertyHandle, FDetailWidgetRow& HeaderRow, IPropertyTypeCustomizationUtils& StructCustomizationUtils) override;
+	virtual void CustomizeChildren(TSharedRef<IPropertyHandle> StructPropertyHandle, IDetailChildrenBuilder& StructBuilder, IPropertyTypeCustomizationUtils& StructCustomizationUtils) override;
 
 private:
 
@@ -29,7 +32,10 @@ private:
 
 	EStateTreeTransitionTrigger GetTrigger() const;
 	bool GetDelayTransition() const;
-	
+
+	void OnCopyTransition() const;
+	void OnPasteTransition() const;
+
 	TSharedPtr<IPropertyHandle> TriggerProperty;
 	TSharedPtr<IPropertyHandle> PriorityProperty;
 	TSharedPtr<IPropertyHandle> EventTagProperty;
@@ -38,7 +44,8 @@ private:
 	TSharedPtr<IPropertyHandle> DelayDurationProperty;
 	TSharedPtr<IPropertyHandle> DelayRandomVarianceProperty;
 	TSharedPtr<IPropertyHandle> ConditionsProperty;
+	TSharedPtr<IPropertyHandle> IDProperty;
 
-	class IPropertyUtilities* PropUtils;
+	TSharedPtr<IPropertyUtilities> PropUtils;
 	TSharedPtr<IPropertyHandle> StructProperty;
 };
diff --git a/Engine/Plugins/Runtime/StateTree/Source/StateTreeEditorModule/Private/Debugger/SStateTreeDebuggerInstanceTree.cpp b/Engine/Plugins/Runtime/StateTree/Source/StateTreeEditorModule/Private/Debugger/SStateTreeDebuggerInstanceTree.cpp
--- a/Engine/Plugins/Runtime/StateTree/Source/StateTreeEditorModule/Private/Debugger/SStateTreeDebuggerInstanceTree.cpp
+++ b/Engine/Plugins/Runtime/StateTree/Source/StateTreeEditorModule/Private/Debugger/SStateTreeDebuggerInstanceTree.cpp
@@ -82,6 +82,11 @@ void SStateTreeDebuggerInstanceTree::SetSelection(const TSharedPtr<RewindDebugge
 	TreeView->SetSelection(SelectedItem);
 }
 
+void SStateTreeDebuggerInstanceTree::ScrollTo(const TSharedPtr<RewindDebugger::FRewindDebuggerTrack>& SelectedItem) const
+{
+	TreeView->RequestScrollIntoView(SelectedItem);
+}
+
 void SStateTreeDebuggerInstanceTree::ScrollTo(const double ScrollOffset) const
 {
 	TreeView->SetScrollOffset(ScrollOffset);
diff --git a/Engine/Plugins/Runtime/StateTree/Source/StateTreeEditorModule/Private/Debugger/SStateTreeDebuggerInstanceTree.h b/Engine/Plugins/Runtime/StateTree/Source/StateTreeEditorModule/Private/Debugger/SStateTreeDebuggerInstanceTree.h
--- a/Engine/Plugins/Runtime/StateTree/Source/StateTreeEditorModule/Private/Debugger/SStateTreeDebuggerInstanceTree.h
+++ b/Engine/Plugins/Runtime/StateTree/Source/StateTreeEditorModule/Private/Debugger/SStateTreeDebuggerInstanceTree.h
@@ -39,6 +39,7 @@ public:
 
 	TSharedPtr<RewindDebugger::FRewindDebuggerTrack> GetSelection() const;
 	void SetSelection(const TSharedPtr<RewindDebugger::FRewindDebuggerTrack>& SelectedItem) const;
+	void ScrollTo(const TSharedPtr<RewindDebugger::FRewindDebuggerTrack>& SelectedItem) const;
 	void ScrollTo(double ScrollOffset) const;
 
 private:
diff --git a/Engine/Plugins/Runtime/StateTree/Source/StateTreeEditorModule/Private/Debugger/SStateTreeDebuggerView.cpp b/Engine/Plugins/Runtime/StateTree/Source/StateTreeEditorModule/Private/Debugger/SStateTreeDebuggerView.cpp
--- a/Engine/Plugins/Runtime/StateTree/Source/StateTreeEditorModule/Private/Debugger/SStateTreeDebuggerView.cpp
+++ b/Engine/Plugins/Runtime/StateTree/Source/StateTreeEditorModule/Private/Debugger/SStateTreeDebuggerView.cpp
@@ -981,11 +981,12 @@ void SStateTreeDebuggerView::OnNewInstance(FStateTreeInstanceDebugId InstanceId)
 
 		// Look at current selection; if nothing selected or stale track then select new track
 		const TSharedPtr<RewindDebugger::FRewindDebuggerTrack> Selection = InstancesTreeView->GetSelection();
-		if (const FStateTreeDebuggerBaseTrack* DebuggerBaseTrack = static_cast<FStateTreeDebuggerBaseTrack*>(ExistingOwnerTrack->Get()))
+		if (const FStateTreeDebuggerBaseTrack* DebuggerBaseTrack = static_cast<FStateTreeDebuggerBaseTrack*>(Selection.Get()))
 		{
 			if (DebuggerBaseTrack->IsStale())
 			{
 				InstancesTreeView->SetSelection(SubTrack);
+				InstancesTreeView->ScrollTo(SubTrack);
 			}
 		}
 		else
diff --git a/Engine/Plugins/Runtime/StateTree/Source/StateTreeEditorModule/Private/Debugger/StateTreeDebuggerUIExtensions.cpp b/Engine/Plugins/Runtime/StateTree/Source/StateTreeEditorModule/Private/Debugger/StateTreeDebuggerUIExtensions.cpp
--- a/Engine/Plugins/Runtime/StateTree/Source/StateTreeEditorModule/Private/Debugger/StateTreeDebuggerUIExtensions.cpp
+++ b/Engine/Plugins/Runtime/StateTree/Source/StateTreeEditorModule/Private/Debugger/StateTreeDebuggerUIExtensions.cpp
@@ -147,7 +147,173 @@ void OnTaskBreakpointToggled(TSharedPtr<IPropertyHandle> StructProperty, UStateT
 #endif // WITH_STATETREE_DEBUGGER
 }
 
-TSharedRef<SWidget> CreateStateWidget(IDetailLayoutBuilder& DetailBuilder, UStateTreeEditorData* TreeData)
+ECheckBoxState GetBreakpointCheckState(
+	const TSharedPtr<IPropertyHandle>& StructPropertyHandle,
+	const UStateTreeEditorData* EditorData,
+	TFunctionRef<ECheckBoxState(const UObject& OuterObject, const IPropertyHandle& PropertyHandle, const UStateTreeEditorData&)> Callback)
+{
+	ECheckBoxState CommonState = ECheckBoxState::Unchecked;
+#if WITH_STATETREE_DEBUGGER
+	if (StructPropertyHandle.IsValid() == false
+		|| StructPropertyHandle->IsValidHandle() == false
+		|| EditorData == nullptr)
+	{
+		return CommonState;
+	}
+
+	TArray<UObject*> OuterObjects;
+	StructPropertyHandle->GetOuterObjects(OuterObjects);
+
+	CommonState = ECheckBoxState::Undetermined;
+	for (const UObject* OuterObject : OuterObjects)
+	{
+		if (OuterObject != nullptr)
+		{
+			const ECheckBoxState State = Callback(*OuterObject, *StructPropertyHandle.Get(), *EditorData);
+			if (CommonState == ECheckBoxState::Undetermined)
+			{
+				CommonState = State;
+			}
+			else if (CommonState != State)
+			{
+				CommonState = ECheckBoxState::Undetermined;
+				break;
+			}
+		}
+	}
+#endif // WITH_STATETREE_DEBUGGER
+	return CommonState;
+}
+
+ECheckBoxState GetTransitionBreakpointCheckState(const TSharedPtr<IPropertyHandle>& StructPropertyHandle, const UStateTreeEditorData* EditorData)
+{
+#if WITH_STATETREE_DEBUGGER	
+	return GetBreakpointCheckState(StructPropertyHandle, EditorData,
+		[](const UObject& OuterObject, const IPropertyHandle& PropertyHandle, const UStateTreeEditorData& EditorData)
+		{
+			const UStateTreeState* TreeState = Cast<UStateTreeState>(&OuterObject);
+			const int32 IndexInArray = PropertyHandle.GetIndexInArray();
+			if (TreeState->Transitions.IsValidIndex(IndexInArray))
+			{
+				return EditorData.HasBreakpoint(TreeState->Transitions[IndexInArray].ID, EStateTreeBreakpointType::OnTransition)
+					? ECheckBoxState::Checked : ECheckBoxState::Unchecked;	
+			}
+
+			return  ECheckBoxState::Unchecked;
+		});
+#else
+	return  ECheckBoxState::Unchecked;
+#endif // WITH_STATETREE_DEBUGGER
+}
+
+ECheckBoxState GetTransitionEnabledCheckState(const TSharedPtr<IPropertyHandle>& StructPropertyHandle)
+{
+	const TSharedPtr<IPropertyHandle> EnabledProperty = StructPropertyHandle->GetChildHandle(GET_MEMBER_NAME_CHECKED(FStateTreeTransition, bTransitionEnabled));
+	if (ensureMsgf(EnabledProperty.IsValid() && EnabledProperty->IsValidHandle(),
+		TEXT("The property needs to be exposed to the Editor to allow us to use a property handle")))
+	{
+		bool bEnabled = false;
+
+		switch (EnabledProperty->GetValue(bEnabled))
+		{
+		case FPropertyAccess::MultipleValues:
+			return ECheckBoxState::Undetermined;
+		case FPropertyAccess::Success:
+			return bEnabled ? ECheckBoxState::Checked : ECheckBoxState::Unchecked;
+		case FPropertyAccess::Fail:
+		default:
+			return ECheckBoxState::Checked;
+		}
+	}
+	return ECheckBoxState::Checked;
+}
+
+void OnTransitionEnableToggled(const TSharedPtr<IPropertyHandle>& StructPropertyHandle)
+{
+	const TSharedPtr<IPropertyHandle> EnabledProperty = StructPropertyHandle->GetChildHandle(GET_MEMBER_NAME_CHECKED(FStateTreeTransition, bTransitionEnabled));
+	if (ensureMsgf(EnabledProperty.IsValid() && EnabledProperty->IsValidHandle(),
+		TEXT("The property needs to be exposed to the Editor to allow us to use a property handle")))
+	{
+		FScopedTransaction Transaction(LOCTEXT("OnStateEnableToggled", "Toggled State Enabled"));
+
+		bool bEnabled = false;
+		const FPropertyAccess::Result Result = EnabledProperty->GetValue(bEnabled);
+
+		if (Result == FPropertyAccess::MultipleValues || bEnabled == false)
+		{
+			EnabledProperty->SetValue(true);
+		}
+		else
+		{
+			EnabledProperty->SetValue(false);
+		}
+	}
+}
+
+void OnTransitionBreakpointToggled(const TSharedPtr<IPropertyHandle>& StructPropertyHandle, UStateTreeEditorData* EditorData)
+{	
+#if WITH_STATETREE_DEBUGGER
+	if (StructPropertyHandle.IsValid() == false
+		|| StructPropertyHandle->IsValidHandle() == false
+		|| EditorData == nullptr)
+	{
+		return;
+	}
+	
+	// Determine desired toggled state. Undermined will switch all to checked (i.e. Enable breakpoint)
+	const ECheckBoxState CommonState = GetTransitionBreakpointCheckState(StructPropertyHandle, EditorData);
+	const ECheckBoxState ToggledState = (CommonState == ECheckBoxState::Checked) ? ECheckBoxState::Unchecked : ECheckBoxState::Checked; 
+	
+	const IPropertyHandle& PropertyHandle = *(StructPropertyHandle.Get());
+
+	TArray<UObject*> OuterObjects;
+	PropertyHandle.GetOuterObjects(OuterObjects);
+
+	for (const UObject* OuterObject : OuterObjects)
+	{
+		if (const UStateTreeState* TreeState = Cast<UStateTreeState>(OuterObject))
+		{
+			const int32 IndexInArray = PropertyHandle.GetIndexInArray();
+
+			if (TreeState->Transitions.IsValidIndex(IndexInArray))
+			{
+				const FStateTreeTransition& Transition = TreeState->Transitions[IndexInArray];
+				const bool bHasBreakpoint = EditorData->HasBreakpoint(Transition.ID, EStateTreeBreakpointType::OnTransition);
+				EditorData->HasBreakpoint(Transition.ID, EStateTreeBreakpointType::OnTransition);
+				if (ToggledState == ECheckBoxState::Checked && bHasBreakpoint == false)
+				{
+					EditorData->AddBreakpoint(Transition.ID, EStateTreeBreakpointType::OnTransition);	
+				}
+				else if (ToggledState == ECheckBoxState::Unchecked && bHasBreakpoint)
+				{
+					EditorData->RemoveBreakpoint(Transition.ID, EStateTreeBreakpointType::OnTransition);
+				}		
+			}
+		}
+	}
+#endif // WITH_STATETREE_DEBUGGER
+}
+
+TSharedRef<SWidget> CreateDebugOptionsWidget(const TSharedRef<SWidget>& ContentWidget)
+{
+	return SNew(SComboButton)
+		.ToolTipText(LOCTEXT("DebugOptions", "Debug Options"))
+		.HasDownArrow(false)
+		.ContentPadding(0)
+		.ButtonStyle(FAppStyle::Get(), "SimpleButton")
+		.MenuContent()
+		[
+			ContentWidget
+		]
+		.ButtonContent()
+		[
+			SNew(SImage)
+			.Image(FStateTreeEditorStyle::Get().GetBrush("StateTreeEditor.DebugOptions"))
+			.ColorAndOpacity(FSlateColor::UseForeground())
+		];
+}
+
+TSharedRef<SWidget> CreateStateWidget(const IDetailLayoutBuilder& DetailBuilder, UStateTreeEditorData* TreeData)
 {
 	TWeakObjectPtr<UStateTreeEditorData> WeakTreeData = TreeData;
 	const TSharedPtr<IPropertyHandle> EnabledProperty = DetailBuilder.GetProperty(GET_MEMBER_NAME_CHECKED(UStateTreeState, bEnabled));
@@ -303,21 +469,7 @@ TSharedRef<SWidget> CreateStateWidget(IDetailLayoutBuilder& DetailBuilder, UStat
 		.VAlign(VAlign_Center)
 		.AutoWidth()
 		[
-			SNew(SComboButton)
-			.ToolTipText(LOCTEXT("DebugOptions", "Debug Options"))
-			.HasDownArrow(false)
-			.ContentPadding(0)
-			.ButtonStyle(FAppStyle::Get(), "SimpleButton")
-			.MenuContent()
-			[
-				DebuggingActions.MakeWidget()
-			]
-			.ButtonContent()
-			[
-				SNew(SImage)
-				.Image(FStateTreeEditorStyle::Get().GetBrush("StateTreeEditor.DebugOptions"))
-				.ColorAndOpacity(FSlateColor::UseForeground())
-			]
+			CreateDebugOptionsWidget(DebuggingActions.MakeWidget())
 		];
 
 	return HeaderContentWidget;
@@ -613,21 +765,123 @@ TSharedRef<SWidget> CreateEditorNodeWidget(const TSharedPtr<IPropertyHandle>& St
 		.VAlign(VAlign_Center)
 		.AutoWidth()
 		[
-			SNew(SComboButton)
-			.ToolTipText(LOCTEXT("DebugOptions", "Debug Options"))
-			.HasDownArrow(false)
-			.ContentPadding(0)
-			.ButtonStyle(FAppStyle::Get(), "SimpleButton")
-			.MenuContent()
+			CreateDebugOptionsWidget(DebuggingActions.MakeWidget())
+		];
+}
+
+TSharedRef<SWidget> CreateTransitionWidget(const TSharedPtr<IPropertyHandle>& StructPropertyHandle, UStateTreeEditorData* TreeData)
+{
+	using namespace EditorNodeUtils;
+
+	TWeakObjectPtr<UStateTreeEditorData> WeakTreeData = TreeData;
+
+	FMenuBuilder DebuggingActions(/*CloseWindowAfterMenuSelection*/ true, /*CommandList*/nullptr);
+	DebuggingActions.BeginSection(FName("Options"), LOCTEXT("TransitionOptions", "Transition Debug Options"));
+	DebuggingActions.AddMenuEntry
+	(
+		LOCTEXT("ToggleTransitionEnabled", "Transition Enabled"),
+		LOCTEXT("ToggleTransitionEnabled_ToolTip", "Enables or disables selected transition(s). StateTree must be recompiled to take effect."),
+		FSlateIcon(),
+		FUIAction
+		(
+			FExecuteAction::CreateLambda([StructPropertyHandle] { OnTransitionEnableToggled(StructPropertyHandle); }),
+			FCanExecuteAction(),
+			FGetActionCheckState::CreateLambda([StructPropertyHandle] { return GetTransitionEnabledCheckState(StructPropertyHandle); })
+		),
+		NAME_None,
+		EUserInterfaceActionType::Check
+	);
+	DebuggingActions.AddSeparator();
+	DebuggingActions.AddMenuEntry
+	(
+		LOCTEXT("ToggleTransitionBreakpoint", "Break on transition"),
+		LOCTEXT("ToggleTransitionBreakpoint_ToolTip", "Enables or disables breakpoint when executing the transition(s)."),
+		FSlateIcon(),
+		FUIAction
+		(
+			FExecuteAction::CreateLambda([StructPropertyHandle, WeakTreeData] { OnTransitionBreakpointToggled(StructPropertyHandle, WeakTreeData.Get()); }),
+			FCanExecuteAction(),
+			FGetActionCheckState::CreateLambda([StructPropertyHandle, WeakTreeData]
+				{
+					return GetTransitionBreakpointCheckState(StructPropertyHandle, WeakTreeData.Get());
+				})
+		),
+		NAME_None,
+		EUserInterfaceActionType::Check
+	);
+	DebuggingActions.EndSection();
+
+	return SNew(SHorizontalBox)
+		+ SHorizontalBox::Slot()
+		.FillWidth(1.0f)
+		.HAlign(HAlign_Right)
+		[
+			SNew(SHorizontalBox)
+			.Clipping(EWidgetClipping::OnDemand)
+
+			// Disabled Label
+			+ SHorizontalBox::Slot()
+			.AutoWidth()
+			.HAlign(HAlign_Left)
+			.VAlign(VAlign_Center)
+			.Padding(4.0f, 0.0f)
 			[
-				DebuggingActions.MakeWidget()
+				SNew(SBorder)
+				.Padding(FMargin(6, 2))
+				.BorderImage_Lambda([StructPropertyHandle]
+					{
+						return GetTransitionEnabledCheckState(StructPropertyHandle) == ECheckBoxState::Unchecked
+						? FStateTreeEditorStyle::Get().GetBrush("StateTree.Node.Label")
+						: FStateTreeEditorStyle::Get().GetBrush("StateTree.Node.Label.Mixed");
+					})
+				.Visibility_Lambda([StructPropertyHandle]
+					{
+						return GetTransitionEnabledCheckState(StructPropertyHandle) == ECheckBoxState::Checked ? EVisibility::Collapsed : EVisibility::Visible;
+					})
+				[
+					SNew(STextBlock)
+					.TextStyle(FStateTreeEditorStyle::Get(), "StateTree.Param.Label")
+					.ColorAndOpacity(FStyleColors::Foreground)
+					.Text(LOCTEXT("LabelDisabled", "DISABLED"))
+					.ToolTipText(LOCTEXT("DisabledTransitionTooltip", "This transition has been disabled."))
+				]
 			]
-			.ButtonContent()
+
+			// Break on transition Label
+			+ SHorizontalBox::Slot()
+			.AutoWidth()
+			.HAlign(HAlign_Left)
+			.VAlign(VAlign_Center)
+			.Padding(4.0f, 0.0f)
 			[
-				SNew(SImage)
-				.Image(FStateTreeEditorStyle::Get().GetBrush("StateTreeEditor.DebugOptions"))
-				.ColorAndOpacity(FSlateColor::UseForeground())
+				SNew(SBorder)
+				.Padding(FMargin(6, 2))
+				.BorderImage_Lambda([StructPropertyHandle, WeakTreeData]
+					{
+						return GetTransitionBreakpointCheckState(StructPropertyHandle, WeakTreeData.Get()) == ECheckBoxState::Checked
+						? FStateTreeEditorStyle::Get().GetBrush("StateTree.Node.Label")
+						: FStateTreeEditorStyle::Get().GetBrush("StateTree.Node.Label.Mixed");
+					})
+				.Visibility_Lambda([StructPropertyHandle, WeakTreeData]
+					{
+						return GetTransitionBreakpointCheckState(StructPropertyHandle, WeakTreeData.Get()) == ECheckBoxState::Unchecked ? EVisibility::Collapsed : EVisibility::Visible;
+					})
+				[
+					SNew(STextBlock)
+					.TextStyle(FStateTreeEditorStyle::Get(), "StateTree.Param.Label")
+					.ColorAndOpacity(FStyleColors::Foreground)
+					.Text(LOCTEXT("LabelTransitionBreakpoint", "BREAK ON TRANSITION"))
+					.ToolTipText(LOCTEXT("BreakTransitionTooltip", "Break when executing this transition."))
+				]
 			]
+		]
+		+ SHorizontalBox::Slot()
+		.Padding(0,0,6,0)
+		.HAlign(HAlign_Right)
+		.VAlign(VAlign_Center)
+		.AutoWidth()
+		[
+			CreateDebugOptionsWidget(DebuggingActions.MakeWidget())
 		];
 }
 
diff --git a/Engine/Plugins/Runtime/StateTree/Source/StateTreeEditorModule/Private/Debugger/StateTreeDebuggerUIExtensions.h b/Engine/Plugins/Runtime/StateTree/Source/StateTreeEditorModule/Private/Debugger/StateTreeDebuggerUIExtensions.h
--- a/Engine/Plugins/Runtime/StateTree/Source/StateTreeEditorModule/Private/Debugger/StateTreeDebuggerUIExtensions.h
+++ b/Engine/Plugins/Runtime/StateTree/Source/StateTreeEditorModule/Private/Debugger/StateTreeDebuggerUIExtensions.h
@@ -17,7 +17,8 @@ class UStateTreeEditorData;
 namespace UE::StateTreeEditor::DebuggerExtensions
 {
 
-TSharedRef<SWidget> CreateStateWidget(IDetailLayoutBuilder& DetailBuilder, UStateTreeEditorData* TreeData);
+TSharedRef<SWidget> CreateStateWidget(const IDetailLayoutBuilder& DetailBuilder, UStateTreeEditorData* TreeData);
 TSharedRef<SWidget> CreateEditorNodeWidget(const TSharedPtr<IPropertyHandle>& StructPropertyHandle, UStateTreeEditorData* TreeData);
+TSharedRef<SWidget> CreateTransitionWidget(const TSharedPtr<IPropertyHandle>& StructPropertyHandle, UStateTreeEditorData* TreeData);
 
 }; // UE::StateTreeEditor::DebuggerExtensions
diff --git a/Engine/Plugins/Runtime/StateTree/Source/StateTreeEditorModule/Private/SStateTreeViewRow.cpp b/Engine/Plugins/Runtime/StateTree/Source/StateTreeEditorModule/Private/SStateTreeViewRow.cpp
--- a/Engine/Plugins/Runtime/StateTree/Source/StateTreeEditorModule/Private/SStateTreeViewRow.cpp
+++ b/Engine/Plugins/Runtime/StateTree/Source/StateTreeEditorModule/Private/SStateTreeViewRow.cpp
@@ -5,6 +5,7 @@
 
 
 #include "EditorFontGlyphs.h"
+#include "StateTreeEditor.h"
 #include "StateTreeEditorData.h"
 #include "StateTreeEditorStyle.h"
 
@@ -259,10 +260,38 @@ void SStateTreeViewRow::Construct(const FArguments& InArgs, const TSharedRef<STa
 				.AutoWidth()
 				.Padding(FMargin(8.0f, 0.0f, 0, 0.0f))
 				[
-					SNew(STextBlock)
-					.Text(this, &SStateTreeViewRow::GetCompletedTransitionsIcon)
-					.TextStyle(FStateTreeEditorStyle::Get(), "StateTree.Icon")
-					.Visibility(this, &SStateTreeViewRow::GetCompletedTransitionVisibility)
+					SNew(SOverlay)
+					+ SOverlay::Slot()
+					[
+						SNew(STextBlock)
+						.Text(this, &SStateTreeViewRow::GetCompletedTransitionsIcon)
+						.TextStyle(FStateTreeEditorStyle::Get(), "StateTree.Icon")
+						.Visibility(this, &SStateTreeViewRow::GetCompletedTransitionVisibility)
+					]
+					+ SOverlay::Slot()
+					[
+						// Breakpoint box
+						SNew(SHorizontalBox)
+						+ SHorizontalBox::Slot()
+						.VAlign(VAlign_Top)
+						.HAlign(HAlign_Left)
+						.AutoWidth()
+						[
+							SNew(SBox)
+							.Padding(FMargin(0.0f, -10.0f, 0.0f, 0.0f))
+							[
+								SNew(SImage)
+								.DesiredSizeOverride(FVector2D(10.f, 10.f))
+								.Image(FStateTreeEditorStyle::Get().GetBrush(TEXT("StateTreeEditor.Debugger.Breakpoint.EnabledAndValid")))
+								.Visibility(this, &SStateTreeViewRow::GetCompletedTransitionBreakpointVisibility)
+								.ToolTipText_Lambda([this]
+									{
+										return FText::Format(LOCTEXT("TransitionBreakpointTooltip","Break when executing transition: {0}"),
+											GetCompletedTransitionWithBreakpointDesc());
+									})
+							]
+						]
+					]
 				]
 				+ SHorizontalBox::Slot()
 				.VAlign(VAlign_Center)
@@ -298,10 +327,38 @@ void SStateTreeViewRow::Construct(const FArguments& InArgs, const TSharedRef<STa
 				.Padding(FMargin(4.0f, 0.0f, 0, 0))
 				.AutoWidth()
 				[
-					SNew(STextBlock)
-					.Text(this, &SStateTreeViewRow::GetSucceededTransitionIcon)
-					.TextStyle(FStateTreeEditorStyle::Get(), "StateTree.Icon")
-					.Visibility(this, &SStateTreeViewRow::GetSucceededTransitionVisibility)
+					SNew(SOverlay)
+					+ SOverlay::Slot()
+					[
+						SNew(STextBlock)
+						.Text(this, &SStateTreeViewRow::GetSucceededTransitionIcon)
+						.TextStyle(FStateTreeEditorStyle::Get(), "StateTree.Icon")
+						.Visibility(this, &SStateTreeViewRow::GetSucceededTransitionVisibility)
+					]
+					+ SOverlay::Slot()
+					[
+						// Breakpoint box
+						SNew(SHorizontalBox)
+						+ SHorizontalBox::Slot()
+						.VAlign(VAlign_Top)
+						.HAlign(HAlign_Left)
+						.AutoWidth()
+						[
+							SNew(SBox)
+							.Padding(FMargin(0.0f, -10.0f, 0.0f, 0.0f))
+							[
+								SNew(SImage)
+								.DesiredSizeOverride(FVector2D(10.f, 10.f))
+								.Image(FStateTreeEditorStyle::Get().GetBrush(TEXT("StateTreeEditor.Debugger.Breakpoint.EnabledAndValid")))
+								.Visibility(this, &SStateTreeViewRow::GetSucceededTransitionBreakpointVisibility)
+								.ToolTipText_Lambda([this]
+									{
+										return FText::Format(LOCTEXT("TransitionBreakpointTooltip", "Break when executing transition: {0}"),
+											GetSucceededTransitionWithBreakpointDesc());
+									})
+							]
+						]
+					]
 				]
 				+ SHorizontalBox::Slot()
 				.VAlign(VAlign_Center)
@@ -337,10 +394,38 @@ void SStateTreeViewRow::Construct(const FArguments& InArgs, const TSharedRef<STa
 				.Padding(FMargin(4.0f, 0.0f, 0, 0))
 				.AutoWidth()
 				[
-					SNew(STextBlock)
-					.Text(this, &SStateTreeViewRow::GetFailedTransitionIcon)
-					.TextStyle(FStateTreeEditorStyle::Get(), "StateTree.Icon")
-					.Visibility(this, &SStateTreeViewRow::GetFailedTransitionVisibility)
+					SNew(SOverlay)
+					+ SOverlay::Slot()
+					[
+						SNew(STextBlock)
+						.Text(this, &SStateTreeViewRow::GetFailedTransitionIcon)
+						.TextStyle(FStateTreeEditorStyle::Get(), "StateTree.Icon")
+						.Visibility(this, &SStateTreeViewRow::GetFailedTransitionVisibility)
+					]
+					+ SOverlay::Slot()
+					[
+						// Breakpoint box
+						SNew(SHorizontalBox)
+						+ SHorizontalBox::Slot()
+						.VAlign(VAlign_Top)
+						.HAlign(HAlign_Left)
+						.AutoWidth()
+						[
+							SNew(SBox)
+							.Padding(FMargin(0.0f, -10.0f, 0.0f, 0.0f))
+							[
+								SNew(SImage)
+								.DesiredSizeOverride(FVector2D(10.f, 10.f))
+								.Image(FStateTreeEditorStyle::Get().GetBrush(TEXT("StateTreeEditor.Debugger.Breakpoint.EnabledAndValid")))
+								.Visibility(this, &SStateTreeViewRow::GetFailedTransitionBreakpointVisibility)
+								.ToolTipText_Lambda([this]
+									{
+										return FText::Format(LOCTEXT("TransitionBreakpointTooltip", "Break when executing transition: {0}"),
+											GetFailedTransitionWithBreakpointDesc());
+									})
+							]
+						]
+					]
 				]
 				+ SHorizontalBox::Slot()
 				.VAlign(VAlign_Center)
@@ -353,6 +438,7 @@ void SStateTreeViewRow::Construct(const FArguments& InArgs, const TSharedRef<STa
 					.Visibility(this, &SStateTreeViewRow::GetFailedTransitionVisibility)
 				]
 			]
+
 			// Transitions
 			+ SHorizontalBox::Slot()
 			.VAlign(VAlign_Center)
@@ -374,10 +460,38 @@ void SStateTreeViewRow::Construct(const FArguments& InArgs, const TSharedRef<STa
 				.Padding(FMargin(4.0f, 0.0f, 0, 0))
 				.AutoWidth()
 				[
-					SNew(STextBlock)
-					.Text(FEditorFontGlyphs::Long_Arrow_Right)
-					.TextStyle(FStateTreeEditorStyle::Get(), "StateTree.Icon")
-					.Visibility(this, &SStateTreeViewRow::GetConditionalTransitionsVisibility)
+					SNew(SOverlay)
+					+ SOverlay::Slot()
+					[
+						SNew(STextBlock)
+						.Text(FEditorFontGlyphs::Long_Arrow_Right)
+						.TextStyle(FStateTreeEditorStyle::Get(), "StateTree.Icon")
+						.Visibility(this, &SStateTreeViewRow::GetConditionalTransitionsVisibility)
+					]
+					+ SOverlay::Slot()
+					[
+						// Breakpoint box
+						SNew(SHorizontalBox)
+						+ SHorizontalBox::Slot()
+						.VAlign(VAlign_Top)
+						.HAlign(HAlign_Left)
+						.AutoWidth()
+						[
+							SNew(SBox)
+							.Padding(FMargin(0.0f, -10.0f, 0.0f, 0.0f))
+							[
+								SNew(SImage)
+								.DesiredSizeOverride(FVector2D(10.f, 10.f))
+								.Image(FStateTreeEditorStyle::Get().GetBrush(TEXT("StateTreeEditor.Debugger.Breakpoint.EnabledAndValid")))
+								.Visibility(this, &SStateTreeViewRow::GetConditionalTransitionsBreakpointVisibility)
+								.ToolTipText_Lambda([this]
+									{
+										return FText::Format(LOCTEXT("TransitionBreakpointTooltip", "Break when executing transition: {0}"),
+											GetConditionalTransitionsWithBreakpointDesc());
+									})
+							]
+						]
+					]
 				]
 				+ SHorizontalBox::Slot()
 				.VAlign(VAlign_Center)
@@ -462,10 +576,15 @@ TSharedRef<SHorizontalBox> SStateTreeViewRow::CreateTasksWidget()
 					return FText::GetEmpty();
 				};
 
-			FText TaskName = FText::FromName(Task->Name);
-
-			// Uncomment to visualize Task Ids
-			// TaskName = FText::FromString(FString::Printf(TEXT("%s (%s)"), *Task->Name.ToString(), *LexToString(TaskId)));
+			FText TaskName;
+			if (UE::StateTree::Editor::GbDisplayItemIds)
+			{
+				TaskName = FText::FromString(FString::Printf(TEXT("%s (%s)"), *Task->Name.ToString(), *LexToString(TaskId)));
+			}
+			else
+			{
+				TaskName = FText::FromName(Task->Name);
+			}
 
 			TasksBox->AddSlot()
 				.AutoWidth()
@@ -571,6 +690,10 @@ FText SStateTreeViewRow::GetStateDesc() const
 {
 	if (const UStateTreeState* State = WeakState.Get())
 	{
+		if (UE::StateTree::Editor::GbDisplayItemIds)
+		{
+			return FText::FromString(FString::Printf(TEXT("%s (%s)"), *State->Name.ToString(), *LexToString(State->ID)));	
+		}
 		return FText::FromName(State->Name);
 	}
 	return FText::FromName(FName());
@@ -795,13 +918,31 @@ FText SStateTreeViewRow::GetLinkDescription(const FStateTreeStateLink& Link)
 	return FText::GetEmpty();
 };
 
-FText SStateTreeViewRow::GetTransitionsDesc(const UStateTreeState& State, const EStateTreeTransitionTrigger Trigger, const bool bUseMask) const
+FText SStateTreeViewRow::GetTransitionsDesc(const UStateTreeState& State, const EStateTreeTransitionTrigger Trigger, const FTransitionDescFilterOptions FilterOptions) const
 {
 	TArray<FText> DescItems;
-	
+	const UStateTreeEditorData* TreeEditorData = WeakTreeData.Get();
+
 	for (const FStateTreeTransition& Transition : State.Transitions)
 	{
-		const bool bMatch = bUseMask ? EnumHasAnyFlags(Transition.Trigger, Trigger) : Transition.Trigger == Trigger;
+		// Apply filter for enabled/disabled transitions
+		if ((FilterOptions.Enabled == ETransitionDescRequirement::RequiredTrue && Transition.bTransitionEnabled == false)
+			|| (FilterOptions.Enabled == ETransitionDescRequirement::RequiredFalse && Transition.bTransitionEnabled))
+		{
+			continue;
+		}
+
+#if WITH_STATETREE_DEBUGGER
+		// Apply filter for transitions with/without breakpoint
+		const bool bHasBreakpoint = TreeEditorData != nullptr && TreeEditorData->HasBreakpoint(Transition.ID, EStateTreeBreakpointType::OnTransition);
+		if ((FilterOptions.WithBreakpoint == ETransitionDescRequirement::RequiredTrue && bHasBreakpoint == false)
+			|| (FilterOptions.WithBreakpoint == ETransitionDescRequirement::RequiredFalse && bHasBreakpoint))
+		{
+			continue;
+		}
+#endif // WITH_STATETREE_DEBUGGER
+
+		const bool bMatch = FilterOptions.bUseMask ? EnumHasAnyFlags(Transition.Trigger, Trigger) : Transition.Trigger == Trigger;
 		if (bMatch)
 		{
 			DescItems.Add(GetLinkDescription(Transition.State));
@@ -859,7 +1000,7 @@ FText SStateTreeViewRow::GetTransitionsDesc(const UStateTreeState& State, const
 	return FText::Join(FText::FromString(TEXT(", ")), DescItems);
 }
 
-FText SStateTreeViewRow::GetTransitionsIcon(const UStateTreeState& State, const EStateTreeTransitionTrigger Trigger, const bool bUseMask) const
+FText SStateTreeViewRow::GetTransitionsIcon(const UStateTreeState& State, const EStateTreeTransitionTrigger Trigger, const FTransitionDescFilterOptions FilterOptions) const
 {
 	enum EIconType
 	{
@@ -871,10 +1012,29 @@ FText SStateTreeViewRow::GetTransitionsIcon(const UStateTreeState& State, const
 	};
 	uint8 IconType = IconNone;
 	
+	const UStateTreeEditorData* TreeEditorData = WeakTreeData.Get();
+	
 	for (const FStateTreeTransition& Transition : State.Transitions)
 	{
+		// Apply filter for enabled/disabled transitions
+		if ((FilterOptions.Enabled == ETransitionDescRequirement::RequiredTrue && Transition.bTransitionEnabled == false)
+			|| (FilterOptions.Enabled == ETransitionDescRequirement::RequiredFalse && Transition.bTransitionEnabled))
+		{
+			continue;
+		}
+
+#if WITH_STATETREE_DEBUGGER
+		// Apply filter for transitions with/without breakpoint
+		const bool bHasBreakpoint = TreeEditorData != nullptr && TreeEditorData->HasBreakpoint(Transition.ID, EStateTreeBreakpointType::OnTransition);
+		if ((FilterOptions.WithBreakpoint == ETransitionDescRequirement::RequiredTrue && bHasBreakpoint == false)
+			|| (FilterOptions.WithBreakpoint == ETransitionDescRequirement::RequiredFalse && bHasBreakpoint))
+		{
+			continue;
+		}
+#endif // WITH_STATETREE_DEBUGGER
+		
 		// The icons here depict "transition direction", not the type specifically.
-		const bool bMatch = bUseMask ? EnumHasAnyFlags(Transition.Trigger, Trigger) : Transition.Trigger == Trigger;
+		const bool bMatch = FilterOptions.bUseMask ? EnumHasAnyFlags(Transition.Trigger, Trigger) : Transition.Trigger == Trigger;
 		if (bMatch)
 		{
 			switch (Transition.State.LinkType)
@@ -950,8 +1110,19 @@ EVisibility SStateTreeViewRow::GetTransitionsVisibility(const UStateTreeState& S
 
 		for (const FStateTreeTransition& Transition : State.Transitions)
 		{
+			// Skip disabled transitions
+			if (Transition.bTransitionEnabled == false)
+			{
+				continue;
+			}
+
 			HandledTriggers |= Transition.Trigger;
 			bExactMatch |= (Transition.Trigger == Trigger);
+
+			if (bExactMatch)
+			{
+				break;
+			}
 		}
 
 		// Assume that leaf states should have completion transitions.
@@ -1008,6 +1179,12 @@ EVisibility SStateTreeViewRow::GetTransitionsVisibility(const UStateTreeState& S
 	// Handle the test
 	for (const FStateTreeTransition& Transition : State.Transitions)
 	{
+		// Skip disabled transitions
+		if (Transition.bTransitionEnabled == false)
+		{
+			continue;
+		}
+
 		if (EnumHasAnyFlags(Trigger, Transition.Trigger))
 		{
 			return EVisibility::Visible;
@@ -1016,6 +1193,27 @@ EVisibility SStateTreeViewRow::GetTransitionsVisibility(const UStateTreeState& S
 	return EVisibility::Collapsed;
 }
 
+EVisibility SStateTreeViewRow::GetTransitionsBreakpointVisibility(const UStateTreeState& State, const EStateTreeTransitionTrigger Trigger) const
+{
+#if WITH_STATETREE_DEBUGGER
+	if (const UStateTreeEditorData* TreeEditorData = WeakTreeData.Get())
+	{
+		for (const FStateTreeTransition& Transition : State.Transitions)
+		{
+			if (Transition.bTransitionEnabled && EnumHasAnyFlags(Trigger, Transition.Trigger))
+			{
+				if (TreeEditorData->HasBreakpoint(Transition.ID, EStateTreeBreakpointType::OnTransition))
+				{
+					return GetTransitionsVisibility(State, Trigger);
+				}
+			}
+		}
+	}
+#endif // WITH_STATETREE_DEBUGGER
+	
+	return EVisibility::Collapsed;
+}
+
 EVisibility SStateTreeViewRow::GetCompletedTransitionVisibility() const
 {
 	if (const UStateTreeState* State = WeakState.Get())
@@ -1025,6 +1223,15 @@ EVisibility SStateTreeViewRow::GetCompletedTransitionVisibility() const
 	return EVisibility::Visible;
 }
 
+EVisibility SStateTreeViewRow::GetCompletedTransitionBreakpointVisibility() const
+{
+	if (const UStateTreeState* State = WeakState.Get())
+	{
+		return GetTransitionsBreakpointVisibility(*State, EStateTreeTransitionTrigger::OnStateCompleted);
+	}
+	return EVisibility::Visible;
+}
+
 FText SStateTreeViewRow::GetCompletedTransitionsDesc() const
 {
 	if (const UStateTreeState* State = WeakState.Get())
@@ -1034,6 +1241,17 @@ FText SStateTreeViewRow::GetCompletedTransitionsDesc() const
 	return LOCTEXT("Invalid", "Invalid");
 }
 
+FText SStateTreeViewRow::GetCompletedTransitionWithBreakpointDesc() const
+{
+	if (const UStateTreeState* State = WeakState.Get())
+	{
+		FTransitionDescFilterOptions FilterOptions;
+		FilterOptions.WithBreakpoint = ETransitionDescRequirement::RequiredTrue;
+		return GetTransitionsDesc(*State, EStateTreeTransitionTrigger::OnStateCompleted, FilterOptions);
+	}
+	return FText::GetEmpty();
+}
+
 FText SStateTreeViewRow::GetCompletedTransitionsIcon() const
 {
 	if (const UStateTreeState* State = WeakState.Get())
@@ -1052,6 +1270,15 @@ EVisibility SStateTreeViewRow::GetSucceededTransitionVisibility() const
 	return EVisibility::Collapsed;
 }
 
+EVisibility SStateTreeViewRow::GetSucceededTransitionBreakpointVisibility() const
+{
+	if (const UStateTreeState* State = WeakState.Get())
+	{
+		return GetTransitionsBreakpointVisibility(*State, EStateTreeTransitionTrigger::OnStateSucceeded);
+	}
+	return EVisibility::Collapsed;
+}
+
 FText SStateTreeViewRow::GetSucceededTransitionDesc() const
 {
 	if (const UStateTreeState* State = WeakState.Get())
@@ -1061,6 +1288,17 @@ FText SStateTreeViewRow::GetSucceededTransitionDesc() const
 	return FText::GetEmpty();
 }
 
+FText SStateTreeViewRow::GetSucceededTransitionWithBreakpointDesc() const
+{
+	if (const UStateTreeState* State = WeakState.Get())
+	{
+		FTransitionDescFilterOptions FilterOptions;
+		FilterOptions.WithBreakpoint = ETransitionDescRequirement::RequiredTrue;
+		return GetTransitionsDesc(*State, EStateTreeTransitionTrigger::OnStateSucceeded, FilterOptions);
+	}
+	return FText::GetEmpty();
+}
+
 FText SStateTreeViewRow::GetSucceededTransitionIcon() const
 {
 	if (const UStateTreeState* State = WeakState.Get())
@@ -1079,6 +1317,15 @@ EVisibility SStateTreeViewRow::GetFailedTransitionVisibility() const
 	return EVisibility::Collapsed;
 }
 
+EVisibility SStateTreeViewRow::GetFailedTransitionBreakpointVisibility() const
+{
+	if (const UStateTreeState* State = WeakState.Get())
+	{
+		return GetTransitionsBreakpointVisibility(*State, EStateTreeTransitionTrigger::OnStateFailed);
+	}
+	return EVisibility::Collapsed;
+}
+
 FText SStateTreeViewRow::GetFailedTransitionDesc() const
 {
 	if (const UStateTreeState* State = WeakState.Get())
@@ -1088,6 +1335,17 @@ FText SStateTreeViewRow::GetFailedTransitionDesc() const
 	return LOCTEXT("Invalid", "Invalid");
 }
 
+FText SStateTreeViewRow::GetFailedTransitionWithBreakpointDesc() const
+{
+	if (const UStateTreeState* State = WeakState.Get())
+	{
+		FTransitionDescFilterOptions FilterOptions;
+		FilterOptions.WithBreakpoint = ETransitionDescRequirement::RequiredTrue;
+		return GetTransitionsDesc(*State, EStateTreeTransitionTrigger::OnStateFailed, FilterOptions);
+	}
+	return FText::GetEmpty();
+}
+
 FText SStateTreeViewRow::GetFailedTransitionIcon() const
 {
 	if (const UStateTreeState* State = WeakState.Get())
@@ -1106,11 +1364,34 @@ EVisibility SStateTreeViewRow::GetConditionalTransitionsVisibility() const
 	return EVisibility::Collapsed;
 }
 
+EVisibility SStateTreeViewRow::GetConditionalTransitionsBreakpointVisibility() const
+{
+	if (const UStateTreeState* State = WeakState.Get())
+	{
+		return GetTransitionsBreakpointVisibility(*State, EStateTreeTransitionTrigger::OnTick | EStateTreeTransitionTrigger::OnEvent);
+	}
+	return EVisibility::Collapsed;
+}
+
 FText SStateTreeViewRow::GetConditionalTransitionsDesc() const
 {
 	if (const UStateTreeState* State = WeakState.Get())
 	{
-		return GetTransitionsDesc(*State, EStateTreeTransitionTrigger::OnTick | EStateTreeTransitionTrigger::OnEvent, /*bUseMask*/true);
+		FTransitionDescFilterOptions FilterOptions;
+		FilterOptions.bUseMask = true;
+		return GetTransitionsDesc(*State, EStateTreeTransitionTrigger::OnTick | EStateTreeTransitionTrigger::OnEvent, FilterOptions);
+	}
+	return FText::GetEmpty();
+}
+
+FText SStateTreeViewRow::GetConditionalTransitionsWithBreakpointDesc() const
+{
+	if (const UStateTreeState* State = WeakState.Get())
+	{
+		FTransitionDescFilterOptions FilterOptions;
+		FilterOptions.WithBreakpoint = ETransitionDescRequirement::RequiredTrue;
+		FilterOptions.bUseMask = true;
+		return GetTransitionsDesc(*State, EStateTreeTransitionTrigger::OnTick | EStateTreeTransitionTrigger::OnEvent, FilterOptions);
 	}
 	return FText::GetEmpty();
 }
diff --git a/Engine/Plugins/Runtime/StateTree/Source/StateTreeEditorModule/Private/SStateTreeViewRow.h b/Engine/Plugins/Runtime/StateTree/Source/StateTreeEditorModule/Private/SStateTreeViewRow.h
--- a/Engine/Plugins/Runtime/StateTree/Source/StateTreeEditorModule/Private/SStateTreeViewRow.h
+++ b/Engine/Plugins/Runtime/StateTree/Source/StateTreeEditorModule/Private/SStateTreeViewRow.h
@@ -50,24 +50,53 @@ private:
 	FText GetLinkedStateDesc() const;
 
 	EVisibility GetCompletedTransitionVisibility() const;
+	EVisibility GetCompletedTransitionBreakpointVisibility() const;
 	FText GetCompletedTransitionsDesc() const;
+	FText GetCompletedTransitionWithBreakpointDesc() const;
 	FText GetCompletedTransitionsIcon() const;
 
 	EVisibility GetSucceededTransitionVisibility() const;
+	EVisibility GetSucceededTransitionBreakpointVisibility() const;
 	FText GetSucceededTransitionDesc() const;
+	FText GetSucceededTransitionWithBreakpointDesc() const;
 	FText GetSucceededTransitionIcon() const;
 
 	EVisibility GetFailedTransitionVisibility() const;
+	EVisibility GetFailedTransitionBreakpointVisibility() const;
 	FText GetFailedTransitionDesc() const;
+	FText GetFailedTransitionWithBreakpointDesc() const;
 	FText GetFailedTransitionIcon() const;
 
 	EVisibility GetConditionalTransitionsVisibility() const;
+	EVisibility GetConditionalTransitionsBreakpointVisibility() const;
 	FText GetConditionalTransitionsDesc() const;
+	FText GetConditionalTransitionsWithBreakpointDesc() const;
+
+	enum class ETransitionDescRequirement : uint8
+	{
+		Any,
+		RequiredTrue,
+		RequiredFalse,
+	};
+
+	/**
+	 * Filtering options used to build the description of the transitions.
+	 * The default setup includes only enabled transition,
+	 * with or without breakpoints and requires exact trigger match (no partial mask)
+	 */
+	struct FTransitionDescFilterOptions
+	{
+		FTransitionDescFilterOptions() {};
+		ETransitionDescRequirement Enabled = ETransitionDescRequirement::RequiredTrue;
+		ETransitionDescRequirement WithBreakpoint = ETransitionDescRequirement::Any;
+		bool bUseMask = false;
+	};
 
 	static FText GetLinkDescription(const FStateTreeStateLink& Link);
-	FText GetTransitionsDesc(const UStateTreeState& State, const EStateTreeTransitionTrigger Trigger, const bool bUseMask = false) const;
-	FText GetTransitionsIcon(const UStateTreeState& State, const EStateTreeTransitionTrigger Trigger, const bool bUseMask = false) const;
+	FText GetTransitionsDesc(const UStateTreeState& State, const EStateTreeTransitionTrigger Trigger, const FTransitionDescFilterOptions FilterOptions = {}) const;
+	FText GetTransitionsIcon(const UStateTreeState& State, const EStateTreeTransitionTrigger Trigger, const FTransitionDescFilterOptions FilterOptions = {}) const;
 	EVisibility GetTransitionsVisibility(const UStateTreeState& State, const EStateTreeTransitionTrigger Trigger) const;
+	EVisibility GetTransitionsBreakpointVisibility(const UStateTreeState& State, const EStateTreeTransitionTrigger Trigger) const;
 
 	bool HasParentTransitionForTrigger(const UStateTreeState& State, const EStateTreeTransitionTrigger Trigger) const;
 
diff --git a/Engine/Plugins/Runtime/StateTree/Source/StateTreeEditorModule/Private/StateTreeCompiler.cpp b/Engine/Plugins/Runtime/StateTree/Source/StateTreeEditorModule/Private/StateTreeCompiler.cpp
--- a/Engine/Plugins/Runtime/StateTree/Source/StateTreeEditorModule/Private/StateTreeCompiler.cpp
+++ b/Engine/Plugins/Runtime/StateTree/Source/StateTreeEditorModule/Private/StateTreeCompiler.cpp
@@ -340,6 +340,12 @@ bool FStateTreeCompiler::Compile(UStateTree& InStateTree)
 		StateTree->IDToStateMappings.Emplace(ToState.Key, FStateTreeStateHandle(ToState.Value));
 	}
 
+	// Store mapping between state transition identifier and compact transition index. Used for debugging purposes.
+	for (const TPair<FGuid, int32>& ToTransition: IDToTransition)
+	{
+		StateTree->IDToTransitionMappings.Emplace(ToTransition.Key, FStateTreeIndex16(ToTransition.Value));
+	}
+
 	UE::StateTree::Compiler::FCheckOutersArchive CheckOuters(*StateTree, *EditorData, Log);
 	StateTree->Serialize(CheckOuters);
 	
@@ -790,10 +796,13 @@ bool FStateTreeCompiler::CreateStateTransitions()
 		
 		for (FStateTreeTransition& Transition : SourceState->Transitions)
 		{
+			IDToTransition.Add(Transition.ID, StateTree->Transitions.Num());
+
 			FCompactStateTransition& CompactTransition = StateTree->Transitions.AddDefaulted_GetRef();
 			CompactTransition.Trigger = Transition.Trigger;
 			CompactTransition.Priority = Transition.Priority;
 			CompactTransition.EventTag = Transition.EventTag;
+			CompactTransition.bTransitionEnabled = Transition.bTransitionEnabled;
 			
 			if (Transition.bDelayTransition)
 			{
diff --git a/Engine/Plugins/Runtime/StateTree/Source/StateTreeEditorModule/Private/StateTreeEditor.cpp b/Engine/Plugins/Runtime/StateTree/Source/StateTreeEditorModule/Private/StateTreeEditor.cpp
--- a/Engine/Plugins/Runtime/StateTree/Source/StateTreeEditorModule/Private/StateTreeEditor.cpp
+++ b/Engine/Plugins/Runtime/StateTree/Source/StateTreeEditorModule/Private/StateTreeEditor.cpp
@@ -44,6 +44,17 @@ const FName FStateTreeEditor::CompilerResultsTabId(TEXT("StateTreeEditor_Compile
 const FName FStateTreeEditor::DebuggerTabId(TEXT("StateTreeEditor_Debugger"));
 #endif // WITH_STATETREE_DEBUGGER
 
+
+namespace UE::StateTree::Editor
+{
+bool GbDisplayItemIds = false;
+
+FAutoConsoleVariableRef CVarDisplayItemIds(
+	TEXT("statetree.displayitemids"),
+	GbDisplayItemIds,
+	TEXT("Appends Id to task and state names in the treeview and expose Ids in the details view."));
+}
+
 void FStateTreeEditor::PostUndo(bool bSuccess)
 {
 }
diff --git a/Engine/Plugins/Runtime/StateTree/Source/StateTreeEditorModule/Private/StateTreeEditor.h b/Engine/Plugins/Runtime/StateTree/Source/StateTreeEditorModule/Private/StateTreeEditor.h
--- a/Engine/Plugins/Runtime/StateTree/Source/StateTreeEditorModule/Private/StateTreeEditor.h
+++ b/Engine/Plugins/Runtime/StateTree/Source/StateTreeEditorModule/Private/StateTreeEditor.h
@@ -20,6 +20,8 @@ namespace UE::StateTree::Editor
 
 	// Calculates editor data hash of the asset. 
 	uint32 CalcAssetHash(const UStateTree& StateTree);
+
+	extern bool GbDisplayItemIds;
 } // UE::StateTree::Editor
 
 
diff --git a/Engine/Plugins/Runtime/StateTree/Source/StateTreeEditorModule/Private/StateTreeEditorStyle.cpp b/Engine/Plugins/Runtime/StateTree/Source/StateTreeEditorModule/Private/StateTreeEditorStyle.cpp
--- a/Engine/Plugins/Runtime/StateTree/Source/StateTreeEditorModule/Private/StateTreeEditorStyle.cpp
+++ b/Engine/Plugins/Runtime/StateTree/Source/StateTreeEditorModule/Private/StateTreeEditorStyle.cpp
@@ -96,6 +96,13 @@ FStateTreeEditorStyle::FStateTreeEditorStyle()
 		Set("StateTree.Details", Details);
 
 		Set("StateTree.Node.Label", new FSlateRoundedBoxBrush(FStyleColors::AccentGray, 6.f));
+
+		// For multi selection with mixed values for a given property
+		Set("StateTree.Node.Label.Mixed", new FSlateRoundedBoxBrush(FStyleColors::Dropdown, 6.f));
+
+		const FLinearColor Color = FStyleColors::Hover.GetSpecifiedColor();
+		const FLinearColor HollowColor = Color.CopyWithNewOpacity(0.0);
+		Set("StateTree.Node.Label.Mixed", new FSlateRoundedBoxBrush(HollowColor, 6.0f, Color, 1.0f));
 	}
 
 	// Task
diff --git a/Engine/Plugins/Runtime/StateTree/Source/StateTreeEditorModule/Private/StateTreeState.cpp b/Engine/Plugins/Runtime/StateTree/Source/StateTreeEditorModule/Private/StateTreeState.cpp
--- a/Engine/Plugins/Runtime/StateTree/Source/StateTreeEditorModule/Private/StateTreeState.cpp
+++ b/Engine/Plugins/Runtime/StateTree/Source/StateTreeEditorModule/Private/StateTreeState.cpp
@@ -26,6 +26,17 @@ FStateTreeTransition::FStateTreeTransition(const EStateTreeTransitionTrigger InT
 	State = InState ? InState->GetLinkToState() : FStateTreeStateLink(InType);
 }
 
+void FStateTreeTransition::PostSerialize(const FArchive& Ar)
+{
+#if WITH_EDITORONLY_DATA
+	const int32 CurrentVersion = Ar.CustomVer(FStateTreeCustomVersion::GUID);
+	if (CurrentVersion < FStateTreeCustomVersion::AddedTransitionIds)
+	{
+		ID = FGuid::NewGuid();
+	}
+#endif // WITH_EDITORONLY_DAT
+}
+
 
 //////////////////////////////////////////////////////////////////////////
 // UStateTreeState
@@ -69,6 +80,7 @@ void UStateTreeState::PostEditChangeChainProperty(FPropertyChangedChainEvent& Pr
 	static const FStateTreeEditPropertyPath StateEnterConditionsPath(UStateTreeState::StaticClass(), TEXT("EnterConditions"));
 	static const FStateTreeEditPropertyPath StateTransitionsPath(UStateTreeState::StaticClass(), TEXT("Transitions"));
 	static const FStateTreeEditPropertyPath StateTransitionsConditionsPath(UStateTreeState::StaticClass(), TEXT("Transitions.Conditions"));
+	static const FStateTreeEditPropertyPath StateTransitionsIDPath(UStateTreeState::StaticClass(), TEXT("Transitions.ID"));
 
 
 	// Broadcast name changes so that the UI can update.
@@ -196,6 +208,17 @@ void UStateTreeState::PostEditChangeChainProperty(FPropertyChangedChainEvent& Pr
 			}
 		}
 
+		// Transitions
+		if (ChangePropertyPath.IsPathExact(StateTransitionsPath))
+		{
+			const int32 TransitionsIndex = ChangePropertyPath.GetPropertyArrayIndex(StateTransitionsPath);
+			if (Transitions.IsValidIndex(TransitionsIndex))
+			{
+				FStateTreeTransition& Transition = Transitions[TransitionsIndex];
+				Transition.ID = FGuid::NewGuid();
+			}
+		}
+
 		// Transition conditions
 		if (ChangePropertyPath.IsPathExact(StateTransitionsConditionsPath))
 		{
@@ -219,7 +242,7 @@ void UStateTreeState::PostEditChangeChainProperty(FPropertyChangedChainEvent& Pr
 		}
 	}
 	
-	// Set default state to root on new transitions.
+	// Set default state to root and Id on new transitions.
 	if (PropertyChangedEvent.ChangeType == EPropertyChangeType::ArrayAdd)
 	{
 		if (ChangePropertyPath.IsPathExact(StateTransitionsPath))
@@ -231,6 +254,7 @@ void UStateTreeState::PostEditChangeChainProperty(FPropertyChangedChainEvent& Pr
 				Transition.Trigger = EStateTreeTransitionTrigger::OnStateCompleted;
 				const UStateTreeState* RootState = GetRootState();
 				Transition.State = RootState->GetLinkToState();
+				Transition.ID = FGuid::NewGuid();
 			}
 		}
 	}
diff --git a/Engine/Plugins/Runtime/StateTree/Source/StateTreeEditorModule/Private/StateTreeViewModel.cpp b/Engine/Plugins/Runtime/StateTree/Source/StateTreeEditorModule/Private/StateTreeViewModel.cpp
--- a/Engine/Plugins/Runtime/StateTree/Source/StateTreeEditorModule/Private/StateTreeViewModel.cpp
+++ b/Engine/Plugins/Runtime/StateTree/Source/StateTreeEditorModule/Private/StateTreeViewModel.cpp
@@ -160,6 +160,9 @@ namespace UE::StateTree::Editor
 
 		for (FStateTreeTransition& Transition : State->Transitions)
 		{
+			// Transition Ids are not used by nodes so no need to add to 'IDsMap'
+			Transition.ID = FGuid::NewGuid();
+
 			FixNodesAfterDuplication(Transition.Conditions, IDsMap, Links);
 			Links.Emplace(&Transition.State);
 		}
@@ -306,17 +309,28 @@ void FStateTreeViewModel::RefreshDebuggerBreakpoints()
 
 		for (const FStateTreeEditorBreakpoint& Breakpoint : TreeData->Breakpoints)
 		{
+			// Test if the ID is associated to a task
 			const FStateTreeIndex16 Index = StateTree->GetNodeIndexFromId(Breakpoint.ID);
 			if (Index.IsValid())
 			{
-				Debugger->SetTaskBreakpoint(Index, Breakpoint.BreakpointType);	
+				Debugger->SetTaskBreakpoint(Index, Breakpoint.BreakpointType);
 			}
 			else
 			{
+				// Then test if the ID is associated to a State
 				FStateTreeStateHandle StateHandle = StateTree->GetStateHandleFromId(Breakpoint.ID);
 				if (StateHandle.IsValid())
 				{
-					Debugger->SetBreakpoint(StateHandle, Breakpoint.BreakpointType);
+					Debugger->SetStateBreakpoint(StateHandle, Breakpoint.BreakpointType);
+				}
+				else
+				{
+					// Then test if the ID is associated to a transition
+					const FStateTreeIndex16 TransitionIndex = StateTree->GetTransitionIndexFromId(Breakpoint.ID);
+					if (TransitionIndex.IsValid())
+					{
+						Debugger->SetTransitionBreakpoint(TransitionIndex, Breakpoint.BreakpointType);
+					}
 				}
 			}
 		}
diff --git a/Engine/Plugins/Runtime/StateTree/Source/StateTreeEditorModule/Public/StateTreeCompiler.h b/Engine/Plugins/Runtime/StateTree/Source/StateTreeEditorModule/Public/StateTreeCompiler.h
--- a/Engine/Plugins/Runtime/StateTree/Source/StateTreeEditorModule/Public/StateTreeCompiler.h
+++ b/Engine/Plugins/Runtime/StateTree/Source/StateTreeEditorModule/Public/StateTreeCompiler.h
@@ -67,6 +67,7 @@ private:
 	UStateTreeEditorData* EditorData = nullptr;
 	TMap<FGuid, int32> IDToNode;
 	TMap<FGuid, int32> IDToState;
+	TMap<FGuid, int32> IDToTransition;
 	TMap<FGuid, const FStateTreeDataView > IDToStructValue;
 	TArray<UStateTreeState*> SourceStates;
 
diff --git a/Engine/Plugins/Runtime/StateTree/Source/StateTreeEditorModule/Public/StateTreeEditorData.h b/Engine/Plugins/Runtime/StateTree/Source/StateTreeEditorModule/Public/StateTreeEditorData.h
--- a/Engine/Plugins/Runtime/StateTree/Source/StateTreeEditorModule/Public/StateTreeEditorData.h
+++ b/Engine/Plugins/Runtime/StateTree/Source/StateTreeEditorModule/Public/StateTreeEditorData.h
@@ -19,14 +19,16 @@ struct FStateTreeEditorBreakpoint
 
 	FStateTreeEditorBreakpoint() = default;
 	explicit FStateTreeEditorBreakpoint(const FGuid& ID, const EStateTreeBreakpointType BreakpointType)
-		: ID(ID),
-		BreakpointType(BreakpointType)
+		: ID(ID)
+		, BreakpointType(BreakpointType)
 	{
 	}
 
+	/** Unique Id of the Node or State associated to the breakpoint. */
 	UPROPERTY()
 	FGuid ID;
 
+	/** The event type that should trigger the breakpoint (e.g. OnEnter, OnExit, etc.). */
 	UPROPERTY()
 	EStateTreeBreakpointType BreakpointType = EStateTreeBreakpointType::Unset;
 };
@@ -218,10 +220,10 @@ PRAGMA_ENABLE_DEPRECATION_WARNINGS
 
 #if WITH_STATETREE_DEBUGGER
 	bool HasAnyBreakpoint(FGuid ID) const;
-	bool HasBreakpoint(FGuid ID, EStateTreeBreakpointType Type) const;
-	const FStateTreeEditorBreakpoint* GetBreakpoint(FGuid ID, EStateTreeBreakpointType Type) const;
-	void AddBreakpoint(FGuid ID, EStateTreeBreakpointType Type);
-	bool RemoveBreakpoint(FGuid ID, EStateTreeBreakpointType Type);
+	bool HasBreakpoint(FGuid ID, EStateTreeBreakpointType BreakpointType) const;
+	const FStateTreeEditorBreakpoint* GetBreakpoint(FGuid ID, EStateTreeBreakpointType BreakpointType) const;
+	void AddBreakpoint(FGuid ID, EStateTreeBreakpointType BreakpointType);
+	bool RemoveBreakpoint(FGuid ID, EStateTreeBreakpointType BreakpointType);
 #endif // WITH_STATETREE_DEBUGGER
 
 	// ~StateTree Builder API
diff --git a/Engine/Plugins/Runtime/StateTree/Source/StateTreeEditorModule/Public/StateTreeState.h b/Engine/Plugins/Runtime/StateTree/Source/StateTreeEditorModule/Public/StateTreeState.h
--- a/Engine/Plugins/Runtime/StateTree/Source/StateTreeEditorModule/Public/StateTreeState.h
+++ b/Engine/Plugins/Runtime/StateTree/Source/StateTreeEditorModule/Public/StateTreeState.h
@@ -18,6 +18,8 @@ struct STATETREEEDITORMODULE_API FStateTreeTransition
 	FStateTreeTransition() = default;
 	FStateTreeTransition(const EStateTreeTransitionTrigger InTrigger, const EStateTreeTransitionType InType, const UStateTreeState* InState = nullptr);
 	FStateTreeTransition(const EStateTreeTransitionTrigger InTrigger, const FGameplayTag InEventTag, const EStateTreeTransitionType InType, const UStateTreeState* InState = nullptr);
+
+	void PostSerialize(const FArchive& Ar);
 	
 	template<typename T, typename... TArgs>
 	TStateTreeEditorNode<T>& AddCondition(TArgs&&... InArgs)
@@ -45,6 +47,9 @@ struct STATETREEEDITORMODULE_API FStateTreeTransition
 	UPROPERTY(EditDefaultsOnly, Category = "Transition", meta=(DisplayName="Transition To"))
 	FStateTreeStateLink State;
 
+	UPROPERTY(EditDefaultsOnly, Category = "Transition")
+	FGuid ID;
+
 	/**
 	 * Transition priority when multiple transitions happen at the same time.
 	 * During transition handling, the transitions are visited from leaf to root.
@@ -68,6 +73,19 @@ struct STATETREEEDITORMODULE_API FStateTreeTransition
 	/** Conditions that must pass so that the transition can be triggered. */
 	UPROPERTY(EditDefaultsOnly, Category = "Transition", meta = (BaseStruct = "/Script/StateTreeModule.StateTreeConditionBase", BaseClass = "/Script/StateTreeModule.StateTreeConditionBlueprintBase"))
 	TArray<FStateTreeEditorNode> Conditions;
+
+	/** True if the Transition is Enabled (i.e. not explicitly disabled in the asset). */
+	UPROPERTY(EditDefaultsOnly, Category = "Debug")
+	bool bTransitionEnabled = true;
+};
+
+template<>
+struct TStructOpsTypeTraits<FStateTreeTransition> : TStructOpsTypeTraitsBase2<FStateTreeTransition>
+{
+	enum
+	{
+		WithPostSerialize = true,
+	};
 };
 
 USTRUCT()
@@ -170,12 +188,16 @@ public:
 	 */
 	FStateTreeTransition& AddTransition(const EStateTreeTransitionTrigger InTrigger, const EStateTreeTransitionType InType, const UStateTreeState* InState = nullptr)
 	{
-		return Transitions.Emplace_GetRef(InTrigger, InType, InState);
+		FStateTreeTransition& Transition = Transitions.Emplace_GetRef(InTrigger, InType, InState);
+		Transition.ID = FGuid::NewGuid();
+		return Transition;
 	}
 
 	FStateTreeTransition& AddTransition(const EStateTreeTransitionTrigger InTrigger, const FGameplayTag InEventTag, const EStateTreeTransitionType InType, const UStateTreeState* InState = nullptr)
 	{
-		return Transitions.Emplace_GetRef(InTrigger, InEventTag, InType, InState);
+		FStateTreeTransition& Transition = Transitions.Emplace_GetRef(InTrigger, InEventTag, InType, InState);
+		Transition.ID = FGuid::NewGuid();
+		return Transition;
 	}
 
 
@@ -196,7 +218,7 @@ public:
 	UPROPERTY(EditDefaultsOnly, Category = "State")
 	FStateTreeStateParameters Parameters;
 
-	UPROPERTY(meta = (IgnoreForMemberInitializationTest))
+	UPROPERTY(EditDefaultsOnly, Category = "State", meta = (IgnoreForMemberInitializationTest))
 	FGuid ID;
 
 	UPROPERTY(EditDefaultsOnly, Category = "Enter Conditions", meta = (BaseStruct = "/Script/StateTreeModule.StateTreeConditionBase", BaseClass = "/Script/StateTreeModule.StateTreeConditionBlueprintBase"))
diff --git a/Engine/Plugins/Runtime/StateTree/Source/StateTreeModule/Private/Debugger/StateTreeDebugger.cpp b/Engine/Plugins/Runtime/StateTree/Source/StateTreeModule/Private/Debugger/StateTreeDebugger.cpp
--- a/Engine/Plugins/Runtime/StateTree/Source/StateTreeModule/Private/Debugger/StateTreeDebugger.cpp
+++ b/Engine/Plugins/Runtime/StateTree/Source/StateTreeModule/Private/Debugger/StateTreeDebugger.cpp
@@ -559,17 +559,22 @@ bool FStateTreeDebugger::HasTransitionBreakpoint(const FStateTreeIndex16 Index,
 	});
 }
 
-void FStateTreeDebugger::SetBreakpoint(FStateTreeStateHandle StateHandle, EStateTreeBreakpointType BreakpointType)
+void FStateTreeDebugger::SetStateBreakpoint(const FStateTreeStateHandle StateHandle, const EStateTreeBreakpointType BreakpointType)
 {
 	Breakpoints.Emplace(StateHandle, BreakpointType);
 }
 
-void FStateTreeDebugger::SetTaskBreakpoint(const FStateTreeIndex16 NodeIndex, EStateTreeBreakpointType BreakpointType)
+void FStateTreeDebugger::SetTransitionBreakpoint(const FStateTreeIndex16 TransitionIndex, const EStateTreeBreakpointType BreakpointType)
+{
+	Breakpoints.Emplace(FStateTreeDebuggerBreakpoint::FStateTreeTransitionIndex(TransitionIndex), BreakpointType);
+}
+
+void FStateTreeDebugger::SetTaskBreakpoint(const FStateTreeIndex16 NodeIndex, const EStateTreeBreakpointType BreakpointType)
 {
 	Breakpoints.Emplace(FStateTreeDebuggerBreakpoint::FStateTreeTaskIndex(NodeIndex), BreakpointType);
 }
 
-void FStateTreeDebugger::ClearBreakpoint(const FStateTreeIndex16 NodeIndex, EStateTreeBreakpointType BreakpointType)
+void FStateTreeDebugger::ClearBreakpoint(const FStateTreeIndex16 NodeIndex, const EStateTreeBreakpointType BreakpointType)
 {
 	const int32 Index = Breakpoints.IndexOfByPredicate([NodeIndex, BreakpointType](const FStateTreeDebuggerBreakpoint& Breakpoint)
 		{
@@ -662,9 +667,9 @@ void FStateTreeDebugger::SendNotifications()
 		check(HitBreakpointInstanceId.IsValid());
 		check(Breakpoints.IsValidIndex(HitBreakpointIndex));
 
-		// Force scrub time to latest read time to reflect most recent events.
+		// Force scrub time to latest simulation time to reflect most recent events.
 		// This will notify scrub position changed and active states
-		SetScrubTime(LastTraceReadTime);
+		SetScrubTime(RecordingDuration);
 
 		// Make sure the instance is selected in case the breakpoint was set for any instances 
 		if (SelectedInstanceId != HitBreakpointInstanceId)
diff --git a/Engine/Plugins/Runtime/StateTree/Source/StateTreeModule/Private/Debugger/StateTreeDebuggerTypes.cpp b/Engine/Plugins/Runtime/StateTree/Source/StateTreeModule/Private/Debugger/StateTreeDebuggerTypes.cpp
--- a/Engine/Plugins/Runtime/StateTree/Source/StateTreeModule/Private/Debugger/StateTreeDebuggerTypes.cpp
+++ b/Engine/Plugins/Runtime/StateTree/Source/StateTreeModule/Private/Debugger/StateTreeDebuggerTypes.cpp
@@ -265,6 +265,8 @@ EStateTreeTraceEventType FStateTreeDebuggerBreakpoint::GetMatchingEventType(cons
 		return EStateTreeTraceEventType::OnEnter;
 	case EStateTreeBreakpointType::OnExit:
 		return EStateTreeTraceEventType::OnExit;
+	case EStateTreeBreakpointType::OnTransition:
+		return EStateTreeTraceEventType::OnTransition;
 	default:
 		return EStateTreeTraceEventType::Unset;
 	}
diff --git a/Engine/Plugins/Runtime/StateTree/Source/StateTreeModule/Private/StateTree.cpp b/Engine/Plugins/Runtime/StateTree/Source/StateTreeModule/Private/StateTree.cpp
--- a/Engine/Plugins/Runtime/StateTree/Source/StateTreeModule/Private/StateTree.cpp
+++ b/Engine/Plugins/Runtime/StateTree/Source/StateTreeModule/Private/StateTree.cpp
@@ -59,6 +59,12 @@ const FCompactStateTransition* UStateTree::GetTransitionFromIndex(const int16 Tr
 	return Transitions.IsValidIndex(TransitionIndex) ? &Transitions[TransitionIndex] : nullptr;
 }
 
+FStateTreeIndex16 UStateTree::GetTransitionIndexFromId(const FGuid Id) const
+{
+	const FStateTreeTransitionIdToIndex* Entry = IDToTransitionMappings.FindByPredicate([Id](const FStateTreeTransitionIdToIndex& Entry){ return Entry.Id == Id; });
+	return Entry != nullptr ? Entry->Index : FStateTreeIndex16::Invalid;
+}
+
 TSharedPtr<FStateTreeInstanceData> UStateTree::GetSharedInstanceData() const
 {
 	// Create a unique index for each thread.
@@ -113,6 +119,7 @@ void UStateTree::ResetCompiled()
 	Parameters.Reset();
 	IDToStateMappings.Reset();
 	IDToNodeMappings.Reset();
+	IDToTransitionMappings.Reset();
 
 	ParametersDataViewIndex = FStateTreeIndex8::Invalid;
 	
diff --git a/Engine/Plugins/Runtime/StateTree/Source/StateTreeModule/Private/StateTreeExecutionContext.cpp b/Engine/Plugins/Runtime/StateTree/Source/StateTreeModule/Private/StateTreeExecutionContext.cpp
--- a/Engine/Plugins/Runtime/StateTree/Source/StateTreeModule/Private/StateTreeExecutionContext.cpp
+++ b/Engine/Plugins/Runtime/StateTree/Source/StateTreeModule/Private/StateTreeExecutionContext.cpp
@@ -419,7 +419,9 @@ EStateTreeRunStatus FStateTreeExecutionContext::Tick(const float DeltaTime)
 		if (TriggerTransitions())
 		{
 			STATETREE_TRACE_SCOPED_PHASE(EStateTreeUpdatePhase::ApplyTransitions);
-			
+			STATETREE_TRACE_TRANSITION_EVENT(NextTransitionIndex, EStateTreeTraceEventType::OnTransition);
+			NextTransitionIndex = FStateTreeIndex16::Invalid;
+
 			// We have committed to state change, consume events that were accumulated during the tick above.
 			EventQueue.Reset();
 
@@ -1663,6 +1665,12 @@ bool FStateTreeExecutionContext::TriggerTransitions()
 				const int16 TransitionIndex = State.TransitionsBegin + i;
 				const FCompactStateTransition& Transition = StateTree.Transitions[TransitionIndex];
 
+				// Skip disabled transitions
+				if (Transition.bTransitionEnabled == false)
+				{
+					continue;
+				}
+				
 				// No need to test the transition if same or higher priority transition has already been processed.
 				if (Transition.Priority <= NextTransition.Priority)
 				{
@@ -1691,7 +1699,10 @@ bool FStateTreeExecutionContext::TriggerTransitions()
 							});
 
 						// Trigger Delayed Transition when the delay has passed.
-						RequestTransition(Transition.State, Transition.Priority);
+						if (RequestTransition(Transition.State, Transition.Priority))
+						{
+							NextTransitionIndex = FStateTreeIndex16(TransitionIndex);
+						}
 						continue;
 					}
 				}
@@ -1738,7 +1749,10 @@ bool FStateTreeExecutionContext::TriggerTransitions()
 						}
 					}
 
-					RequestTransition(Transition.State, Transition.Priority);
+					if (RequestTransition(Transition.State, Transition.Priority))
+					{
+						NextTransitionIndex = FStateTreeIndex16(TransitionIndex);
+					}
 				}
 			}
 		}
@@ -1793,6 +1807,12 @@ bool FStateTreeExecutionContext::TriggerTransitions()
 				const int16 TransitionIndex = State.TransitionsBegin + i;
 				const FCompactStateTransition& Transition = StateTree.Transitions[TransitionIndex];
 
+				// Skip disabled transitions
+				if (Transition.bTransitionEnabled == false)
+				{
+					continue;
+				}
+
 				if (EnumHasAnyFlags(Transition.Trigger, CompletionTrigger))
 				{
 					bool bPassed = false;
@@ -1808,6 +1828,7 @@ bool FStateTreeExecutionContext::TriggerTransitions()
 						// No priority on completion transitions, use the priority to signal that state is selected.
 						if (RequestTransition(Transition.State, EStateTreeTransitionPriority::Normal))
 						{
+							NextTransitionIndex = FStateTreeIndex16(TransitionIndex);
 							break;
 						}
 					}
@@ -2035,6 +2056,12 @@ bool FStateTreeExecutionContext::SelectStateInternal(const FStateTreeStateHandle
 				const int16 TransitionIndex = State.TransitionsBegin + i;
 				const FCompactStateTransition& Transition = StateTree.Transitions[TransitionIndex];
 
+				// Skip disabled transitions
+				if (Transition.bTransitionEnabled == false)
+				{
+					continue;
+				}
+
 				// No need to test the transition if same or higher priority transition has already been processed.
 				if (Transition.Priority <= CurrentPriority)
 				{
diff --git a/Engine/Plugins/Runtime/StateTree/Source/StateTreeModule/Public/Debugger/StateTreeDebugger.h b/Engine/Plugins/Runtime/StateTree/Source/StateTreeModule/Public/Debugger/StateTreeDebugger.h
--- a/Engine/Plugins/Runtime/StateTree/Source/StateTreeModule/Public/Debugger/StateTreeDebugger.h
+++ b/Engine/Plugins/Runtime/StateTree/Source/StateTreeModule/Public/Debugger/StateTreeDebugger.h
@@ -92,7 +92,8 @@ struct STATETREEMODULE_API FStateTreeDebugger : FTickableGameObject
 	bool HasStateBreakpoint(FStateTreeStateHandle StateHandle, EStateTreeBreakpointType BreakpointType) const;
 	bool HasTaskBreakpoint(FStateTreeIndex16 Index, EStateTreeBreakpointType BreakpointType) const;
 	bool HasTransitionBreakpoint(FStateTreeIndex16 Index, EStateTreeBreakpointType BreakpointType) const;
-	void SetBreakpoint(FStateTreeStateHandle StateHandle, EStateTreeBreakpointType BreakpointType);
+	void SetStateBreakpoint(FStateTreeStateHandle StateHandle, EStateTreeBreakpointType BreakpointType);
+	void SetTransitionBreakpoint(FStateTreeIndex16 SubIndex, EStateTreeBreakpointType BreakpointType);
 	void SetTaskBreakpoint(FStateTreeIndex16 NodeIndex, EStateTreeBreakpointType BreakpointType);
 	void ClearBreakpoint(FStateTreeIndex16 NodeIndex, EStateTreeBreakpointType BreakpointType);
 	void ClearAllBreakpoints();
diff --git a/Engine/Plugins/Runtime/StateTree/Source/StateTreeModule/Public/Debugger/StateTreeTraceTypes.h b/Engine/Plugins/Runtime/StateTree/Source/StateTreeModule/Public/Debugger/StateTreeTraceTypes.h
--- a/Engine/Plugins/Runtime/StateTree/Source/StateTreeModule/Public/Debugger/StateTreeTraceTypes.h
+++ b/Engine/Plugins/Runtime/StateTree/Source/StateTreeModule/Public/Debugger/StateTreeTraceTypes.h
@@ -28,7 +28,8 @@ enum class EStateTreeTraceEventType : uint8
 	OnTaskTicked		UMETA(DisplayName = "Ticked"),		// for Task events
 	Passed				UMETA(DisplayName = "Passed"),		// for Condition events
 	Failed				UMETA(DisplayName = "Failed"),		// for Condition events
-	OnEvaluating		UMETA(DisplayName = "Evaluating")	// for Transition events
+	OnEvaluating		UMETA(DisplayName = "Evaluating"),	// for Transition events
+	OnTransition		UMETA(DisplayName = "Transition")	// for Transition events
 };
 
 #if WITH_STATETREE_DEBUGGER
diff --git a/Engine/Plugins/Runtime/StateTree/Source/StateTreeModule/Public/StateTree.h b/Engine/Plugins/Runtime/StateTree/Source/StateTreeModule/Public/StateTree.h
--- a/Engine/Plugins/Runtime/StateTree/Source/StateTreeModule/Public/StateTree.h
+++ b/Engine/Plugins/Runtime/StateTree/Source/StateTreeModule/Public/StateTree.h
@@ -35,6 +35,8 @@ struct STATETREEMODULE_API FStateTreeCustomVersion
 		AddedExternalTransitions,
 		// Changed how bindings are represented
 		ChangedBindingsRepresentation,
+		// Added guid to transitions
+		AddedTransitionIds,
 
 		// -----<new versions can be added above this line>-------------------------------------------------
 		VersionPlusOne,
@@ -121,7 +123,7 @@ public:
 	/** @return Struct views of all nodes */
 	const FInstancedStructContainer& GetNodes() const { return Nodes; }
 
-	/** @return Node index matching a given Id; invalid handle if node not found. */
+	/** @return Node index matching a given Id; invalid index if node not found. */
 	FStateTreeIndex16 GetNodeIndexFromId(const FGuid Id) const;
 
 	/** @return View of all states. */
@@ -129,6 +131,9 @@ public:
 
 	/** @return Pointer to the transition at a given index; null if not found. */ 
 	const FCompactStateTransition* GetTransitionFromIndex(const int16 TransitionIndex) const;
+	
+	/** @return Runtime transition index matching a given Id; invalid index if node not found. */
+	FStateTreeIndex16 GetTransitionIndexFromId(const FGuid Id) const;
 
 #if WITH_EDITOR
 	/** Resets the compiled data to empty. */
@@ -231,6 +236,10 @@ private:
 	/** Mapping of node guid for the Editor and node index, created at compilation. */
 	UPROPERTY()
 	TArray<FStateTreeNodeIdToIndex> IDToNodeMappings;
+	
+	/** Mapping of state transition identifiers and runtime compact transition index, created at compilation. */
+	UPROPERTY()
+	TArray<FStateTreeTransitionIdToIndex> IDToTransitionMappings;
 
 	/**
 	 * Parameters that could be used for bindings within the Tree.
diff --git a/Engine/Plugins/Runtime/StateTree/Source/StateTreeModule/Public/StateTreeExecutionContext.h b/Engine/Plugins/Runtime/StateTree/Source/StateTreeModule/Public/StateTreeExecutionContext.h
--- a/Engine/Plugins/Runtime/StateTree/Source/StateTreeModule/Public/StateTreeExecutionContext.h
+++ b/Engine/Plugins/Runtime/StateTree/Source/StateTreeModule/Public/StateTreeExecutionContext.h
@@ -442,6 +442,9 @@ protected:
 	/** Next transition, used by RequestTransition(). */
 	FStateTreeTransitionResult NextTransition;
 
+	/** Index of the state transition that caused the state change. Will be invalid for framework transitions and external requests. */
+	FStateTreeIndex16 NextTransitionIndex = FStateTreeIndex16::Invalid;
+
 	/** Current state we're processing, or invalid if not applicable. */
 	FStateTreeStateHandle CurrentlyProcessedState;
 
diff --git a/Engine/Plugins/Runtime/StateTree/Source/StateTreeModule/Public/StateTreeTypes.h b/Engine/Plugins/Runtime/StateTree/Source/StateTreeModule/Public/StateTreeTypes.h
--- a/Engine/Plugins/Runtime/StateTree/Source/StateTreeModule/Public/StateTreeTypes.h
+++ b/Engine/Plugins/Runtime/StateTree/Source/StateTreeModule/Public/StateTreeTypes.h
@@ -539,8 +539,8 @@ struct STATETREEMODULE_API FStateTreeTransitionResult
 		TargetState = FStateTreeStateHandle::Invalid;
 		NextActiveStates.Reset();
 		CurrentState = FStateTreeStateHandle::Invalid;
-		ChangeType = EStateTreeStateChangeType::Changed; 
-		Priority = EStateTreeTransitionPriority::None; 
+		ChangeType = EStateTreeStateChangeType::Changed;
+		Priority = EStateTreeTransitionPriority::None;
 	}
 	
 	/** Current active states, where the transition started. */
@@ -573,7 +573,7 @@ struct STATETREEMODULE_API FStateTreeTransitionResult
 
 	/** Priority of the transition that caused the state change. */
 	UPROPERTY(EditDefaultsOnly, Category = "Default", BlueprintReadOnly)
-	EStateTreeTransitionPriority Priority = EStateTreeTransitionPriority::None; 
+	EStateTreeTransitionPriority Priority = EStateTreeTransitionPriority::None;
 };
 
 
@@ -585,6 +585,11 @@ struct STATETREEMODULE_API FCompactStateTransition
 {
 	GENERATED_BODY()
 
+	explicit FCompactStateTransition()
+		: bTransitionEnabled(true)
+	{
+	}
+
 	/** @return True if the transition has delay. */
 	bool HasDelay() const
 	{
@@ -618,6 +623,10 @@ struct STATETREEMODULE_API FCompactStateTransition
 	/** Number of conditions to test. */
 	UPROPERTY()
 	uint8 ConditionsNum = 0;
+
+	/** Indicates if the transition is enabled and should be considered. */
+	UPROPERTY()
+	uint8 bTransitionEnabled : 1;
 };
 
 /**
@@ -840,6 +849,28 @@ struct STATETREEMODULE_API FStateTreeNodeIdToIndex
 	FStateTreeIndex16 Index;
 };
 
+/**
+ * Pair of transition id and its associated compact transition index created at compilation.
+ */
+USTRUCT()
+struct STATETREEMODULE_API FStateTreeTransitionIdToIndex
+{
+	GENERATED_BODY()
+
+	FStateTreeTransitionIdToIndex() = default;
+	explicit FStateTreeTransitionIdToIndex(const FGuid& Id, const FStateTreeIndex16 Index)
+		: Id(Id)
+		, Index(Index)
+	{
+	}
+
+	UPROPERTY();
+	FGuid Id;
+	
+	UPROPERTY();
+	FStateTreeIndex16 Index;
+};
+
 /**
  * Describes an external data. The data can point to a struct or object.
  * The code that handles StateTree ticking is responsible for passing in the actually data, see FStateTreeExecutionContext.
