Static Analysis fix

diff --git a/Engine/Plugins/Runtime/StateTree/Source/StateTreeEditorModule/Private/Debugger/StateTreeDebuggerTrack.cpp b/Engine/Plugins/Runtime/StateTree/Source/StateTreeEditorModule/Private/Debugger/StateTreeDebuggerTrack.cpp
--- a/Engine/Plugins/Runtime/StateTree/Source/StateTreeEditorModule/Private/Debugger/StateTreeDebuggerTrack.cpp
+++ b/Engine/Plugins/Runtime/StateTree/Source/StateTreeEditorModule/Private/Debugger/StateTreeDebuggerTrack.cpp
@@ -66,10 +66,11 @@ bool FStateTreeDebuggerInstanceTrack::UpdateInternal()
 			const FStateTreeTraceActiveStatesEvent& Event = Events[EventIndex].Get<FStateTreeTraceActiveStatesEvent>();
 				
 			FString StatePath;
-			for (int32 StateIndex = 0; StateIndex < Event.ActiveStates.Num(); StateIndex++)
+			const TConstArrayView<FCompactStateTreeState> States = StateTree->GetStates();
+			for (int32 ActiveStateIndex = 0; ActiveStateIndex < Event.ActiveStates.Num(); ActiveStateIndex++)
 			{
-				const FCompactStateTreeState& State = StateTree->GetStates()[Event.ActiveStates[StateIndex].Index];
-				StatePath.Appendf(TEXT("%s%s"), StateIndex == 0 ? TEXT("") : TEXT("."), *State.Name.ToString());
+				const FCompactStateTreeState& State = States[Event.ActiveStates[ActiveStateIndex].Index];
+				StatePath.Appendf(TEXT("%s%s"), ActiveStateIndex == 0 ? TEXT("") : TEXT("."), *State.Name.ToString());
 			}
 
 			UE::StateTreeDebugger::FFrameSpan Span = EventCollection.FrameSpans[SpanIndex];
