StateTree: Improved reconciling property names when they change
- Added PropertyGuid to FStateTreePropertyPathSegment to allow reconsile changed Blueprint class and User Defined Struct property names
- Changed FStateTreePropertyPath to optionally handle property redirects and BP/UDS name changes
- Renamed UpdateInstanceStructsFromValue to UpdateSegmentsFromValue to better reflect the use
- Improved State Tree editor bindings to update when BP/UDS properties are changed

diff --git a/Engine/Plugins/Runtime/StateTree/Source/StateTreeEditorModule/Private/Customizations/StateTreeBindingExtension.cpp b/Engine/Plugins/Runtime/StateTree/Source/StateTreeEditorModule/Private/Customizations/StateTreeBindingExtension.cpp
--- a/Engine/Plugins/Runtime/StateTree/Source/StateTreeEditorModule/Private/Customizations/StateTreeBindingExtension.cpp
+++ b/Engine/Plugins/Runtime/StateTree/Source/StateTreeEditorModule/Private/Customizations/StateTreeBindingExtension.cpp
@@ -105,7 +105,7 @@ void MakeStructPropertyPathFromBindingChain(const FGuid StructID, const TArray<F
 		}
 	}
 
-	OutPath.UpdateInstanceStructsFromValue(DataView);
+	OutPath.UpdateSegmentsFromValue(DataView);
 }
 
 EStateTreePropertyUsage MakeStructPropertyPathFromPropertyHandle(TSharedPtr<const IPropertyHandle> InPropertyHandle, FStateTreePropertyPath& OutPath)
diff --git a/Engine/Plugins/Runtime/StateTree/Source/StateTreeEditorModule/Private/StateTreeCompiler.cpp b/Engine/Plugins/Runtime/StateTree/Source/StateTreeEditorModule/Private/StateTreeCompiler.cpp
--- a/Engine/Plugins/Runtime/StateTree/Source/StateTreeEditorModule/Private/StateTreeCompiler.cpp
+++ b/Engine/Plugins/Runtime/StateTree/Source/StateTreeEditorModule/Private/StateTreeCompiler.cpp
@@ -1515,7 +1515,7 @@ bool FStateTreeCompiler::GetAndValidateBindings(const FStateTreeBindableStructDe
 		// Update path instance types from latest data. E.g. binding may have been created for instanced object of type FooB, and changed to FooA.
  		// @todo: not liking how this mutates the Binding.TargetPath, but currently we dont track well the instanced object changes.
 
-		if (!Binding.GetMutableTargetPath().UpdateInstanceStructsFromValue(TargetValue))
+		if (!Binding.GetMutableTargetPath().UpdateSegmentsFromValue(TargetValue))
 		{
 			Log.Reportf(EMessageSeverity::Error, TargetStruct,
 						TEXT("Malformed target property path for binding source property '%s' for target '%s:%s'."),
@@ -1554,7 +1554,7 @@ bool FStateTreeCompiler::GetAndValidateBindings(const FStateTreeBindableStructDe
 		const FStateTreeDataView SourceValue = IDToStructValue[SourceStructID];
 		if (SourceValue.IsValid())
 		{
-			if (!Binding.GetMutableSourcePath().UpdateInstanceStructsFromValue(SourceValue))
+			if (!Binding.GetMutableSourcePath().UpdateSegmentsFromValue(SourceValue))
 			{
 				Log.Reportf(EMessageSeverity::Error, TargetStruct,
 					TEXT("Malformed target property path for binding source property '%s' for source '%s:%s'."),
diff --git a/Engine/Plugins/Runtime/StateTree/Source/StateTreeEditorModule/Private/StateTreeEditorData.cpp b/Engine/Plugins/Runtime/StateTree/Source/StateTreeEditorModule/Private/StateTreeEditorData.cpp
--- a/Engine/Plugins/Runtime/StateTree/Source/StateTreeEditorModule/Private/StateTreeEditorData.cpp
+++ b/Engine/Plugins/Runtime/StateTree/Source/StateTreeEditorModule/Private/StateTreeEditorData.cpp
@@ -9,6 +9,11 @@
 #include "Algo/LevenshteinDistance.h"
 #include "StateTreeEditorModule.h"
 
+#if WITH_EDITOR
+#include "Engine/UserDefinedStruct.h"
+#include "StructUtilsDelegates.h"
+#endif
+
 #include UE_INLINE_GENERATED_CPP_BY_NAME(StateTreeEditorData)
 
 void UStateTreeEditorData::PostInitProperties()
@@ -16,9 +21,97 @@ void UStateTreeEditorData::PostInitProperties()
 	Super::PostInitProperties();
 
 	RootParameters.ID = FGuid::NewGuid();
+
+#if WITH_EDITOR
+	OnObjectsReinstancedHandle = FCoreUObjectDelegates::OnObjectsReinstanced.AddUObject(this, &UStateTreeEditorData::OnObjectsReinstanced);
+	OnUserDefinedStructReinstancedHandle = UE::StructUtils::Delegates::OnUserDefinedStructReinstanced.AddUObject(this, &UStateTreeEditorData::OnUserDefinedStructReinstanced);
+#endif
 }
 
 #if WITH_EDITOR
+
+void UStateTreeEditorData::BeginDestroy()
+{
+	if (OnObjectsReinstancedHandle.IsValid())
+	{
+		FCoreUObjectDelegates::OnObjectsReinstanced.Remove(OnObjectsReinstancedHandle);
+		OnObjectsReinstancedHandle.Reset();
+	}
+	if (OnUserDefinedStructReinstancedHandle.IsValid())
+	{
+		UE::StructUtils::Delegates::OnUserDefinedStructReinstanced.Remove(OnUserDefinedStructReinstancedHandle);
+		OnUserDefinedStructReinstancedHandle.Reset();
+	}
+	
+	Super::BeginDestroy();
+}
+
+void UStateTreeEditorData::OnObjectsReinstanced(const FReplacementObjectMap& ObjectMap)
+{
+	if (ObjectMap.IsEmpty())
+	{
+		return;
+	}
+	
+	TSet<const UStruct*> Structs;
+	for (TMap<UObject*, UObject*>::TConstIterator It(ObjectMap); It; ++It)
+	{
+		if (const UObject* ObjectToBeReplaced = It->Value)
+		{
+			Structs.Add(ObjectToBeReplaced->GetClass());
+		}
+	}
+
+	bool bShouldUpdate = false;
+	VisitAllNodes([&Structs, &bShouldUpdate](const UStateTreeState* State, const FStateTreeBindableStructDesc& Desc, const FStateTreeDataView Value)
+	{
+		if (Structs.Contains(Value.GetStruct()))
+		{
+			bShouldUpdate = true;
+			return EStateTreeVisitor::Break; 
+		}
+		return EStateTreeVisitor::Continue;
+	});
+
+	if (!bShouldUpdate)
+	{
+		bShouldUpdate = EditorBindings.ContainsAnyStruct(Structs);
+	}
+	
+	if (bShouldUpdate)
+	{
+		UpdateBindingsInstanceStructs();
+	}
+}
+
+void UStateTreeEditorData::OnUserDefinedStructReinstanced(const UUserDefinedStruct& UserDefinedStruct)
+{
+	TSet<const UStruct*> Structs;
+	Structs.Add(&UserDefinedStruct);
+
+	bool bShouldUpdate = false;
+	VisitAllNodes([&Structs, &bShouldUpdate](const UStateTreeState* State, const FStateTreeBindableStructDesc& Desc, const FStateTreeDataView Value)
+	{
+		if (Structs.Contains(Value.GetStruct()))
+		{
+			bShouldUpdate = true;
+			return EStateTreeVisitor::Break; 
+		}
+		return EStateTreeVisitor::Continue;
+	});
+
+	if (!bShouldUpdate)
+	{
+		bShouldUpdate = EditorBindings.ContainsAnyStruct(Structs);
+	}
+	
+	if (bShouldUpdate)
+	{
+		UpdateBindingsInstanceStructs();
+	}
+}
+
+
 void UStateTreeEditorData::PostLoad()
 {
 	Super::PostLoad();
@@ -437,12 +530,12 @@ void UStateTreeEditorData::UpdateBindingsInstanceStructs()
 	{
 		if (AllValues.Contains(Binding.GetSourcePath().GetStructID()))
 		{
-			Binding.GetMutableSourcePath().UpdateInstanceStructsFromValue(AllValues[Binding.GetSourcePath().GetStructID()]);
+			Binding.GetMutableSourcePath().UpdateSegmentsFromValue(AllValues[Binding.GetSourcePath().GetStructID()]);
 		}
 
 		if (AllValues.Contains(Binding.GetTargetPath().GetStructID()))
 		{
-			Binding.GetMutableTargetPath().UpdateInstanceStructsFromValue(AllValues[Binding.GetTargetPath().GetStructID()]);
+			Binding.GetMutableTargetPath().UpdateSegmentsFromValue(AllValues[Binding.GetTargetPath().GetStructID()]);
 		}
 	}
 }
diff --git a/Engine/Plugins/Runtime/StateTree/Source/StateTreeEditorModule/Private/StateTreeEditorPropertyBindings.cpp b/Engine/Plugins/Runtime/StateTree/Source/StateTreeEditorModule/Private/StateTreeEditorPropertyBindings.cpp
--- a/Engine/Plugins/Runtime/StateTree/Source/StateTreeEditorModule/Private/StateTreeEditorPropertyBindings.cpp
+++ b/Engine/Plugins/Runtime/StateTree/Source/StateTreeEditorModule/Private/StateTreeEditorPropertyBindings.cpp
@@ -123,6 +123,35 @@ void FStateTreeEditorPropertyBindings::RemoveUnusedBindings(const TMap<FGuid, co
 		});
 }
 
+bool FStateTreeEditorPropertyBindings::ContainsAnyStruct(const TSet<const UStruct*>& Structs)
+{
+	auto PathContainsStruct = [&Structs](const FStateTreePropertyPath& PropertyPath)
+	{
+		for (const FStateTreePropertyPathSegment& Segment : PropertyPath.GetSegments())
+		{
+			if (Structs.Contains(Segment.GetInstanceStruct()))
+			{
+				return true;
+			}
+		}
+		return false;
+	};
+	
+	for (FStateTreePropertyPathBinding& PropertyPathBinding : PropertyBindings)
+	{
+		if (PathContainsStruct(PropertyPathBinding.GetSourcePath()))
+		{
+			return true;
+		}
+		if (PathContainsStruct(PropertyPathBinding.GetTargetPath()))
+		{
+			return true;
+		}
+	}
+	
+	return false;
+}
+
 PRAGMA_DISABLE_DEPRECATION_WARNINGS
 void FStateTreeEditorPropertyBindings::AddPropertyBinding(const FStateTreeEditorPropertyPath& SourcePath, const FStateTreeEditorPropertyPath& TargetPath)
 {
diff --git a/Engine/Plugins/Runtime/StateTree/Source/StateTreeEditorModule/Public/StateTreeEditorData.h b/Engine/Plugins/Runtime/StateTree/Source/StateTreeEditorModule/Public/StateTreeEditorData.h
--- a/Engine/Plugins/Runtime/StateTree/Source/StateTreeEditorModule/Public/StateTreeEditorData.h
+++ b/Engine/Plugins/Runtime/StateTree/Source/StateTreeEditorModule/Public/StateTreeEditorData.h
@@ -59,6 +59,10 @@ public:
 	// ~IStateTreeEditorPropertyBindingsOwner
 
 #if WITH_EDITOR
+	using FReplacementObjectMap = TMap<UObject*, UObject*>;
+	void OnObjectsReinstanced(const FReplacementObjectMap& ObjectMap);
+	void OnUserDefinedStructReinstanced(const UUserDefinedStruct& UserDefinedStruct);
+	virtual void BeginDestroy() override;
 	virtual void PostLoad() override;
 	virtual void PostEditChangeChainProperty(FPropertyChangedChainEvent& PropertyChangedEvent) override;
 #endif
@@ -233,6 +237,11 @@ private:
 	void FixObjectNodes();
 	void UpdateBindingsInstanceStructs();
 
+#if WITH_EDITORONLY_DATA
+	FDelegateHandle OnObjectsReinstancedHandle;
+	FDelegateHandle OnUserDefinedStructReinstancedHandle;
+#endif
+
 public:
 	/** Schema describing which inputs, evaluators, and tasks a StateTree can contain */	
 	UPROPERTY(EditDefaultsOnly, Category = Common, Instanced)
diff --git a/Engine/Plugins/Runtime/StateTree/Source/StateTreeEditorModule/Public/StateTreeEditorPropertyBindings.h b/Engine/Plugins/Runtime/StateTree/Source/StateTreeEditorModule/Public/StateTreeEditorPropertyBindings.h
--- a/Engine/Plugins/Runtime/StateTree/Source/StateTreeEditorModule/Public/StateTreeEditorPropertyBindings.h
+++ b/Engine/Plugins/Runtime/StateTree/Source/StateTreeEditorModule/Public/StateTreeEditorPropertyBindings.h
@@ -59,6 +59,9 @@ struct STATETREEEDITORMODULE_API FStateTreeEditorPropertyBindings
 	 */
 	void RemoveUnusedBindings(const TMap<FGuid, const FStateTreeDataView>& ValidStructs);
 
+	/** @return true if any of the bindings references any of the Structs. */
+	bool ContainsAnyStruct(const TSet<const UStruct*>& Structs);
+
 	/** @return array view to all bindings. */
 	TConstArrayView<FStateTreePropertyPathBinding> GetBindings() const { return PropertyBindings; }
 
diff --git a/Engine/Plugins/Runtime/StateTree/Source/StateTreeEditorModule/StateTreeEditorModule.Build.cs b/Engine/Plugins/Runtime/StateTree/Source/StateTreeEditorModule/StateTreeEditorModule.Build.cs
--- a/Engine/Plugins/Runtime/StateTree/Source/StateTreeEditorModule/StateTreeEditorModule.Build.cs
+++ b/Engine/Plugins/Runtime/StateTree/Source/StateTreeEditorModule/StateTreeEditorModule.Build.cs
@@ -30,6 +30,7 @@ namespace UnrealBuildTool.Rules
 				"BlueprintGraph",
 				"PropertyAccessEditor",
 				"StructUtils",
+				"StructUtilsEngine",
 				"StructUtilsEditor",
 				"GameplayTags",
 			}
diff --git a/Engine/Plugins/Runtime/StateTree/Source/StateTreeModule/Private/StateTree.cpp b/Engine/Plugins/Runtime/StateTree/Source/StateTreeModule/Private/StateTree.cpp
--- a/Engine/Plugins/Runtime/StateTree/Source/StateTreeModule/Private/StateTree.cpp
+++ b/Engine/Plugins/Runtime/StateTree/Source/StateTreeModule/Private/StateTree.cpp
@@ -639,13 +639,13 @@ bool UStateTree::PatchBindings()
 		{
 			FStateTreePropertyPathBinding& Binding = PropertyPathBindings[Index];
 
-			if (!Binding.GetMutableSourcePath().UpdateInstanceStructsFromValue(DataViews[Binding.GetCompiledSourceStructIndex().Get()]))
+			if (!Binding.GetMutableSourcePath().UpdateSegmentsFromValue(DataViews[Binding.GetCompiledSourceStructIndex().Get()]))
 			{
 				UE_LOG(LogStateTree, Error, TEXT("%hs: Failed to update source instance structs for property binding '%s'."), __FUNCTION__, *Binding.GetTargetPath().ToString());
 				return false;
 			}
 
-			if (!Binding.GetMutableTargetPath().UpdateInstanceStructsFromValue(TargetView))
+			if (!Binding.GetMutableTargetPath().UpdateSegmentsFromValue(TargetView))
 			{
 				UE_LOG(LogStateTree, Error, TEXT("%hs: Failed to update target instance structs for property binding '%s'."), __FUNCTION__, *Binding.GetTargetPath().ToString());
 				return false;
diff --git a/Engine/Plugins/Runtime/StateTree/Source/StateTreeModule/Private/StateTreePropertyBindings.cpp b/Engine/Plugins/Runtime/StateTree/Source/StateTreeModule/Private/StateTreePropertyBindings.cpp
--- a/Engine/Plugins/Runtime/StateTree/Source/StateTreeModule/Private/StateTreePropertyBindings.cpp
+++ b/Engine/Plugins/Runtime/StateTree/Source/StateTreeModule/Private/StateTreePropertyBindings.cpp
@@ -4,6 +4,14 @@
 #include "Misc/EnumerateRange.h"
 #include "PropertyPathHelpers.h"
 
+#if WITH_EDITOR
+#include "UObject/CoreRedirects.h"
+#include "UObject/Package.h"
+#include "Engine/BlueprintGeneratedClass.h"
+#include "Engine/UserDefinedStruct.h"
+#include "Kismet2/StructureEditorUtils.h"
+#endif
+
 #include UE_INLINE_GENERATED_CPP_BY_NAME(StateTreePropertyBindings)
 
 namespace UE::StateTree::Private
@@ -1211,10 +1219,15 @@ bool FStateTreePropertyPath::FromString(const FString& InPath)
 	return bResult;
 }
 
-bool FStateTreePropertyPath::UpdateInstanceStructsFromValue(const FStateTreeDataView BaseValueView, FString* OutError)
+bool FStateTreePropertyPath::UpdateSegments(const UStruct* BaseStruct, FString* OutError)
+{
+	return UpdateSegmentsFromValue(FStateTreeDataView(BaseStruct, nullptr), OutError);
+}
+
+bool FStateTreePropertyPath::UpdateSegmentsFromValue(const FStateTreeDataView BaseValueView, FString* OutError)
 {
 	TArray<FStateTreePropertyPathIndirection> Indirections;
-	if (!ResolveIndirectionsWithValue(BaseValueView, Indirections, OutError))
+	if (!ResolveIndirectionsWithValue(BaseValueView, Indirections, OutError, /*bHandleRedirects*/true))
 	{
 		return false;
 	}
@@ -1230,6 +1243,13 @@ bool FStateTreePropertyPath::UpdateInstanceStructsFromValue(const FStateTreeData
 		{
 			Segments[Indirection.PathSegmentIndex].SetInstanceStruct(Indirection.InstanceStruct);
 		}
+#if WITH_EDITORONLY_DATA		
+		if (!Indirection.GetRedirectedName().IsNone())
+		{
+			Segments[Indirection.PathSegmentIndex].SetName(Indirection.GetRedirectedName());
+		}
+		Segments[Indirection.PathSegmentIndex].SetPropertyGuid(Indirection.GetPropertyGuid());
+#endif			
 	}
 
 	return true;
@@ -1272,12 +1292,12 @@ FString FStateTreePropertyPath::ToString(const int32 HighlightedSegment, const T
 }
 
 
-bool FStateTreePropertyPath::ResolveIndirections(const UStruct* BaseStruct, TArray<FStateTreePropertyPathIndirection>& OutIndirections, FString* OutError) const
+bool FStateTreePropertyPath::ResolveIndirections(const UStruct* BaseStruct, TArray<FStateTreePropertyPathIndirection>& OutIndirections, FString* OutError, bool bHandleRedirects) const
 {
-	return ResolveIndirectionsWithValue(FStateTreeDataView(BaseStruct, nullptr), OutIndirections, OutError);
+	return ResolveIndirectionsWithValue(FStateTreeDataView(BaseStruct, nullptr), OutIndirections, OutError, bHandleRedirects);
 }
 
-bool FStateTreePropertyPath::ResolveIndirectionsWithValue(const FStateTreeDataView BaseValueView, TArray<FStateTreePropertyPathIndirection>& OutIndirections, FString* OutError) const
+bool FStateTreePropertyPath::ResolveIndirectionsWithValue(const FStateTreeDataView BaseValueView, TArray<FStateTreePropertyPathIndirection>& OutIndirections, FString* OutError, bool bHandleRedirects) const
 {
 	OutIndirections.Reset();
 	if (OutError)
@@ -1310,7 +1330,79 @@ bool FStateTreePropertyPath::ResolveIndirectionsWithValue(const FStateTreeDataVi
 		}
 
 		const FProperty* Property = CurrentStruct->FindPropertyByName(Segment->GetName());
-		if (Property == nullptr)
+
+#if WITH_EDITORONLY_DATA
+		FName RedirectedName;
+		FGuid PropertyGuid = Segment->GetPropertyGuid();
+
+		// Try to fix the path in editor.
+		if (bHandleRedirects)
+		{
+			
+			// Check if there's a core redirect for it.
+			if (!Property)
+			{
+				// Try to match by property ID (Blueprint or User Defined Struct).
+				if (Segment->GetPropertyGuid().IsValid())
+				{
+					if (const UBlueprintGeneratedClass* BlueprintClass = Cast<UBlueprintGeneratedClass>(CurrentStruct))
+					{
+						if (const FName* Name = BlueprintClass->PropertyGuids.FindKey(Segment->GetPropertyGuid()))
+						{
+							RedirectedName = *Name;
+							Property = CurrentStruct->FindPropertyByName(RedirectedName);
+						}
+					}
+					else if (const UUserDefinedStruct* UserDefinedStruct = Cast<UUserDefinedStruct>(CurrentStruct))
+					{
+						if (FProperty* FoundProperty = FStructureEditorUtils::GetPropertyByGuid(UserDefinedStruct, Segment->GetPropertyGuid()))
+						{
+							RedirectedName = FoundProperty->GetFName();
+							Property = FoundProperty;
+						}
+					}
+				}
+				else
+				{
+					// Try core redirect
+					const FCoreRedirectObjectName OldPropertyName(Segment->GetName(), CurrentStruct->GetFName(), *CurrentStruct->GetOutermost()->GetPathName());
+					const FCoreRedirectObjectName NewPropertyName = FCoreRedirects::GetRedirectedName(ECoreRedirectFlags::Type_Property, OldPropertyName);
+					if (OldPropertyName != NewPropertyName)
+					{
+						// Cached the result for later use.
+						RedirectedName = NewPropertyName.ObjectName;
+
+						Property = CurrentStruct->FindPropertyByName(RedirectedName);
+					}
+				}
+			}
+
+			// Update PropertyGuid 
+			if (Property)
+			{
+				const FName PropertyName = !RedirectedName.IsNone() ? RedirectedName : Segment->GetName();
+				if (const UBlueprintGeneratedClass* BlueprintClass = Cast<UBlueprintGeneratedClass>(CurrentStruct))
+				{
+					if (const FGuid* VarGuid = BlueprintClass->PropertyGuids.Find(PropertyName))
+					{
+						PropertyGuid = *VarGuid;
+					}
+				}
+				else if (const UUserDefinedStruct* UserDefinedStruct = Cast<UUserDefinedStruct>(CurrentStruct))
+				{
+					// Parse Guid from UDS property name.
+					PropertyGuid = FStructureEditorUtils::GetGuidFromPropertyName(PropertyName);
+				}
+			}
+		}
+#else
+		if (bHandleRedirects)
+		{
+			ensureMsgf(false, TEXT("The bHandleRedirects option is only available on editor builds."));
+		}
+#endif // WITH_EDITORONLY_DATA
+
+		if (!Property)
 		{
 			if (OutError)
 			{
@@ -1335,7 +1427,11 @@ bool FStateTreePropertyPath::ResolveIndirectionsWithValue(const FStateTreeDataVi
 			Indirection.ArrayIndex = Segment->GetArrayIndex();
 			Indirection.PropertyOffset = ArrayProperty->GetOffset_ForInternal();
 			Indirection.PathSegmentIndex = Segment.GetIndex();
-			Indirection.AccessType = EStateTreePropertyAccessType::IndexArray; 
+			Indirection.AccessType = EStateTreePropertyAccessType::IndexArray;
+#if WITH_EDITORONLY_DATA
+			Indirection.RedirectedName = RedirectedName;
+			Indirection.PropertyGuid = PropertyGuid;
+#endif
 			
 			ArrayIndex = 0;
 			Offset = 0;
@@ -1381,7 +1477,10 @@ bool FStateTreePropertyPath::ResolveIndirectionsWithValue(const FStateTreeDataVi
 		Indirection.PropertyOffset = Offset;
 		Indirection.PathSegmentIndex = Segment.GetIndex();
 		Indirection.AccessType = EStateTreePropertyAccessType::Offset; 
-
+#if WITH_EDITORONLY_DATA
+		Indirection.RedirectedName = RedirectedName;
+		Indirection.PropertyGuid = PropertyGuid;
+#endif
 		const bool bLastSegment = Segment.GetIndex() == (Segments.Num() - 1);
 
 		if (!bLastSegment)
diff --git a/Engine/Plugins/Runtime/StateTree/Source/StateTreeModule/Public/StateTreePropertyBindings.h b/Engine/Plugins/Runtime/StateTree/Source/StateTreeModule/Public/StateTreePropertyBindings.h
--- a/Engine/Plugins/Runtime/StateTree/Source/StateTreeModule/Public/StateTreePropertyBindings.h
+++ b/Engine/Plugins/Runtime/StateTree/Source/StateTreeModule/Public/StateTreePropertyBindings.h
@@ -209,12 +209,29 @@ struct STATETREEMODULE_API FStateTreePropertyPathSegment
 	{
 		return InstanceStruct;
 	}
+
+#if WITH_EDITORONLY_DATA
+	FGuid GetPropertyGuid() const
+	{
+		return PropertyGuid;
+	}
+
+	void SetPropertyGuid(const FGuid NewGuid)
+	{
+		PropertyGuid = NewGuid;
+	}
+#endif
 	
 private:
 	/** Name of the property */
 	UPROPERTY()
 	FName Name;
 
+#if WITH_EDITORONLY_DATA
+	/** Guid of the property for Blueprint classes or User Defined Structs. */
+	FGuid PropertyGuid;
+#endif
+
 	/** Array index if the property is dynamic or static array. */
 	UPROPERTY()
 	int32 ArrayIndex = INDEX_NONE;
@@ -246,6 +263,11 @@ struct FStateTreePropertyPathIndirection
 	int32 GetPathSegmentIndex() const { return PathSegmentIndex; }
 	EStateTreePropertyAccessType GetAccessType() const { return AccessType; }
 	const uint8* GetPropertyAddress() const { return ContainerAddress + PropertyOffset; }
+
+#if WITH_EDITORONLY_DATA
+	FName GetRedirectedName() const { return RedirectedName; }
+	FGuid GetPropertyGuid() const { return PropertyGuid; }
+#endif
 	
 private:
 	/** Property at the indirection. */
@@ -259,6 +281,14 @@ private:
 	
 	/** Type of the instance class/struct of when AccessType is ObjectInstance or StructInstance. */
 	const UStruct* InstanceStruct = nullptr;
+
+#if WITH_EDITORONLY_DATA
+	/** Redirected name, if the give property name was not found but was reconciled using core redirect or property Guid. Requires ResolveIndirections/WithValue() to be called with bHandleRedirects = true. */
+	FName RedirectedName;
+
+	/** Guid of the property for Blueprint classes or User Defined Structs. Requires ResolveIndirections/WithValue() to be called with bHandleRedirects = true. */
+	FGuid PropertyGuid;
+#endif
 	
 	/** Array index for static and dynamic arrays. Note: static array indexing is baked in the PropertyOffset. */
 	int32 ArrayIndex = 0;
@@ -326,28 +356,46 @@ struct STATETREEMODULE_API FStateTreePropertyPath
 	 * @param BaseStruct Base struct/class type the path is relative to.
 	 * @param OutIndirections Indirections describing how the properties were accessed. 
 	 * @param OutError Optional, pointer to string where error will be logged if update fails.
+	 * @param bHandleRedirects If true, the method will try to resolve missing properties using core redirects, and properties on Blueprint and User Defined Structs by ID. Available only in editor builds!
 	 * @return true of path could be resolved against the base value, and instance types were updated.
 	 */
-	bool ResolveIndirections(const UStruct* BaseStruct, TArray<FStateTreePropertyPathIndirection>& OutIndirections, FString* OutError = nullptr) const;
+	bool ResolveIndirections(const UStruct* BaseStruct, TArray<FStateTreePropertyPathIndirection>& OutIndirections, FString* OutError = nullptr, bool bHandleRedirects = false) const;
 	
 	/**
 	 * Resolves the property path against base value. The path is assumed to be relative to the BaseValueView.
 	 * @param BaseValueView Base value the path is relative to.
 	 * @param OutIndirections Indirections describing how the properties were accessed. 
 	 * @param OutError Optional, pointer to string where error will be logged if update fails.
+	 * @param bHandleRedirects If true, the method will try to resolve missing properties using core redirects, and properties on Blueprint and User Defined Structs by ID. Available only in editor builds!
 	 * @return true of path could be resolved against the base value, and instance types were updated.
 	 */
-	bool ResolveIndirectionsWithValue(const FStateTreeDataView BaseValueView, TArray<FStateTreePropertyPathIndirection>& OutIndirections, FString* OutError = nullptr) const;
+	bool ResolveIndirectionsWithValue(const FStateTreeDataView BaseValueView, TArray<FStateTreePropertyPathIndirection>& OutIndirections, FString* OutError = nullptr, bool bHandleRedirects = false) const;
 
 	/**
-	 * Updates property segment InstanceStruct types from value. The path is expected to be relative to the base value.
+	 * Updates property segments from base struct type. The path is expected to be relative to the BaseStruct.
+	 * The method handles renamed properties (core redirect, Blueprint and User Defined Structs by ID).
+	 * @param BaseValueView Base value the path is relative to.
+	 * @param OutError Optional, pointer to string where error will be logged if update fails.
+	 * @return true of path could be resolved against the base value, and instance types were updated.
+	 */
+	bool UpdateSegments(const UStruct* BaseStruct, FString* OutError = nullptr);
+
+	/**
+	 * Updates property segments from base value. The path is expected to be relative to the base value.
+	 * The method updates instance types, and handles renamed properties (core redirect, Blueprint and User Defined Structs by ID).
 	 * By storing the instance types on the path, we can resolve the path without the base value later.
 	 * @param BaseValueView Base value the path is relative to.
 	 * @param OutError Optional, pointer to string where error will be logged if update fails.
 	 * @return true of path could be resolved against the base value, and instance types were updated.
 	 */
-	bool UpdateInstanceStructsFromValue(const FStateTreeDataView BaseValueView, FString* OutError = nullptr);
+	bool UpdateSegmentsFromValue(const FStateTreeDataView BaseValueView, FString* OutError = nullptr);
 
+	UE_DEPRECATED(5.3, "Use UpdateSegmentsFromValue instead")
+	bool UpdateInstanceStructsFromValue(const FStateTreeDataView BaseValueView, FString* OutError = nullptr)
+	{
+		return UpdateSegmentsFromValue(BaseValueView, OutError);
+	}
+	
 	/** @return true if the path is empty. In that case the path points to the struct. */
 	bool IsPathEmpty() const { return Segments.IsEmpty(); }
 
diff --git a/Engine/Plugins/Runtime/StateTree/Source/StateTreeTestSuite/Private/StateTreeTest.cpp b/Engine/Plugins/Runtime/StateTree/Source/StateTreeTestSuite/Private/StateTreeTest.cpp
--- a/Engine/Plugins/Runtime/StateTree/Source/StateTreeTestSuite/Private/StateTreeTest.cpp
+++ b/Engine/Plugins/Runtime/StateTree/Source/StateTreeTestSuite/Private/StateTreeTest.cpp
@@ -1081,7 +1081,7 @@ struct FStateTreeTest_PropertyPathObject : FAITestBase
 		UStateTreeTest_PropertyObject* Object = NewObject<UStateTreeTest_PropertyObject>();
 		Object->InstancedObject = NewObject<UStateTreeTest_PropertyObjectInstanced>();
 		
-		const bool bUpdateResult = Path.UpdateInstanceStructsFromValue(FStateTreeDataView(Object));
+		const bool bUpdateResult = Path.UpdateSegmentsFromValue(FStateTreeDataView(Object));
 
 		AITEST_TRUE("Update instance types should succeeed", bUpdateResult);
 		AITEST_TRUE("Path segment 0 instance type should be UStateTreeTest_PropertyObjectInstanced", Path.GetSegment(0).GetInstanceStruct() == UStateTreeTest_PropertyObjectInstanced::StaticClass());
@@ -1256,7 +1256,7 @@ struct FStateTreeTest_PropertyPathArrayOfInstancedObjects : FAITestBase
 		UStateTreeTest_PropertyObject* Object = NewObject<UStateTreeTest_PropertyObject>();
 		Object->ArrayOfInstancedStructs.Emplace(FConstStructView::Make(Struct));
 
-		const bool bUpdateResult = Path.UpdateInstanceStructsFromValue(FStateTreeDataView(Object));
+		const bool bUpdateResult = Path.UpdateSegmentsFromValue(FStateTreeDataView(Object));
 		AITEST_TRUE("Update instance types should succeeed", bUpdateResult);
 		AITEST_EQUAL("Should have 2 path segments", Path.NumSegments(), 2);
 		AITEST_TRUE("Path segment 0 instance type should be FStateTreeTest_PropertyStruct", Path.GetSegment(0).GetInstanceStruct() == FStateTreeTest_PropertyStruct::StaticStruct());
