StateTree: Fixed ID clash when duplicating transitions, plus misc editor fixes
- Fix duplicate state
- Fix creating new node IDs when a whole transitions is duplicated
- Fixed state ID leaking into state names, when ID debug display is on
- Fixed updating the state tree view when task type is changed

diff --git a/Engine/Plugins/Runtime/StateTree/Source/StateTreeEditorModule/Private/Customizations/StateTreeTransitionDetails.cpp b/Engine/Plugins/Runtime/StateTree/Source/StateTreeEditorModule/Private/Customizations/StateTreeTransitionDetails.cpp
--- a/Engine/Plugins/Runtime/StateTree/Source/StateTreeEditorModule/Private/Customizations/StateTreeTransitionDetails.cpp
+++ b/Engine/Plugins/Runtime/StateTree/Source/StateTreeEditorModule/Private/Customizations/StateTreeTransitionDetails.cpp
@@ -242,8 +242,38 @@ void FStateTreeTransitionDetails::OnCopyTransition() const
 	} 
 }
 
+UStateTreeEditorData* FStateTreeTransitionDetails::GetEditorData() const
+{
+	TArray<UObject*> OuterObjects;
+	StructProperty->GetOuterObjects(OuterObjects);
+	for (UObject* Outer : OuterObjects)
+	{
+		UStateTreeEditorData* OuterEditorData = Cast<UStateTreeEditorData>(Outer);
+		if (OuterEditorData == nullptr)
+		{
+			OuterEditorData = Outer->GetTypedOuter<UStateTreeEditorData>();
+		}
+		if (OuterEditorData)
+		{
+			return OuterEditorData;
+		}
+	}
+	return nullptr;
+}
+
 void FStateTreeTransitionDetails::OnPasteTransition() const
 {
+	UStateTreeEditorData* EditorData = GetEditorData();
+	if (!EditorData)
+	{
+		return;
+	}
+	FStateTreeEditorPropertyBindings* Bindings = EditorData->GetPropertyEditorBindings();
+	if (!Bindings)
+	{
+		return;
+	}
+
 	FString PastedText;
 	FPlatformApplicationMisc::ClipboardPaste(PastedText);
 
@@ -267,6 +297,16 @@ void FStateTreeTransitionDetails::OnPasteTransition() const
 		if (FStateTreeTransition* Transition = static_cast<FStateTreeTransition*>(RawData[Index]))
 		{
 			Transition->ID = FGuid::NewGuid();
+
+			for (FStateTreeEditorNode& Condition : Transition->Conditions)
+			{
+				const FGuid OldStructID = Condition.ID;
+				Condition.ID = FGuid::NewGuid();
+				if (OldStructID.IsValid())
+				{
+					Bindings->CopyBindings(OldStructID, Condition.ID);
+				}
+			}
 		}
 	}
 
diff --git a/Engine/Plugins/Runtime/StateTree/Source/StateTreeEditorModule/Private/Customizations/StateTreeTransitionDetails.h b/Engine/Plugins/Runtime/StateTree/Source/StateTreeEditorModule/Private/Customizations/StateTreeTransitionDetails.h
--- a/Engine/Plugins/Runtime/StateTree/Source/StateTreeEditorModule/Private/Customizations/StateTreeTransitionDetails.h
+++ b/Engine/Plugins/Runtime/StateTree/Source/StateTreeEditorModule/Private/Customizations/StateTreeTransitionDetails.h
@@ -8,8 +8,7 @@ class IPropertyHandle;
 class IPropertyUtilities;
 class IDetailChildrenBuilder;
 class FDetailWidgetRow;
-class UStateTree;
-class UStateTreeState;
+class UStateTreeEditorData;
 enum class EStateTreeTransitionTrigger : uint8;
 
 /**
@@ -28,6 +27,8 @@ public:
 
 private:
 
+	UStateTreeEditorData* GetEditorData() const;
+
 	FText GetDescription() const;
 
 	EStateTreeTransitionTrigger GetTrigger() const;
diff --git a/Engine/Plugins/Runtime/StateTree/Source/StateTreeEditorModule/Private/SStateTreeView.cpp b/Engine/Plugins/Runtime/StateTree/Source/StateTreeEditorModule/Private/SStateTreeView.cpp
--- a/Engine/Plugins/Runtime/StateTree/Source/StateTreeEditorModule/Private/SStateTreeView.cpp
+++ b/Engine/Plugins/Runtime/StateTree/Source/StateTreeEditorModule/Private/SStateTreeView.cpp
@@ -359,10 +359,7 @@ void SStateTreeView::HandleModelStatesChanged(const TSet<UStateTreeState*>& Affe
 	// This method is called when anything in a state changes, make sure to only rebuild when needed.
 	if (PropertyChangedEvent.MemberProperty->GetFName() == GET_MEMBER_NAME_CHECKED(UStateTreeState, Tasks))
 	{
-		if (PropertyChangedEvent.ChangeType != EPropertyChangeType::ValueSet)
-		{
-			bArraysChanged = true;
-		}
+		bArraysChanged = true;
 	}
 		
 	if (bArraysChanged)
diff --git a/Engine/Plugins/Runtime/StateTree/Source/StateTreeEditorModule/Private/SStateTreeViewRow.cpp b/Engine/Plugins/Runtime/StateTree/Source/StateTreeEditorModule/Private/SStateTreeViewRow.cpp
--- a/Engine/Plugins/Runtime/StateTree/Source/StateTreeEditorModule/Private/SStateTreeViewRow.cpp
+++ b/Engine/Plugins/Runtime/StateTree/Source/StateTreeEditorModule/Private/SStateTreeViewRow.cpp
@@ -168,6 +168,20 @@ void SStateTreeViewRow::Construct(const FArguments& InArgs, const TSharedRef<STa
 											.Clipping(EWidgetClipping::ClipToBounds)
 											.IsSelected(this, &SStateTreeViewRow::IsStateSelected)
 										]
+
+										// State ID
+										+ SHorizontalBox::Slot()
+										.VAlign(VAlign_Center)
+										.AutoWidth()
+										[
+											SNew(STextBlock)
+											.Visibility_Lambda([]()
+											{
+												return UE::StateTree::Editor::GbDisplayItemIds ? EVisibility::Visible : EVisibility::Collapsed;
+											})
+											.Text(this, &SStateTreeViewRow::GetStateIDDesc)
+											.TextStyle(FStateTreeEditorStyle::Get(), "StateTree.Details")
+										]
 									]
 									+ SOverlay::Slot()
 									[
@@ -690,15 +704,20 @@ FText SStateTreeViewRow::GetStateDesc() const
 {
 	if (const UStateTreeState* State = WeakState.Get())
 	{
-		if (UE::StateTree::Editor::GbDisplayItemIds)
-		{
-			return FText::FromString(FString::Printf(TEXT("%s (%s)"), *State->Name.ToString(), *LexToString(State->ID)));	
-		}
 		return FText::FromName(State->Name);
 	}
 	return FText::FromName(FName());
 }
 
+FText SStateTreeViewRow::GetStateIDDesc() const
+{
+	if (const UStateTreeState* State = WeakState.Get())
+	{
+		return FText::FromString(*LexToString(State->ID));
+	}
+	return FText::FromName(FName());
+}
+
 EVisibility SStateTreeViewRow::GetConditionVisibility() const
 {
 	if (const UStateTreeState* State = WeakState.Get())
diff --git a/Engine/Plugins/Runtime/StateTree/Source/StateTreeEditorModule/Private/SStateTreeViewRow.h b/Engine/Plugins/Runtime/StateTree/Source/StateTreeEditorModule/Private/SStateTreeViewRow.h
--- a/Engine/Plugins/Runtime/StateTree/Source/StateTreeEditorModule/Private/SStateTreeViewRow.h
+++ b/Engine/Plugins/Runtime/StateTree/Source/StateTreeEditorModule/Private/SStateTreeViewRow.h
@@ -35,6 +35,7 @@ private:
 	FSlateColor GetActiveStateColor() const;
 	FSlateColor GetSubTreeMarkerColor() const;
 	FText GetStateDesc() const;
+	FText GetStateIDDesc() const;
 
 	EVisibility GetConditionVisibility() const;
 	EVisibility GetStateBreakpointVisibility() const;
diff --git a/Engine/Plugins/Runtime/StateTree/Source/StateTreeEditorModule/Private/StateTreeState.cpp b/Engine/Plugins/Runtime/StateTree/Source/StateTreeEditorModule/Private/StateTreeState.cpp
--- a/Engine/Plugins/Runtime/StateTree/Source/StateTreeEditorModule/Private/StateTreeState.cpp
+++ b/Engine/Plugins/Runtime/StateTree/Source/StateTreeEditorModule/Private/StateTreeState.cpp
@@ -216,6 +216,14 @@ void UStateTreeState::PostEditChangeChainProperty(FPropertyChangedChainEvent& Pr
 			{
 				FStateTreeTransition& Transition = Transitions[TransitionsIndex];
 				Transition.ID = FGuid::NewGuid();
+
+				// Update conditions
+				for (FStateTreeEditorNode& Condition : Transition.Conditions)
+				{
+					const FGuid OldStructID = Condition.ID;
+					Condition.ID = FGuid::NewGuid();
+					CopyBindings(OldStructID, Condition.ID);
+				}
 			}
 		}
 
diff --git a/Engine/Plugins/Runtime/StateTree/Source/StateTreeEditorModule/Private/StateTreeViewModel.cpp b/Engine/Plugins/Runtime/StateTree/Source/StateTreeEditorModule/Private/StateTreeViewModel.cpp
--- a/Engine/Plugins/Runtime/StateTree/Source/StateTreeEditorModule/Private/StateTreeViewModel.cpp
+++ b/Engine/Plugins/Runtime/StateTree/Source/StateTreeEditorModule/Private/StateTreeViewModel.cpp
@@ -836,11 +836,8 @@ void FStateTreeViewModel::DuplicateSelectedStates()
 	const int32 Index = AfterState->Parent ? AfterState->Parent->Children.Find(AfterState) : TreeData->SubTrees.Find(AfterState);
 	if (Index != INDEX_NONE)
 	{
-		FString TextToImport;
-		FPlatformApplicationMisc::ClipboardPaste(TextToImport);
-			
 		const FScopedTransaction Transaction(LOCTEXT("DuplicateStatesTransaction", "Duplicate State(s)"));
-		PasteStatesAsChildrenFromText(TextToImport, AfterState->Parent, Index + 1);
+		PasteStatesAsChildrenFromText(ExportedText, AfterState->Parent, Index + 1);
 	}
 }
 
