[StateTreeDebugger]
- enabled by default in the StateTreeEditor
- traces can be automatically started on non-editor targets by setting 'bAutoStartDebuggerTracesOnNonEditorTargets = true in StateTreeSettings or' under [/Script/StateTreeModule.StateTreeSettings]
in the .ini config file.
- traces can be manually started/stopped using `statetree.startdebuggertraces` and `statetree.stopdebuggertraces` from the console

diff --git a/Engine/Plugins/Runtime/StateTree/Source/StateTreeEditorModule/Private/SStateTreeView.cpp b/Engine/Plugins/Runtime/StateTree/Source/StateTreeEditorModule/Private/SStateTreeView.cpp
--- a/Engine/Plugins/Runtime/StateTree/Source/StateTreeEditorModule/Private/SStateTreeView.cpp
+++ b/Engine/Plugins/Runtime/StateTree/Source/StateTreeEditorModule/Private/SStateTreeView.cpp
@@ -472,12 +472,9 @@ TSharedPtr<SWidget> SStateTreeView::HandleContextMenuOpening()
 	MenuBuilder.AddMenuEntry(FStateTreeEditorCommands::Get().EnableStates);
 
 #if WITH_STATETREE_DEBUGGER
-	if (UStateTreeSettings::Get().bUseDebugger)
-	{
-		MenuBuilder.AddSeparator();
-		MenuBuilder.AddMenuEntry(FStateTreeDebuggerCommands::Get().EnableOnEnterStateBreakpoint);
-		MenuBuilder.AddMenuEntry(FStateTreeDebuggerCommands::Get().EnableOnExitStateBreakpoint);
-	}
+	MenuBuilder.AddSeparator();
+	MenuBuilder.AddMenuEntry(FStateTreeDebuggerCommands::Get().EnableOnEnterStateBreakpoint);
+	MenuBuilder.AddMenuEntry(FStateTreeDebuggerCommands::Get().EnableOnExitStateBreakpoint);
 #endif // WITH_STATETREE_DEBUGGER
 	
 	return MenuBuilder.MakeWidget();
diff --git a/Engine/Plugins/Runtime/StateTree/Source/StateTreeEditorModule/Private/StateTreeEditor.cpp b/Engine/Plugins/Runtime/StateTree/Source/StateTreeEditorModule/Private/StateTreeEditor.cpp
--- a/Engine/Plugins/Runtime/StateTree/Source/StateTreeEditorModule/Private/StateTreeEditor.cpp
+++ b/Engine/Plugins/Runtime/StateTree/Source/StateTreeEditorModule/Private/StateTreeEditor.cpp
@@ -104,13 +104,10 @@ void FStateTreeEditor::RegisterTabSpawners(const TSharedRef<class FTabManager>&
 		.SetIcon(FSlateIcon(FAppStyle::GetAppStyleSetName(), "LevelEditor.Tabs.Outliner"));
 
 #if WITH_STATETREE_DEBUGGER
-	if (UStateTreeSettings::Get().bUseDebugger)
-	{
-		InTabManager->RegisterTabSpawner(DebuggerTabId, FOnSpawnTab::CreateSP(this, &FStateTreeEditor::SpawnTab_Debugger))
-		   .SetDisplayName(NSLOCTEXT("StateTreeEditor", "DebuggerTab", "Debugger"))
-		   .SetGroup(WorkspaceMenuCategoryRef)
-		   .SetIcon(FSlateIcon(FAppStyle::GetAppStyleSetName(), "LevelEditor.Tabs.Outliner"));
-	}
+	InTabManager->RegisterTabSpawner(DebuggerTabId, FOnSpawnTab::CreateSP(this, &FStateTreeEditor::SpawnTab_Debugger))
+	   .SetDisplayName(NSLOCTEXT("StateTreeEditor", "DebuggerTab", "Debugger"))
+	   .SetGroup(WorkspaceMenuCategoryRef)
+	   .SetIcon(FSlateIcon(FAppStyle::GetAppStyleSetName(), "LevelEditor.Tabs.Outliner"));
 #endif // WITH_STATETREE_DEBUGGER
 }
 
diff --git a/Engine/Plugins/Runtime/StateTree/Source/StateTreeModule/Private/StateTreeModule.cpp b/Engine/Plugins/Runtime/StateTree/Source/StateTreeModule/Private/StateTreeModule.cpp
--- a/Engine/Plugins/Runtime/StateTree/Source/StateTreeModule/Private/StateTreeModule.cpp
+++ b/Engine/Plugins/Runtime/StateTree/Source/StateTreeModule/Private/StateTreeModule.cpp
@@ -52,19 +52,25 @@ class FStateTreeModule : public IStateTreeModule
 	
 	FStateTreeTraceModule StateTreeTraceModule;
 
-	FAutoConsoleCommand EnableDebugger = FAutoConsoleCommand(
-		TEXT("statetree.enabledebugger"),
-		TEXT("Sets `bUseDebugger` to true in StateTreeSettings and turns on traces if not already active."),
-		FConsoleCommandDelegate::CreateLambda(
-			[]()
+	/** Keep track if StartTraces was explicitly called. */
+	bool bIsTracing = false;
+
+	FAutoConsoleCommand StartDebuggerTracesCommand = FAutoConsoleCommand(
+		TEXT("statetree.startdebuggertraces"),
+		TEXT("Turns on StateTree debugger traces if not already active."),
+		FConsoleCommandDelegate::CreateLambda([]
 		{
-				UStateTreeSettings& Settings = UStateTreeSettings::Get();
-				if (Settings.bUseDebugger == false)
-				{
-					Settings.bUseDebugger = true;
-					FStateTreeModule& StateTreeModule = FModuleManager::GetModuleChecked<FStateTreeModule>("StateTree");
-					StateTreeModule.StartTraces();	
-				}
+			FStateTreeModule& StateTreeModule = FModuleManager::GetModuleChecked<FStateTreeModule>("StateTreeModule");
+			StateTreeModule.StartTraces();
+		}));
+
+	FAutoConsoleCommand StopDebuggerTracesCommand = FAutoConsoleCommand(
+		TEXT("statetree.stopdebuggertraces"),
+		TEXT("Turns off StateTree debugger traces if active."),
+		FConsoleCommandDelegate::CreateLambda([]
+		{
+			FStateTreeModule& StateTreeModule = FModuleManager::GetModuleChecked<FStateTreeModule>("StateTreeModule");
+			StateTreeModule.StopTraces();
 		}));
 #endif // WITH_STATETREE_DEBUGGER
 };
@@ -84,7 +90,10 @@ void FStateTreeModule::StartupModule()
 #if !WITH_EDITOR
 	// We don't automatically start traces for Editor targets since we rely on the debugger
 	// to start recording either on user action or on PIE session start.
-	StartTraces();
+	if (UStateTreeSettings::Get().bAutoStartDebuggerTracesOnNonEditorTargets)
+	{
+		StartTraces();
+	}
 #endif // !WITH_EDITOR
 
 #endif // WITH_STATETREE_DEBUGGER
@@ -109,7 +118,7 @@ void FStateTreeModule::ShutdownModule()
 void FStateTreeModule::StartTraces()
 {
 #if WITH_STATETREE_DEBUGGER
-	if (!UStateTreeSettings::Get().bUseDebugger || IsRunningCommandlet())
+	if (IsRunningCommandlet() || bIsTracing)
 	{
 		return;
 	}
@@ -151,12 +160,18 @@ void FStateTreeModule::StartTraces()
 		Options.bExcludeTail = true;
 		FTraceAuxiliary::Start(FTraceAuxiliary::EConnectionType::Network, TEXT("localhost"), TEXT(""), &Options, LogStateTree);
 	}
+
+	bIsTracing = true;
 #endif // WITH_STATETREE_DEBUGGER
 }
 
 void FStateTreeModule::StopTraces()
 {
 #if WITH_STATETREE_DEBUGGER
+	if (bIsTracing == false)
+	{
+		return;
+	}
 
 	UE::Trace::ToggleChannel(TEXT("StateTreeDebugChannel"), false);
 	UE::Trace::ToggleChannel(TEXT("FrameChannel"), false);
@@ -175,6 +190,8 @@ void FStateTreeModule::StopTraces()
 	{
 		FTraceAuxiliary::Stop();
 	}
+
+	bIsTracing = false;
 #endif // WITH_STATETREE_DEBUGGER
 }
 
diff --git a/Engine/Plugins/Runtime/StateTree/Source/StateTreeModule/Public/StateTreeSettings.h b/Engine/Plugins/Runtime/StateTree/Source/StateTreeModule/Public/StateTreeSettings.h
--- a/Engine/Plugins/Runtime/StateTree/Source/StateTreeModule/Public/StateTreeSettings.h
+++ b/Engine/Plugins/Runtime/StateTree/Source/StateTreeModule/Public/StateTreeSettings.h
@@ -14,9 +14,13 @@ class STATETREEMODULE_API UStateTreeSettings : public UDeveloperSettings
 	GENERATED_BODY()
 
 public:
-	static UStateTreeSettings& Get() { return *CastChecked<UStateTreeSettings>(UStateTreeSettings::StaticClass()->GetDefaultObject()); }
+	static UStateTreeSettings& Get() { return *CastChecked<UStateTreeSettings>(StaticClass()->GetDefaultObject()); }
 
-	/** The debugger is still an experimental feature, hence not active by default. */
+	/**
+	 * Editor targets relies on PIE and StateTreeEditor to start/stop traces.
+	 * This is to start traces automatically when launching Standalone, Client or Server builds. 
+	 * It's also possible to do it manually using 'statetree.startdebuggertraces' and 'statetree.stopdebuggertraces' in the console.
+	 */
 	UPROPERTY(EditDefaultsOnly, Category = StateTree, config)
-	bool bUseDebugger = false;
+	bool bAutoStartDebuggerTracesOnNonEditorTargets = false;
 };
\ No newline at end of file
