StateTree: Fix handling of transition tasks during State Tree compile
- restored tracking of struct/object instances per state (got lost in a merge)
- fixed tracking of bHasTransitionTasks (was storing only last task)

diff --git a/Engine/Plugins/Runtime/StateTree/Source/StateTreeEditorModule/Private/StateTreeCompiler.cpp b/Engine/Plugins/Runtime/StateTree/Source/StateTreeEditorModule/Private/StateTreeCompiler.cpp
--- a/Engine/Plugins/Runtime/StateTree/Source/StateTreeEditorModule/Private/StateTreeCompiler.cpp
+++ b/Engine/Plugins/Runtime/StateTree/Source/StateTreeEditorModule/Private/StateTreeCompiler.cpp
@@ -539,10 +539,14 @@ bool FStateTreeCompiler::CreateGlobalTasks()
 	StateTree->bHasGlobalTransitionTasks = false;
 	for (FStateTreeEditorNode& TaskNode : EditorData->GlobalTasks)
 	{
-		if (!CreateTask(nullptr, TaskNode, StateTree->bHasGlobalTransitionTasks))
+		if (!CreateTask(nullptr, TaskNode))
 		{
 			return false;
 		}
+		
+		const FStateTreeTaskBase& LastAddedTask = Nodes.Last().Get<FStateTreeTaskBase>();
+		
+		StateTree->bHasGlobalTransitionTasks |= LastAddedTask.bShouldAffectTransitions;
 	}
 	
 	const int32 GlobalTasksNum = Nodes.Num() - GlobalTasksBegin;
@@ -667,24 +671,49 @@ bool FStateTreeCompiler::CreateStateTasksAndParameters()
 			return false;
 		}
 		CompactState.TasksBegin = uint16(TasksBegin);
-		CompactState.bHasTransitionTasks = false;
 		
-		for (FStateTreeEditorNode& TaskNode : SourceState->Tasks)
+		int32 TaskInstanceStructNum = 0;
+		int32 TaskInstanceObjectNum = 0;
+
+		// Update instance data num for each state.
+		if (CompactState.Type == EStateTreeStateType::Linked)
+		{
+			// Linked state parameters.
+			TaskInstanceStructNum++;
+		}
+
+		TArrayView<FStateTreeEditorNode> Tasks;
+		if (SourceState->Tasks.Num())
+		{
+			Tasks = SourceState->Tasks;
+		}
+		else if (SourceState->SingleTask.Node.IsValid())
 		{
-			bool bHasTransitionTasks = CompactState.bHasTransitionTasks; // bHasTransitionTasks is a bit flag.
-			if (!CreateTask(SourceState, TaskNode, bHasTransitionTasks))
+			Tasks = TArrayView<FStateTreeEditorNode>(&SourceState->SingleTask, 1);
+		}
+		
+		bool bStateHasTransitionTasks = false;
+		for (FStateTreeEditorNode& TaskNode : Tasks)
+		{
+			if (!CreateTask(SourceState, TaskNode))
 			{
 				return false;
 			}
-			CompactState.bHasTransitionTasks = bHasTransitionTasks; 
-		}
 
-		bool bSingleStateHasTransitionTasks = CompactState.bHasTransitionTasks; // bHasTransitionTasks is a bit flag.
-		if (!CreateTask(SourceState, SourceState->SingleTask, bSingleStateHasTransitionTasks))
-		{
-			return false;
+			const FStateTreeTaskBase& LastAddedTask = Nodes.Last().Get<FStateTreeTaskBase>();
+			if (LastAddedTask.bInstanceIsObject)
+			{
+				TaskInstanceObjectNum++;
+			}
+			else
+			{
+				TaskInstanceStructNum++;
+			}
+			
+			bStateHasTransitionTasks |= LastAddedTask.bShouldAffectTransitions;
 		}
-		CompactState.bHasTransitionTasks = bSingleStateHasTransitionTasks;
+
+		CompactState.bHasTransitionTasks = bStateHasTransitionTasks;
 	
 		const int32 TasksNum = Nodes.Num() - TasksBegin;
 		if (const auto Validation = UE::StateTree::Compiler::IsValidCount8(TasksNum); Validation.DidFail())
@@ -692,7 +721,20 @@ bool FStateTreeCompiler::CreateStateTasksAndParameters()
 			Validation.Log(Log, TEXT("TasksNum"));
 			return false;
 		}
+		if (const auto Validation = UE::StateTree::Compiler::IsValidCount8(TaskInstanceObjectNum); Validation.DidFail())
+		{
+			Validation.Log(Log, TEXT("TaskInstanceObjectNum"));
+			return false;
+		}
+		if (const auto Validation = UE::StateTree::Compiler::IsValidCount8(TaskInstanceStructNum); Validation.DidFail())
+		{
+			Validation.Log(Log, TEXT("TaskInstanceStructNum"));
+			return false;
+		}
+
 		CompactState.TasksNum = uint8(TasksNum);
+		CompactState.TaskInstanceStructNum = uint8(TaskInstanceStructNum);
+		CompactState.TaskInstanceObjectNum = uint8(TaskInstanceObjectNum);
 	}
 	
 	return true;
@@ -1125,7 +1167,7 @@ bool FStateTreeCompiler::CompileAndValidateNode(const UStateTreeState* SourceSta
 	return Result != EDataValidationResult::Invalid;
 }
 
-bool FStateTreeCompiler::CreateTask(UStateTreeState* State, const FStateTreeEditorNode& TaskNode, bool& bOutHasTransitionTasks)
+bool FStateTreeCompiler::CreateTask(UStateTreeState* State, const FStateTreeEditorNode& TaskNode)
 {
 	// Silently ignore empty nodes.
 	if (!TaskNode.Node.IsValid())
@@ -1231,11 +1273,6 @@ bool FStateTreeCompiler::CreateTask(UStateTreeState* State, const FStateTreeEdit
 		return false;
 	}
 	Task.DataViewIndex = FStateTreeIndex16(SourceStructIndex);
-
-	if (Task.bShouldAffectTransitions)
-	{
-		bOutHasTransitionTasks = true;
-	}
 	
 	return true;
 }
diff --git a/Engine/Plugins/Runtime/StateTree/Source/StateTreeEditorModule/Public/StateTreeCompiler.h b/Engine/Plugins/Runtime/StateTree/Source/StateTreeEditorModule/Public/StateTreeCompiler.h
--- a/Engine/Plugins/Runtime/StateTree/Source/StateTreeEditorModule/Public/StateTreeCompiler.h
+++ b/Engine/Plugins/Runtime/StateTree/Source/StateTreeEditorModule/Public/StateTreeCompiler.h
@@ -50,7 +50,7 @@ private:
 	
 	bool CreateConditions(UStateTreeState& State, TConstArrayView<FStateTreeEditorNode> Conditions);
 	bool CreateCondition(UStateTreeState& State, const FStateTreeEditorNode& CondNode, const EStateTreeConditionOperand Operand, const int8 DeltaIndent);
-	bool CreateTask(UStateTreeState* State, const FStateTreeEditorNode& TaskNode, bool& bOutHasTransitionTasks);
+	bool CreateTask(UStateTreeState* State, const FStateTreeEditorNode& TaskNode);
 	bool CreateEvaluator(const FStateTreeEditorNode& EvalNode);
 	bool GetAndValidateBindings(const FStateTreeBindableStructDesc& TargetStruct, FStateTreeDataView TargetValue, TArray<FStateTreePropertyPathBinding>& OutBindings) const;
 	bool IsPropertyAnyEnum(const FStateTreeBindableStructDesc& Struct, FStateTreePropertyPath Path) const;
