StateTree: Fix out of range access when compiling empty task.

diff --git a/Engine/Plugins/Runtime/StateTree/Source/StateTreeEditorModule/Private/StateTreeCompiler.cpp b/Engine/Plugins/Runtime/StateTree/Source/StateTreeEditorModule/Private/StateTreeCompiler.cpp
--- a/Engine/Plugins/Runtime/StateTree/Source/StateTreeEditorModule/Private/StateTreeCompiler.cpp
+++ b/Engine/Plugins/Runtime/StateTree/Source/StateTreeEditorModule/Private/StateTreeCompiler.cpp
@@ -539,6 +539,12 @@ bool FStateTreeCompiler::CreateGlobalTasks()
 	StateTree->bHasGlobalTransitionTasks = false;
 	for (FStateTreeEditorNode& TaskNode : EditorData->GlobalTasks)
 	{
+		// Silently ignore empty nodes.
+		if (!TaskNode.Node.IsValid())
+		{
+			continue;
+		}
+
 		if (!CreateTask(nullptr, TaskNode))
 		{
 			return false;
@@ -695,6 +701,12 @@ bool FStateTreeCompiler::CreateStateTasksAndParameters()
 		bool bStateHasTransitionTasks = false;
 		for (FStateTreeEditorNode& TaskNode : Tasks)
 		{
+			// Silently ignore empty nodes.
+			if (!TaskNode.Node.IsValid())
+			{
+				continue;
+			}
+
 			if (!CreateTask(SourceState, TaskNode))
 			{
 				return false;
@@ -1169,10 +1181,9 @@ bool FStateTreeCompiler::CompileAndValidateNode(const UStateTreeState* SourceSta
 
 bool FStateTreeCompiler::CreateTask(UStateTreeState* State, const FStateTreeEditorNode& TaskNode)
 {
-	// Silently ignore empty nodes.
 	if (!TaskNode.Node.IsValid())
 	{
-		return true;
+		return false;
 	}
 	
 	// Create binding source struct descriptor.
