[StateTree] Added option to construct a FStateTreeIndex16 directly from a uint16.
This fixes cases where we go from FStateTreeIndex16::Invalid to uint16 then back to FStateTreeIndex16.

diff --git a/Engine/Plugins/Runtime/StateTree/Source/StateTreeModule/Public/StateTreeIndexTypes.h b/Engine/Plugins/Runtime/StateTree/Source/StateTreeModule/Public/StateTreeIndexTypes.h
--- a/Engine/Plugins/Runtime/StateTree/Source/StateTreeModule/Public/StateTreeIndexTypes.h
+++ b/Engine/Plugins/Runtime/StateTree/Source/StateTreeModule/Public/StateTreeIndexTypes.h
@@ -20,14 +20,26 @@ struct STATETREEMODULE_API FStateTreeIndex16
 	}
 
 	FStateTreeIndex16() = default;
-	
+
+	/**
+	 * Construct from a uint16 index where MAX_uint16 is considered an invalid index
+	 * (i.e FStateTreeIndex16::InvalidValue).
+	 */
+	explicit FStateTreeIndex16(const uint16 InIndex) : Value(InIndex)
+	{
+	}
+
+	/**
+	 * Construct from a int32 index where INDEX_NONE is considered an invalid index
+	 * and converted to FStateTreeIndex16::InvalidValue (i.e MAX_uint16).
+	 */
 	explicit FStateTreeIndex16(const int32 InIndex)
 	{
 		check(InIndex == INDEX_NONE || IsValidIndex(InIndex));
 		Value = InIndex == INDEX_NONE ? InvalidValue : (uint16)InIndex;
 	}
 
-	/** @retrun value of the index. */
+	/** @retrun value of the index or FStateTreeIndex16::InvalidValue (i.e. MAX_uint16) if invalid. */
 	uint16 Get() const { return Value; }
 	
 	/** @return the index value as int32, mapping invalid value to INDEX_NONE. */
