[StateTreeDebugger]
- created a dedicated option to auto-scroll and converted debugger pause to control event processing
- clear tracks when selecting new valid session
- scrub buttons can now be used as soon as analysis is pause of auto-scroll disabled
- fixed potential access to invalid struct view when processing a FStateTreeTraceNodeEvent
- added return value and output parameter to IStateTreeModule::StartTraces to take appropriate actions if traces was already active or not.
- added new button to the DebuggerView to control the session analysis state (pause or not). Mainly required when performing a remove analysis of non editor target.
- added new button to the DebuggerView to toggle auto-scroll. Manual selection in the timeline will turn it off so it can now be turned back on.

diff --git a/Engine/Plugins/Runtime/StateTree/Source/StateTreeEditorModule/Private/Debugger/SStateTreeDebuggerView.cpp b/Engine/Plugins/Runtime/StateTree/Source/StateTreeEditorModule/Private/Debugger/SStateTreeDebuggerView.cpp
--- a/Engine/Plugins/Runtime/StateTree/Source/StateTreeEditorModule/Private/Debugger/SStateTreeDebuggerView.cpp
+++ b/Engine/Plugins/Runtime/StateTree/Source/StateTreeEditorModule/Private/Debugger/SStateTreeDebuggerView.cpp
@@ -21,6 +21,7 @@
 #include "StateTreeDebuggerTrack.h"
 #include "StateTreeEditorData.h"
 #include "StateTreeEditorSettings.h"
+#include "StateTreeEditorStyle.h"
 #include "StateTreeExecutionContext.h"
 #include "StateTreeModule.h"
 #include "StateTreeState.h"
@@ -161,6 +162,7 @@ SStateTreeDebuggerView::~SStateTreeDebuggerView()
 	{
 		Debugger->OnScrubStateChanged.Unbind();
 		Debugger->OnBreakpointHit.Unbind();
+		Debugger->OnNewSession.Unbind();
 		Debugger->OnNewInstance.Unbind();
 		Debugger->OnSelectedInstanceCleared.Unbind();
 	}
@@ -174,24 +176,19 @@ SStateTreeDebuggerView::~SStateTreeDebuggerView()
 
 void SStateTreeDebuggerView::StartRecording()
 {
-	// We give priority to the Editor actions even if an analysis was active (remote process)
-	// This will stop current analysis and connect to the new live trace.
-	Debugger->RequestAnalysisOfNextLiveSession();
-	
-	Debugger->ClearSelection();
-
 	if (CanStartRecording())
 	{
-		IStateTreeModule& StateTreeModule = FModuleManager::GetModuleChecked<IStateTreeModule>("StateTreeModule");
-		StateTreeModule.StartTraces();
+		Debugger->ClearSelection();
 
-		bRecording = true;
+		// We give priority to the Editor actions even if an analysis was active (remote process)
+		// This will stop current analysis and connect to the new live trace.
+		bRecording = Debugger->RequestAnalysisOfNextLiveSession();
 	}
 }
 
 void SStateTreeDebuggerView::StopRecording()
 {
-	if (bRecording)
+	if (CanStopRecording())
 	{
 		IStateTreeModule& StateTreeModule = FModuleManager::GetModuleChecked<IStateTreeModule>("StateTreeModule");
 		StateTreeModule.StopTraces();
@@ -220,18 +217,17 @@ void SStateTreeDebuggerView::OnPIEStarted(const bool bIsSimulating)
 void SStateTreeDebuggerView::OnPIEStopped(const bool bIsSimulating)
 {
 	StopRecording();
-	Debugger->StopAnalysis();
-	Debugger->Unpause();
+	Debugger->StopSessionAnalysis();
 }
 
 void SStateTreeDebuggerView::OnPIEPaused(const bool bIsSimulating) const
 {
-	Debugger->Pause();
+	Debugger->PauseSessionAnalysis();
 }
 
 void SStateTreeDebuggerView::OnPIEResumed(const bool bIsSimulating) const
 {
-	Debugger->Unpause();
+	Debugger->ResumeSessionAnalysis();
 }
 
 void SStateTreeDebuggerView::OnPIESingleStepped(bool bSimulating) const
@@ -246,9 +242,11 @@ void SStateTreeDebuggerView::Construct(const FArguments& InArgs, const UStateTre
 	StateTreeEditorData = Cast<UStateTreeEditorData>(InStateTree.EditorData.Get()); 
 
 	Debugger = InStateTreeViewModel->GetDebugger();
-	bRecording = FTraceAuxiliary::IsConnected();
+	IStateTreeModule& StateTreeModule = FModuleManager::GetModuleChecked<IStateTreeModule>("StateTreeModule");
+	bRecording = StateTreeModule.IsTracing();
 
 	// Bind callbacks to the debugger delegates
+	Debugger->OnNewSession.BindSP(this, &SStateTreeDebuggerView::OnNewSession);
 	Debugger->OnNewInstance.BindSP(this, &SStateTreeDebuggerView::OnNewInstance);
 	Debugger->OnScrubStateChanged.BindSP(this, &SStateTreeDebuggerView::OnDebuggerScrubStateChanged);
 	Debugger->OnBreakpointHit.BindSP(this, &SStateTreeDebuggerView::OnBreakpointHit, InCommandList);
@@ -260,7 +258,7 @@ void SStateTreeDebuggerView::Construct(const FArguments& InArgs, const UStateTre
 	// Put debugger in proper simulation state when view is constructed after PIE/SIE was started
 	if (FPlayWorldCommandCallbacks::HasPlayWorldAndPaused())
 	{
-		Debugger->Pause();
+		Debugger->PauseSessionAnalysis();
 	}
 
 	// Add & Bind commands
@@ -284,35 +282,86 @@ void SStateTreeDebuggerView::Construct(const FArguments& InArgs, const UStateTre
 		FGetActionCheckState::CreateLambda([this] { return GetStateBreakpointCheckState(EStateTreeBreakpointType::OnExit); }),
 		FIsActionButtonVisible::CreateLambda([this] { return CanAddStateBreakpoint(EStateTreeBreakpointType::OnExit) || CanRemoveStateBreakpoint(EStateTreeBreakpointType::OnExit); }));
 	
-	// Toolbar
-	FSlimHorizontalToolBarBuilder ToolbarBuilder(InCommandList, FMultiBoxCustomization::None, /*InExtender*/ nullptr, /*InForceSmallIcons*/ true);
-	ToolbarBuilder.BeginSection(TEXT("Debugging"));
+	// Toolbars
+
+	FUIAction AnalysisPauseToggleButtonAction;
+	AnalysisPauseToggleButtonAction.GetActionCheckState.BindLambda([this]
+	{
+		return Debugger && Debugger->IsAnalysisSessionPaused() ? ECheckBoxState::Checked : ECheckBoxState::Unchecked;
+	});
+	AnalysisPauseToggleButtonAction.ExecuteAction.BindLambda([this]
+	{
+		if (Debugger)
+		{
+			if (Debugger->IsAnalysisSessionPaused())
+			{
+				Debugger->ResumeSessionAnalysis();
+			}
+			else
+			{
+				Debugger->PauseSessionAnalysis();
+			}
+		}
+	});
+	
+	FSlimHorizontalToolBarBuilder LeftToolbar(InCommandList, FMultiBoxCustomization::None, /*InExtender*/ nullptr, /*InForceSmallIcons*/ true);
+	LeftToolbar.BeginSection(TEXT("Debugging"));
 	{
 		const FStateTreeDebuggerCommands& DebuggerCommands = FStateTreeDebuggerCommands::Get();
-		ToolbarBuilder.AddToolBarButton(DebuggerCommands.StartRecording);
-		ToolbarBuilder.AddToolBarButton(DebuggerCommands.StopRecording);
+		LeftToolbar.AddToolBarButton(DebuggerCommands.StartRecording);
+		LeftToolbar.AddToolBarButton(DebuggerCommands.StopRecording);
 
 		const FPlayWorldCommands& PlayWorldCommands = FPlayWorldCommands::Get();
-		ToolbarBuilder.AddToolBarButton(PlayWorldCommands.RepeatLastPlay);
-		ToolbarBuilder.AddToolBarButton(PlayWorldCommands.PausePlaySession,
+		LeftToolbar.AddToolBarButton(PlayWorldCommands.RepeatLastPlay);
+		LeftToolbar.AddToolBarButton(PlayWorldCommands.PausePlaySession,
 					NAME_None, TAttribute<FText>(), TAttribute<FText>(), FSlateIcon(FAppStyle::GetAppStyleSetName(), "PlayWorld.PausePlaySession.Small"));
-		ToolbarBuilder.AddToolBarButton(PlayWorldCommands.ResumePlaySession,
+		LeftToolbar.AddToolBarButton(PlayWorldCommands.ResumePlaySession,
 					NAME_None, TAttribute<FText>(), TAttribute<FText>(), FSlateIcon(FAppStyle::GetAppStyleSetName(), "PlayWorld.ResumePlaySession.Small"));	
-		ToolbarBuilder.AddToolBarButton(PlayWorldCommands.StopPlaySession,
+		LeftToolbar.AddToolBarButton(PlayWorldCommands.StopPlaySession,
 					NAME_None, TAttribute<FText>(), TAttribute<FText>(), FSlateIcon(FAppStyle::GetAppStyleSetName(), "PlayWorld.StopPlaySession.Small"));
-		ToolbarBuilder.AddSeparator();
-		ToolbarBuilder.AddToolBarButton(DebuggerCommands.PreviousFrameWithStateChange);
-		ToolbarBuilder.AddToolBarButton(DebuggerCommands.PreviousFrameWithEvents);
-		ToolbarBuilder.AddToolBarButton(DebuggerCommands.NextFrameWithEvents);
-		ToolbarBuilder.AddToolBarButton(DebuggerCommands.NextFrameWithStateChange);
+		LeftToolbar.AddSeparator();
+		LeftToolbar.AddToolBarButton(DebuggerCommands.PreviousFrameWithStateChange);
+		LeftToolbar.AddToolBarButton(DebuggerCommands.PreviousFrameWithEvents);
+		LeftToolbar.AddToolBarButton(DebuggerCommands.NextFrameWithEvents);
+		LeftToolbar.AddToolBarButton(DebuggerCommands.NextFrameWithStateChange);
+		LeftToolbar.AddSeparator();
+		LeftToolbar.AddToolBarButton(AnalysisPauseToggleButtonAction,
+			NAME_None,
+			TAttribute<FText>(),
+			LOCTEXT("DebuggerPauseToolTip", "Debugger Paused"),
+			FSlateIcon(FStateTreeEditorStyle::Get().GetStyleSetName(), "StateTreeEditor.ToggleDebuggerPause"),
+			EUserInterfaceActionType::ToggleButton);
 	}
-	ToolbarBuilder.EndSection();
+	LeftToolbar.EndSection();
+
+	FSlimHorizontalToolBarBuilder RightToolbar(nullptr, FMultiBoxCustomization::None);
+
+	FUIAction AutoScrollToggleButtonAction;
+	AutoScrollToggleButtonAction.GetActionCheckState.BindLambda([this]
+	{
+		return bAutoScroll ? ECheckBoxState::Checked : ECheckBoxState::Unchecked;
+	});
+	AutoScrollToggleButtonAction.ExecuteAction.BindLambda([this]
+	{
+		bAutoScroll = !bAutoScroll;
+	});
+
+	RightToolbar.BeginSection("Auto-Scroll");
+	RightToolbar.AddToolBarButton(
+		AutoScrollToggleButtonAction,
+		NAME_None,
+		TAttribute<FText>(),
+		LOCTEXT("AutoScrollToolTip", "Auto-Scroll"),
+		FSlateIcon(FStateTreeEditorStyle::Get().GetStyleSetName(),"StateTreeEditor.AutoScroll"),
+		EUserInterfaceActionType::ToggleButton);
+	RightToolbar.EndSection();
 
 	// Auto-select session if there is only one available
 	TArray<FStateTreeDebugger::FTraceDescriptor> TraceDescriptors;
 	Debugger->GetLiveTraces(TraceDescriptors);
-	
-	if (TraceDescriptors.Num() == 1)
+
+	// Make sure that the active trace contains statetree specific traces
+	if (TraceDescriptors.Num() == 1 && StateTreeModule.IsTracing())
 	{
 		Debugger->StartSessionAnalysis(TraceDescriptors.Last());
 	}
@@ -478,9 +527,22 @@ void SStateTreeDebuggerView::Construct(const FArguments& InArgs, const UStateTre
 			SNew(SVerticalBox)
 			+ SVerticalBox::Slot()
 			.AutoHeight()
+			.HAlign(HAlign_Fill)
 			.Padding(0.0f, 0.0f, 0.0f, 4.0f)
 			[
-				ToolbarBuilder.MakeWidget()
+				SNew(SHorizontalBox)
+				+ SHorizontalBox::Slot()
+				.HAlign(HAlign_Left)
+				.FillWidth(1.0f)
+				[
+					LeftToolbar.MakeWidget()
+				]
+				+SHorizontalBox::Slot()
+				.HAlign(HAlign_Right)
+				.AutoWidth()
+				[
+					RightToolbar.MakeWidget()
+				]
 			]
 			+ SVerticalBox::Slot().AutoHeight()
 			[
@@ -591,8 +653,8 @@ void SStateTreeDebuggerView::Tick(const FGeometry& AllottedGeometry, const doubl
 {
 	QUICK_SCOPE_CYCLE_COUNTER(STAT_StateTreeDebuggerView_TickView);
 
-	//Stick to most recent data
-	if (Debugger->IsPaused() == false)
+	// Stick to most recent data if required
+	if (bAutoScroll && !Debugger->IsAnalysisSessionPaused())
 	{
 		Debugger->SetScrubTime(Debugger->GetRecordingDuration());
 	}
@@ -657,7 +719,7 @@ void SStateTreeDebuggerView::BindDebuggerToolbarCommands(const TSharedRef<FUICom
 
 bool SStateTreeDebuggerView::CanStepBackToPreviousStateWithEvents() const
 {
-	return Debugger->CanStepBackToPreviousStateWithEvents();
+	return (bAutoScroll == false || Debugger->IsAnalysisSessionPaused()) && Debugger->CanStepBackToPreviousStateWithEvents();
 }
 
 void SStateTreeDebuggerView::StepBackToPreviousStateWithEvents()
@@ -667,7 +729,7 @@ void SStateTreeDebuggerView::StepBackToPreviousStateWithEvents()
 
 bool SStateTreeDebuggerView::CanStepForwardToNextStateWithEvents() const
 {
-	return Debugger->CanStepForwardToNextStateWithEvents();
+	return (bAutoScroll == false || Debugger->IsAnalysisSessionPaused()) && Debugger->CanStepForwardToNextStateWithEvents();
 }
 
 void SStateTreeDebuggerView::StepForwardToNextStateWithEvents()
@@ -677,7 +739,7 @@ void SStateTreeDebuggerView::StepForwardToNextStateWithEvents()
 
 bool SStateTreeDebuggerView::CanStepBackToPreviousStateChange() const
 {
-	return Debugger->CanStepBackToPreviousStateChange();
+	return (bAutoScroll == false || Debugger->IsAnalysisSessionPaused()) && Debugger->CanStepBackToPreviousStateChange();
 }
 
 void SStateTreeDebuggerView::StepBackToPreviousStateChange()
@@ -687,7 +749,7 @@ void SStateTreeDebuggerView::StepBackToPreviousStateChange()
 
 bool SStateTreeDebuggerView::CanStepForwardToNextStateChange() const
 {
-	return Debugger->CanStepForwardToNextStateChange();
+	return (bAutoScroll == false || Debugger->IsAnalysisSessionPaused()) && Debugger->CanStepForwardToNextStateChange();
 }
 
 void SStateTreeDebuggerView::StepForwardToNextStateChange()
@@ -838,7 +900,7 @@ void SStateTreeDebuggerView::HandleEnableStateBreakpoint(EStateTreeBreakpointTyp
 
 void SStateTreeDebuggerView::OnTimeLineScrubPositionChanged(double Time, bool bIsScrubbing)
 {
-	Debugger->Pause();
+	bAutoScroll = false;
 	Debugger->SetScrubTime(Time);
 }
 
@@ -954,6 +1016,23 @@ void SStateTreeDebuggerView::OnBreakpointHit(const FStateTreeInstanceDebugId Ins
 	}
 }
 
+void SStateTreeDebuggerView::OnNewSession()
+{
+	// Clear tracks
+	InstanceOwnerTracks.Reset();
+
+	// Refresh tree viewa
+	if (InstancesTreeView)
+	{
+		InstancesTreeView->Refresh();
+	}
+
+	if (InstanceTimelinesTreeView)
+	{
+		InstanceTimelinesTreeView->Refresh();
+	}
+}
+
 void SStateTreeDebuggerView::OnNewInstance(FStateTreeInstanceDebugId InstanceId)
 {
 	const UE::StateTreeDebugger::FInstanceDescriptor* FoundDescriptor = Debugger->GetInstanceDescriptor(InstanceId);
diff --git a/Engine/Plugins/Runtime/StateTree/Source/StateTreeEditorModule/Private/Debugger/SStateTreeDebuggerView.h b/Engine/Plugins/Runtime/StateTree/Source/StateTreeEditorModule/Private/Debugger/SStateTreeDebuggerView.h
--- a/Engine/Plugins/Runtime/StateTree/Source/StateTreeEditorModule/Private/Debugger/SStateTreeDebuggerView.h
+++ b/Engine/Plugins/Runtime/StateTree/Source/StateTreeEditorModule/Private/Debugger/SStateTreeDebuggerView.h
@@ -64,6 +64,7 @@ private:
 	void OnPIESingleStepped(bool bIsSimulating) const;
 	
 	void OnBreakpointHit(const FStateTreeInstanceDebugId InstanceId, const FStateTreeDebuggerBreakpoint Breakpoint, const TSharedRef<FUICommandList> ActionList) const;
+	void OnNewSession();
 	void OnNewInstance(FStateTreeInstanceDebugId InstanceId);
 	void OnSelectedInstanceCleared();
 
@@ -150,10 +151,13 @@ private:
 	FInstancedStruct SelectedNodeDataStruct;
 
 	/** Object created from the event data when statetree node was holding an object. */
-	TWeakObjectPtr<UObject> SelectedNodeDataObject;
+	TWeakObjectPtr<> SelectedNodeDataObject;
 
 	/** Indicates that a live session was started (record button or auto record in PIE) to generate StateTree traces. */
 	bool bRecording = false;
+
+	/** True if the timelines scrolls automatically. */
+	bool bAutoScroll = true;
 };
 
 #endif // WITH_STATETREE_DEBUGGER
\ No newline at end of file
diff --git a/Engine/Plugins/Runtime/StateTree/Source/StateTreeEditorModule/Private/StateTreeEditorStyle.cpp b/Engine/Plugins/Runtime/StateTree/Source/StateTreeEditorModule/Private/StateTreeEditorStyle.cpp
--- a/Engine/Plugins/Runtime/StateTree/Source/StateTreeEditorModule/Private/StateTreeEditorStyle.cpp
+++ b/Engine/Plugins/Runtime/StateTree/Source/StateTreeEditorModule/Private/StateTreeEditorStyle.cpp
@@ -194,6 +194,7 @@ FStateTreeEditorStyle::FStateTreeEditorStyle()
 		Set("StateTreeEditor.DuplicateStates", new IMAGE_BRUSH_SVG("Starship/Common/Duplicate", CoreStyleConstants::Icon16x16));
 		Set("StateTreeEditor.DeleteStates", new IMAGE_BRUSH_SVG("Starship/Common/Delete", CoreStyleConstants::Icon16x16));
 		Set("StateTreeEditor.RenameState", new IMAGE_BRUSH_SVG("Starship/Common/Rename", CoreStyleConstants::Icon16x16));
+		Set("StateTreeEditor.AutoScroll", new IMAGE_BRUSH_SVG("Starship/Insights/AutoScrollRight_20", CoreStyleConstants::Icon16x16));		
 	}
 
 	{
@@ -218,6 +219,7 @@ FStateTreeEditorStyle::FStateTreeEditorStyle()
 
 		Set("StateTreeEditor.EnableStates", new IMAGE_BRUSH("Icons/Empty_16x", CoreStyleConstants::Icon16x16));
 		Set( "StateTreeEditor.Debugger.Breakpoint.EnabledAndValid", new IMAGE_BRUSH_SVG( "Starship/Blueprints/Breakpoint_Valid", CoreStyleConstants::Icon16x16, FStyleColors::AccentRed));
+		Set("StateTreeEditor.ToggleDebuggerPause", new IMAGE_BRUSH("Icons/generic_pause_16x", CoreStyleConstants::Icon16x16));
 	}
 
 	{
diff --git a/Engine/Plugins/Runtime/StateTree/Source/StateTreeModule/Private/Debugger/StateTreeDebugger.cpp b/Engine/Plugins/Runtime/StateTree/Source/StateTreeModule/Private/Debugger/StateTreeDebugger.cpp
--- a/Engine/Plugins/Runtime/StateTree/Source/StateTreeModule/Private/Debugger/StateTreeDebugger.cpp
+++ b/Engine/Plugins/Runtime/StateTree/Source/StateTreeModule/Private/Debugger/StateTreeDebugger.cpp
@@ -81,7 +81,7 @@ FStateTreeDebugger::FStateTreeDebugger()
 
 FStateTreeDebugger::~FStateTreeDebugger()
 {
-	StopAnalysis();
+	StopSessionAnalysis();
 }
 
 void FStateTreeDebugger::Tick(const float DeltaTime)
@@ -100,27 +100,22 @@ void FStateTreeDebugger::Tick(const float DeltaTime)
 	}
 	
 	UpdateInstances();
-	SyncToCurrentSessionDuration();
+
+	if (bSessionAnalysisPaused == false)
+	{
+		SyncToCurrentSessionDuration();
+	}
 }
 
-void FStateTreeDebugger::StopAnalysis()
+void FStateTreeDebugger::StopSessionAnalysis()
 {
 	if (const TraceServices::IAnalysisSession* Session = GetAnalysisSession())
 	{
 		Session->Stop(true);
 		AnalysisSession.Reset();	
 	}
-}
 
-void FStateTreeDebugger::Pause()
-{
-	bPaused = true;
-}
-
-void FStateTreeDebugger::Unpause()
-{
-	// unpause and let next tick read the traces
-	bPaused = false;	
+	bSessionAnalysisPaused = false;
 }
 
 void FStateTreeDebugger::SyncToCurrentSessionDuration()
@@ -206,21 +201,49 @@ void FStateTreeDebugger::UpdateInstances()
 	}
 }
 
-void FStateTreeDebugger::RequestAnalysisOfNextLiveSession()
+bool FStateTreeDebugger::RequestAnalysisOfNextLiveSession()
 {
-	// Invalidate our current active session
-	ActiveSessionTraceDescriptor = FTraceDescriptor();
+	// Get snapshot of current trace to help identify the next live one
+	LastLiveSessionId = INDEX_NONE;
+	TArray<FTraceDescriptor> TraceDescriptors;
+	GetLiveTraces(TraceDescriptors);
 
-	// Stop current analysis if any
-	StopAnalysis();
+	// 0 is the invalid value used for Trace Id
+	constexpr int32 InvalidTraceId = 0;
+	int32 ActiveTraceId = InvalidTraceId;
 
-	TArray<FTraceDescriptor> Traces;
-	GetLiveTraces(Traces);
-	
-	LastLiveSessionId = Traces.Num() ? Traces.Last().TraceId : INDEX_NONE;
+	// StartTraces returns true if a new connection was created. In this case
+	// we try to start an analysis on that new connection as soon as possible.
+	// Otherwise it might have been able to use an active connection in which
+	// case it was returned in the output parameter.
+	if (StateTreeModule.StartTraces(ActiveTraceId))
+	{
+		// Invalidate our current active session
+		ActiveSessionTraceDescriptor = FTraceDescriptor();
+
+		// Stop current analysis if any
+		StopSessionAnalysis();
 
-	// This won't succeed yet but will schedule our next retry
-	TryStartNewLiveSessionAnalysis(1.0f);
+		LastLiveSessionId = TraceDescriptors.Num() ? TraceDescriptors.Last().TraceId : INDEX_NONE;
+
+		// This won't succeed yet but will schedule our next retry
+		TryStartNewLiveSessionAnalysis(1.0f);
+		return true;
+	}
+
+	// Otherwise we start analysis of the already active trace, if any.
+	if (ActiveTraceId != InvalidTraceId)
+	{
+		if (const FTraceDescriptor* Descriptor = TraceDescriptors.FindByPredicate([ActiveTraceId](const FTraceDescriptor& Descriptor)
+			{
+				return Descriptor.TraceId == ActiveTraceId;
+			}))
+		{
+			return StartSessionAnalysis(*Descriptor);
+		}
+	}
+
+	return false;
 }
 
 bool FStateTreeDebugger::TryStartNewLiveSessionAnalysis(const float RetryPollingDuration)
@@ -248,14 +271,21 @@ bool FStateTreeDebugger::StartSessionAnalysis(const FTraceDescriptor& TraceDescr
 	}
 
 	ActiveSessionTraceDescriptor = FTraceDescriptor();
+
+	// Make sure any active analysis is stopped
+	StopSessionAnalysis();
+
 	UE::Trace::FStoreClient* StoreClient = GetStoreClient();
 	if (StoreClient == nullptr)
 	{
 		return false;
 	}
 
-	// Make sure any active analysis is stopped
-	StopAnalysis();
+	// If new trace descriptor is not valid no need to continue
+	if (TraceDescriptor.IsValid() == false)
+	{
+		return false;
+	}
 
 	RecordingDuration = 0;
 	AnalysisDuration = 0;
@@ -291,13 +321,14 @@ bool FStateTreeDebugger::StartSessionAnalysis(const FTraceDescriptor& TraceDescr
 		{
 			checkf(!AnalysisSession.IsValid(), TEXT("Must make sure that current session was properly stopped before starting a new one otherwise it can cause threading issues"));
 			AnalysisSession = TraceAnalysisService->StartAnalysis(TraceId, *TraceName, MoveTemp(TraceData));
-			
-			SyncToCurrentSessionDuration();
 		}
 
 		if (AnalysisSession.IsValid())
 		{
-			ActiveSessionTraceDescriptor = TraceDescriptor;	
+			ActiveSessionTraceDescriptor = TraceDescriptor;
+			EventCollections.Reset();
+			ScrubState.SetEventCollectionIndex(INDEX_NONE);
+			OnNewSession.ExecuteIfBound();
 		}
 	}
 
@@ -470,7 +501,7 @@ void FStateTreeDebugger::RefreshActiveStates()
 
 bool FStateTreeDebugger::CanStepBackToPreviousStateWithEvents() const
 {
-	return bPaused ? ScrubState.HasPreviousFrame() : false;
+	return ScrubState.HasPreviousFrame();
 }
 
 void FStateTreeDebugger::StepBackToPreviousStateWithEvents()
@@ -483,7 +514,7 @@ void FStateTreeDebugger::StepBackToPreviousStateWithEvents()
 
 bool FStateTreeDebugger::CanStepForwardToNextStateWithEvents() const
 {
-	return bPaused ? ScrubState.HasNextFrame() : false;
+	return ScrubState.HasNextFrame();
 }
 
 void FStateTreeDebugger::StepForwardToNextStateWithEvents()
@@ -496,7 +527,7 @@ void FStateTreeDebugger::StepForwardToNextStateWithEvents()
 
 bool FStateTreeDebugger::CanStepBackToPreviousStateChange() const
 {
-	return bPaused ? ScrubState.HasPreviousActiveStates() : false;
+	return ScrubState.HasPreviousActiveStates();
 }
 
 void FStateTreeDebugger::StepBackToPreviousStateChange()
@@ -509,7 +540,7 @@ void FStateTreeDebugger::StepBackToPreviousStateChange()
 
 bool FStateTreeDebugger::CanStepForwardToNextStateChange() const
 {
-	return bPaused ? ScrubState.HasNextActiveStates() : false;
+	return ScrubState.HasNextActiveStates();
 }
 
 void FStateTreeDebugger::StepForwardToNextStateChange()
@@ -681,6 +712,8 @@ void FStateTreeDebugger::SendNotifications()
 
 		HitBreakpointIndex = INDEX_NONE;
 		HitBreakpointInstanceId.Reset();
+
+		PauseSessionAnalysis();
 	}
 }
 
@@ -702,15 +735,14 @@ void FStateTreeDebugger::ReadTrace(
 	}
 }
 
-void FStateTreeDebugger::EvaluateBreakpoints(const FStateTreeInstanceDebugId InstanceId, const FStateTreeTraceEventVariantType& Event)
+bool FStateTreeDebugger::EvaluateBreakpoints(const FStateTreeInstanceDebugId InstanceId, const FStateTreeTraceEventVariantType& Event)
 {
-	if (bPaused // ignored when scrubbing a paused session
-		|| StateTreeAsset == nullptr // asset is required to properly match state handles
-		|| HitBreakpointIndex != INDEX_NONE // stop on first hit breakpoint
+	if (StateTreeAsset == nullptr // asset is required to properly match state handles
+		|| HitBreakpointIndex != INDEX_NONE // Only consider first hit breakpoint in the frame
 		|| Breakpoints.IsEmpty()
 		|| (SelectedInstanceId.IsValid() && InstanceId != SelectedInstanceId)) // ignore events not for the selected instances		
 	{
-		return;
+		return false;
 	}
 	
 	EStateTreeTraceEventType EventType = EStateTreeTraceEventType::Unset;
@@ -757,6 +789,8 @@ void FStateTreeDebugger::EvaluateBreakpoints(const FStateTreeInstanceDebugId Ins
 			}
 		}
 	}
+
+	return HitBreakpointInstanceId.IsValid();
 }
 
 bool FStateTreeDebugger::ProcessEvent(const FStateTreeInstanceDebugId InstanceId, const TraceServices::FFrame& Frame, const FStateTreeTraceEventVariantType& Event)
diff --git a/Engine/Plugins/Runtime/StateTree/Source/StateTreeModule/Private/Debugger/StateTreeTraceTypes.cpp b/Engine/Plugins/Runtime/StateTree/Source/StateTreeModule/Private/Debugger/StateTreeTraceTypes.cpp
--- a/Engine/Plugins/Runtime/StateTree/Source/StateTreeModule/Private/Debugger/StateTreeTraceTypes.cpp
+++ b/Engine/Plugins/Runtime/StateTree/Source/StateTreeModule/Private/Debugger/StateTreeTraceTypes.cpp
@@ -59,13 +59,13 @@ FString FStateTreeTraceTransitionEvent::ToString(const UStateTree& StateTree) co
 //----------------------------------------------------------------------//
 FString FStateTreeTraceNodeEvent::ToString(const UStateTree& StateTree) const
 {
-	const FConstStructView NodeView = StateTree.GetNode(Index.Get());
+	const FConstStructView NodeView = Index.IsValid() ? StateTree.GetNode(Index.Get()) : FConstStructView();
 	const FStateTreeNodeBase* Node = NodeView.GetPtr<const FStateTreeNodeBase>();
 
 	return FString::Printf(TEXT("%s '%s (%s)'"),
 			*UEnum::GetDisplayValueAsText(EventType).ToString(),
 			Node != nullptr ? *Node->Name.ToString() : *LexToString(Index.Get()),
-			*NodeView.GetScriptStruct()->GetName());
+			NodeView.IsValid() ? *NodeView.GetScriptStruct()->GetName() : TEXT("Invalid Node"));
 }
 
 
diff --git a/Engine/Plugins/Runtime/StateTree/Source/StateTreeModule/Private/StateTreeModule.cpp b/Engine/Plugins/Runtime/StateTree/Source/StateTreeModule/Private/StateTreeModule.cpp
--- a/Engine/Plugins/Runtime/StateTree/Source/StateTreeModule/Private/StateTreeModule.cpp
+++ b/Engine/Plugins/Runtime/StateTree/Source/StateTreeModule/Private/StateTreeModule.cpp
@@ -26,7 +26,8 @@ class FStateTreeModule : public IStateTreeModule
 	virtual void StartupModule() override;
 	virtual void ShutdownModule() override;
 
-	virtual void StartTraces() override;
+	virtual bool StartTraces(int32& OutTraceId) override;
+	virtual bool IsTracing() const override;
 	virtual void StopTraces() override;
 
 #if WITH_STATETREE_DEBUGGER
@@ -59,19 +60,20 @@ class FStateTreeModule : public IStateTreeModule
 		TEXT("statetree.startdebuggertraces"),
 		TEXT("Turns on StateTree debugger traces if not already active."),
 		FConsoleCommandDelegate::CreateLambda([]
-		{
-			FStateTreeModule& StateTreeModule = FModuleManager::GetModuleChecked<FStateTreeModule>("StateTreeModule");
-			StateTreeModule.StartTraces();
-		}));
+			{
+				IStateTreeModule& StateTreeModule = FModuleManager::GetModuleChecked<IStateTreeModule>("StateTreeModule");
+				int32 TraceId = 0;
+				StateTreeModule.StartTraces(TraceId);
+			}));
 
 	FAutoConsoleCommand StopDebuggerTracesCommand = FAutoConsoleCommand(
 		TEXT("statetree.stopdebuggertraces"),
 		TEXT("Turns off StateTree debugger traces if active."),
 		FConsoleCommandDelegate::CreateLambda([]
-		{
-			FStateTreeModule& StateTreeModule = FModuleManager::GetModuleChecked<FStateTreeModule>("StateTreeModule");
-			StateTreeModule.StopTraces();
-		}));
+			{
+				IStateTreeModule& StateTreeModule = FModuleManager::GetModuleChecked<IStateTreeModule>("StateTreeModule");
+				StateTreeModule.StopTraces();
+			}));
 #endif // WITH_STATETREE_DEBUGGER
 };
 
@@ -92,7 +94,8 @@ void FStateTreeModule::StartupModule()
 	// to start recording either on user action or on PIE session start.
 	if (UStateTreeSettings::Get().bAutoStartDebuggerTracesOnNonEditorTargets)
 	{
-		StartTraces();
+		int32 TraceId = INDEX_NONE;
+		StartTraces(TraceId);
 	}
 #endif // !WITH_EDITOR
 
@@ -115,15 +118,24 @@ void FStateTreeModule::ShutdownModule()
 #endif // WITH_STATETREE_DEBUGGER
 }
 
-void FStateTreeModule::StartTraces()
+bool FStateTreeModule::StartTraces(int32& OutTraceId)
 {
+	OutTraceId = INDEX_NONE;
 #if WITH_STATETREE_DEBUGGER
 	if (IsRunningCommandlet() || bIsTracing)
 	{
-		return;
+		return false;
 	}
 
-	const bool bAlreadyConnected = FTraceAuxiliary::IsConnected();
+	FGuid SessionGuid, TraceGuid;
+	const bool bAlreadyConnected = FTraceAuxiliary::IsConnected(SessionGuid, TraceGuid);
+
+	if (UE::Trace::FStoreClient* Client = GetStoreClient())
+	{
+		const UE::Trace::FStoreClient::FSessionInfo* SessionInfo = StoreClient->GetSessionInfoByGuid(TraceGuid);
+		// Note that 0 is returned instead of INDEX_NONE to match default invalid value for GetTraceId 
+		OutTraceId = SessionInfo != nullptr ? SessionInfo->GetTraceId(): 0;
+	}
 
 	// If trace is already connected let's keep track of enabled channels to restore them when we stop recording
 	if (bAlreadyConnected)
@@ -154,14 +166,27 @@ void FStateTreeModule::StartTraces()
 	UE::Trace::ToggleChannel(TEXT("StateTreeDebugChannel"), true);
 	UE::Trace::ToggleChannel(TEXT("FrameChannel"), true);
 
+	bool bAreTracesStarted = false;
 	if (bAlreadyConnected == false)
 	{
 		FTraceAuxiliary::FOptions Options;
 		Options.bExcludeTail = true;
-		FTraceAuxiliary::Start(FTraceAuxiliary::EConnectionType::Network, TEXT("localhost"), TEXT(""), &Options, LogStateTree);
+		bAreTracesStarted = FTraceAuxiliary::Start(FTraceAuxiliary::EConnectionType::Network, TEXT("localhost"), TEXT(""), &Options, LogStateTree);
 	}
 
 	bIsTracing = true;
+	return bAreTracesStarted;
+#else
+	return false;
+#endif // WITH_STATETREE_DEBUGGER
+}
+
+bool FStateTreeModule::IsTracing() const
+{
+#if WITH_STATETREE_DEBUGGER
+	return bIsTracing;
+#else
+	return false;
 #endif // WITH_STATETREE_DEBUGGER
 }
 
diff --git a/Engine/Plugins/Runtime/StateTree/Source/StateTreeModule/Public/Debugger/StateTreeDebugger.h b/Engine/Plugins/Runtime/StateTree/Source/StateTreeModule/Public/Debugger/StateTreeDebugger.h
--- a/Engine/Plugins/Runtime/StateTree/Source/StateTreeModule/Public/Debugger/StateTreeDebugger.h
+++ b/Engine/Plugins/Runtime/StateTree/Source/StateTreeModule/Public/Debugger/StateTreeDebugger.h
@@ -28,6 +28,7 @@ DECLARE_DELEGATE_OneParam(FOnStateTreeDebuggerScrubStateChanged, const UE::State
 DECLARE_DELEGATE_TwoParams(FOnStateTreeDebuggerBreakpointHit, FStateTreeInstanceDebugId InstanceId, const FStateTreeDebuggerBreakpoint Breakpoint);
 DECLARE_DELEGATE_OneParam(FOnStateTreeDebuggerActiveStatesChanges, TConstArrayView<FStateTreeStateHandle> ActiveStates);
 DECLARE_DELEGATE_OneParam(FOnStateTreeDebuggerNewInstance, FStateTreeInstanceDebugId InstanceId);
+DECLARE_DELEGATE(FOnStateTreeDebuggerNewSession);
 DECLARE_DELEGATE(FOnStateTreeDebuggerDebuggedInstanceSet);
 
 struct STATETREEMODULE_API FStateTreeDebugger : FTickableGameObject
@@ -53,15 +54,6 @@ struct STATETREEMODULE_API FStateTreeDebugger : FTickableGameObject
 	const UStateTree* GetAsset() const { return StateTreeAsset.Get(); }
 	void SetAsset(const UStateTree* InStateTreeAsset) { StateTreeAsset = InStateTreeAsset; }
 	
-	/** Stops reading traces every frame to preserve current state. */
-	void Pause();
-
-	/** Indicates that the debugger was explicitly paused and is no longer fetching new events from the analysis session. */
-	bool IsPaused() const { return bPaused; }
-
-	/** Resumes reading traces every frame. */
-	void Unpause();
-
 	/** Forces a single refresh to latest state. Useful when simulation is paused. */
 	void SyncToCurrentSessionDuration();
 	
@@ -77,9 +69,6 @@ struct STATETREEMODULE_API FStateTreeDebugger : FTickableGameObject
 	bool CanStepForwardToNextStateChange() const;
 	void StepForwardToNextStateChange();
 
-	bool IsAnalysisSessionActive() const { return GetAnalysisSession() != nullptr; }
-	const TraceServices::IAnalysisSession* GetAnalysisSession() const;
-
 	bool IsActiveInstance(double Time, FStateTreeInstanceDebugId InstanceId) const;
 	FText GetInstanceName(FStateTreeInstanceDebugId InstanceId) const;
 	FText GetInstanceDescription(FStateTreeInstanceDebugId InstanceId) const;
@@ -125,14 +114,21 @@ struct STATETREEMODULE_API FStateTreeDebugger : FTickableGameObject
 	/**
 	 * Queue a request to auto start an analysis session on the next available live trace.
 	 * This will replace the current analysis session if any.
+	 * @return True if connection was successfully requested or was able to use active trace, false otherwise.
 	 */
-	void RequestAnalysisOfNextLiveSession();
+	bool RequestAnalysisOfNextLiveSession();
 
+	bool IsAnalysisSessionActive() const { return GetAnalysisSession() != nullptr; }
+	bool IsAnalysisSessionPaused() const { return bSessionAnalysisPaused; }
+	const TraceServices::IAnalysisSession* GetAnalysisSession() const;
 	bool StartSessionAnalysis(const FTraceDescriptor& TraceDescriptor);
-	void StopAnalysis();
+	void PauseSessionAnalysis()  { bSessionAnalysisPaused = true; }
+	void ResumeSessionAnalysis() { bSessionAnalysisPaused = false; }
+	void StopSessionAnalysis();
 	FTraceDescriptor GetSelectedTraceDescriptor() const { return ActiveSessionTraceDescriptor; }
 	FText GetSelectedTraceDescription() const;
 
+	FOnStateTreeDebuggerNewSession OnNewSession;
 	FOnStateTreeDebuggerNewInstance OnNewInstance;
 	FOnStateTreeDebuggerDebuggedInstanceSet OnSelectedInstanceCleared;
 	FOnStateTreeDebuggerScrubStateChanged OnScrubStateChanged;
@@ -155,7 +151,13 @@ private:
 		const TraceServices::FFrame& Frame
 		);
 
-	void EvaluateBreakpoints(FStateTreeInstanceDebugId InstanceId, const FStateTreeTraceEventVariantType& Event);
+	/**
+	 * Tests event for given instance id against breakpoints.
+	 * @param InstanceId Id of the statetree instance that produces the event.
+	 * @param Event The event received from the instance.
+	 * @return True if a breakpoint has been it, false otherwise.
+	 */
+	bool EvaluateBreakpoints(FStateTreeInstanceDebugId InstanceId, const FStateTreeTraceEventVariantType& Event);
 
 	void SendNotifications();
 
@@ -243,8 +245,11 @@ private:
 	/** List of new instances discovered by processing event in the analysis session. */
 	TArray<FStateTreeInstanceDebugId> NewInstances;
 
-	/** Indicates that the debugger was explicitly paused and will wait before fetching new events from the analysis session provider. */
-	bool bPaused = false;
+	/**
+	 * Indicates that the debugger no longer process new events from the analysis session until it gets resumed.
+	 * This can be an external explicit request or after hitting a breakpoint.
+	 */
+	bool bSessionAnalysisPaused = false;
 };
 
 #endif // WITH_STATETREE_DEBUGGER
diff --git a/Engine/Plugins/Runtime/StateTree/Source/StateTreeModule/Public/StateTreeModule.h b/Engine/Plugins/Runtime/StateTree/Source/StateTreeModule/Public/StateTreeModule.h
--- a/Engine/Plugins/Runtime/StateTree/Source/StateTreeModule/Public/StateTreeModule.h
+++ b/Engine/Plugins/Runtime/StateTree/Source/StateTreeModule/Public/StateTreeModule.h
@@ -39,10 +39,14 @@ public:
 	}
 
 	/**
-	 * Start tracing and enables StateTree debugging related channels (frame + statetree).
+	 * Creates new tracing connection if necessary and enables StateTree debugging related channels (frame + statetree).
 	 * If traces are already active we keep track of all channels previously activated to restore them on stop.
+	 * @param OutTraceId In case a connection is already active this indicates its id, 0 otherwise.
+	 * @return True if a new trace connection was created, false otherwise (already active or not created)
+	 * @note Even if it returns false this can still have enabled StateTree related channels, OutTraceId will indicate
+	 * which trace to use and StopTraces should be called to restore to previous setup.
 	 */
-	virtual void StartTraces() = 0;
+	virtual bool StartTraces(int32& OutTraceId) = 0;
 
 	/**
 	 * Stops the trace service if it was not already connected when StartTraces was called.
@@ -50,6 +54,12 @@ public:
 	 */
 	virtual void StopTraces() = 0;
 
+	/**
+	 * Indicates if the statetree specific traces are active (explicitly started by StartTraces).
+	 * @return True is StartTraces was called, false otherwise.
+	 */
+	virtual bool IsTracing() const = 0;
+
 #if WITH_STATETREE_DEBUGGER
 	/**
 	 * Gets the store client.
