StateTree: Show mismatching property types on property binding
- Improved formatting of State Tree error messages
- Fixed GetDataViewByID to return true on null values but valid types (e.g. context data)
- Added check if property types match for property binding widget
- Cache current property binding data for binding widget

diff --git a/Engine/Plugins/Runtime/StateTree/Source/StateTreeEditorModule/Private/Customizations/StateTreeBindingExtension.cpp b/Engine/Plugins/Runtime/StateTree/Source/StateTreeEditorModule/Private/Customizations/StateTreeBindingExtension.cpp
--- a/Engine/Plugins/Runtime/StateTree/Source/StateTreeEditorModule/Private/Customizations/StateTreeBindingExtension.cpp
+++ b/Engine/Plugins/Runtime/StateTree/Source/StateTreeEditorModule/Private/Customizations/StateTreeBindingExtension.cpp
@@ -250,6 +250,224 @@ FText GetPropertyTypeText(const FProperty* Property)
 
 FOnStateTreePropertyBindingChanged STATETREEEDITORMODULE_API OnStateTreePropertyBindingChanged;
 
+struct FCachedBindingData : public TSharedFromThis<FCachedBindingData>
+{
+	FCachedBindingData(UObject* InOwnerObject, const FStateTreePropertyPath& InTargetPath, const TSharedPtr<const IPropertyHandle>& InPropertyHandle)
+		: WeakOwnerObject(InOwnerObject)
+		, TargetPath(InTargetPath)
+		, PropertyHandle(InPropertyHandle)
+	{
+	}
+	
+	void UpdateData()
+	{
+		static FName PropertyIcon(TEXT("Kismet.Tabs.Variables"));
+
+		Text = FText::GetEmpty();
+		TooltipText = FText::GetEmpty();
+		Color = FLinearColor::White;
+		Image = nullptr;
+
+		if (!PropertyHandle.IsValid())
+		{
+			return;
+		}
+
+		const FProperty* Property = PropertyHandle->GetProperty();
+		if (!Property)
+		{
+			return;
+		}
+
+		UObject* OwnerObject = WeakOwnerObject.Get();
+		if (!OwnerObject)
+		{
+			return;
+		}
+		
+		IStateTreeEditorPropertyBindingsOwner* BindingOwner = Cast<IStateTreeEditorPropertyBindingsOwner>(OwnerObject);
+		if (!BindingOwner)
+		{
+			return;
+		}
+
+		FStateTreeEditorPropertyBindings* EditorBindings = BindingOwner->GetPropertyEditorBindings();
+		if (!EditorBindings)
+		{
+			return;
+		}
+
+		TArray<FStateTreeBindableStructDesc> AccessibleStructs;
+		BindingOwner->GetAccessibleStructs(TargetPath.GetStructID(), AccessibleStructs);
+		
+		const UEdGraphSchema_K2* Schema = GetDefault<UEdGraphSchema_K2>();
+		check(Schema);
+
+		FEdGraphPinType PinType;
+		Schema->ConvertPropertyToPinType(Property, PinType);
+
+		if (const FStateTreePropertyPath* SourcePath = EditorBindings->GetPropertyBindingSource(TargetPath))
+		{
+			const FStateTreeBindableStructDesc* SourceDesc = UE::StateTree::PropertyBinding::FindStruct(AccessibleStructs, SourcePath->GetStructID());
+			if (SourceDesc)
+			{
+				bool bIsValidBinding = false;
+
+				// Check that the binding is valid.
+				FStateTreeDataView SourceDataView;
+				FStateTreeDataView TargetDataView;
+				const FProperty* SourceLeafProperty = nullptr;
+				const UStruct* SourceStruct = nullptr;
+				if (BindingOwner->GetDataViewByID(SourcePath->GetStructID(), SourceDataView)
+					&& BindingOwner->GetDataViewByID(TargetPath.GetStructID(), TargetDataView))
+				{
+					TArray<FStateTreePropertyPathIndirection> SourceIndirections;
+					TArray<FStateTreePropertyPathIndirection> TargetIndirections;
+					
+					if (SourcePath->ResolveIndirectionsWithValue(SourceDataView, SourceIndirections)
+						&& TargetPath.ResolveIndirectionsWithValue(TargetDataView, TargetIndirections))
+					{
+						FStateTreePropertyCopy DummyCopy;
+						FStateTreePropertyPathIndirection LastSourceIndirection = !SourceIndirections.IsEmpty() ? SourceIndirections.Last() : FStateTreePropertyPathIndirection(SourceDataView.GetStruct());
+						FStateTreePropertyPathIndirection LastTargetIndirection = !TargetIndirections.IsEmpty() ? TargetIndirections.Last() : FStateTreePropertyPathIndirection(TargetDataView.GetStruct());
+						if (FStateTreePropertyBindings::ResolveCopyType(LastSourceIndirection, LastTargetIndirection, DummyCopy))
+						{
+							bIsValidBinding = true;
+						}
+					}
+
+					if (!SourcePath->IsPathEmpty())
+					{
+						if (SourceIndirections.Num() > 0)
+						{
+							SourceLeafProperty = SourceIndirections.Last().GetProperty();
+						}
+					}
+					else
+					{
+						SourceStruct = SourceDataView.GetStruct();
+					}
+				}
+
+				FString SourcePropertyName;
+				SourcePropertyName += SourceDesc->Name.ToString();
+				if (!SourcePath->IsPathEmpty())
+				{
+					SourcePropertyName += TEXT(" ") + SourcePath->ToString();
+				}
+
+				if (bIsValidBinding)
+				{
+					Text = FText::FromString(SourcePropertyName);
+
+					if (SourcePath->IsPathEmpty())
+					{
+						TooltipText = FText::Format(LOCTEXT("ExistingBindingTooltip", "Property is bound to {0}."), FText::FromString(SourceDesc->ToString()));
+					}
+					else
+					{
+						TooltipText = FText::Format(LOCTEXT("ExistingBindingWithPropertyTooltip", "Property is bound to {0} property {1}."), FText::FromString(SourceDesc->ToString()), FText::FromString(SourcePath->ToString()));
+					}
+					
+					Image = FAppStyle::GetBrush(PropertyIcon);
+					Color = Schema->GetPinTypeColor(PinType);
+				}
+				else
+				{
+					FText SourceType;
+					if (SourceLeafProperty)
+					{
+						SourceType = UE::StateTree::PropertyBinding::GetPropertyTypeText(SourceLeafProperty);
+					}
+					else if (SourceStruct)
+					{
+						SourceType = SourceStruct->GetDisplayNameText();
+					}
+					FText TargetType = UE::StateTree::PropertyBinding::GetPropertyTypeText(Property);
+					
+					Text = FText::FromString(SourcePropertyName);
+
+					if (SourcePath->IsPathEmpty())
+					{
+						TooltipText = FText::Format(LOCTEXT("MismatchingBindingTooltip", "Property is bound to {0}, but binding source type '{1}' does not match property type '{2}'."),
+							FText::FromString(SourceDesc->ToString()), SourceType, TargetType);
+					}
+					else
+					{
+						TooltipText = FText::Format(LOCTEXT("MismatchingBindingTooltipWithProperty", "Property is bound to {0} property {1}, but binding source type '{2}' does not match property type '{3}'."),
+							FText::FromString(SourceDesc->ToString()), FText::FromString(SourcePath->ToString()), SourceType, TargetType);
+					}
+
+					Image = FCoreStyle::Get().GetBrush("Icons.ErrorWithColor");
+					Color = FLinearColor::White;
+				}
+			}
+			else
+			{
+				// Missing source
+				Text = FText::Format(LOCTEXT("MissingSource", "???.{0}"), FText::FromString(SourcePath->ToString()));
+				TooltipText = FText::Format(LOCTEXT("MissingBindingTooltip", "Missing binding source for property path '{0}'."), FText::FromString(SourcePath->ToString()));
+				Image = FCoreStyle::Get().GetBrush("Icons.ErrorWithColor");
+				Color = FLinearColor::White;
+			}
+		}
+		else
+		{
+			// No bindings
+			Text = FText::GetEmpty();
+			TooltipText = FText::Format(LOCTEXT("BindTooltip", "Bind {0} to value from another property."), UE::StateTree::PropertyBinding::GetPropertyTypeText(Property));
+			Image = FAppStyle::GetBrush(PropertyIcon);
+			Color = Schema->GetPinTypeColor(PinType);
+		}
+
+		bIsDataCached = true;
+	}
+
+	FText GetText()
+	{
+		ConditionallyUpdateData();
+		return Text;
+	}
+	
+	FText GetTooltipText()
+	{
+		ConditionallyUpdateData();
+		return TooltipText;
+	}
+	
+	FLinearColor GetColor()
+	{
+		ConditionallyUpdateData();
+		return Color;
+	}
+	
+	const FSlateBrush* GetImage()
+	{
+		ConditionallyUpdateData();
+		return Image;
+	}
+	
+private:
+
+	void ConditionallyUpdateData()
+	{
+		if (!bIsDataCached)
+		{
+			UpdateData();
+		}
+	}
+	
+	TWeakObjectPtr<UObject> WeakOwnerObject = nullptr;
+	FStateTreePropertyPath TargetPath;
+	TSharedPtr<const IPropertyHandle> PropertyHandle;
+
+	FText Text;
+	FText TooltipText;
+	FLinearColor Color = FLinearColor::White;
+	const FSlateBrush* Image = nullptr;
+	bool bIsDataCached = false;
+};
+
 } // UE::StateTree::PropertyBinding
 
 bool FStateTreeBindingExtension::IsPropertyExtendable(const UClass* InObjectClass, const IPropertyHandle& PropertyHandle) const
@@ -362,6 +580,8 @@ void FStateTreeBindingExtension::ExtendWidgetRow(FDetailWidgetRow& InWidgetRow,
 		FString BaseStructName;
 		StructRefBaseStruct = UE::StateTree::Compiler::GetBaseStructFromMetaData(Property, BaseStructName);
 	}
+
+	TSharedPtr<UE::StateTree::PropertyBinding::FCachedBindingData> CachedBindingData = MakeShared<UE::StateTree::PropertyBinding::FCachedBindingData>(OwnerObject, TargetPath, InPropertyHandle);
 	
 	FPropertyBindingWidgetArgs Args;
 	Args.Property = InPropertyHandle->GetProperty();
@@ -471,7 +691,7 @@ void FStateTreeBindingExtension::ExtendWidgetRow(FDetailWidgetRow& InWidgetRow,
 			return true;
 		});
 
-	Args.OnAddBinding = FOnAddBinding::CreateLambda([EditorBindings, OwnerObject, BindingOwner, TargetPath, AccessibleStructs](FName InPropertyName, const TArray<FBindingChainElement>& InBindingChain)
+	Args.OnAddBinding = FOnAddBinding::CreateLambda([EditorBindings, OwnerObject, BindingOwner, TargetPath, AccessibleStructs, CachedBindingData](FName InPropertyName, const TArray<FBindingChainElement>& InBindingChain)
 		{
 			IPropertyAccessEditor& PropertyAccessEditor = IModularFeatures::Get().GetModularFeature<IPropertyAccessEditor>("PropertyAccessEditor");
 			if (EditorBindings && OwnerObject)
@@ -502,9 +722,11 @@ void FStateTreeBindingExtension::ExtendWidgetRow(FDetailWidgetRow& InWidgetRow,
 					UE::StateTree::PropertyBinding::OnStateTreePropertyBindingChanged.Broadcast(SourcePath, TargetPath);
 				}
 			}
+
+			CachedBindingData->UpdateData();
 		});
 
-	Args.OnRemoveBinding = FOnRemoveBinding::CreateLambda([EditorBindings, OwnerObject, TargetPath](FName InPropertyName)
+	Args.OnRemoveBinding = FOnRemoveBinding::CreateLambda([EditorBindings, OwnerObject, TargetPath, CachedBindingData](FName InPropertyName)
 		{
 			if (EditorBindings && OwnerObject)
 			{
@@ -514,6 +736,8 @@ void FStateTreeBindingExtension::ExtendWidgetRow(FDetailWidgetRow& InWidgetRow,
 				const FStateTreePropertyPath SourcePath; // Null path
 				UE::StateTree::PropertyBinding::OnStateTreePropertyBindingChanged.Broadcast(SourcePath, TargetPath);
 			}
+
+			CachedBindingData->UpdateData();
 		});
 
 	Args.OnCanRemoveBinding = FOnCanRemoveBinding::CreateLambda([EditorBindings, TargetPath](FName InPropertyName)
@@ -521,114 +745,24 @@ void FStateTreeBindingExtension::ExtendWidgetRow(FDetailWidgetRow& InWidgetRow,
 			return EditorBindings && EditorBindings->HasPropertyBinding(TargetPath);
 		});
 
-	Args.CurrentBindingText = MakeAttributeLambda([EditorBindings, TargetPath, AccessibleStructs]()
+	Args.CurrentBindingText = MakeAttributeLambda([CachedBindingData]()
 		{
-			FText CurrentValue = FText::GetEmpty();
-
-			if (EditorBindings)
-			{
-				if (const FStateTreePropertyPath* SourcePath = EditorBindings->GetPropertyBindingSource(TargetPath))
-				{
-					const FStateTreeBindableStructDesc* SourceDesc = UE::StateTree::PropertyBinding::FindStruct(AccessibleStructs, SourcePath->GetStructID());
-					if (SourceDesc)
-					{
-						FString PropertyName = SourceDesc->Name.ToString();
-						if (!SourcePath->IsPathEmpty())
-						{
-							PropertyName += TEXT(".") + SourcePath->ToString();
-						}
-						CurrentValue = FText::FromString(PropertyName);
-					}
-					else
-					{
-						CurrentValue = FText::Format(LOCTEXT("MissingSource", "???.{0}"), FText::FromString(SourcePath->ToString()));
-					}
-				}
-			}
-
-			return CurrentValue;
+			return CachedBindingData->GetText();
 		});
 
-	Args.CurrentBindingToolTipText = MakeAttributeLambda([EditorBindings, TargetPath, AccessibleStructs, Property]()
+	Args.CurrentBindingToolTipText = MakeAttributeLambda([CachedBindingData]()
 		{
-			FText CurrentValue = FText::GetEmpty();
-
-			if (EditorBindings)
-			{
-				if (const FStateTreePropertyPath* SourcePath = EditorBindings->GetPropertyBindingSource(TargetPath))
-				{
-					const FStateTreeBindableStructDesc* SourceDesc = UE::StateTree::PropertyBinding::FindStruct(AccessibleStructs, SourcePath->GetStructID());
-					if (SourceDesc)
-					{
-						FString PropertyName = SourceDesc->Name.ToString();
-						if (!SourcePath->IsPathEmpty())
-						{
-							PropertyName += TEXT(".") + SourcePath->ToString();
-						}
-						CurrentValue = FText::Format(LOCTEXT("ExistingBindingTooltip", "Property is bound to {0}."), FText::FromString(PropertyName));
-					}
-					else
-					{
-						CurrentValue = FText::Format(LOCTEXT("MissingBindingTooltip", "Missing binding source for property path '{0}'."), FText::FromString(SourcePath->ToString()));
-					}
-				}
-				else if (Property != nullptr)
-				{
-					CurrentValue = FText::Format(LOCTEXT("BindTooltip", "Bind value to {0} property."), UE::StateTree::PropertyBinding::GetPropertyTypeText(Property));
-				}
-			}
-
-			return CurrentValue;
+			return CachedBindingData->GetTooltipText();
 		});
 
-	Args.CurrentBindingImage = MakeAttributeLambda([EditorBindings, TargetPath, AccessibleStructs]() -> const FSlateBrush*
+	Args.CurrentBindingImage = MakeAttributeLambda([CachedBindingData]() -> const FSlateBrush*
 		{
-			static FName PropertyIcon(TEXT("Kismet.Tabs.Variables"));
-
-			const FSlateBrush* Result = FAppStyle::GetBrush(PropertyIcon);
-		
-			if (EditorBindings)
-			{
-				if (const FStateTreePropertyPath* SourcePath = EditorBindings->GetPropertyBindingSource(TargetPath))
-				{
-					const FStateTreeBindableStructDesc* SourceDesc = UE::StateTree::PropertyBinding::FindStruct(AccessibleStructs, SourcePath->GetStructID());
-					if (!SourceDesc)
-					{
-						// Invalid binding
-						Result = FCoreStyle::Get().GetBrush("Icons.ErrorWithColor");
-					}
-				}
-			}
-
-			return Result;
+			return CachedBindingData->GetImage();
 		});
 
-	Args.CurrentBindingColor = MakeAttributeLambda([InPropertyHandle, EditorBindings, TargetPath, AccessibleStructs]() -> FLinearColor
+	Args.CurrentBindingColor = MakeAttributeLambda([CachedBindingData]() -> FLinearColor
 		{
-			const UEdGraphSchema_K2* Schema = GetDefault<UEdGraphSchema_K2>();
-
-			FEdGraphPinType PinType;
-			Schema->ConvertPropertyToPinType(InPropertyHandle->GetProperty(), PinType);
-		
-			FLinearColor BindingColor = Schema->GetPinTypeColor(PinType);
-
-			// TODO: Handle coloring of type promotion
-
-			if (EditorBindings)
-			{
-				if (const FStateTreePropertyPath* SourcePath = EditorBindings->GetPropertyBindingSource(TargetPath))
-				{
-					const FStateTreeBindableStructDesc* SourceDesc = UE::StateTree::PropertyBinding::FindStruct(AccessibleStructs, SourcePath->GetStructID());
-					if (!SourceDesc)
-					{
-						// Invalid binding
-						BindingColor = FLinearColor::White;
-					}
-				}
-			}
-
-
-			return BindingColor;
+			return CachedBindingData->GetColor();
 		});
 
 	if (BindingOwner)
diff --git a/Engine/Plugins/Runtime/StateTree/Source/StateTreeEditorModule/Private/StateTreeCompiler.cpp b/Engine/Plugins/Runtime/StateTree/Source/StateTreeEditorModule/Private/StateTreeCompiler.cpp
--- a/Engine/Plugins/Runtime/StateTree/Source/StateTreeEditorModule/Private/StateTreeCompiler.cpp
+++ b/Engine/Plugins/Runtime/StateTree/Source/StateTreeEditorModule/Private/StateTreeCompiler.cpp
@@ -1427,7 +1427,7 @@ bool FStateTreeCompiler::ValidateStructRef(const FStateTreeBindableStructDesc& S
 	if (!TargetPath.ResolveIndirections(TargetStruct.Struct, TargetIndirection, &ResolveError))
 	{
 		// This will later be reported by the bindings compiler.
-		Log.Reportf(EMessageSeverity::Error, TargetStruct, TEXT("Failed to resolve binding path in %s: %s"), *TargetStruct.Name.ToString(), *ResolveError);
+		Log.Reportf(EMessageSeverity::Error, TargetStruct, TEXT("Failed to resolve binding path in %s: %s"), *TargetStruct.ToString(), *ResolveError);
 		return false;
 	}
 	const FProperty* TargetLeafProperty = TargetIndirection.Num() > 0 ? TargetIndirection.Last().GetProperty() : nullptr;
@@ -1444,8 +1444,8 @@ bool FStateTreeCompiler::ValidateStructRef(const FStateTreeBindableStructDesc& S
 	if (TargetBaseStruct == nullptr)
 	{
 		Log.Reportf(EMessageSeverity::Error, TargetStruct,
-				TEXT("Could not find base struct '%s' for target '%s:%s'."),
-				*TargetBaseStructName, *TargetStruct.Name.ToString(), *TargetPath.ToString());
+				TEXT("Could not find base struct type '%s' for target %s'."),
+				*TargetBaseStructName, *UE::StateTree::GetDescAndPathAsString(TargetStruct, TargetPath));
 		return false;
 	}
 
@@ -1453,7 +1453,7 @@ bool FStateTreeCompiler::ValidateStructRef(const FStateTreeBindableStructDesc& S
 	if (!SourcePath.ResolveIndirections(SourceStruct.Struct, SourceIndirection, &ResolveError))
 	{
 		// This will later be reported by the bindings compiler.
-		Log.Reportf(EMessageSeverity::Error, SourceStruct, TEXT("Failed to resolve binding path in %s: %s"), *SourceStruct.Name.ToString(), *ResolveError);
+		Log.Reportf(EMessageSeverity::Error, SourceStruct, TEXT("Failed to resolve binding path in %s: %s"), *SourceStruct.ToString(), *ResolveError);
 		return false;
 	}
 	const FProperty* SourceLeafProperty = SourceIndirection.Num() > 0 ? SourceIndirection.Last().GetProperty() : nullptr;
@@ -1473,17 +1473,17 @@ bool FStateTreeCompiler::ValidateStructRef(const FStateTreeBindableStructDesc& S
 		if (SourceBaseStruct == nullptr)
 		{
 			Log.Reportf(EMessageSeverity::Error, TargetStruct,
-					TEXT("Could not find base struct '%s' for bidning source '%s:%s'."),
-					*SourceBaseStructName, *SourceStruct.Name.ToString(), *SourcePath.ToString());
+					TEXT("Could not find base struct '%s' for binding source %s."),
+					*SourceBaseStructName, *UE::StateTree::GetDescAndPathAsString(SourceStruct, SourcePath));
 			return false;
 		}
 
 		if (SourceBaseStruct->IsChildOf(TargetBaseStruct) == false)
 		{
 			Log.Reportf(EMessageSeverity::Error, TargetStruct,
-						TEXT("Type mismatch between source '%s:%s' and target '%s:%s' types, '%s' is not child of '%s'."),
-						*SourceStruct.Name.ToString(), *SourcePath.ToString(),
-						*TargetStruct.Name.ToString(), *TargetPath.ToString(),
+						TEXT("Type mismatch between source %s and target %s types, '%s' is not child of '%s'."),
+						*UE::StateTree::GetDescAndPathAsString(SourceStruct, SourcePath),
+						*UE::StateTree::GetDescAndPathAsString(TargetStruct, TargetPath),
 						*GetNameSafe(SourceBaseStruct), *GetNameSafe(TargetBaseStruct));
 			return false;
 		}
@@ -1493,9 +1493,9 @@ bool FStateTreeCompiler::ValidateStructRef(const FStateTreeBindableStructDesc& S
 		if (!SourceStructProperty->Struct || SourceStructProperty->Struct->IsChildOf(TargetBaseStruct) == false)
 		{
 			Log.Reportf(EMessageSeverity::Error, TargetStruct,
-						TEXT("Type mismatch between source '%s:%s' and target '%s:%s' types, '%s' is not child of '%s'."),
-						*SourceStruct.Name.ToString(), *SourcePath.ToString(),
-						*TargetStruct.Name.ToString(), *TargetPath.ToString(),
+						TEXT("Type mismatch between source %s and target %s types, '%s' is not child of '%s'."),
+						*UE::StateTree::GetDescAndPathAsString(SourceStruct, SourcePath),
+						*UE::StateTree::GetDescAndPathAsString(TargetStruct, TargetPath),
 						*GetNameSafe(SourceStructProperty->Struct), *GetNameSafe(TargetBaseStruct));
 			return false;
 		}
@@ -1534,8 +1534,8 @@ bool FStateTreeCompiler::GetAndValidateBindings(const FStateTreeBindableStructDe
 	if (TargetStruct.Struct == nullptr)
 	{
 		Log.Reportf(EMessageSeverity::Error, TargetStruct,
-				TEXT("The type of binding target '%s' is invalid."),
-				*TargetStruct.Name.ToString());
+				TEXT("The type of binding target %s is invalid."),
+				*TargetStruct.ToString());
 		return false;
 	}
 	
@@ -1554,8 +1554,8 @@ bool FStateTreeCompiler::GetAndValidateBindings(const FStateTreeBindableStructDe
 		if (SourceStructIdx == INDEX_NONE)
 		{
 			Log.Reportf(EMessageSeverity::Error, TargetStruct,
-						TEXT("Failed to find binding source property '%s' for target '%s:%s'."),
-						*Binding.GetSourcePath().ToString(), *TargetStruct.Name.ToString(), *Binding.GetTargetPath().ToString());
+						TEXT("Failed to find binding source property '%s' for target %s."),
+						*Binding.GetSourcePath().ToString(), *UE::StateTree::GetDescAndPathAsString(TargetStruct, Binding.GetTargetPath()));
 			return false;
 		}
 		const FStateTreeBindableStructDesc& SourceStruct = BindingsCompiler.GetSourceStructDesc(SourceStructIdx);
@@ -1566,8 +1566,8 @@ bool FStateTreeCompiler::GetAndValidateBindings(const FStateTreeBindableStructDe
 		if (!Binding.GetMutableTargetPath().UpdateSegmentsFromValue(TargetValue))
 		{
 			Log.Reportf(EMessageSeverity::Error, TargetStruct,
-						TEXT("Malformed target property path for binding source property '%s' for target '%s:%s'."),
-						*Binding.GetSourcePath().ToString(), *TargetStruct.Name.ToString(), *Binding.GetTargetPath().ToString());
+						TEXT("Malformed target property path for binding source property '%s' for target %s."),
+						*Binding.GetSourcePath().ToString(), *UE::StateTree::GetDescAndPathAsString(TargetStruct, Binding.GetTargetPath()));
 			return false;
 		}
 		
@@ -1583,18 +1583,18 @@ bool FStateTreeCompiler::GetAndValidateBindings(const FStateTreeBindableStructDe
 		if (!bSourceAccessible)
 		{
 			Log.Reportf(EMessageSeverity::Error, TargetStruct,
-						TEXT("Property '%s:%s' cannot be bound to '%s:%s', because the binding source '%s' is not updated before '%s' in the tree."),
-						*SourceStruct.Name.ToString(), *Binding.GetSourcePath().ToString(),
-						*TargetStruct.Name.ToString(), *Binding.GetTargetPath().ToString(),
-						*SourceStruct.Name.ToString(), *TargetStruct.Name.ToString());
+						TEXT("Property at %s cannot be bound to %s, because the binding source %s is not updated before %s in the tree."),
+						*UE::StateTree::GetDescAndPathAsString(SourceStruct, Binding.GetSourcePath()),
+						*UE::StateTree::GetDescAndPathAsString(TargetStruct, Binding.GetTargetPath()),
+						*SourceStruct.ToString(), *TargetStruct.ToString());
 			return false;
 		}
 
 		if (!IDToStructValue.Contains(SourceStructID))
 		{
 			Log.Reportf(EMessageSeverity::Error, TargetStruct,
-				TEXT("Failed to find value for binding source property '%s' for target '%s:%s'."),
-				*Binding.GetSourcePath().ToString(), *TargetStruct.Name.ToString(), *Binding.GetTargetPath().ToString());
+				TEXT("Failed to find value for binding source property '%s' for target %s."),
+				*Binding.GetSourcePath().ToString(), *UE::StateTree::GetDescAndPathAsString(TargetStruct, Binding.GetTargetPath()));
 			return false;
 		}
 
@@ -1605,8 +1605,8 @@ bool FStateTreeCompiler::GetAndValidateBindings(const FStateTreeBindableStructDe
 			if (!Binding.GetMutableSourcePath().UpdateSegmentsFromValue(SourceValue))
 			{
 				Log.Reportf(EMessageSeverity::Error, TargetStruct,
-					TEXT("Malformed target property path for binding source property '%s' for source '%s:%s'."),
-					*Binding.GetSourcePath().ToString(), *TargetStruct.Name.ToString(), *Binding.GetTargetPath().ToString());
+					TEXT("Malformed target property path for binding source property '%s' for source %s."),
+					*Binding.GetSourcePath().ToString(), *UE::StateTree::GetDescAndPathAsString(TargetStruct, Binding.GetTargetPath()));
 				return false;
 			}
 		}
@@ -1669,8 +1669,8 @@ bool FStateTreeCompiler::GetAndValidateBindings(const FStateTreeBindableStructDe
 			if (bIsOptional == false && !IsPropertyBound(PropertyName))
 			{
 				Log.Reportf(EMessageSeverity::Error, TargetStruct,
-					TEXT("Input property '%s:%s' is expected to have a binding."),
-					*TargetStruct.Name.ToString(), *PropertyName.ToString());
+					TEXT("Input property '%s' on %s is expected to have a binding."),
+					*PropertyName.ToString(), *TargetStruct.ToString());
 				bResult = false;
 			}
 		}
@@ -1690,8 +1690,8 @@ bool FStateTreeCompiler::GetAndValidateBindings(const FStateTreeBindableStructDe
 			if (ContextObjectType == nullptr)
 			{
 				Log.Reportf(EMessageSeverity::Error, TargetStruct,
-					TEXT("Context property '%s:%s' type is expected to be Object Reference or Struct."),
-					*TargetStruct.Name.ToString(), *PropertyName.ToString());
+					TEXT("The type of Context property '%s' on %s is expected to be Object Reference or Struct."),
+					*PropertyName.ToString(), *TargetStruct.ToString());
 				bResult = false;
 				continue;
 			}
@@ -1710,8 +1710,8 @@ bool FStateTreeCompiler::GetAndValidateBindings(const FStateTreeBindableStructDe
 				else
 				{
 					Log.Reportf(EMessageSeverity::Error, TargetStruct,
-						TEXT("Cound not find matching Context object for Context property '%s:%s'. Property must have manual binding."),
-						*TargetStruct.Name.ToString(), *PropertyName.ToString());
+						TEXT("Cound not find matching Context object for Context property '%s' on %s. Property must have manual binding."),
+						*PropertyName.ToString(), *TargetStruct.ToString());
 					bResult = false;
 				}
 			}
diff --git a/Engine/Plugins/Runtime/StateTree/Source/StateTreeEditorModule/Private/StateTreeEditorData.cpp b/Engine/Plugins/Runtime/StateTree/Source/StateTreeEditorModule/Private/StateTreeEditorData.cpp
--- a/Engine/Plugins/Runtime/StateTree/Source/StateTreeEditorModule/Private/StateTreeEditorData.cpp
+++ b/Engine/Plugins/Runtime/StateTree/Source/StateTreeEditorModule/Private/StateTreeEditorData.cpp
@@ -355,17 +355,19 @@ bool UStateTreeEditorData::GetStructByID(const FGuid StructID, FStateTreeBindabl
 
 bool UStateTreeEditorData::GetDataViewByID(const FGuid StructID, FStateTreeDataView& OutDataView) const
 {
-	VisitAllNodes([&OutDataView, StructID](const UStateTreeState* State, const FStateTreeBindableStructDesc& Desc, const FStateTreeDataView Value)
+	bool bFound = false;
+	VisitAllNodes([&OutDataView, &bFound, StructID](const UStateTreeState* State, const FStateTreeBindableStructDesc& Desc, const FStateTreeDataView Value)
 	{
 		if (Desc.ID == StructID)
 		{
+			bFound = true;
 			OutDataView = Value;
 			return EStateTreeVisitor::Break;
 		}
 		return EStateTreeVisitor::Continue;
 	});
 
-	return OutDataView.IsValid();
+	return bFound;
 }
 
 const UStateTreeState* UStateTreeEditorData::GetStateByStructID(const FGuid TargetStructID) const
diff --git a/Engine/Plugins/Runtime/StateTree/Source/StateTreeEditorModule/Private/StateTreePropertyBindingCompiler.cpp b/Engine/Plugins/Runtime/StateTree/Source/StateTreeEditorModule/Private/StateTreePropertyBindingCompiler.cpp
--- a/Engine/Plugins/Runtime/StateTree/Source/StateTreeEditorModule/Private/StateTreePropertyBindingCompiler.cpp
+++ b/Engine/Plugins/Runtime/StateTree/Source/StateTreeEditorModule/Private/StateTreePropertyBindingCompiler.cpp
@@ -59,25 +59,25 @@ bool FStateTreePropertyBindingCompiler::CompileBatch(const FStateTreeBindableStr
 		
 		if (!Binding.GetSourcePath().ResolveIndirections(SourceStruct.Struct, SourceIndirections, &Error))
 		{
-			Log->Reportf(EMessageSeverity::Error, TargetStruct, TEXT("Resolving path in %s: %s"), *SourceStruct.Name.ToString(), *Error);
+			Log->Reportf(EMessageSeverity::Error, TargetStruct, TEXT("Resolving path in %s: %s"), *SourceStruct.ToString(), *Error);
 			return false;
 		}
 
 		if (!Binding.GetTargetPath().ResolveIndirections(TargetStruct.Struct, TargetIndirections, &Error))
 		{
-			Log->Reportf(EMessageSeverity::Error, TargetStruct, TEXT("Resolving path in %s: %s"), *TargetStruct.Name.ToString(), *Error);
+			Log->Reportf(EMessageSeverity::Error, TargetStruct, TEXT("Resolving path in %s: %s"), *TargetStruct.ToString(), *Error);
 			return false;
 		}
 
 		FStateTreePropertyCopy DummyCopy;
 		FStateTreePropertyPathIndirection LastSourceIndirection = !SourceIndirections.IsEmpty() ? SourceIndirections.Last() : FStateTreePropertyPathIndirection(SourceStruct.Struct);
 		FStateTreePropertyPathIndirection LastTargetIndirection = !TargetIndirections.IsEmpty() ? TargetIndirections.Last() : FStateTreePropertyPathIndirection(TargetStruct.Struct);
-		if (!PropertyBindings->ResolveCopyType(DummyCopy, LastSourceIndirection, LastTargetIndirection))
+		if (!PropertyBindings->ResolveCopyType(LastSourceIndirection, LastTargetIndirection, DummyCopy))
 		{
 			Log->Reportf(EMessageSeverity::Error, TargetStruct,
-			TEXT("Failed to resolve property copy type between %s:%s and %s:%s."),
-				*SourceStruct.Name.ToString(), *Binding.GetSourcePath().ToString(),
-				*TargetStruct.Name.ToString(), *Binding.GetTargetPath().ToString());
+			TEXT("Cannot copy properties between %s and %s, properties are incompatible."),
+				*UE::StateTree::GetDescAndPathAsString(SourceStruct, Binding.GetSourcePath()),
+				*UE::StateTree::GetDescAndPathAsString(TargetStruct, Binding.GetTargetPath()));
 			return false;
 		}
 
@@ -143,12 +143,11 @@ void FStateTreePropertyBindingCompiler::Finalize()
 
 int32 FStateTreePropertyBindingCompiler::AddSourceStruct(const FStateTreeBindableStructDesc& SourceStruct)
 {
-	const int32 ExistingIndex = SourceStructs.IndexOfByPredicate([&SourceStruct](const FStateTreeBindableStructDesc& Struct) { return (Struct.ID == SourceStruct.ID); });
-	if (ExistingIndex != INDEX_NONE)
+	const FStateTreeBindableStructDesc* ExistingStruct = SourceStructs.FindByPredicate([&SourceStruct](const FStateTreeBindableStructDesc& Struct) { return (Struct.ID == SourceStruct.ID); });
+	if (ExistingStruct)
 	{
-		const FStateTreeBindableStructDesc& ExistingStruct = SourceStructs[ExistingIndex];
-		UE_LOG(LogStateTree, Error, TEXT("Struct '%s' already exists as '%s'"),
-			*SourceStruct.Name.ToString(), *SourceStruct.ID.ToString(), *ExistingStruct.Name.ToString());
+		UE_LOG(LogStateTree, Error, TEXT("%s already exists as %s"),
+			*SourceStruct.ToString(), *ExistingStruct->ToString());
 	}
 	
 	SourceStructs.Add(SourceStruct);
diff --git a/Engine/Plugins/Runtime/StateTree/Source/StateTreeModule/Private/StateTreePropertyBindings.cpp b/Engine/Plugins/Runtime/StateTree/Source/StateTreeModule/Private/StateTreePropertyBindings.cpp
--- a/Engine/Plugins/Runtime/StateTree/Source/StateTreeModule/Private/StateTreePropertyBindings.cpp
+++ b/Engine/Plugins/Runtime/StateTree/Source/StateTreeModule/Private/StateTreePropertyBindings.cpp
@@ -14,6 +14,25 @@
 
 #include UE_INLINE_GENERATED_CPP_BY_NAME(StateTreePropertyBindings)
 
+namespace UE::StateTree
+{
+	FString GetDescAndPathAsString(const FStateTreeBindableStructDesc& Desc, const FStateTreePropertyPath& Path)
+	{
+		FStringBuilderBase Result;
+
+		Result += Desc.ToString();
+
+		if (!Path.IsPathEmpty())
+		{
+			Result += TEXT(" ");
+			Result += Path.ToString();
+		}
+
+		return Result.ToString();
+	}
+
+} // UE::StateTree
+
 namespace UE::StateTree::Private
 {
 
@@ -56,8 +75,24 @@ PRAGMA_DISABLE_DEPRECATION_WARNINGS
 	}
 PRAGMA_ENABLE_DEPRECATION_WARNINGS
 #endif // WITH_EDITORONLY_DATA
-}; // UE::StateTree::Private 
+} // UE::StateTree::Private 
+
+
+//----------------------------------------------------------------//
+//  FStateTreeBindableStructDesc
+//----------------------------------------------------------------//
 
+FString FStateTreeBindableStructDesc::ToString() const
+{
+	FStringBuilderBase Result;
+
+	Result += UEnum::GetDisplayValueAsText(DataSource).ToString();
+	Result += TEXT(" '");
+	Result += Name.ToString();
+	Result += TEXT("'");
+
+	return Result.ToString();
+}
 
 //----------------------------------------------------------------//
 //  FStateTreePropertyPathBinding
@@ -141,7 +176,7 @@ bool FStateTreePropertyBindings::ResolvePaths()
 			FStateTreePropertyPathIndirection TargetLeafIndirection;
 			bSuccess = bSuccess && ResolvePath(SourceStruct, Binding.GetSourcePath(), Copy.SourceIndirection, SourceLeafIndirection);
 			bSuccess = bSuccess && ResolvePath(TargetStruct, Binding.GetTargetPath(), Copy.TargetIndirection, TargetLeafIndirection);
-			bSuccess = bSuccess && ResolveCopyType(Copy, SourceLeafIndirection, TargetLeafIndirection);
+			bSuccess = bSuccess && ResolveCopyType(SourceLeafIndirection, TargetLeafIndirection, Copy);
 			if (!bSuccess) 
 			{
 				// Resolving or validating failed, make the copy a nop.
@@ -297,7 +332,7 @@ bool FStateTreePropertyBindings::ResolvePath(const UStruct* Struct, const FState
 	return true;
 }
 
-bool FStateTreePropertyBindings::ResolveCopyType(FStateTreePropertyCopy& Copy, const FStateTreePropertyPathIndirection& SourceIndirection,const FStateTreePropertyPathIndirection& TargetIndirection)
+bool FStateTreePropertyBindings::ResolveCopyType(const FStateTreePropertyPathIndirection& SourceIndirection,const FStateTreePropertyPathIndirection& TargetIndirection, FStateTreePropertyCopy& OutCopy)
 {
 	// @todo: see if GetPropertyCompatibility() can be implemented as call to ResolveCopyType() instead so that we write this logic just once.
 	
@@ -312,10 +347,10 @@ bool FStateTreePropertyBindings::ResolveCopyType(FStateTreePropertyCopy& Copy, c
 		return false;
 	}
 
-	Copy.SourceLeafProperty = SourceProperty;
-	Copy.TargetLeafProperty = TargetProperty;
-	Copy.CopySize = 0;
-	Copy.Type = EStateTreePropertyCopyType::None;
+	OutCopy.SourceLeafProperty = SourceProperty;
+	OutCopy.TargetLeafProperty = TargetProperty;
+	OutCopy.CopySize = 0;
+	OutCopy.Type = EStateTreePropertyCopyType::None;
 	
 	if (SourceProperty == nullptr)
 	{
@@ -324,7 +359,7 @@ bool FStateTreePropertyBindings::ResolveCopyType(FStateTreePropertyCopy& Copy, c
 		{
 			if (TargetStructProperty->Struct == SourceStruct)
 			{
-				Copy.Type = EStateTreePropertyCopyType::CopyStruct;
+				OutCopy.Type = EStateTreePropertyCopyType::CopyStruct;
 				return true;
 			}
 		}
@@ -332,7 +367,7 @@ bool FStateTreePropertyBindings::ResolveCopyType(FStateTreePropertyCopy& Copy, c
 		{
 			if (SourceStruct->IsChildOf(TargetObjectProperty->PropertyClass))
 			{
-				Copy.Type = EStateTreePropertyCopyType::CopyObject;
+				OutCopy.Type = EStateTreePropertyCopyType::CopyObject;
 				return true;
 			}
 		}
@@ -350,7 +385,7 @@ bool FStateTreePropertyBindings::ResolveCopyType(FStateTreePropertyCopy& Copy, c
 				// FStateTreeStructRef to FStateTreeStructRef is copied as usual.
 				if (SourceStructProperty->Struct != TBaseStructure<FStateTreeStructRef>::Get())
 				{
-					Copy.Type = EStateTreePropertyCopyType::StructReference;
+					OutCopy.Type = EStateTreePropertyCopyType::StructReference;
 					return true;
 				}
 			}
@@ -374,39 +409,39 @@ bool FStateTreePropertyBindings::ResolveCopyType(FStateTreePropertyCopy& Copy, c
 	{
 		if (CastField<FNameProperty>(TargetProperty))
 		{
-			Copy.Type = EStateTreePropertyCopyType::CopyName;
+			OutCopy.Type = EStateTreePropertyCopyType::CopyName;
 			return true;
 		}
 		else if (CastField<FBoolProperty>(TargetProperty))
 		{
-			Copy.Type = EStateTreePropertyCopyType::CopyBool;
+			OutCopy.Type = EStateTreePropertyCopyType::CopyBool;
 			return true;
 		}
 		else if (CastField<FStructProperty>(TargetProperty))
 		{
-			Copy.Type = EStateTreePropertyCopyType::CopyStruct;
+			OutCopy.Type = EStateTreePropertyCopyType::CopyStruct;
 			return true;
 		}
 		else if (CastField<FObjectPropertyBase>(TargetProperty))
 		{
-			Copy.Type = EStateTreePropertyCopyType::CopyObject;
+			OutCopy.Type = EStateTreePropertyCopyType::CopyObject;
 			return true;
 		}
 		else if (CastField<FArrayProperty>(TargetProperty) && TargetProperty->HasAnyPropertyFlags(CPF_EditFixedSize))
 		{
 			// only apply array copying rules if the destination array is fixed size, otherwise it will be 'complex'
-			Copy.Type = EStateTreePropertyCopyType::CopyFixedArray;
+			OutCopy.Type = EStateTreePropertyCopyType::CopyFixedArray;
 			return true;
 		}
 		else if (TargetProperty->PropertyFlags & CPF_IsPlainOldData)
 		{
-			Copy.Type = EStateTreePropertyCopyType::CopyPlain;
-			Copy.CopySize = SourceProperty->ElementSize * SourceProperty->ArrayDim;
+			OutCopy.Type = EStateTreePropertyCopyType::CopyPlain;
+			OutCopy.CopySize = SourceProperty->ElementSize * SourceProperty->ArrayDim;
 			return true;
 		}
 		else
 		{
-			Copy.Type = EStateTreePropertyCopyType::CopyComplex;
+			OutCopy.Type = EStateTreePropertyCopyType::CopyComplex;
 			return true;
 		}
 	}
@@ -416,32 +451,32 @@ bool FStateTreePropertyBindings::ResolveCopyType(FStateTreePropertyCopy& Copy, c
 		{
 			if (TargetProperty->IsA<FByteProperty>())
 			{
-				Copy.Type = EStateTreePropertyCopyType::PromoteBoolToByte;
+				OutCopy.Type = EStateTreePropertyCopyType::PromoteBoolToByte;
 				return true;
 			}
 			else if (TargetProperty->IsA<FIntProperty>())
 			{
-				Copy.Type = EStateTreePropertyCopyType::PromoteBoolToInt32;
+				OutCopy.Type = EStateTreePropertyCopyType::PromoteBoolToInt32;
 				return true;
 			}
 			else if (TargetProperty->IsA<FUInt32Property>())
 			{
-				Copy.Type = EStateTreePropertyCopyType::PromoteBoolToUInt32;
+				OutCopy.Type = EStateTreePropertyCopyType::PromoteBoolToUInt32;
 				return true;
 			}
 			else if (TargetProperty->IsA<FInt64Property>())
 			{
-				Copy.Type = EStateTreePropertyCopyType::PromoteBoolToInt64;
+				OutCopy.Type = EStateTreePropertyCopyType::PromoteBoolToInt64;
 				return true;
 			}
 			else if (TargetProperty->IsA<FFloatProperty>())
 			{
-				Copy.Type = EStateTreePropertyCopyType::PromoteBoolToFloat;
+				OutCopy.Type = EStateTreePropertyCopyType::PromoteBoolToFloat;
 				return true;
 			}
 			else if (TargetProperty->IsA<FDoubleProperty>())
 			{
-				Copy.Type = EStateTreePropertyCopyType::PromoteBoolToDouble;
+				OutCopy.Type = EStateTreePropertyCopyType::PromoteBoolToDouble;
 				return true;
 			}
 		}
@@ -449,27 +484,27 @@ bool FStateTreePropertyBindings::ResolveCopyType(FStateTreePropertyCopy& Copy, c
 		{
 			if (TargetProperty->IsA<FIntProperty>())
 			{
-				Copy.Type = EStateTreePropertyCopyType::PromoteByteToInt32;
+				OutCopy.Type = EStateTreePropertyCopyType::PromoteByteToInt32;
 				return true;
 			}
 			else if (TargetProperty->IsA<FUInt32Property>())
 			{
-				Copy.Type = EStateTreePropertyCopyType::PromoteByteToUInt32;
+				OutCopy.Type = EStateTreePropertyCopyType::PromoteByteToUInt32;
 				return true;
 			}
 			else if (TargetProperty->IsA<FInt64Property>())
 			{
-				Copy.Type = EStateTreePropertyCopyType::PromoteByteToInt64;
+				OutCopy.Type = EStateTreePropertyCopyType::PromoteByteToInt64;
 				return true;
 			}
 			else if (TargetProperty->IsA<FFloatProperty>())
 			{
-				Copy.Type = EStateTreePropertyCopyType::PromoteByteToFloat;
+				OutCopy.Type = EStateTreePropertyCopyType::PromoteByteToFloat;
 				return true;
 			}
 			else if (TargetProperty->IsA<FDoubleProperty>())
 			{
-				Copy.Type = EStateTreePropertyCopyType::PromoteByteToDouble;
+				OutCopy.Type = EStateTreePropertyCopyType::PromoteByteToDouble;
 				return true;
 			}
 		}
@@ -477,17 +512,17 @@ bool FStateTreePropertyBindings::ResolveCopyType(FStateTreePropertyCopy& Copy, c
 		{
 			if (TargetProperty->IsA<FInt64Property>())
 			{
-				Copy.Type = EStateTreePropertyCopyType::PromoteInt32ToInt64;
+				OutCopy.Type = EStateTreePropertyCopyType::PromoteInt32ToInt64;
 				return true;
 			}
 			else if (TargetProperty->IsA<FFloatProperty>())
 			{
-				Copy.Type = EStateTreePropertyCopyType::PromoteInt32ToFloat;
+				OutCopy.Type = EStateTreePropertyCopyType::PromoteInt32ToFloat;
 				return true;
 			}
 			else if (TargetProperty->IsA<FDoubleProperty>())
 			{
-				Copy.Type = EStateTreePropertyCopyType::PromoteInt32ToDouble;
+				OutCopy.Type = EStateTreePropertyCopyType::PromoteInt32ToDouble;
 				return true;
 			}
 		}
@@ -495,17 +530,17 @@ bool FStateTreePropertyBindings::ResolveCopyType(FStateTreePropertyCopy& Copy, c
 		{
 			if (TargetProperty->IsA<FInt64Property>())
 			{
-				Copy.Type = EStateTreePropertyCopyType::PromoteUInt32ToInt64;
+				OutCopy.Type = EStateTreePropertyCopyType::PromoteUInt32ToInt64;
 				return true;
 			}
 			else if (TargetProperty->IsA<FFloatProperty>())
 			{
-				Copy.Type = EStateTreePropertyCopyType::PromoteUInt32ToFloat;
+				OutCopy.Type = EStateTreePropertyCopyType::PromoteUInt32ToFloat;
 				return true;
 			}
 			else if (TargetProperty->IsA<FDoubleProperty>())
 			{
-				Copy.Type = EStateTreePropertyCopyType::PromoteUInt32ToDouble;
+				OutCopy.Type = EStateTreePropertyCopyType::PromoteUInt32ToDouble;
 				return true;
 			}
 		}
@@ -513,17 +548,17 @@ bool FStateTreePropertyBindings::ResolveCopyType(FStateTreePropertyCopy& Copy, c
 		{
 			if (TargetProperty->IsA<FIntProperty>())
 			{
-				Copy.Type = EStateTreePropertyCopyType::PromoteFloatToInt32;
+				OutCopy.Type = EStateTreePropertyCopyType::PromoteFloatToInt32;
 				return true;
 			}
 			else if (TargetProperty->IsA<FInt64Property>())
 			{
-				Copy.Type = EStateTreePropertyCopyType::PromoteFloatToInt64;
+				OutCopy.Type = EStateTreePropertyCopyType::PromoteFloatToInt64;
 				return true;
 			}
 			else if (TargetProperty->IsA<FDoubleProperty>())
 			{
-				Copy.Type = EStateTreePropertyCopyType::PromoteFloatToDouble;
+				OutCopy.Type = EStateTreePropertyCopyType::PromoteFloatToDouble;
 				return true;
 			}
 		}
@@ -531,17 +566,17 @@ bool FStateTreePropertyBindings::ResolveCopyType(FStateTreePropertyCopy& Copy, c
 		{
 			if (TargetProperty->IsA<FIntProperty>())
 			{
-				Copy.Type = EStateTreePropertyCopyType::DemoteDoubleToInt32;
+				OutCopy.Type = EStateTreePropertyCopyType::DemoteDoubleToInt32;
 				return true;
 			}
 			else if (TargetProperty->IsA<FInt64Property>())
 			{
-				Copy.Type = EStateTreePropertyCopyType::DemoteDoubleToInt64;
+				OutCopy.Type = EStateTreePropertyCopyType::DemoteDoubleToInt64;
 				return true;
 			}
 			else if (TargetProperty->IsA<FFloatProperty>())
 			{
-				Copy.Type = EStateTreePropertyCopyType::DemoteDoubleToFloat;
+				OutCopy.Type = EStateTreePropertyCopyType::DemoteDoubleToFloat;
 				return true;
 			}
 		}
diff --git a/Engine/Plugins/Runtime/StateTree/Source/StateTreeModule/Public/StateTreePropertyBindings.h b/Engine/Plugins/Runtime/StateTree/Source/StateTreeModule/Public/StateTreePropertyBindings.h
--- a/Engine/Plugins/Runtime/StateTree/Source/StateTreeModule/Public/StateTreePropertyBindings.h
+++ b/Engine/Plugins/Runtime/StateTree/Source/StateTreeModule/Public/StateTreePropertyBindings.h
@@ -136,6 +136,8 @@ struct STATETREEMODULE_API FStateTreeBindableStructDesc
 #endif
 
 	bool IsValid() const { return Struct != nullptr; }
+
+	FString ToString() const;
 	
 	/** The type of the struct or class. */
 	UPROPERTY()
@@ -782,9 +784,17 @@ struct STATETREEMODULE_API FStateTreePropertyBindings
 	/** @return how properties are compatible for copying. */
 	static EStateTreePropertyAccessCompatibility GetPropertyCompatibility(const FProperty* FromProperty, const FProperty* ToProperty);
 
+	/**
+	 * Resolves what kind of copy type to use between specified property indirections.
+	 * @param SourceIndirection Property path indirections of the copy source,
+	 * @param TargetIndirection Property path indirections of the copy target.
+	 * @param OutCopy Resulting copy type.
+	 * @return true if copy was resolved, or false if no copy could be resolved between paths.
+	 */
+	[[nodiscard]] static bool ResolveCopyType(const FStateTreePropertyPathIndirection& SourceIndirection, const FStateTreePropertyPathIndirection& TargetIndirection, FStateTreePropertyCopy& OutCopy);
+
 private:
 	[[nodiscard]] bool ResolvePath(const UStruct* Struct, const FStateTreePropertyPath& Path, FStateTreePropertyIndirection& OutFirstIndirection, FStateTreePropertyPathIndirection& OutLeafIndirection);
-	[[nodiscard]] bool ResolveCopyType(FStateTreePropertyCopy& Copy, const FStateTreePropertyPathIndirection& SourceIndirection,const FStateTreePropertyPathIndirection& TargetIndirection);
 
 	void PerformCopy(const FStateTreePropertyCopy& Copy, uint8* SourceAddress, uint8* TargetAddress) const;
 	void PerformResetObjects(const FStateTreePropertyCopy& Copy, uint8* TargetAddress) const;
@@ -859,6 +869,12 @@ PRAGMA_ENABLE_DEPRECATION_WARNINGS
 };
 
 
+namespace UE::StateTree
+{
+	/** @return desc and path as a display string. */
+	extern STATETREEMODULE_API FString GetDescAndPathAsString(const FStateTreeBindableStructDesc& Desc, const FStateTreePropertyPath& Path);
+} // UE::StateTree
+
 namespace UE::StateTree::Private
 {
 #if WITH_EDITORONLY_DATA
@@ -870,7 +886,7 @@ PRAGMA_DISABLE_DEPRECATION_WARNINGS
 	extern STATETREEMODULE_API FStateTreeEditorPropertyPath ConvertEditorPath(const FStateTreePropertyPath& InPath);
 PRAGMA_ENABLE_DEPRECATION_WARNINGS
 #endif // WITH_EDITORONLY_DATA
-}; // UE::StateTree::Private 
+} // UE::StateTree::Private 
 
 #if UE_ENABLE_INCLUDE_ORDER_DEPRECATED_IN_5_2
 #include "CoreMinimal.h"
