StateTree: Remove ensure from ResolveIndirectionsWithValue()
- fails silenty, as we need to call it during runtiem load
- there's already a note on the method documentation about the behavior

diff --git a/Engine/Plugins/Runtime/StateTree/Source/StateTreeModule/Private/StateTreePropertyBindings.cpp b/Engine/Plugins/Runtime/StateTree/Source/StateTreeModule/Private/StateTreePropertyBindings.cpp
--- a/Engine/Plugins/Runtime/StateTree/Source/StateTreeModule/Private/StateTreePropertyBindings.cpp
+++ b/Engine/Plugins/Runtime/StateTree/Source/StateTreeModule/Private/StateTreePropertyBindings.cpp
@@ -1430,11 +1430,6 @@ bool FStateTreePropertyPath::ResolveIndirectionsWithValue(const FStateTreeDataVi
 				}
 			}
 		}
-#else
-		if (bHandleRedirects)
-		{
-			ensureMsgf(false, TEXT("The bHandleRedirects option is only available on editor builds."));
-		}
 #endif // WITH_EDITORONLY_DATA
 
 		if (!Property)
