[StateTree] moved some types to the execution types header

diff --git a/Engine/Plugins/Runtime/GameplayStateTree/Source/GameplayStateTreeModule/Public/Components/StateTreeComponentSchema.h b/Engine/Plugins/Runtime/GameplayStateTree/Source/GameplayStateTreeModule/Public/Components/StateTreeComponentSchema.h
--- a/Engine/Plugins/Runtime/GameplayStateTree/Source/GameplayStateTreeModule/Public/Components/StateTreeComponentSchema.h
+++ b/Engine/Plugins/Runtime/GameplayStateTree/Source/GameplayStateTreeModule/Public/Components/StateTreeComponentSchema.h
@@ -4,7 +4,7 @@
 
 #include "StateTreeSchema.h"
 #include "GameFramework/Actor.h"
-#include "StateTreeTypes.h"
+#include "StateTreeExecutionTypes.h"
 #include "StateTreeComponentSchema.generated.h"
 
 /**
@@ -25,7 +25,7 @@ protected:
 	virtual bool IsClassAllowed(const UClass* InScriptStruct) const override;
 	virtual bool IsExternalItemAllowed(const UStruct& InStruct) const override;
 	
-	virtual TConstArrayView<FStateTreeExternalDataDesc> GetContextDataDescs() const;
+	virtual TConstArrayView<FStateTreeExternalDataDesc> GetContextDataDescs() const override;
 
 	virtual void PostLoad() override;
 
diff --git a/Engine/Plugins/Runtime/StateTree/Source/StateTreeModule/Private/StateTreeExecutionTypes.cpp b/Engine/Plugins/Runtime/StateTree/Source/StateTreeModule/Private/StateTreeExecutionTypes.cpp
new file mode 100644
--- /dev/null
+++ b/Engine/Plugins/Runtime/StateTree/Source/StateTreeModule/Private/StateTreeExecutionTypes.cpp
@@ -0,0 +1,9 @@
+// Copyright Epic Games, Inc. All Rights Reserved.
+
+#include "StateTreeExecutionTypes.h"
+
+const FStateTreeExternalDataHandle FStateTreeExternalDataHandle::Invalid = FStateTreeExternalDataHandle();
+
+#if WITH_STATETREE_DEBUGGER
+const FStateTreeInstanceDebugId FStateTreeInstanceDebugId::Invalid = FStateTreeInstanceDebugId();
+#endif // WITH_STATETREE_DEBUGGER
\ No newline at end of file
diff --git a/Engine/Plugins/Runtime/StateTree/Source/StateTreeModule/Private/StateTreeInstanceData.cpp b/Engine/Plugins/Runtime/StateTree/Source/StateTreeModule/Private/StateTreeInstanceData.cpp
--- a/Engine/Plugins/Runtime/StateTree/Source/StateTreeModule/Private/StateTreeInstanceData.cpp
+++ b/Engine/Plugins/Runtime/StateTree/Source/StateTreeModule/Private/StateTreeInstanceData.cpp
@@ -1,6 +1,7 @@
 // Copyright Epic Games, Inc. All Rights Reserved.
 
 #include "StateTreeInstanceData.h"
+#include "StateTreeExecutionTypes.h"
 #include "Serialization/MemoryReader.h"
 #include "Serialization/MemoryWriter.h"
 #include "VisualLogger/VisualLogger.h"
diff --git a/Engine/Plugins/Runtime/StateTree/Source/StateTreeModule/Private/StateTreeTypes.cpp b/Engine/Plugins/Runtime/StateTree/Source/StateTreeModule/Private/StateTreeTypes.cpp
--- a/Engine/Plugins/Runtime/StateTree/Source/StateTreeModule/Private/StateTreeTypes.cpp
+++ b/Engine/Plugins/Runtime/StateTree/Source/StateTreeModule/Private/StateTreeTypes.cpp
@@ -1,30 +1,63 @@
 // Copyright Epic Games, Inc. All Rights Reserved.
 
 #include "StateTreeTypes.h"
-#include "StateTree.h"
+#include "StateTree.h" // FStateTreeCustomVersion
 
 #include UE_INLINE_GENERATED_CPP_BY_NAME(StateTreeTypes)
 
 DEFINE_LOG_CATEGORY(LogStateTree);
 
 const FStateTreeStateHandle FStateTreeStateHandle::Invalid = FStateTreeStateHandle();
-const FStateTreeStateHandle FStateTreeStateHandle::Succeeded = FStateTreeStateHandle(FStateTreeStateHandle::SucceededIndex);
-const FStateTreeStateHandle FStateTreeStateHandle::Failed = FStateTreeStateHandle(FStateTreeStateHandle::FailedIndex);
-const FStateTreeStateHandle FStateTreeStateHandle::Stopped = FStateTreeStateHandle(FStateTreeStateHandle::StoppedIndex);
+const FStateTreeStateHandle FStateTreeStateHandle::Succeeded = FStateTreeStateHandle(SucceededIndex);
+const FStateTreeStateHandle FStateTreeStateHandle::Failed = FStateTreeStateHandle(FailedIndex);
+const FStateTreeStateHandle FStateTreeStateHandle::Stopped = FStateTreeStateHandle(StoppedIndex);
 const FStateTreeStateHandle FStateTreeStateHandle::Root = FStateTreeStateHandle(0);
 
 const FStateTreeIndex16 FStateTreeIndex16::Invalid = FStateTreeIndex16();
 const FStateTreeIndex8 FStateTreeIndex8::Invalid = FStateTreeIndex8();
 
-const FStateTreeExternalDataHandle FStateTreeExternalDataHandle::Invalid = FStateTreeExternalDataHandle();
-
-#if WITH_STATETREE_DEBUGGER
-const FStateTreeInstanceDebugId FStateTreeInstanceDebugId::Invalid = FStateTreeInstanceDebugId();
-#endif // WITH_STATETREE_DEBUGGER
 
 //////////////////////////////////////////////////////////////////////////
 // FStateTreeStateLink
 
+EStateTreeRunStatus FStateTreeStateHandle::ToCompletionStatus() const
+{
+	if (Index == SucceededIndex)
+	{
+		return EStateTreeRunStatus::Succeeded;
+	}
+
+	if (Index == FailedIndex)
+	{
+		return EStateTreeRunStatus::Failed;
+	}
+
+	if (Index == StoppedIndex)
+	{
+		return EStateTreeRunStatus::Stopped;
+	}
+	return EStateTreeRunStatus::Unset;
+}
+
+FStateTreeStateHandle FStateTreeStateHandle::FromCompletionStatus(const EStateTreeRunStatus Status)
+{
+	if (Status == EStateTreeRunStatus::Succeeded)
+	{
+		return Succeeded;
+	}
+
+	if (Status == EStateTreeRunStatus::Failed)
+	{
+		return Failed;
+	}
+
+	if (Status == EStateTreeRunStatus::Stopped)
+	{
+		return Stopped;
+	}
+	return {};
+}
+
 bool FStateTreeStateLink::Serialize(FStructuredArchive::FSlot Slot)
 {
 	Slot.GetUnderlyingArchive().UsingCustomVersion(FStateTreeCustomVersion::GUID);
diff --git a/Engine/Plugins/Runtime/StateTree/Source/StateTreeModule/Public/Blueprint/StateTreeNodeBlueprintBase.h b/Engine/Plugins/Runtime/StateTree/Source/StateTreeModule/Public/Blueprint/StateTreeNodeBlueprintBase.h
--- a/Engine/Plugins/Runtime/StateTree/Source/StateTreeModule/Public/Blueprint/StateTreeNodeBlueprintBase.h
+++ b/Engine/Plugins/Runtime/StateTree/Source/StateTreeModule/Public/Blueprint/StateTreeNodeBlueprintBase.h
@@ -2,7 +2,7 @@
 
 #pragma once
 
-#include "StateTreeTypes.h"
+#include "StateTreeExecutionTypes.h"
 #include "StateTreeNodeBlueprintBase.generated.h"
 
 struct FStateTreeEvent;
diff --git a/Engine/Plugins/Runtime/StateTree/Source/StateTreeModule/Public/StateTreeConditionBase.h b/Engine/Plugins/Runtime/StateTree/Source/StateTreeModule/Public/StateTreeConditionBase.h
--- a/Engine/Plugins/Runtime/StateTree/Source/StateTreeModule/Public/StateTreeConditionBase.h
+++ b/Engine/Plugins/Runtime/StateTree/Source/StateTreeModule/Public/StateTreeConditionBase.h
@@ -3,6 +3,7 @@
 #pragma once
 
 #include "StateTreeNodeBase.h"
+#include "StateTreeExecutionTypes.h"
 #include "StateTreeConditionBase.generated.h"
 
 struct FStateTreeExecutionContext;
diff --git a/Engine/Plugins/Runtime/StateTree/Source/StateTreeModule/Public/StateTreeExecutionTypes.h b/Engine/Plugins/Runtime/StateTree/Source/StateTreeModule/Public/StateTreeExecutionTypes.h
--- a/Engine/Plugins/Runtime/StateTree/Source/StateTreeModule/Public/StateTreeExecutionTypes.h
+++ b/Engine/Plugins/Runtime/StateTree/Source/StateTreeModule/Public/StateTreeExecutionTypes.h
@@ -1,6 +1,7 @@
 // Copyright Epic Games, Inc. All Rights Reserved.
 
 #pragma once
+#include "StateTreeTypes.h"
 
 #include "StateTreeExecutionTypes.generated.h"
 
@@ -24,4 +25,485 @@ enum class EStateTreeUpdatePhase : uint8
 	TransitionConditions	UMETA(DisplayName = "Transition conditions"),
 	StateSelection			UMETA(DisplayName = "Try Enter"),
 	EnterConditions			UMETA(DisplayName = "Enter conditions"),
-};
\ No newline at end of file
+};
+
+
+/** Status describing current run status of a State Tree. */
+UENUM(BlueprintType)
+enum class EStateTreeRunStatus : uint8
+{
+	/** Tree is still running. */
+	Running,
+	
+	/** Tree execution has stopped on failure. */
+	Failed,
+	
+	/** Tree execution has stopped on success. */
+	Succeeded,
+
+	/** The State Tree was requested to stop without a particular success or failure state. */
+	Stopped,
+
+	/** Status not set. */
+	Unset,
+};
+
+
+/** State change type. Passed to EnterState() and ExitState() to indicate how the state change affects the state and Evaluator or Task is on. */
+UENUM()
+enum class EStateTreeStateChangeType : uint8
+{
+	/** Not an activation */
+	None,
+	
+	/** The state became activated or deactivated. */
+	Changed,
+	
+	/** The state is parent of new active state and sustained previous active state. */
+	Sustained,
+};
+
+
+/** Defines how to assign the result of a condition to evaluate.  */
+UENUM()
+enum class EStateTreeConditionEvaluationMode : uint8
+{
+	/** Condition is evaluated to set the result. This is the normal behavior. */
+	Evaluated,
+	
+	/** Do not evaluate the condition and force result to 'true'. */
+	ForcedTrue,
+	
+	/** Do not evaluate the condition and force result to 'false'. */
+	ForcedFalse,
+};
+
+
+/**
+ * Handle to access an external struct or object.
+ * Note: Use the templated version below. 
+ */
+USTRUCT()
+struct STATETREEMODULE_API FStateTreeExternalDataHandle
+{
+	GENERATED_BODY()
+
+	static const FStateTreeExternalDataHandle Invalid;
+	
+	static bool IsValidIndex(const int32 Index) { return FStateTreeIndex16::IsValidIndex(Index); }
+	bool IsValid() const { return DataViewIndex.IsValid(); }
+
+	UPROPERTY()
+	FStateTreeIndex16 DataViewIndex = FStateTreeIndex16::Invalid;
+};
+
+
+/**
+ * Handle to access an external struct or object.
+ * This reference handle can be used in StateTree tasks and evaluators to have quick access to external data.
+ * The type provided to the template is used by the linker and context to pass along the type.
+ *
+ * USTRUCT()
+ * struct FExampleTask : public FStateTreeTaskBase
+ * {
+ *    ...
+ *
+ *    bool Link(FStateTreeLinker& Linker)
+ *    {
+ *      Linker.LinkExternalData(ExampleSubsystemHandle);
+ *      return true;
+ *    }
+ * 
+ *    EStateTreeRunStatus EnterState(FStateTreeExecutionContext& Context, const EStateTreeStateChangeType ChangeType, const FStateTreeTransitionResult& Transition)
+ *    {
+ *      const UExampleSubsystem& ExampleSubsystem = Context.GetExternalData(ExampleSubsystemHandle);
+ *      ...
+ *    }
+ *
+ *    TStateTreeExternalDataHandle<UExampleSubsystem> ExampleSubsystemHandle;
+ * }
+ */
+template<typename T, EStateTreeExternalDataRequirement Req = EStateTreeExternalDataRequirement::Required>
+struct TStateTreeExternalDataHandle : FStateTreeExternalDataHandle
+{
+	typedef T DataType;
+	static constexpr EStateTreeExternalDataRequirement DataRequirement = Req;
+};
+
+
+/**
+ * Describes an external data. The data can point to a struct or object.
+ * The code that handles StateTree ticking is responsible for passing in the actually data, see FStateTreeExecutionContext.
+ */
+USTRUCT()
+struct STATETREEMODULE_API FStateTreeExternalDataDesc
+{
+	GENERATED_BODY()
+
+	FStateTreeExternalDataDesc() = default;
+	FStateTreeExternalDataDesc(const UStruct* InStruct, const EStateTreeExternalDataRequirement InRequirement) : Struct(InStruct), Requirement(InRequirement) {}
+
+	FStateTreeExternalDataDesc(const FName InName, const UStruct* InStruct, const FGuid InGuid)
+		: Struct(InStruct)
+		, Name(InName)
+#if WITH_EDITORONLY_DATA
+		, ID(InGuid)
+#endif
+	{}
+
+	bool operator==(const FStateTreeExternalDataDesc& Other) const
+	{
+		return Struct == Other.Struct && Requirement == Other.Requirement;
+	}
+	
+	/** Class or struct of the external data. */
+	UPROPERTY();
+	TObjectPtr<const UStruct> Struct = nullptr;
+
+	/**
+	 * Name of the external data. Used only for bindable external data (enforced by the schema).
+	 * External data linked explicitly by the nodes (i.e. LinkExternalData) are identified only
+	 * by their type since they are used for unique instance of a given type.  
+	 */
+	UPROPERTY(VisibleAnywhere, Category = Common)
+	FName Name;
+	
+	/** Handle/Index to the StateTreeExecutionContext data views array */
+	UPROPERTY();
+	FStateTreeExternalDataHandle Handle;
+
+	/** Describes if the data is required or not. */
+	UPROPERTY();
+	EStateTreeExternalDataRequirement Requirement = EStateTreeExternalDataRequirement::Required;
+
+#if WITH_EDITORONLY_DATA
+	/** Unique identifier. Used only for bindable external data. */
+	UPROPERTY()
+	FGuid ID;
+#endif
+};
+
+
+/** Transition request */
+USTRUCT(BlueprintType)
+struct STATETREEMODULE_API FStateTreeTransitionRequest
+{
+	GENERATED_BODY()
+
+	FStateTreeTransitionRequest() = default;
+
+	explicit FStateTreeTransitionRequest(const FStateTreeStateHandle InTargetState, const EStateTreeTransitionPriority InPriority = EStateTreeTransitionPriority::Normal)
+		: TargetState(InTargetState)
+		, Priority(InPriority)
+	{
+	}
+
+	/** Source state of the transition. Filled in by the StateTree execution context. */
+	UPROPERTY(EditDefaultsOnly, Category = "Default")
+	FStateTreeStateHandle SourceState;
+
+	/** Target state of the transition. */
+	UPROPERTY(EditDefaultsOnly, Category = "Default")
+	FStateTreeStateHandle TargetState;
+	
+	/** Priority of the transition. */
+	UPROPERTY(EditDefaultsOnly, Category = "Default")
+	EStateTreeTransitionPriority Priority = EStateTreeTransitionPriority::Normal;
+};
+
+
+/**
+ * Describes an array of active states in a State Tree.
+ */
+USTRUCT(BlueprintType)
+struct STATETREEMODULE_API FStateTreeActiveStates
+{
+	GENERATED_BODY()
+
+	static constexpr uint8 MaxStates = 8;	// Max number of active states
+
+	FStateTreeActiveStates() = default;
+	
+	explicit FStateTreeActiveStates(const FStateTreeStateHandle StateHandle)
+	{
+		Push(StateHandle);
+	}
+	
+	/** Resets the active state array to empty. */
+	void Reset()
+	{
+		NumStates = 0;
+	}
+
+	/** Pushes new state at the back of the array and returns true if there was enough space. */
+	bool Push(const FStateTreeStateHandle StateHandle)
+	{
+		if ((NumStates + 1) > MaxStates)
+		{
+			return false;
+		}
+		
+		States[NumStates++] = StateHandle;
+		
+		return true;
+	}
+
+	/** Pushes new state at the front of the array and returns true if there was enough space. */
+	bool PushFront(const FStateTreeStateHandle StateHandle)
+	{
+		if ((NumStates + 1) > MaxStates)
+		{
+			return false;
+		}
+
+		NumStates++;
+		for (int32 Index = (int32)NumStates - 1; Index > 0; Index--)
+		{
+			States[Index] = States[Index - 1];
+		}
+		States[0] = StateHandle;
+		
+		return true;
+	}
+
+	/** Pops a state from the back of the array and returns the popped value, or invalid handle if the array was empty. */
+	FStateTreeStateHandle Pop()
+	{
+		if (NumStates == 0)
+		{
+			return FStateTreeStateHandle::Invalid;			
+		}
+
+		const FStateTreeStateHandle Ret = States[NumStates - 1];
+		NumStates--;
+		return Ret;
+	}
+
+	/** Sets the number of states, new states are set to invalid state. */
+	void SetNum(const int32 NewNum)
+	{
+		check(NewNum >= 0 && NewNum <= MaxStates);
+		if (NewNum > (int32)NumStates)
+		{
+			for (int32 Index = NumStates; Index < NewNum; Index++)
+			{
+				States[Index] = FStateTreeStateHandle::Invalid;
+			}
+		}
+		NumStates = static_cast<uint8>(NewNum);
+	}
+
+	/** Returns true of the array contains specified state. */
+	bool Contains(const FStateTreeStateHandle StateHandle) const
+	{
+		for (const FStateTreeStateHandle& Handle : *this)
+		{
+			if (Handle == StateHandle)
+			{
+				return true;
+			}
+		}
+		return false;
+	}
+
+	/** Returns index of a state, searching in reverse order. */
+	int32 IndexOfReverse(const FStateTreeStateHandle StateHandle) const
+	{
+		for (int32 Index = (int32)NumStates - 1; Index >= 0; Index--)
+		{
+			if (States[Index] == StateHandle)
+				return Index;
+		}
+		return INDEX_NONE;
+	}
+	
+	/** Returns last state in the array, or invalid state if the array is empty. */
+	FStateTreeStateHandle Last() const { return NumStates > 0 ? States[NumStates - 1] : FStateTreeStateHandle::Invalid; }
+	
+	/** Returns number of states in the array. */
+	int32 Num() const { return NumStates; }
+
+	/** Returns true if the index is within array bounds. */
+	bool IsValidIndex(const int32 Index) const { return Index >= 0 && Index < (int32)NumStates; }
+	
+	/** Returns true if the array is empty. */
+	bool IsEmpty() const { return NumStates == 0; } 
+
+	/** Returns a specified state in the array. */
+	FORCEINLINE FStateTreeStateHandle operator[](const int32 Index) const
+	{
+		check(Index >= 0 && Index < (int32)NumStates);
+		return States[Index];
+	}
+
+	/** Returns mutable reference to a specified state in the array. */
+	FORCEINLINE FStateTreeStateHandle& operator[](const int32 Index)
+	{
+		check(Index >= 0 && Index < (int32)NumStates);
+		return States[Index];
+	}
+
+	/** Returns a specified state in the array, or FStateTreeStateHandle::Invalid if Index is out of array bounds. */
+	FStateTreeStateHandle GetStateSafe(const int32 Index) const
+	{
+		return (Index >= 0 && Index < (int32)NumStates) ? States[Index] : FStateTreeStateHandle::Invalid;
+	}
+
+	/**
+	 * DO NOT USE DIRECTLY
+	 * STL-like iterators to enable range-based for loop support.
+	 */
+	FORCEINLINE FStateTreeStateHandle* begin() { return &States[0]; }
+	FORCEINLINE FStateTreeStateHandle* end  () { return &States[0] + Num(); }
+	FORCEINLINE const FStateTreeStateHandle* begin() const { return &States[0]; }
+	FORCEINLINE const FStateTreeStateHandle* end  () const { return &States[0] + Num(); }
+
+
+	UPROPERTY(EditDefaultsOnly, Category = Default)
+	FStateTreeStateHandle States[MaxStates];
+
+	UPROPERTY(EditDefaultsOnly, Category = Default)
+	uint8 NumStates = 0;
+};
+
+
+/**
+ * Describes a state tree transition. Source is the state where the transition started, Target describes the state where the transition pointed at,
+ * and Next describes the selected state. The reason Transition and Next are different is that Transition state can be a selector state,
+ * in which case the children will be visited until a leaf state is found, which will be the next state.
+ */
+USTRUCT(BlueprintType)
+struct STATETREEMODULE_API FStateTreeTransitionResult
+{
+	GENERATED_BODY()
+
+	FStateTreeTransitionResult() = default;
+
+	void Reset()
+	{
+		CurrentActiveStates.Reset();
+		CurrentRunStatus = EStateTreeRunStatus::Unset;
+		TargetState = FStateTreeStateHandle::Invalid;
+		NextActiveStates.Reset();
+		CurrentState = FStateTreeStateHandle::Invalid;
+		ChangeType = EStateTreeStateChangeType::Changed;
+		Priority = EStateTreeTransitionPriority::None;
+	}
+	
+	/** Current active states, where the transition started. */
+	UPROPERTY(EditDefaultsOnly, Category = "Default", BlueprintReadOnly)
+	FStateTreeActiveStates CurrentActiveStates;
+
+	/** Current Run status. */
+	UPROPERTY(EditDefaultsOnly, Category = "Default", BlueprintReadOnly)
+	EStateTreeRunStatus CurrentRunStatus = EStateTreeRunStatus::Unset;
+
+	/** Transition source state */
+	UPROPERTY(EditDefaultsOnly, Category = "Default", BlueprintReadOnly)
+	FStateTreeStateHandle SourceState = FStateTreeStateHandle::Invalid;
+
+	/** Transition target state */
+	UPROPERTY(EditDefaultsOnly, Category = "Default", BlueprintReadOnly)
+	FStateTreeStateHandle TargetState = FStateTreeStateHandle::Invalid;
+
+	/** States selected as result of the transition. */
+	UPROPERTY(EditDefaultsOnly, Category = "Default", BlueprintReadOnly)
+	FStateTreeActiveStates NextActiveStates;
+
+	/** The current state being executed. On enter/exit callbacks this is the state of the task. */
+	UPROPERTY(EditDefaultsOnly, Category = "Default", BlueprintReadOnly)
+	FStateTreeStateHandle CurrentState = FStateTreeStateHandle::Invalid;
+
+	/** If the change type is Sustained, then the CurrentState was reselected, or if Changed then the state was just activated. */
+	UPROPERTY(EditDefaultsOnly, Category = "Default", BlueprintReadOnly)
+	EStateTreeStateChangeType ChangeType = EStateTreeStateChangeType::Changed; 
+
+	/** Priority of the transition that caused the state change. */
+	UPROPERTY(EditDefaultsOnly, Category = "Default", BlueprintReadOnly)
+	EStateTreeTransitionPriority Priority = EStateTreeTransitionPriority::None;
+};
+
+
+#if WITH_STATETREE_DEBUGGER
+struct STATETREEMODULE_API FStateTreeInstanceDebugId
+{
+	FStateTreeInstanceDebugId() = default;
+	FStateTreeInstanceDebugId(const uint32 InstanceId, const uint32 SerialNumber)
+		: Id(InstanceId), SerialNumber(SerialNumber)
+	{
+	}
+	
+	bool IsValid() const { return Id != INDEX_NONE && SerialNumber != INDEX_NONE; }
+	bool IsInvalid() const { return !IsValid(); }
+	void Reset() { *this = Invalid; }
+
+	bool operator==(const FStateTreeInstanceDebugId& Other) const
+	{
+		return Id == Other.Id && SerialNumber == Other.SerialNumber;
+	}
+
+	bool operator!=(const FStateTreeInstanceDebugId& Other) const
+	{
+		return !(*this == Other);
+	}
+
+	friend uint32 GetTypeHash(const FStateTreeInstanceDebugId InstanceDebugId)
+	{
+		return HashCombine(InstanceDebugId.Id, InstanceDebugId.SerialNumber);
+	}
+
+	friend FString LexToString(const FStateTreeInstanceDebugId InstanceDebugId)
+	{
+		return FString::Printf(TEXT("0x%x | %d"), InstanceDebugId.Id, InstanceDebugId.SerialNumber);
+	}
+
+	static const FStateTreeInstanceDebugId Invalid;
+	
+	uint32 Id = INDEX_NONE;
+	uint32 SerialNumber = INDEX_NONE;
+};
+#endif // WITH_STATETREE_DEBUGGER
+
+USTRUCT()
+struct STATETREEMODULE_API FStateTreeExecutionState
+{
+	GENERATED_BODY()
+
+	/** @returns Delayed transition state for a specific transition, or nullptr if it does not exists. */
+	FStateTreeTransitionDelayedState* FindDelayedTransition(const FStateTreeIndex16 TransitionIndex)
+	{
+		return DelayedTransitions.FindByPredicate([TransitionIndex](const FStateTreeTransitionDelayedState& State){ return State.TransitionIndex == TransitionIndex; });
+	}
+
+	/** Currently active states */
+	FStateTreeActiveStates ActiveStates;
+
+#if WITH_STATETREE_DEBUGGER
+	/** Id for the active instance used for debugging. */
+	mutable FStateTreeInstanceDebugId InstanceDebugId;
+#endif
+
+	/** Index of the first task struct in the currently initialized instance data. */
+	FStateTreeIndex16 FirstTaskStructIndex = FStateTreeIndex16::Invalid;
+	
+	/** Index of the first task object in the currently initialized instance data. */
+	FStateTreeIndex16 FirstTaskObjectIndex = FStateTreeIndex16::Invalid;
+
+	/** The index of the task that failed during enter state. Exit state uses it to call ExitState() symmetrically. */
+	FStateTreeIndex16 EnterStateFailedTaskIndex = FStateTreeIndex16::Invalid;
+
+	/** Result of last tick */
+	EStateTreeRunStatus LastTickStatus = EStateTreeRunStatus::Failed;
+
+	/** Running status of the instance */
+	EStateTreeRunStatus TreeRunStatus = EStateTreeRunStatus::Unset;
+
+	/** Handle of the state that was first to report state completed (success or failure), used to trigger completion transitions. */
+	FStateTreeStateHandle CompletedStateHandle = FStateTreeStateHandle::Invalid;
+
+	/** Number of times a new state has been changed. */
+	uint16 StateChangeCount = 0;
+
+	/** Running time of the delayed transition */
+	TArray<FStateTreeTransitionDelayedState> DelayedTransitions;
+};
diff --git a/Engine/Plugins/Runtime/StateTree/Source/StateTreeModule/Public/StateTreeInstanceData.h b/Engine/Plugins/Runtime/StateTree/Source/StateTreeModule/Public/StateTreeInstanceData.h
--- a/Engine/Plugins/Runtime/StateTree/Source/StateTreeModule/Public/StateTreeInstanceData.h
+++ b/Engine/Plugins/Runtime/StateTree/Source/StateTreeModule/Public/StateTreeInstanceData.h
@@ -5,6 +5,7 @@
 #include "InstancedStructContainer.h"
 #include "StateTreeEvents.h"
 #include "StateTreeTypes.h"
+#include "StateTreeExecutionTypes.h"
 #include "StateTreeInstanceData.generated.h"
 
 struct FStateTreeTransitionRequest;
diff --git a/Engine/Plugins/Runtime/StateTree/Source/StateTreeModule/Public/StateTreeLinker.h b/Engine/Plugins/Runtime/StateTree/Source/StateTreeModule/Public/StateTreeLinker.h
--- a/Engine/Plugins/Runtime/StateTree/Source/StateTreeModule/Public/StateTreeLinker.h
+++ b/Engine/Plugins/Runtime/StateTree/Source/StateTreeModule/Public/StateTreeLinker.h
@@ -3,7 +3,7 @@
 #pragma once
 
 #include "StateTreeSchema.h"
-#include "StateTreeTypes.h"
+#include "StateTreeExecutionTypes.h"
 #include "StateTreeLinker.generated.h"
 
 UENUM()
diff --git a/Engine/Plugins/Runtime/StateTree/Source/StateTreeModule/Public/StateTreeTaskBase.h b/Engine/Plugins/Runtime/StateTree/Source/StateTreeModule/Public/StateTreeTaskBase.h
--- a/Engine/Plugins/Runtime/StateTree/Source/StateTreeModule/Public/StateTreeTaskBase.h
+++ b/Engine/Plugins/Runtime/StateTree/Source/StateTreeModule/Public/StateTreeTaskBase.h
@@ -3,6 +3,7 @@
 #pragma once
 
 #include "StateTreeNodeBase.h"
+#include "StateTreeExecutionTypes.h"
 #include "StateTreeTaskBase.generated.h"
 
 struct FStateTreeExecutionContext;
diff --git a/Engine/Plugins/Runtime/StateTree/Source/StateTreeModule/Public/StateTreeTypes.h b/Engine/Plugins/Runtime/StateTree/Source/StateTreeModule/Public/StateTreeTypes.h
--- a/Engine/Plugins/Runtime/StateTree/Source/StateTreeModule/Public/StateTreeTypes.h
+++ b/Engine/Plugins/Runtime/StateTree/Source/StateTreeModule/Public/StateTreeTypes.h
@@ -20,40 +20,7 @@ namespace UE::StateTree
 	inline const FName SchemaTag(TEXT("Schema"));
 }; // UE::StateTree
 
-
-/** Status describing current run status of a State Tree. */
-UENUM(BlueprintType)
-enum class EStateTreeRunStatus : uint8
-{
-	/** Tree is still running. */
-	Running,
-	
-	/** Tree execution has stopped on failure. */
-	Failed,
-	
-	/** Tree execution has stopped on success. */
-	Succeeded,
-
-	/** The State Tree was requested to stop without a particular success or failure state. */
-	Stopped,
-
-	/** Status not set. */
-	Unset,
-};
-
-/** State change type. Passed to EnterState() and ExitState() to indicate how the state change affects the state and Evaluator or Task is on. */
-UENUM()
-enum class EStateTreeStateChangeType : uint8
-{
-	/** Not an activation */
-	None,
-	
-	/** The state became activated or deactivated. */
-	Changed,
-	
-	/** The state is parent of new active state and sustained previous active state. */
-    Sustained,
-};
+enum class EStateTreeRunStatus : uint8;
 
 /** Transitions behavior. */
 UENUM()
@@ -91,20 +58,6 @@ enum class EStateTreeConditionOperand : uint8
 	Or,
 };
 
-/** Defines how to assign the result of a condition to evaluate.  */
-UENUM()
-enum class EStateTreeConditionEvaluationMode : uint8
-{
-	/** Condition is evaluated to set the result. This is the normal behavior. */
-	Evaluated,
-	
-	/** Do not evaluate the condition and force result to 'true'. */
-	ForcedTrue,
-	
-	/** Do not evaluate the condition and force result to 'false'. */
-	ForcedFalse,
-};
-
 UENUM()
 enum class EStateTreeStateType : uint8
 {
@@ -224,40 +177,9 @@ struct STATETREEMODULE_API FStateTreeStateHandle
 	bool IsValid() const { return Index != InvalidIndex; }
 	void Invalidate() { Index = InvalidIndex; }
 	bool IsCompletionState() const { return Index == SucceededIndex || Index == FailedIndex || Index == StoppedIndex; }
-	EStateTreeRunStatus ToCompletionStatus() const
-	{
-		if (Index == SucceededIndex)
-		{
-			return EStateTreeRunStatus::Succeeded;
-		}
-		else if (Index == FailedIndex)
-		{
-			return EStateTreeRunStatus::Failed;
-		}
-		else if (Index == StoppedIndex)
-		{
-			return EStateTreeRunStatus::Stopped;
-		}
-		return EStateTreeRunStatus::Unset;
-	}
+	EStateTreeRunStatus ToCompletionStatus() const;
+	static FStateTreeStateHandle FromCompletionStatus(const EStateTreeRunStatus Status);
 
-	static FStateTreeStateHandle FromCompletionStatus(const EStateTreeRunStatus Status)
-	{
-		if (Status == EStateTreeRunStatus::Succeeded)
-		{
-			return FStateTreeStateHandle::Succeeded;
-		}
-		else if (Status == EStateTreeRunStatus::Failed)
-		{
-			return FStateTreeStateHandle::Failed;
-		}
-		else if (Status == EStateTreeRunStatus::Stopped)
-		{
-			return FStateTreeStateHandle::Stopped;
-		}
-		return {};
-	}
-	
 	bool operator==(const FStateTreeStateHandle& RHS) const { return Index == RHS.Index; }
 	bool operator!=(const FStateTreeStateHandle& RHS) const { return Index != RHS.Index; }
 
@@ -278,186 +200,6 @@ struct STATETREEMODULE_API FStateTreeStateHandle
 };
 
 
-/** Transition request */
-USTRUCT(BlueprintType)
-struct STATETREEMODULE_API FStateTreeTransitionRequest
-{
-	GENERATED_BODY()
-
-	FStateTreeTransitionRequest() = default;
-	FStateTreeTransitionRequest(const FStateTreeStateHandle InTargetState, const EStateTreeTransitionPriority InPriority = EStateTreeTransitionPriority::Normal)
-		: TargetState(InTargetState)
-		, Priority(InPriority)
-	{
-	}
-
-	/** Source state of the transition. Filled in by the StateTree execution context. */
-	UPROPERTY(EditDefaultsOnly, Category = "Default")
-	FStateTreeStateHandle SourceState;
-
-	/** Target state of the transition. */
-	UPROPERTY(EditDefaultsOnly, Category = "Default")
-	FStateTreeStateHandle TargetState;
-	
-	/** Priority of the transition. */
-	UPROPERTY(EditDefaultsOnly, Category = "Default")
-	EStateTreeTransitionPriority Priority = EStateTreeTransitionPriority::Normal;
-};
-
-/**
- * Describes an array of active states in a State Tree.
- */
-USTRUCT(BlueprintType)
-struct STATETREEMODULE_API FStateTreeActiveStates
-{
-	GENERATED_BODY()
-
-	static constexpr uint8 MaxStates = 8;	// Max number of active states
-
-	FStateTreeActiveStates() = default;
-	
-	explicit FStateTreeActiveStates(const FStateTreeStateHandle StateHandle)
-	{
-		Push(StateHandle);
-	}
-	
-	/** Resets the active state array to empty. */
-	void Reset()
-	{
-		NumStates = 0;
-	}
-
-	/** Pushes new state at the back of the array and returns true if there was enough space. */
-	bool Push(const FStateTreeStateHandle StateHandle)
-	{
-		if ((NumStates + 1) > MaxStates)
-		{
-			return false;
-		}
-		
-		States[NumStates++] = StateHandle;
-		
-		return true;
-	}
-
-	/** Pushes new state at the front of the array and returns true if there was enough space. */
-	bool PushFront(const FStateTreeStateHandle StateHandle)
-	{
-		if ((NumStates + 1) > MaxStates)
-		{
-			return false;
-		}
-
-		NumStates++;
-		for (int32 Index = (int32)NumStates - 1; Index > 0; Index--)
-		{
-			States[Index] = States[Index - 1];
-		}
-		States[0] = StateHandle;
-		
-		return true;
-	}
-
-	/** Pops a state from the back of the array and returns the popped value, or invalid handle if the array was empty. */
-	FStateTreeStateHandle Pop()
-	{
-		if (NumStates == 0)
-		{
-			return FStateTreeStateHandle::Invalid;			
-		}
-
-		const FStateTreeStateHandle Ret = States[NumStates - 1];
-		NumStates--;
-		return Ret;
-	}
-
-	/** Sets the number of states, new states are set to invalid state. */
-	void SetNum(const int32 NewNum)
-	{
-		check(NewNum >= 0 && NewNum <= MaxStates);
-		if (NewNum > (int32)NumStates)
-		{
-			for (int32 Index = NumStates; Index < NewNum; Index++)
-			{
-				States[Index] = FStateTreeStateHandle::Invalid;
-			}
-		}
-		NumStates = static_cast<uint8>(NewNum);
-	}
-
-	/** Returns true of the array contains specified state. */
-	bool Contains(const FStateTreeStateHandle StateHandle) const
-	{
-		for (const FStateTreeStateHandle& Handle : *this)
-		{
-			if (Handle == StateHandle)
-			{
-				return true;
-			}
-		}
-		return false;
-	}
-
-	/** Returns index of a state, searching in reverse order. */
-	int32 IndexOfReverse(const FStateTreeStateHandle StateHandle) const
-	{
-		for (int32 Index = (int32)NumStates - 1; Index >= 0; Index--)
-		{
-			if (States[Index] == StateHandle)
-				return Index;
-		}
-		return INDEX_NONE;
-	}
-	
-	/** Returns last state in the array, or invalid state if the array is empty. */
-	FStateTreeStateHandle Last() const { return NumStates > 0 ? States[NumStates - 1] : FStateTreeStateHandle::Invalid; }
-	
-	/** Returns number of states in the array. */
-	int32 Num() const { return NumStates; }
-
-	/** Returns true if the index is within array bounds. */
-	bool IsValidIndex(const int32 Index) const { return Index >= 0 && Index < (int32)NumStates; }
-	
-	/** Returns true if the array is empty. */
-	bool IsEmpty() const { return NumStates == 0; } 
-
-	/** Returns a specified state in the array. */
-	FORCEINLINE FStateTreeStateHandle operator[](const int32 Index) const
-	{
-		check(Index >= 0 && Index < (int32)NumStates);
-		return States[Index];
-	}
-
-	/** Returns mutable reference to a specified state in the array. */
-	FORCEINLINE FStateTreeStateHandle& operator[](const int32 Index)
-	{
-		check(Index >= 0 && Index < (int32)NumStates);
-		return States[Index];
-	}
-
-	/** Returns a specified state in the array, or FStateTreeStateHandle::Invalid if Index is out of array bounds. */
-	FStateTreeStateHandle GetStateSafe(const int32 Index) const
-	{
-		return (Index >= 0 && Index < (int32)NumStates) ? States[Index] : FStateTreeStateHandle::Invalid;
-	}
-
-	/**
-	 * DO NOT USE DIRECTLY
-	 * STL-like iterators to enable range-based for loop support.
-	 */
-	FORCEINLINE FStateTreeStateHandle* begin() { return &States[0]; }
-	FORCEINLINE FStateTreeStateHandle* end  () { return &States[0] + Num(); }
-	FORCEINLINE const FStateTreeStateHandle* begin() const { return &States[0]; }
-	FORCEINLINE const FStateTreeStateHandle* end  () const { return &States[0] + Num(); }
-
-
-	UPROPERTY(EditDefaultsOnly, Category = Default)
-	FStateTreeStateHandle States[MaxStates];
-
-	UPROPERTY(EditDefaultsOnly, Category = Default)
-	uint8 NumStates = 0;
-};
-
 /**
  * Time duration with random variance. Stored compactly as two uint16s, which gives time range of about 650 seconds.
  * The variance is symmetric (+-) around the specified duration.
@@ -520,62 +262,6 @@ protected:
 	uint16 RandomVariance = 0;
 };
 
-/**
- * Describes a state tree transition. Source is the state where the transition started, Target describes the state where the transition pointed at,
- * and Next describes the selected state. The reason Transition and Next are different is that Transition state can be a selector state,
- * in which case the children will be visited until a leaf state is found, which will be the next state.
- */
-USTRUCT(BlueprintType)
-struct STATETREEMODULE_API FStateTreeTransitionResult
-{
-	GENERATED_BODY()
-
-	FStateTreeTransitionResult() = default;
-
-	void Reset()
-	{
-		CurrentActiveStates.Reset();
-		CurrentRunStatus = EStateTreeRunStatus::Unset;
-		TargetState = FStateTreeStateHandle::Invalid;
-		NextActiveStates.Reset();
-		CurrentState = FStateTreeStateHandle::Invalid;
-		ChangeType = EStateTreeStateChangeType::Changed;
-		Priority = EStateTreeTransitionPriority::None;
-	}
-	
-	/** Current active states, where the transition started. */
-	UPROPERTY(EditDefaultsOnly, Category = "Default", BlueprintReadOnly)
-	FStateTreeActiveStates CurrentActiveStates;
-
-	/** Current Run status. */
-	UPROPERTY(EditDefaultsOnly, Category = "Default", BlueprintReadOnly)
-	EStateTreeRunStatus CurrentRunStatus = EStateTreeRunStatus::Unset;
-
-	/** Transition source state */
-	UPROPERTY(EditDefaultsOnly, Category = "Default", BlueprintReadOnly)
-	FStateTreeStateHandle SourceState = FStateTreeStateHandle::Invalid;
-
-	/** Transition target state */
-	UPROPERTY(EditDefaultsOnly, Category = "Default", BlueprintReadOnly)
-	FStateTreeStateHandle TargetState = FStateTreeStateHandle::Invalid;
-
-	/** States selected as result of the transition. */
-	UPROPERTY(EditDefaultsOnly, Category = "Default", BlueprintReadOnly)
-	FStateTreeActiveStates NextActiveStates;
-
-	/** The current state being executed. On enter/exit callbacks this is the state of the task. */
-	UPROPERTY(EditDefaultsOnly, Category = "Default", BlueprintReadOnly)
-	FStateTreeStateHandle CurrentState = FStateTreeStateHandle::Invalid;
-
-	/** If the change type is Sustained, then the CurrentState was reselected, or if Changed then the state was just activated. */
-	UPROPERTY(EditDefaultsOnly, Category = "Default", BlueprintReadOnly)
-	EStateTreeStateChangeType ChangeType = EStateTreeStateChangeType::Changed; 
-
-	/** Priority of the transition that caused the state change. */
-	UPROPERTY(EditDefaultsOnly, Category = "Default", BlueprintReadOnly)
-	EStateTreeTransitionPriority Priority = EStateTreeTransitionPriority::None;
-};
-
 
 /**
  *  Runtime representation of a StateTree transition.
@@ -745,55 +431,6 @@ enum class EStateTreeExternalDataRequirement : uint8
 	Optional,	// Data is optional for StateTree execution.
 };
 
-/**
- * Handle to access an external struct or object.
- * Note: Use the templated version below. 
- */
-USTRUCT()
-struct STATETREEMODULE_API FStateTreeExternalDataHandle
-{
-	GENERATED_BODY()
-
-	static const FStateTreeExternalDataHandle Invalid;
-	
-	static bool IsValidIndex(const int32 Index) { return FStateTreeIndex16::IsValidIndex(Index); }
-	bool IsValid() const { return DataViewIndex.IsValid(); }
-
-	UPROPERTY()
-	FStateTreeIndex16 DataViewIndex = FStateTreeIndex16::Invalid;
-};
-
-/**
- * Handle to access an external struct or object.
- * This reference handle can be used in StateTree tasks and evaluators to have quick access to external data.
- * The type provided to the template is used by the linker and context to pass along the type.
- *
- * USTRUCT()
- * struct FExampleTask : public FStateTreeTaskBase
- * {
- *    ...
- *
- *    bool Link(FStateTreeLinker& Linker)
- *    {
- *      Linker.LinkExternalData(ExampleSubsystemHandle);
- *      return true;
- *    }
- * 
- *    EStateTreeRunStatus EnterState(FStateTreeExecutionContext& Context, const EStateTreeStateChangeType ChangeType, const FStateTreeTransitionResult& Transition)
- *    {
- *      const UExampleSubsystem& ExampleSubsystem = Context.GetExternalData(ExampleSubsystemHandle);
- *      ...
- *    }
- *
- *    TStateTreeExternalDataHandle<UExampleSubsystem> ExampleSubsystemHandle;
- * }
- */
-template<typename T, EStateTreeExternalDataRequirement Req = EStateTreeExternalDataRequirement::Required>
-struct TStateTreeExternalDataHandle : FStateTreeExternalDataHandle
-{
-	typedef T DataType;
-	static constexpr EStateTreeExternalDataRequirement DataRequirement = Req;
-};
 
 UENUM()
 enum class EStateTreePropertyUsage : uint8
@@ -871,58 +508,6 @@ struct STATETREEMODULE_API FStateTreeTransitionIdToIndex
 	FStateTreeIndex16 Index;
 };
 
-/**
- * Describes an external data. The data can point to a struct or object.
- * The code that handles StateTree ticking is responsible for passing in the actually data, see FStateTreeExecutionContext.
- */
-USTRUCT()
-struct STATETREEMODULE_API FStateTreeExternalDataDesc
-{
-	GENERATED_BODY()
-
-	FStateTreeExternalDataDesc() = default;
-	FStateTreeExternalDataDesc(const UStruct* InStruct, const EStateTreeExternalDataRequirement InRequirement) : Struct(InStruct), Requirement(InRequirement) {}
-
-	FStateTreeExternalDataDesc(const FName InName, const UStruct* InStruct, const FGuid InGuid)
-		: Struct(InStruct)
-		, Name(InName)
-#if WITH_EDITORONLY_DATA
-		, ID(InGuid)
-#endif
-	{}
-
-	bool operator==(const FStateTreeExternalDataDesc& Other) const
-	{
-		return Struct == Other.Struct && Requirement == Other.Requirement;
-	}
-	
-	/** Class or struct of the external data. */
-	UPROPERTY();
-	TObjectPtr<const UStruct> Struct = nullptr;
-
-	/**
-	 * Name of the external data. Used only for bindable external data (enforced by the schema).
-	 * External data linked explicitly by the nodes (i.e. LinkExternalData) are identified only
-	 * by their type since they are used for unique instance of a given type.  
-	 */
-	UPROPERTY(VisibleAnywhere, Category = Common)
-	FName Name;
-	
-	/** Handle/Index to the StateTreeExecutionContext data views array */
-	UPROPERTY();
-	FStateTreeExternalDataHandle Handle;
-
-	/** Describes if the data is required or not. */
-	UPROPERTY();
-	EStateTreeExternalDataRequirement Requirement = EStateTreeExternalDataRequirement::Required;
-
-#if WITH_EDITORONLY_DATA
-	/** Unique identifier. Used only for bindable external data. */
-	UPROPERTY()
-	FGuid ID;
-#endif
-};
-
 
 /**
  * StateTree struct ref allows to get a reference/pointer to a specified type via property binding.
@@ -1144,89 +729,6 @@ struct STATETREEMODULE_API FStateTreeTransitionDelayedState
 	float TimeLeft = 0.0f;
 };
 
-#if WITH_STATETREE_DEBUGGER
-struct STATETREEMODULE_API FStateTreeInstanceDebugId
-{
-	FStateTreeInstanceDebugId() = default;
-	FStateTreeInstanceDebugId(const uint32 InstanceId, const uint32 SerialNumber)
-		: Id(InstanceId), SerialNumber(SerialNumber)
-	{
-	}
-	
-	bool IsValid() const { return Id != INDEX_NONE && SerialNumber != INDEX_NONE; }
-	bool IsInvalid() const { return !IsValid(); }
-	void Reset() { *this = Invalid; }
-
-	bool operator==(const FStateTreeInstanceDebugId& Other) const
-	{
-		return Id == Other.Id && SerialNumber == Other.SerialNumber;
-	}
-
-	bool operator!=(const FStateTreeInstanceDebugId& Other) const
-	{
-		return !(*this == Other);
-	}
-
-	friend uint32 GetTypeHash(const FStateTreeInstanceDebugId InstanceDebugId)
-	{
-		return HashCombine(InstanceDebugId.Id, InstanceDebugId.SerialNumber);
-	}
-
-	friend FString LexToString(const FStateTreeInstanceDebugId InstanceDebugId)
-	{
-		return FString::Printf(TEXT("0x%x | %d"), InstanceDebugId.Id, InstanceDebugId.SerialNumber);
-	}
-
-	static const FStateTreeInstanceDebugId Invalid;
-	
-	uint32 Id = INDEX_NONE;
-	uint32 SerialNumber = INDEX_NONE;
-};
-#endif // WITH_STATETREE_DEBUGGER
-
-USTRUCT()
-struct STATETREEMODULE_API FStateTreeExecutionState
-{
-	GENERATED_BODY()
-
-	/** @returns Delayed transition state for a specific transition, or nullptr if it does not exists. */
-	FStateTreeTransitionDelayedState* FindDelayedTransition(const FStateTreeIndex16 TransitionIndex)
-	{
-		return DelayedTransitions.FindByPredicate([TransitionIndex](const FStateTreeTransitionDelayedState& State){ return State.TransitionIndex == TransitionIndex; });
-	}
-
-	/** Currently active states */
-	FStateTreeActiveStates ActiveStates;
-
-#if WITH_STATETREE_DEBUGGER
-	/** Id for the active instance used for debugging. */
-	mutable FStateTreeInstanceDebugId InstanceDebugId;
-#endif
-
-	/** Index of the first task struct in the currently initialized instance data. */
-	FStateTreeIndex16 FirstTaskStructIndex = FStateTreeIndex16::Invalid;
-	
-	/** Index of the first task object in the currently initialized instance data. */
-	FStateTreeIndex16 FirstTaskObjectIndex = FStateTreeIndex16::Invalid;
-
-	/** The index of the task that failed during enter state. Exit state uses it to call ExitState() symmetrically. */
-	FStateTreeIndex16 EnterStateFailedTaskIndex = FStateTreeIndex16::Invalid;
-
-	/** Result of last tick */
-	EStateTreeRunStatus LastTickStatus = EStateTreeRunStatus::Failed;
-
-	/** Running status of the instance */
-	EStateTreeRunStatus TreeRunStatus = EStateTreeRunStatus::Unset;
-
-	/** Handle of the state that was first to report state completed (success or failure), used to trigger completion transitions. */
-	FStateTreeStateHandle CompletedStateHandle = FStateTreeStateHandle::Invalid;
-
-	/** Number of times a new state has been changed. */
-	uint16 StateChangeCount = 0;
-
-	/** Running time of the delayed transition */
-	TArray<FStateTreeTransitionDelayedState> DelayedTransitions;
-};
 
 /**
  * Link to another state in StateTree
