Fixed CIS Issue 326584

diff --git a/Engine/Plugins/Runtime/StateTree/Source/StateTreeEditorModule/Private/Debugger/StateTreeDebuggerTrack.h b/Engine/Plugins/Runtime/StateTree/Source/StateTreeEditorModule/Private/Debugger/StateTreeDebuggerTrack.h
--- a/Engine/Plugins/Runtime/StateTree/Source/StateTreeEditorModule/Private/Debugger/StateTreeDebuggerTrack.h
+++ b/Engine/Plugins/Runtime/StateTree/Source/StateTreeEditorModule/Private/Debugger/StateTreeDebuggerTrack.h
@@ -6,7 +6,7 @@
 
 #include "RewindDebuggerTrack.h"
 #include "SStateTreeDebuggerEventTimelineView.h"
-#include "StateTreeTypes.h"
+#include "StateTreeExecutionTypes.h"
 
 struct FStateTreeDebugger;
 
