Fixed duplicated ID in statetree asset
#changelist validated

diff --git a/Engine/Plugins/Runtime/StateTree/Source/StateTreeEditorModule/Private/StateTreePropertyBindingCompiler.cpp b/Engine/Plugins/Runtime/StateTree/Source/StateTreeEditorModule/Private/StateTreePropertyBindingCompiler.cpp
--- a/Engine/Plugins/Runtime/StateTree/Source/StateTreeEditorModule/Private/StateTreePropertyBindingCompiler.cpp
+++ b/Engine/Plugins/Runtime/StateTree/Source/StateTreeEditorModule/Private/StateTreePropertyBindingCompiler.cpp
@@ -146,8 +146,8 @@ int32 FStateTreePropertyBindingCompiler::AddSourceStruct(const FStateTreeBindabl
 	const FStateTreeBindableStructDesc* ExistingStruct = SourceStructs.FindByPredicate([&SourceStruct](const FStateTreeBindableStructDesc& Struct) { return (Struct.ID == SourceStruct.ID); });
 	if (ExistingStruct)
 	{
-		UE_LOG(LogStateTree, Error, TEXT("%s already exists as %s"),
-			*SourceStruct.ToString(), *ExistingStruct->ToString());
+		UE_LOG(LogStateTree, Error, TEXT("%s already exists as %s using ID '%s'"),
+			*SourceStruct.ToString(), *ExistingStruct->ToString(), *ExistingStruct->ID.ToString());
 	}
 	
 	SourceStructs.Add(SourceStruct);
