[StateTreeDebugger]
- On breakpoint hit the tree view will select and focus on the associated state and event treeview will select the event that triggered the breakpoint
- added missing trace for exiting tasks
- scrub to the record time when breakpoint is hit so events tree view is synchronized to the right time (not the latest processed time)

diff --git a/Engine/Plugins/Runtime/StateTree/Source/StateTreeEditorModule/Private/Debugger/SStateTreeDebuggerView.cpp b/Engine/Plugins/Runtime/StateTree/Source/StateTreeEditorModule/Private/Debugger/SStateTreeDebuggerView.cpp
--- a/Engine/Plugins/Runtime/StateTree/Source/StateTreeEditorModule/Private/Debugger/SStateTreeDebuggerView.cpp
+++ b/Engine/Plugins/Runtime/StateTree/Source/StateTreeEditorModule/Private/Debugger/SStateTreeDebuggerView.cpp
@@ -11,7 +11,6 @@
 #include "IStructureDetailsView.h"
 #include "Kismet2/DebuggerCommands.h"
 #include "Modules/ModuleManager.h"
-#include "ProfilingDebugging/TraceAuxiliary.h"
 #include "PropertyEditorModule.h"
 #include "ScopedTransaction.h"
 #include "SStateTreeDebuggerInstanceTree.h"
@@ -72,6 +71,55 @@ struct FEventTreeElement : TSharedFromThis<FEventTreeElement>
 	FStateTreeTraceEventVariantType Event;
 	TArray<TSharedPtr<FEventTreeElement>> Children;
 };
+
+
+/**
+ * Iterates over all tree elements for the frame events
+ * @param Elements Container of hierarchical tree element to visit
+ * @param InFunc function called at each element, should return true if visiting is continued or false to stop.
+ */
+void VisitEventTreeElements(const TConstArrayView<TSharedPtr<FEventTreeElement>> Elements, TFunctionRef<bool(TSharedPtr<FEventTreeElement>& VisitedElement)> InFunc)
+{
+	TArray<TSharedPtr<FEventTreeElement>> Stack;
+	bool bContinue = true;
+
+	for (const TSharedPtr<FEventTreeElement>& RootElement : Elements)
+	{
+		if (RootElement == nullptr)
+		{
+			continue;
+		}
+
+		Stack.Add(RootElement);
+
+		while (!Stack.IsEmpty() && bContinue)
+		{
+			TSharedPtr<FEventTreeElement> StackedElement = Stack[0];
+			check(StackedElement);
+
+			Stack.RemoveAt(0);
+
+			bContinue = InFunc(StackedElement);
+
+			if (bContinue)
+			{
+				for (const TSharedPtr<FEventTreeElement>& Child : StackedElement->Children)
+				{
+					if (Child.IsValid())
+					{
+						Stack.Add(Child);
+					}
+				}
+			}
+		}
+
+		if (!bContinue)
+		{
+			break;
+		}
+	}
+}
+
 } // UE::StateTreeDebugger
 
 
@@ -898,6 +946,60 @@ void SStateTreeDebuggerView::HandleEnableStateBreakpoint(EStateTreeBreakpointTyp
 	}
 }
 
+UStateTreeState* SStateTreeDebuggerView::FindStateAssociatedToBreakpoint(FStateTreeDebuggerBreakpoint Breakpoint) const
+{
+	const UStateTree* Tree = StateTree.Get();
+	UStateTreeEditorData* TreeEditorData = StateTreeEditorData.Get();
+	if (Tree == nullptr || TreeEditorData == nullptr)
+	{
+		return nullptr;
+	}
+
+	UStateTreeState* StateTreeState = nullptr;
+
+	if (const FStateTreeStateHandle* StateHandle = Breakpoint.ElementIdentifier.TryGet<FStateTreeStateHandle>())
+	{
+		const FGuid StateId = StateTree->GetStateIdFromHandle(*StateHandle);
+		StateTreeState = TreeEditorData->GetMutableStateByID(StateId);
+	}
+	else if (const FStateTreeDebuggerBreakpoint::FStateTreeTaskIndex* TaskIndex = Breakpoint.ElementIdentifier.TryGet<FStateTreeDebuggerBreakpoint::FStateTreeTaskIndex>())
+	{
+		const FGuid TaskId = StateTree->GetNodeIdFromIndex(TaskIndex->Index);
+
+		TreeEditorData->VisitHierarchy([&TaskId, &StateTreeState](UStateTreeState& State, UStateTreeState* /*ParentState*/)
+			{
+				for (const FStateTreeEditorNode& EditorNode : State.Tasks)
+				{
+					if (EditorNode.ID == TaskId)
+					{
+						StateTreeState = &State;
+						return EStateTreeVisitor::Break;
+					}
+				}
+				return EStateTreeVisitor::Continue;
+			});
+	}
+	else if (const FStateTreeDebuggerBreakpoint::FStateTreeTransitionIndex* TransitionIndex = Breakpoint.ElementIdentifier.TryGet<FStateTreeDebuggerBreakpoint::FStateTreeTransitionIndex>())
+	{
+		const FGuid TransitionId = StateTree->GetTransitionIdFromIndex(TransitionIndex->Index);
+
+		TreeEditorData->VisitHierarchy([&TransitionId, &StateTreeState](UStateTreeState& State, UStateTreeState* /*ParentState*/)
+			{
+				for (const FStateTreeTransition& StateTransition : State.Transitions)
+				{
+					if (StateTransition.ID == TransitionId)
+					{
+						StateTreeState = &State;
+						return EStateTreeVisitor::Break;
+					}
+				}
+				return EStateTreeVisitor::Continue;
+			});
+	}
+
+	return StateTreeState;
+}
+
 void SStateTreeDebuggerView::OnTimeLineScrubPositionChanged(double Time, bool bIsScrubbing)
 {
 	bAutoScroll = false;
@@ -1007,6 +1109,7 @@ void SStateTreeDebuggerView::ExpandAll(const TArray<TSharedPtr<UE::StateTreeDebu
 
 void SStateTreeDebuggerView::OnBreakpointHit(const FStateTreeInstanceDebugId InstanceId, const FStateTreeDebuggerBreakpoint Breakpoint, const TSharedRef<FUICommandList> ActionList) const
 {
+	// Pause PIE session if possible
 	if (FPlayWorldCommands::Get().PausePlaySession.IsValid())
 	{
 		if (ActionList->CanExecuteAction(FPlayWorldCommands::Get().PausePlaySession.ToSharedRef()))
@@ -1014,6 +1117,35 @@ void SStateTreeDebuggerView::OnBreakpointHit(const FStateTreeInstanceDebugId Ins
 			ActionList->ExecuteAction(FPlayWorldCommands::Get().PausePlaySession.ToSharedRef());
 		}
 	}
+
+	// Extract associated UStateTreeState to focus on it.
+	if (UStateTreeState* AssociatedState = FindStateAssociatedToBreakpoint(Breakpoint))
+	{
+		FStateTreeViewModel* ViewModel = StateTreeViewModel.Get();
+		check(ViewModel);
+		ViewModel->SetSelection(AssociatedState);
+	}
+
+	// Find matching event in the tree view and select it
+	if (EventsTreeView.IsValid())
+	{
+		TSharedPtr<UE::StateTreeDebugger::FEventTreeElement> MatchingElement = nullptr;
+		VisitEventTreeElements(EventsTreeElements, [&MatchingElement, Breakpoint](const TSharedPtr<UE::StateTreeDebugger::FEventTreeElement>& VisitedElement)
+			{
+				if (Breakpoint.IsMatchingEvent(VisitedElement->Event))
+				{
+					MatchingElement = VisitedElement;
+				}
+
+				// Continue visit until we find a matching event
+				return MatchingElement.IsValid() == false;
+			});
+
+		if (MatchingElement.IsValid())
+		{
+			EventsTreeView->SetSelection(MatchingElement);
+		}
+	}
 }
 
 void SStateTreeDebuggerView::OnNewSession()
@@ -1021,7 +1153,7 @@ void SStateTreeDebuggerView::OnNewSession()
 	// Clear tracks
 	InstanceOwnerTracks.Reset();
 
-	// Refresh tree viewa
+	// Refresh tree view
 	if (InstancesTreeView)
 	{
 		InstancesTreeView->Refresh();
diff --git a/Engine/Plugins/Runtime/StateTree/Source/StateTreeEditorModule/Private/Debugger/SStateTreeDebuggerView.h b/Engine/Plugins/Runtime/StateTree/Source/StateTreeEditorModule/Private/Debugger/SStateTreeDebuggerView.h
--- a/Engine/Plugins/Runtime/StateTree/Source/StateTreeEditorModule/Private/Debugger/SStateTreeDebuggerView.h
+++ b/Engine/Plugins/Runtime/StateTree/Source/StateTreeEditorModule/Private/Debugger/SStateTreeDebuggerView.h
@@ -94,7 +94,18 @@ private:
 	bool CanRemoveStateBreakpoint(EStateTreeBreakpointType Type) const;
 	ECheckBoxState GetStateBreakpointCheckState(EStateTreeBreakpointType Type) const;
 	void HandleEnableStateBreakpoint(EStateTreeBreakpointType Type);
-	
+
+	UStateTreeState* FindStateAssociatedToBreakpoint(FStateTreeDebuggerBreakpoint Breakpoint) const;
+
+	/** Callback from timeline widgets to update the debugger scrub state. */
+	void OnTimeLineScrubPositionChanged(double Time, bool bIsScrubbing);
+
+	/** Callback used to reflect debugger scrub state in the UI. */
+	void OnDebuggerScrubStateChanged(const UE::StateTreeDebugger::FScrubState& ScrubState);
+
+	/** Recursively sets tree items as expanded. */
+	void ExpandAll(const TArray<TSharedPtr<UE::StateTreeDebugger::FEventTreeElement>>& Items);
+
 	TSharedPtr<FStateTreeDebugger> Debugger;
 	TSharedPtr<FStateTreeViewModel> StateTreeViewModel;
 	TWeakObjectPtr<const UStateTree> StateTree;
@@ -135,15 +146,6 @@ private:
 	/** Attribute provided by the debugger scrub position to control cursor and timelines positions. */
 	TAttribute<double> ScrubTimeAttribute;
 
-	/** Callback from timeline widgets to update the debugger scrub state. */
-	void OnTimeLineScrubPositionChanged(double Time, bool bIsScrubbing);
-
-	/** Callback used to reflect debugger scrub state in the UI. */
-	void OnDebuggerScrubStateChanged(const UE::StateTreeDebugger::FScrubState& ScrubState);
-
-	/** Recursively sets tree items as expanded. */
-	void ExpandAll(const TArray<TSharedPtr<UE::StateTreeDebugger::FEventTreeElement>>& Items);
-
 	/** Range controlled by the timeline widgets and used to adjust cursor position and track content. */
 	TRange<double> ViewRange = TRange<double>(0, 10);
 
diff --git a/Engine/Plugins/Runtime/StateTree/Source/StateTreeEditorModule/Private/SStateTreeView.cpp b/Engine/Plugins/Runtime/StateTree/Source/StateTreeEditorModule/Private/SStateTreeView.cpp
--- a/Engine/Plugins/Runtime/StateTree/Source/StateTreeEditorModule/Private/SStateTreeView.cpp
+++ b/Engine/Plugins/Runtime/StateTree/Source/StateTreeEditorModule/Private/SStateTreeView.cpp
@@ -379,7 +379,12 @@ void SStateTreeView::HandleModelSelectionChanged(const TArray<TWeakObjectPtr<USt
 
 	if (SelectedStates.Num() > 0)
 	{
-		TreeView->SetItemSelection(SelectedStates, true);
+		TreeView->SetItemSelection(SelectedStates, /*bSelected*/true);
+
+		if (SelectedStates.Num() == 1)
+		{
+			TreeView->RequestScrollIntoView(SelectedStates[0]);	
+		}
 	}
 }
 
diff --git a/Engine/Plugins/Runtime/StateTree/Source/StateTreeEditorModule/Private/StateTreeEditorData.cpp b/Engine/Plugins/Runtime/StateTree/Source/StateTreeEditorModule/Private/StateTreeEditorData.cpp
--- a/Engine/Plugins/Runtime/StateTree/Source/StateTreeEditorModule/Private/StateTreeEditorData.cpp
+++ b/Engine/Plugins/Runtime/StateTree/Source/StateTreeEditorModule/Private/StateTreeEditorData.cpp
@@ -390,9 +390,14 @@ const UStateTreeState* UStateTreeEditorData::GetStateByStructID(const FGuid Targ
 
 const UStateTreeState* UStateTreeEditorData::GetStateByID(const FGuid StateID) const
 {
-	const UStateTreeState* Result = nullptr;
+	return const_cast<UStateTreeEditorData*>(this)->GetMutableStateByID(StateID);
+}
+
+UStateTreeState* UStateTreeEditorData::GetMutableStateByID(const FGuid StateID)
+{
+	UStateTreeState* Result = nullptr;
 	
-	VisitHierarchy([&Result, &StateID](const UStateTreeState& State, UStateTreeState* /*ParentState*/)
+	VisitHierarchy([&Result, &StateID](UStateTreeState& State, UStateTreeState* /*ParentState*/)
 	{
 		if (State.ID == StateID)
 		{
diff --git a/Engine/Plugins/Runtime/StateTree/Source/StateTreeEditorModule/Public/StateTreeEditorData.h b/Engine/Plugins/Runtime/StateTree/Source/StateTreeEditorModule/Public/StateTreeEditorData.h
--- a/Engine/Plugins/Runtime/StateTree/Source/StateTreeEditorModule/Public/StateTreeEditorData.h
+++ b/Engine/Plugins/Runtime/StateTree/Source/StateTreeEditorModule/Public/StateTreeEditorData.h
@@ -73,6 +73,9 @@ public:
 	/** @returns state based on its ID, or nullptr if not found. */
 	const UStateTreeState* GetStateByID(const FGuid StateID) const;
 
+	/** @returns mutable state based on its ID, or nullptr if not found. */
+	UStateTreeState* GetMutableStateByID(const FGuid StateID);
+
 	/** @returns the IDs and instance values of all bindable structs in the StateTree. */
 	void GetAllStructValues(TMap<FGuid, const FStateTreeDataView>& AllValues) const;
 
diff --git a/Engine/Plugins/Runtime/StateTree/Source/StateTreeModule/Private/Debugger/StateTreeDebugger.cpp b/Engine/Plugins/Runtime/StateTree/Source/StateTreeModule/Private/Debugger/StateTreeDebugger.cpp
--- a/Engine/Plugins/Runtime/StateTree/Source/StateTreeModule/Private/Debugger/StateTreeDebugger.cpp
+++ b/Engine/Plugins/Runtime/StateTree/Source/StateTreeModule/Private/Debugger/StateTreeDebugger.cpp
@@ -693,25 +693,24 @@ void FStateTreeDebugger::SendNotifications()
 		NewInstances.Reset();
 	}
 
-	if (HitBreakpointIndex != INDEX_NONE)
+	if (HitBreakpoint.IsSet())
 	{
-		check(HitBreakpointInstanceId.IsValid());
-		check(Breakpoints.IsValidIndex(HitBreakpointIndex));
+		check(HitBreakpoint.InstanceId.IsValid());
+		check(Breakpoints.IsValidIndex(HitBreakpoint.Index));
 
 		// Force scrub time to latest simulation time to reflect most recent events.
 		// This will notify scrub position changed and active states
-		SetScrubTime(RecordingDuration);
+		SetScrubTime(HitBreakpoint.Time);
 
 		// Make sure the instance is selected in case the breakpoint was set for any instances 
-		if (SelectedInstanceId != HitBreakpointInstanceId)
+		if (SelectedInstanceId != HitBreakpoint.InstanceId)
 		{
-			SelectInstance(HitBreakpointInstanceId);
+			SelectInstance(HitBreakpoint.InstanceId);
 		}
 
-		OnBreakpointHit.ExecuteIfBound(HitBreakpointInstanceId, Breakpoints[HitBreakpointIndex]);
+		OnBreakpointHit.ExecuteIfBound(HitBreakpoint.InstanceId, Breakpoints[HitBreakpoint.Index]);
 
-		HitBreakpointIndex = INDEX_NONE;
-		HitBreakpointInstanceId.Reset();
+		HitBreakpoint.Reset();
 
 		PauseSessionAnalysis();
 	}
@@ -738,65 +737,29 @@ void FStateTreeDebugger::ReadTrace(
 bool FStateTreeDebugger::EvaluateBreakpoints(const FStateTreeInstanceDebugId InstanceId, const FStateTreeTraceEventVariantType& Event)
 {
 	if (StateTreeAsset == nullptr // asset is required to properly match state handles
-		|| HitBreakpointIndex != INDEX_NONE // Only consider first hit breakpoint in the frame
+		|| HitBreakpoint.IsSet() // Only consider first hit breakpoint in the frame
 		|| Breakpoints.IsEmpty()
 		|| (SelectedInstanceId.IsValid() && InstanceId != SelectedInstanceId)) // ignore events not for the selected instances		
 	{
 		return false;
 	}
-	
-	EStateTreeTraceEventType EventType = EStateTreeTraceEventType::Unset;
-	
-	Visit([&EventType](auto& TypedEvent) { EventType = TypedEvent.EventType; }, Event); 
 
 	for (int BreakpointIndex = 0; BreakpointIndex < Breakpoints.Num(); ++BreakpointIndex)
 	{
 		const FStateTreeDebuggerBreakpoint Breakpoint = Breakpoints[BreakpointIndex];
-		if (Breakpoint.EventType == EventType)
+		if (Breakpoint.IsMatchingEvent(Event))
 		{
-			if (const FStateTreeStateHandle* StateHandle = Breakpoint.ElementIdentifier.TryGet<FStateTreeStateHandle>())
-			{
-				if (const FStateTreeTraceStateEvent* StateEvent = Event.TryGet<FStateTreeTraceStateEvent>())
-				{
-					if (StateEvent->GetStateHandle() == *StateHandle)
-					{
-						HitBreakpointIndex = BreakpointIndex;
-						HitBreakpointInstanceId = InstanceId;
-					}
-				}
-			}
-			else if (const FStateTreeDebuggerBreakpoint::FStateTreeTaskIndex* TaskIndex = Breakpoint.ElementIdentifier.TryGet<FStateTreeDebuggerBreakpoint::FStateTreeTaskIndex>())
-			{
-				if (const FStateTreeTraceTaskEvent* TaskEvent = Event.TryGet<FStateTreeTraceTaskEvent>())
-				{
-					if (TaskEvent->Index == TaskIndex->Index)
-					{
-						HitBreakpointIndex = BreakpointIndex;
-						HitBreakpointInstanceId = InstanceId;
-					}
-				}
-			}
-			else if (const FStateTreeDebuggerBreakpoint::FStateTreeTransitionIndex* TransitionIndex = Breakpoint.ElementIdentifier.TryGet<FStateTreeDebuggerBreakpoint::FStateTreeTransitionIndex>())
-			{
-				if (const FStateTreeTraceTransitionEvent* TransitionEvent = Event.TryGet<FStateTreeTraceTransitionEvent>())
-				{
-					if (TransitionEvent->TransitionIndex == TransitionIndex->Index)
-					{
-						HitBreakpointIndex = BreakpointIndex;
-						HitBreakpointInstanceId = InstanceId;
-					}
-				}
-			}
+			HitBreakpoint.Index = BreakpointIndex;
+            HitBreakpoint.InstanceId = InstanceId;
+            HitBreakpoint.Time = RecordingDuration;
 		}
 	}
 
-	return HitBreakpointInstanceId.IsValid();
+	return HitBreakpoint.IsSet();
 }
 
 bool FStateTreeDebugger::ProcessEvent(const FStateTreeInstanceDebugId InstanceId, const TraceServices::FFrame& Frame, const FStateTreeTraceEventVariantType& Event)
 {
-	EvaluateBreakpoints(InstanceId, Event);
-
 	UE::StateTreeDebugger::FInstanceEventCollection* ExistingCollection = EventCollections.FindByPredicate(
 		[InstanceId](const UE::StateTreeDebugger::FInstanceEventCollection& Entry)
 		{
@@ -859,6 +822,9 @@ bool FStateTreeDebugger::ProcessEvent(const FStateTreeInstanceDebugId InstanceId
 	// Store event in the collection
 	Events.Emplace(Event);
 
+	// Process at the end so RecordingDuration is up to date and we can associate it to a hit breakpoint if necessary.
+	EvaluateBreakpoints(InstanceId, Event);
+
 	return /*bKeepProcessing*/true;
 }
 
diff --git a/Engine/Plugins/Runtime/StateTree/Source/StateTreeModule/Private/Debugger/StateTreeDebuggerTypes.cpp b/Engine/Plugins/Runtime/StateTree/Source/StateTreeModule/Private/Debugger/StateTreeDebuggerTypes.cpp
--- a/Engine/Plugins/Runtime/StateTree/Source/StateTreeModule/Private/Debugger/StateTreeDebuggerTypes.cpp
+++ b/Engine/Plugins/Runtime/StateTree/Source/StateTreeModule/Private/Debugger/StateTreeDebuggerTypes.cpp
@@ -257,6 +257,33 @@ FStateTreeDebuggerBreakpoint::FStateTreeDebuggerBreakpoint(const FStateTreeTrans
 	EventType = GetMatchingEventType(BreakpointType);
 }
 
+bool FStateTreeDebuggerBreakpoint::IsMatchingEvent(const FStateTreeTraceEventVariantType& Event) const
+{
+	EStateTreeTraceEventType ReceivedEventType = EStateTreeTraceEventType::Unset;
+	Visit([&ReceivedEventType](auto& TypedEvent) { ReceivedEventType = TypedEvent.EventType; }, Event);
+
+	bool bIsMatching = false;
+	if (EventType == ReceivedEventType)
+	{
+		if (const FStateTreeStateHandle* StateHandle = ElementIdentifier.TryGet<FStateTreeStateHandle>())
+		{
+			const FStateTreeTraceStateEvent* StateEvent = Event.TryGet<FStateTreeTraceStateEvent>();
+			bIsMatching = StateEvent != nullptr && StateEvent->GetStateHandle() == *StateHandle;
+		}
+		else if (const FStateTreeTaskIndex* TaskIndex = ElementIdentifier.TryGet<FStateTreeTaskIndex>())
+		{
+			const FStateTreeTraceTaskEvent* TaskEvent = Event.TryGet<FStateTreeTraceTaskEvent>();
+			bIsMatching = TaskEvent != nullptr && TaskEvent->Index == TaskIndex->Index;
+		}
+		else if (const FStateTreeTransitionIndex* TransitionIndex = ElementIdentifier.TryGet<FStateTreeTransitionIndex>())
+		{
+			const FStateTreeTraceTransitionEvent* TransitionEvent = Event.TryGet<FStateTreeTraceTransitionEvent>();
+			bIsMatching = TransitionEvent != nullptr && TransitionEvent->TransitionIndex == TransitionIndex->Index;
+		}
+	}
+	return bIsMatching;
+}
+
 EStateTreeTraceEventType FStateTreeDebuggerBreakpoint::GetMatchingEventType(const EStateTreeBreakpointType BreakpointType)
 {
 	switch (BreakpointType)
diff --git a/Engine/Plugins/Runtime/StateTree/Source/StateTreeModule/Private/Debugger/StateTreeTraceTypes.cpp b/Engine/Plugins/Runtime/StateTree/Source/StateTreeModule/Private/Debugger/StateTreeTraceTypes.cpp
--- a/Engine/Plugins/Runtime/StateTree/Source/StateTreeModule/Private/Debugger/StateTreeTraceTypes.cpp
+++ b/Engine/Plugins/Runtime/StateTree/Source/StateTreeModule/Private/Debugger/StateTreeTraceTypes.cpp
@@ -33,7 +33,7 @@ FString FStateTreeTraceLogEvent::ToString(const UStateTree& StateTree) const
 //----------------------------------------------------------------------//
 FString FStateTreeTraceTransitionEvent::ToString(const UStateTree& StateTree) const
 {
-	if (const FCompactStateTransition* Transition = StateTree.GetTransitionFromIndex(TransitionIndex.Get()))
+	if (const FCompactStateTransition* Transition = StateTree.GetTransitionFromIndex(TransitionIndex))
 	{
 		const FCompactStateTreeState* CompactState = StateTree.GetStateFromHandle(Transition->State);
 		FStringBuilderBase StrBuilder;
diff --git a/Engine/Plugins/Runtime/StateTree/Source/StateTreeModule/Private/StateTree.cpp b/Engine/Plugins/Runtime/StateTree/Source/StateTreeModule/Private/StateTree.cpp
--- a/Engine/Plugins/Runtime/StateTree/Source/StateTreeModule/Private/StateTree.cpp
+++ b/Engine/Plugins/Runtime/StateTree/Source/StateTreeModule/Private/StateTree.cpp
@@ -37,6 +37,14 @@ FStateTreeIndex16 UStateTree::GetNodeIndexFromId(const FGuid Id) const
 	return Entry != nullptr ? Entry->Index : FStateTreeIndex16::Invalid;
 }
 
+FGuid UStateTree::GetNodeIdFromIndex(const FStateTreeIndex16 NodeIndex) const
+{
+	const FStateTreeNodeIdToIndex* Entry = NodeIndex.IsValid()
+		? IDToNodeMappings.FindByPredicate([NodeIndex](const FStateTreeNodeIdToIndex& Entry){ return Entry.Index == NodeIndex; })
+		: nullptr;
+	return Entry != nullptr ? Entry->Id : FGuid();
+}
+
 const FCompactStateTreeState* UStateTree::GetStateFromHandle(const FStateTreeStateHandle StateHandle) const
 {
 	return States.IsValidIndex(StateHandle.Index) ? &States[StateHandle.Index] : nullptr;
@@ -54,9 +62,9 @@ FGuid UStateTree::GetStateIdFromHandle(const FStateTreeStateHandle Handle) const
 	return Entry != nullptr ? Entry->Id : FGuid();
 }
 
-const FCompactStateTransition* UStateTree::GetTransitionFromIndex(const int16 TransitionIndex) const
+const FCompactStateTransition* UStateTree::GetTransitionFromIndex(const FStateTreeIndex16 TransitionIndex) const
 {
-	return Transitions.IsValidIndex(TransitionIndex) ? &Transitions[TransitionIndex] : nullptr;
+	return TransitionIndex.IsValid() && Transitions.IsValidIndex(TransitionIndex.Get()) ? &Transitions[TransitionIndex.Get()] : nullptr;
 }
 
 FStateTreeIndex16 UStateTree::GetTransitionIndexFromId(const FGuid Id) const
@@ -65,6 +73,12 @@ FStateTreeIndex16 UStateTree::GetTransitionIndexFromId(const FGuid Id) const
 	return Entry != nullptr ? Entry->Index : FStateTreeIndex16::Invalid;
 }
 
+FGuid UStateTree::GetTransitionIdFromIndex(const FStateTreeIndex16 Index) const
+{
+	const FStateTreeTransitionIdToIndex* Entry = IDToTransitionMappings.FindByPredicate([Index](const FStateTreeTransitionIdToIndex& Entry){ return Entry.Index == Index; });
+	return Entry != nullptr ? Entry->Id : FGuid();
+}
+
 TSharedPtr<FStateTreeInstanceData> UStateTree::GetSharedInstanceData() const
 {
 	// Create a unique index for each thread.
diff --git a/Engine/Plugins/Runtime/StateTree/Source/StateTreeModule/Private/StateTreeExecutionContext.cpp b/Engine/Plugins/Runtime/StateTree/Source/StateTreeModule/Private/StateTreeExecutionContext.cpp
--- a/Engine/Plugins/Runtime/StateTree/Source/StateTreeModule/Private/StateTreeExecutionContext.cpp
+++ b/Engine/Plugins/Runtime/StateTree/Source/StateTreeModule/Private/StateTreeExecutionContext.cpp
@@ -968,6 +968,7 @@ void FStateTreeExecutionContext::ExitState(const FStateTreeTransitionResult& Tra
 						CSV_SCOPED_TIMING_STAT_EXCLUSIVE(StateTree_Task_ExitState);
 						Task.ExitState(*this, CurrentTransition);
 					}
+					STATETREE_TRACE_TASK_EVENT(TaskIndex, DataViews[Task.DataViewIndex.Get()], EStateTreeTraceEventType::OnExit, Transition.CurrentRunStatus);
 				}
 			}
 		}
@@ -1252,6 +1253,7 @@ void FStateTreeExecutionContext::StopEvaluatorsAndGlobalTasks(const EStateTreeRu
 				QUICK_SCOPE_CYCLE_COUNTER(StateTree_Task_TreeStop);
 				Task.ExitState(*this, Transition);
 			}
+			STATETREE_TRACE_TASK_EVENT(TaskIndex, DataViews[Task.DataViewIndex.Get()], EStateTreeTraceEventType::OnExit, Transition.CurrentRunStatus);
 		}
 	}
 
diff --git a/Engine/Plugins/Runtime/StateTree/Source/StateTreeModule/Public/Debugger/StateTreeDebugger.h b/Engine/Plugins/Runtime/StateTree/Source/StateTreeModule/Public/Debugger/StateTreeDebugger.h
--- a/Engine/Plugins/Runtime/StateTree/Source/StateTreeModule/Public/Debugger/StateTreeDebugger.h
+++ b/Engine/Plugins/Runtime/StateTree/Source/StateTreeModule/Public/Debugger/StateTreeDebugger.h
@@ -47,7 +47,30 @@ struct STATETREEMODULE_API FStateTreeDebugger : FTickableGameObject
 
 		TraceServices::FSessionInfo SessionInfo;
 	};
-	
+
+	struct FHitBreakpoint
+	{
+		bool IsSet() const { return Index != INDEX_NONE; }
+		void Reset()
+		{
+			InstanceId = FStateTreeInstanceDebugId::Invalid;
+			Time = 0;
+			Index = INDEX_NONE;
+		}
+
+		/** Indicates the instance for which the breakpoint has been hit */
+		FStateTreeInstanceDebugId InstanceId = FStateTreeInstanceDebugId::Invalid;
+
+		/**
+		 * Store the time at which the breakpoint was hit since we might have process more events before
+		 * sending the notifications.
+		 */
+		double Time = 0;
+
+		/** Indicates the index of the breakpoint that has been hit */
+		int32 Index = INDEX_NONE;
+	};
+
 	FStateTreeDebugger();
 	virtual ~FStateTreeDebugger() override;
 
@@ -236,11 +259,8 @@ private:
 	/** Combined information regarding current scrub time (e.g. frame index, event collection index, etc.) */ 
 	UE::StateTreeDebugger::FScrubState ScrubState;
 
-	/** Indicates the instance for which a breakpoint has been hit */
-	FStateTreeInstanceDebugId HitBreakpointInstanceId = FStateTreeInstanceDebugId::Invalid;
-	
-	/** Indicates the index of the breakpoint that has been hit */
-	int32 HitBreakpointIndex = INDEX_NONE;
+	/** Information stored when a breakpoint is hit while processing events and used to send notifications. */
+	FHitBreakpoint HitBreakpoint;
 
 	/** List of new instances discovered by processing event in the analysis session. */
 	TArray<FStateTreeInstanceDebugId> NewInstances;
diff --git a/Engine/Plugins/Runtime/StateTree/Source/StateTreeModule/Public/Debugger/StateTreeDebuggerTypes.h b/Engine/Plugins/Runtime/StateTree/Source/StateTreeModule/Public/Debugger/StateTreeDebuggerTypes.h
--- a/Engine/Plugins/Runtime/StateTree/Source/StateTreeModule/Public/Debugger/StateTreeDebuggerTypes.h
+++ b/Engine/Plugins/Runtime/StateTree/Source/StateTreeModule/Public/Debugger/StateTreeDebuggerTypes.h
@@ -295,6 +295,8 @@ struct FStateTreeDebuggerBreakpoint
 	explicit FStateTreeDebuggerBreakpoint(const FStateTreeTaskIndex Index, const EStateTreeBreakpointType BreakpointType);
 	explicit FStateTreeDebuggerBreakpoint(const FStateTreeTransitionIndex Index, const EStateTreeBreakpointType BreakpointType);
 
+	STATETREEMODULE_API bool IsMatchingEvent(const FStateTreeTraceEventVariantType& Event) const;
+
 	FIdentifierVariantType ElementIdentifier;
 	EStateTreeBreakpointType BreakpointType;
 	EStateTreeTraceEventType EventType;
diff --git a/Engine/Plugins/Runtime/StateTree/Source/StateTreeModule/Public/StateTree.h b/Engine/Plugins/Runtime/StateTree/Source/StateTreeModule/Public/StateTree.h
--- a/Engine/Plugins/Runtime/StateTree/Source/StateTreeModule/Public/StateTree.h
+++ b/Engine/Plugins/Runtime/StateTree/Source/StateTreeModule/Public/StateTree.h
@@ -126,15 +126,21 @@ public:
 	/** @return Node index matching a given Id; invalid index if node not found. */
 	FStateTreeIndex16 GetNodeIndexFromId(const FGuid Id) const;
 
+	/** @return Id of the node matching a given node index; invalid Id if node not found. */
+	FGuid GetNodeIdFromIndex(const FStateTreeIndex16 NodeIndex) const;
+
 	/** @return View of all states. */
 	TConstArrayView<FCompactStateTreeState> GetStates() const { return States; }
 
 	/** @return Pointer to the transition at a given index; null if not found. */ 
-	const FCompactStateTransition* GetTransitionFromIndex(const int16 TransitionIndex) const;
+	const FCompactStateTransition* GetTransitionFromIndex(const FStateTreeIndex16 TransitionIndex) const;
 	
 	/** @return Runtime transition index matching a given Id; invalid index if node not found. */
 	FStateTreeIndex16 GetTransitionIndexFromId(const FGuid Id) const;
 
+	/** @return Id of the transition matching a given runtime transition index; invalid Id if transition not found. */
+	FGuid GetTransitionIdFromIndex(const FStateTreeIndex16 Index) const;	
+
 #if WITH_EDITOR
 	/** Resets the compiled data to empty. */
 	void ResetCompiled();
