[StateTree] ResolveLeafValueStructType returns nullptr for data view with type only (null memory)

diff --git a/Engine/Plugins/Runtime/StateTree/Source/StateTreeEditorModule/Private/Customizations/StateTreeBindingExtension.cpp b/Engine/Plugins/Runtime/StateTree/Source/StateTreeEditorModule/Private/Customizations/StateTreeBindingExtension.cpp
--- a/Engine/Plugins/Runtime/StateTree/Source/StateTreeEditorModule/Private/Customizations/StateTreeBindingExtension.cpp
+++ b/Engine/Plugins/Runtime/StateTree/Source/StateTreeEditorModule/Private/Customizations/StateTreeBindingExtension.cpp
@@ -40,6 +40,11 @@ UObject* FindEditorBindingsOwner(UObject* InObject)
 
 UStruct* ResolveLeafValueStructType(FStateTreeDataView ValueView, const TArray<FBindingChainElement>& InBindingChain)
 {
+	if (ValueView.GetMemory() == nullptr)
+	{
+		return nullptr;
+	}
+	
 	FStateTreePropertyPath Path;
 
 	for (const FBindingChainElement& Element : InBindingChain)
