[StateTreeDebugger] added missing transaction when toggling Transition breakpoints (undo/redo). Also prevented State and Task breakpoint toggle from dirtying the asset.

diff --git a/Engine/Plugins/Runtime/StateTree/Source/StateTreeEditorModule/Private/Debugger/StateTreeDebuggerUIExtensions.cpp b/Engine/Plugins/Runtime/StateTree/Source/StateTreeEditorModule/Private/Debugger/StateTreeDebuggerUIExtensions.cpp
--- a/Engine/Plugins/Runtime/StateTree/Source/StateTreeEditorModule/Private/Debugger/StateTreeDebuggerUIExtensions.cpp
+++ b/Engine/Plugins/Runtime/StateTree/Source/StateTreeEditorModule/Private/Debugger/StateTreeDebuggerUIExtensions.cpp
@@ -101,7 +101,7 @@ void OnStateBreakpointToggled(const TArray<TWeakObjectPtr<>>& States, UStateTree
 	}
 
 	const FScopedTransaction Transaction(LOCTEXT("ToggleStateBreakpoint", "Toggle State Breakpoint"));
-	EditorData->Modify();
+	EditorData->Modify(/*bAlwaysMarkDirty*/false);
 
 	for (const TWeakObjectPtr<>& WeakStateObject : States)
 	{
@@ -136,7 +136,7 @@ void OnTaskBreakpointToggled(TSharedPtr<IPropertyHandle> StructProperty, UStateT
 	if (EditorData != nullptr && Node != nullptr)
 	{
 		const FScopedTransaction Transaction(LOCTEXT("ToggleTaskBreakpoint", "Toggle Task Breakpoint"));
-		EditorData->Modify();
+		EditorData->Modify(/*bAlwaysMarkDirty*/false);
 		
 		const bool bRemoved = EditorData->RemoveBreakpoint(Node->ID, Type);
 		if (bRemoved == false)
@@ -269,25 +269,31 @@ void OnTransitionBreakpointToggled(const TSharedPtr<IPropertyHandle>& StructProp
 	TArray<UObject*> OuterObjects;
 	PropertyHandle.GetOuterObjects(OuterObjects);
 
-	for (const UObject* OuterObject : OuterObjects)
+	if (OuterObjects.Num() > 0)
 	{
-		if (const UStateTreeState* TreeState = Cast<UStateTreeState>(OuterObject))
+		const FScopedTransaction Transaction(LOCTEXT("ToggleTransitionBreakpoint", "Toggle Transition Breakpoint"));
+		EditorData->Modify(/*bAlwaysMarkDirty*/false);
+		
+		for (const UObject* OuterObject : OuterObjects)
 		{
-			const int32 IndexInArray = PropertyHandle.GetIndexInArray();
-
-			if (TreeState->Transitions.IsValidIndex(IndexInArray))
+			if (const UStateTreeState* TreeState = Cast<UStateTreeState>(OuterObject))
 			{
-				const FStateTreeTransition& Transition = TreeState->Transitions[IndexInArray];
-				const bool bHasBreakpoint = EditorData->HasBreakpoint(Transition.ID, EStateTreeBreakpointType::OnTransition);
-				EditorData->HasBreakpoint(Transition.ID, EStateTreeBreakpointType::OnTransition);
-				if (ToggledState == ECheckBoxState::Checked && bHasBreakpoint == false)
+				const int32 IndexInArray = PropertyHandle.GetIndexInArray();
+
+				if (TreeState->Transitions.IsValidIndex(IndexInArray))
 				{
-					EditorData->AddBreakpoint(Transition.ID, EStateTreeBreakpointType::OnTransition);	
+					const FStateTreeTransition& Transition = TreeState->Transitions[IndexInArray];
+					const bool bHasBreakpoint = EditorData->HasBreakpoint(Transition.ID, EStateTreeBreakpointType::OnTransition);
+					EditorData->HasBreakpoint(Transition.ID, EStateTreeBreakpointType::OnTransition);
+					if (ToggledState == ECheckBoxState::Checked && bHasBreakpoint == false)
+					{
+						EditorData->AddBreakpoint(Transition.ID, EStateTreeBreakpointType::OnTransition);	
+					}
+					else if (ToggledState == ECheckBoxState::Unchecked && bHasBreakpoint)
+					{
+						EditorData->RemoveBreakpoint(Transition.ID, EStateTreeBreakpointType::OnTransition);
+					}		
 				}
-				else if (ToggledState == ECheckBoxState::Unchecked && bHasBreakpoint)
-				{
-					EditorData->RemoveBreakpoint(Transition.ID, EStateTreeBreakpointType::OnTransition);
-				}		
 			}
 		}
 	}
