[StateTreeDebugger]
- active states text is truncated to the window size in the timeline
- fixed cursor tracking when scrubbing with the buttons
- fixed TrackedCursor when scrubbing to the beginning of the timeline
- replaced toggle pause button by a single button changing icon

diff --git a/Engine/Plugins/Runtime/StateTree/Source/StateTreeEditorModule/Private/Debugger/SStateTreeDebuggerEventTimelineView.cpp b/Engine/Plugins/Runtime/StateTree/Source/StateTreeEditorModule/Private/Debugger/SStateTreeDebuggerEventTimelineView.cpp
--- a/Engine/Plugins/Runtime/StateTree/Source/StateTreeEditorModule/Private/Debugger/SStateTreeDebuggerEventTimelineView.cpp
+++ b/Engine/Plugins/Runtime/StateTree/Source/StateTreeEditorModule/Private/Debugger/SStateTreeDebuggerEventTimelineView.cpp
@@ -4,6 +4,7 @@
 #include "Framework/Application/SlateApplication.h"
 #include "Rendering/DrawElements.h"
 #include "SSimpleTimeSlider.h"
+#include "Fonts/FontMeasure.h"
 
 #define LOCTEXT_NAMESPACE "SStateTreeDebuggerEventTimelineView"
 
@@ -64,15 +65,26 @@ int32 SStateTreeDebuggerEventTimelineView::PaintEvents(const FGeometry& Allotted
 
 				if (Window.Description.IsEmpty() == false)
 				{
-					FSlateDrawElement::MakeText
-					(
-						OutDrawElements,
-						LayerId,
-						AllottedGeometry.ToPaintGeometry(EventSize, FSlateLayoutTransform(FVector2D(XStart + 8, Y))),
-						Window.Description,
-						FCoreStyle::GetDefaultFontStyle("Regular", 8),
-						ESlateDrawEffect::None,
-						FLinearColor::Black);
+					FString Description = Window.Description.ToString();					
+					FSlateFontInfo FontInfo = FCoreStyle::GetDefaultFontStyle("Regular", 8);
+					const TSharedRef<FSlateFontMeasure> FontMeasureService = FSlateApplication::Get().GetRenderer()->GetFontMeasureService();
+					constexpr int32 TextLeftOffset = 8;
+					const int32 HorizontalOffset = FMath::RoundToInt((XEnd - XStart) - TextLeftOffset);
+					const int32 LastWholeCharacterIndex = FontMeasureService->FindLastWholeCharacterIndexBeforeOffset(Description, FontInfo, HorizontalOffset);
+
+					if (LastWholeCharacterIndex >= 0)
+					{
+						Description = Description.Left(LastWholeCharacterIndex + 1);
+
+						FSlateDrawElement::MakeText(
+							OutDrawElements,
+							LayerId,
+							AllottedGeometry.ToPaintGeometry(EventSize, FSlateLayoutTransform(FVector2D(XStart + TextLeftOffset, Y))),
+							Description,
+							FontInfo,
+							ESlateDrawEffect::None,
+							FLinearColor::Black);
+					}
 				}
 			}
 		}
diff --git a/Engine/Plugins/Runtime/StateTree/Source/StateTreeEditorModule/Private/Debugger/SStateTreeDebuggerView.cpp b/Engine/Plugins/Runtime/StateTree/Source/StateTreeEditorModule/Private/Debugger/SStateTreeDebuggerView.cpp
--- a/Engine/Plugins/Runtime/StateTree/Source/StateTreeEditorModule/Private/Debugger/SStateTreeDebuggerView.cpp
+++ b/Engine/Plugins/Runtime/StateTree/Source/StateTreeEditorModule/Private/Debugger/SStateTreeDebuggerView.cpp
@@ -254,6 +254,35 @@ void SStateTreeDebuggerView::StopRecording()
 	}
 }
 
+bool SStateTreeDebuggerView::CanToggleDebuggerAnalysis() const
+{
+	return Debugger && Debugger->IsAnalysisSessionActive(); 
+}
+
+void SStateTreeDebuggerView::ToggleDebuggerAnalysis() const
+{
+	if (Debugger)
+	{
+		if (Debugger->IsAnalysisSessionPaused())
+		{
+			Debugger->ResumeSessionAnalysis();
+		}
+		else
+		{
+			Debugger->PauseSessionAnalysis();
+		}
+	}
+}
+
+FSlateIcon SStateTreeDebuggerView::GetDebuggerAnalysisIcon() const
+{
+	if (Debugger && Debugger->IsAnalysisSessionActive() && !Debugger->IsAnalysisSessionPaused())
+	{
+		return FSlateIcon(FStateTreeEditorStyle::Get().GetStyleSetName(), "StateTreeEditor.PauseDebuggerAnalysis");
+	}
+	return FSlateIcon(FStateTreeEditorStyle::Get().GetStyleSetName(), "StateTreeEditor.ResumeDebuggerAnalysis");
+}
+
 void SStateTreeDebuggerView::OnPIEStarted(const bool bIsSimulating)
 {
 	if (UStateTreeEditorSettings::Get().bShouldDebuggerAutoRecordOnPIE)
@@ -331,34 +360,9 @@ void SStateTreeDebuggerView::Construct(const FArguments& InArgs, const UStateTre
 		FIsActionButtonVisible::CreateLambda([this] { return CanAddStateBreakpoint(EStateTreeBreakpointType::OnExit) || CanRemoveStateBreakpoint(EStateTreeBreakpointType::OnExit); }));
 	
 	// Toolbars
-
-	FUIAction AnalysisPauseToggleButtonAction;
-	AnalysisPauseToggleButtonAction.GetActionCheckState.BindLambda([this]
-	{
-		return Debugger && Debugger->IsAnalysisSessionPaused() ? ECheckBoxState::Checked : ECheckBoxState::Unchecked;
-	});
-	AnalysisPauseToggleButtonAction.ExecuteAction.BindLambda([this]
-	{
-		if (Debugger)
-		{
-			if (Debugger->IsAnalysisSessionPaused())
-			{
-				Debugger->ResumeSessionAnalysis();
-			}
-			else
-			{
-				Debugger->PauseSessionAnalysis();
-			}
-		}
-	});
-	
 	FSlimHorizontalToolBarBuilder LeftToolbar(InCommandList, FMultiBoxCustomization::None, /*InExtender*/ nullptr, /*InForceSmallIcons*/ true);
 	LeftToolbar.BeginSection(TEXT("Debugging"));
 	{
-		const FStateTreeDebuggerCommands& DebuggerCommands = FStateTreeDebuggerCommands::Get();
-		LeftToolbar.AddToolBarButton(DebuggerCommands.StartRecording);
-		LeftToolbar.AddToolBarButton(DebuggerCommands.StopRecording);
-
 		const FPlayWorldCommands& PlayWorldCommands = FPlayWorldCommands::Get();
 		LeftToolbar.AddToolBarButton(PlayWorldCommands.RepeatLastPlay);
 		LeftToolbar.AddToolBarButton(PlayWorldCommands.PausePlaySession,
@@ -368,17 +372,20 @@ void SStateTreeDebuggerView::Construct(const FArguments& InArgs, const UStateTre
 		LeftToolbar.AddToolBarButton(PlayWorldCommands.StopPlaySession,
 					NAME_None, TAttribute<FText>(), TAttribute<FText>(), FSlateIcon(FAppStyle::GetAppStyleSetName(), "PlayWorld.StopPlaySession.Small"));
 		LeftToolbar.AddSeparator();
+
+		const FStateTreeDebuggerCommands& DebuggerCommands = FStateTreeDebuggerCommands::Get();
+		LeftToolbar.AddToolBarButton(DebuggerCommands.StartRecording);
+		LeftToolbar.AddToolBarButton(DebuggerCommands.StopRecording);
 		LeftToolbar.AddToolBarButton(DebuggerCommands.PreviousFrameWithStateChange);
 		LeftToolbar.AddToolBarButton(DebuggerCommands.PreviousFrameWithEvents);
 		LeftToolbar.AddToolBarButton(DebuggerCommands.NextFrameWithEvents);
 		LeftToolbar.AddToolBarButton(DebuggerCommands.NextFrameWithStateChange);
 		LeftToolbar.AddSeparator();
-		LeftToolbar.AddToolBarButton(AnalysisPauseToggleButtonAction,
+		LeftToolbar.AddToolBarButton(DebuggerCommands.ToggleAnalysis,
 			NAME_None,
 			TAttribute<FText>(),
-			LOCTEXT("DebuggerPauseToolTip", "Debugger Paused"),
-			FSlateIcon(FStateTreeEditorStyle::Get().GetStyleSetName(), "StateTreeEditor.ToggleDebuggerPause"),
-			EUserInterfaceActionType::ToggleButton);
+			TAttribute<FText>(),
+			TAttribute<FSlateIcon>::CreateSP(this, &SStateTreeDebuggerView::GetDebuggerAnalysisIcon));
 	}
 	LeftToolbar.EndSection();
 
@@ -763,6 +770,11 @@ void SStateTreeDebuggerView::BindDebuggerToolbarCommands(const TSharedRef<FUICom
 		Commands.NextFrameWithStateChange,
 		FExecuteAction::CreateSP(this, &SStateTreeDebuggerView::StepForwardToNextStateChange),
 		FCanExecuteAction::CreateSP(this, &SStateTreeDebuggerView::CanStepForwardToNextStateChange));
+
+	ToolkitCommands->MapAction(
+		Commands.ToggleAnalysis,
+		FExecuteAction::CreateSP(this, &SStateTreeDebuggerView::ToggleDebuggerAnalysis),
+		FCanExecuteAction::CreateSP(this, &SStateTreeDebuggerView::CanToggleDebuggerAnalysis));
 }
 
 bool SStateTreeDebuggerView::CanStepBackToPreviousStateWithEvents() const
@@ -1008,6 +1020,8 @@ void SStateTreeDebuggerView::OnTimeLineScrubPositionChanged(double Time, bool bI
 
 void SStateTreeDebuggerView::OnDebuggerScrubStateChanged(const UE::StateTreeDebugger::FScrubState& ScrubState)
 {
+	TrackCursor();
+
 	// Rebuild frame details from the events of that frame
 	EventsTreeElements.Reset();
 	ON_SCOPE_EXIT
@@ -1271,15 +1285,15 @@ void SStateTreeDebuggerView::TrackCursor()
 	static constexpr double LeadingMarginFraction = 0.05;
 	static constexpr double TrailingMarginFraction = 0.01;
 
-	if (ScrubTime > CurrentViewRange.GetUpperBoundValue() - ViewRangeDuration * LeadingMarginFraction)
+	if (ScrubTime > (CurrentViewRange.GetUpperBoundValue() - (ViewRangeDuration * LeadingMarginFraction)))
 	{
-		CurrentViewRange.SetUpperBound(ScrubTime + ViewRangeDuration * LeadingMarginFraction);
+		CurrentViewRange.SetUpperBound(ScrubTime + (ViewRangeDuration * LeadingMarginFraction));
 		CurrentViewRange.SetLowerBound(CurrentViewRange.GetUpperBoundValue() - ViewRangeDuration);
 	}
 
-	if (ScrubTime < CurrentViewRange.GetLowerBoundValue() - ViewRangeDuration * TrailingMarginFraction)
+	if (ScrubTime < (CurrentViewRange.GetLowerBoundValue() + (ViewRangeDuration * TrailingMarginFraction)))
 	{
-		CurrentViewRange.SetLowerBound(ScrubTime);
+		CurrentViewRange.SetLowerBound(ScrubTime - (ViewRangeDuration * TrailingMarginFraction));
 		CurrentViewRange.SetUpperBound(CurrentViewRange.GetLowerBoundValue() + ViewRangeDuration);
 	}
 
diff --git a/Engine/Plugins/Runtime/StateTree/Source/StateTreeEditorModule/Private/Debugger/SStateTreeDebuggerView.h b/Engine/Plugins/Runtime/StateTree/Source/StateTreeEditorModule/Private/Debugger/SStateTreeDebuggerView.h
--- a/Engine/Plugins/Runtime/StateTree/Source/StateTreeEditorModule/Private/Debugger/SStateTreeDebuggerView.h
+++ b/Engine/Plugins/Runtime/StateTree/Source/StateTreeEditorModule/Private/Debugger/SStateTreeDebuggerView.h
@@ -21,6 +21,7 @@ namespace UE::StateTreeDebugger
 
 enum class EStateTreeBreakpointType : uint8;
 
+struct FSlateIcon;
 struct FStateTreeDebugger;
 struct FStateTreeInstanceDebugId;
 struct FStateTreeDebuggerBreakpoint;
@@ -75,7 +76,11 @@ private:
 
 	bool CanStopRecording() const { return IsRecording(); }	
 	void StopRecording();
-	
+
+	bool CanToggleDebuggerAnalysis() const;
+	void ToggleDebuggerAnalysis() const;
+	FSlateIcon GetDebuggerAnalysisIcon() const;
+
 	void BindDebuggerToolbarCommands(const TSharedRef<FUICommandList>& ToolkitCommands);
 
 	bool CanStepBackToPreviousStateWithEvents() const;
diff --git a/Engine/Plugins/Runtime/StateTree/Source/StateTreeEditorModule/Private/Debugger/StateTreeDebuggerCommands.cpp b/Engine/Plugins/Runtime/StateTree/Source/StateTreeEditorModule/Private/Debugger/StateTreeDebuggerCommands.cpp
--- a/Engine/Plugins/Runtime/StateTree/Source/StateTreeEditorModule/Private/Debugger/StateTreeDebuggerCommands.cpp
+++ b/Engine/Plugins/Runtime/StateTree/Source/StateTreeEditorModule/Private/Debugger/StateTreeDebuggerCommands.cpp
@@ -32,6 +32,7 @@ void FStateTreeDebuggerCommands::RegisterCommands()
 	UI_COMMAND(PreviousFrameWithEvents, "Previous Frame", "Step one frame back", EUserInterfaceActionType::Button, FInputChord(EKeys::Left));
 	UI_COMMAND(NextFrameWithEvents, "Next Frame", "Step one frame forward", EUserInterfaceActionType::Button, FInputChord(EKeys::Right));
 	UI_COMMAND(NextFrameWithStateChange, "Next State Change", "Jump to next state changed frame", EUserInterfaceActionType::Button, FInputChord(EModifierKey::Control, EKeys::Right));
+	UI_COMMAND(ToggleAnalysis, "Toggle Analysis",  "Pause or resume processing events from the selected trace.", EUserInterfaceActionType::Button, FInputChord());
 }
 
 #undef LOCTEXT_NAMESPACE
diff --git a/Engine/Plugins/Runtime/StateTree/Source/StateTreeEditorModule/Private/Debugger/StateTreeDebuggerCommands.h b/Engine/Plugins/Runtime/StateTree/Source/StateTreeEditorModule/Private/Debugger/StateTreeDebuggerCommands.h
--- a/Engine/Plugins/Runtime/StateTree/Source/StateTreeEditorModule/Private/Debugger/StateTreeDebuggerCommands.h
+++ b/Engine/Plugins/Runtime/StateTree/Source/StateTreeEditorModule/Private/Debugger/StateTreeDebuggerCommands.h
@@ -25,6 +25,7 @@ public:
 	TSharedPtr<FUICommandInfo> PreviousFrameWithEvents;
 	TSharedPtr<FUICommandInfo> NextFrameWithEvents;
 	TSharedPtr<FUICommandInfo> NextFrameWithStateChange;
+	TSharedPtr<FUICommandInfo> ToggleAnalysis;
 };
 
 #endif // WITH_STATETREE_DEBUGGER
\ No newline at end of file
diff --git a/Engine/Plugins/Runtime/StateTree/Source/StateTreeEditorModule/Private/StateTreeEditorStyle.cpp b/Engine/Plugins/Runtime/StateTree/Source/StateTreeEditorModule/Private/StateTreeEditorStyle.cpp
--- a/Engine/Plugins/Runtime/StateTree/Source/StateTreeEditorModule/Private/StateTreeEditorStyle.cpp
+++ b/Engine/Plugins/Runtime/StateTree/Source/StateTreeEditorModule/Private/StateTreeEditorStyle.cpp
@@ -194,7 +194,7 @@ FStateTreeEditorStyle::FStateTreeEditorStyle()
 		Set("StateTreeEditor.DuplicateStates", new IMAGE_BRUSH_SVG("Starship/Common/Duplicate", CoreStyleConstants::Icon16x16));
 		Set("StateTreeEditor.DeleteStates", new IMAGE_BRUSH_SVG("Starship/Common/Delete", CoreStyleConstants::Icon16x16));
 		Set("StateTreeEditor.RenameState", new IMAGE_BRUSH_SVG("Starship/Common/Rename", CoreStyleConstants::Icon16x16));
-		Set("StateTreeEditor.AutoScroll", new IMAGE_BRUSH_SVG("Starship/Insights/AutoScrollRight_20", CoreStyleConstants::Icon16x16));		
+		Set("StateTreeEditor.AutoScroll", new IMAGE_BRUSH_SVG("Starship/Insights/AutoScrollRight_20", CoreStyleConstants::Icon16x16));
 	}
 
 	{
@@ -219,7 +219,8 @@ FStateTreeEditorStyle::FStateTreeEditorStyle()
 
 		Set("StateTreeEditor.EnableStates", new IMAGE_BRUSH("Icons/Empty_16x", CoreStyleConstants::Icon16x16));
 		Set( "StateTreeEditor.Debugger.Breakpoint.EnabledAndValid", new IMAGE_BRUSH_SVG( "Starship/Blueprints/Breakpoint_Valid", CoreStyleConstants::Icon16x16, FStyleColors::AccentRed));
-		Set("StateTreeEditor.ToggleDebuggerPause", new IMAGE_BRUSH("Icons/generic_pause_16x", CoreStyleConstants::Icon16x16));
+		Set("StateTreeEditor.PauseDebuggerAnalysis", new IMAGE_BRUSH("Icons/generic_pause_16x", CoreStyleConstants::Icon16x16));
+		Set("StateTreeEditor.ResumeDebuggerAnalysis", new IMAGE_BRUSH_SVG("Starship/Common/Timeline", CoreStyleConstants::Icon16x16));
 	}
 
 	{
