Fix CIS 328889:  Localization warnings

diff --git a/Engine/Plugins/Runtime/StateTree/Source/StateTreeEditorModule/Private/Debugger/StateTreeDebuggerUIExtensions.cpp b/Engine/Plugins/Runtime/StateTree/Source/StateTreeEditorModule/Private/Debugger/StateTreeDebuggerUIExtensions.cpp
--- a/Engine/Plugins/Runtime/StateTree/Source/StateTreeEditorModule/Private/Debugger/StateTreeDebuggerUIExtensions.cpp
+++ b/Engine/Plugins/Runtime/StateTree/Source/StateTreeEditorModule/Private/Debugger/StateTreeDebuggerUIExtensions.cpp
@@ -271,7 +271,7 @@ void OnTransitionBreakpointToggled(const TSharedPtr<IPropertyHandle>& StructProp
 
 	if (OuterObjects.Num() > 0)
 	{
-		const FScopedTransaction Transaction(LOCTEXT("ToggleTransitionBreakpoint", "Toggle Transition Breakpoint"));
+		const FScopedTransaction Transaction(LOCTEXT("TransitionBreakpointToggled", "Transition Breakpoint Toggled"));
 		EditorData->Modify(/*bAlwaysMarkDirty*/false);
 		
 		for (const UObject* OuterObject : OuterObjects)
