Reuse existing View in StateTreeEditor when debugger tab is reopened

diff --git a/Engine/Plugins/Runtime/StateTree/Source/StateTreeEditorModule/Private/StateTreeEditor.cpp b/Engine/Plugins/Runtime/StateTree/Source/StateTreeEditorModule/Private/StateTreeEditor.cpp
--- a/Engine/Plugins/Runtime/StateTree/Source/StateTreeEditorModule/Private/StateTreeEditor.cpp
+++ b/Engine/Plugins/Runtime/StateTree/Source/StateTreeEditorModule/Private/StateTreeEditor.cpp
@@ -392,7 +392,11 @@ TSharedRef<SDockTab> FStateTreeEditor::SpawnTab_Debugger(const FSpawnTabArgs& Ar
 	TSharedPtr<SWidget> Widget = SNullWidget::NullWidget;
 	if (StateTree != nullptr)
 	{
-		DebuggerView = SNew(SStateTreeDebuggerView, *StateTree, StateTreeViewModel.ToSharedRef(), TreeViewCommandList);
+		// Reuse existing view if Tab is reopened
+		if (!DebuggerView.IsValid())
+		{
+			DebuggerView = SNew(SStateTreeDebuggerView, *StateTree, StateTreeViewModel.ToSharedRef(), TreeViewCommandList);
+		}
 		Widget = DebuggerView;
 	}
 	
