[StateTreeDebugger]
- Tracks are no longer cleared when running multiple consecutive PIE sessions unless 'bShouldDebuggerResetDataOnNewPIESession' is set to true in the StateTree Editor Settings. Non editor session analysis will always clear tracks.
- Added dedicated button to clear track data.
- Since tracks can be preserved between PIE session clamp range of the timelines is now the longest recording duration instead of the current session duration.
- Allow goto previous active states when scrubbing past the recording duration of the current analysis session
- Scrub buttons can now be used when the session is no longer active, regardless of the autoscroll state
- Manual scrubbing no longer change auto scroll status if session is not active
#rnx

diff --git a/Engine/Plugins/Runtime/StateTree/Source/StateTreeEditorModule/Private/Debugger/SStateTreeDebuggerView.cpp b/Engine/Plugins/Runtime/StateTree/Source/StateTreeEditorModule/Private/Debugger/SStateTreeDebuggerView.cpp
--- a/Engine/Plugins/Runtime/StateTree/Source/StateTreeEditorModule/Private/Debugger/SStateTreeDebuggerView.cpp
+++ b/Engine/Plugins/Runtime/StateTree/Source/StateTreeEditorModule/Private/Debugger/SStateTreeDebuggerView.cpp
@@ -21,7 +21,6 @@
 #include "StateTreeEditorData.h"
 #include "StateTreeEditorSettings.h"
 #include "StateTreeEditorStyle.h"
-#include "StateTreeExecutionContext.h"
 #include "StateTreeModule.h"
 #include "StateTreeState.h"
 #include "StateTreeViewModel.h"
@@ -167,6 +166,7 @@ protected:
 					.Text_Lambda([&Event=Item->Event, this]()
 					{
 						FString EventDescription;
+						check(StateTreeViewModel);
 						if (const UStateTree* StateTree = StateTreeViewModel->GetStateTree())
 						{
 							Visit([&EventDescription, StateTree](auto& TypedEvent)
@@ -206,14 +206,12 @@ SStateTreeDebuggerView::~SStateTreeDebuggerView()
 		SelectedNodeDataObject->RemoveFromRoot();
 	}
 
-	if (Debugger)
-	{
-		Debugger->OnScrubStateChanged.Unbind();
-		Debugger->OnBreakpointHit.Unbind();
-		Debugger->OnNewSession.Unbind();
-		Debugger->OnNewInstance.Unbind();
-		Debugger->OnSelectedInstanceCleared.Unbind();
-	}
+	check(Debugger);
+	Debugger->OnScrubStateChanged.Unbind();
+	Debugger->OnBreakpointHit.Unbind();
+	Debugger->OnNewSession.Unbind();
+	Debugger->OnNewInstance.Unbind();
+	Debugger->OnSelectedInstanceCleared.Unbind();
 
 	FEditorDelegates::BeginPIE.RemoveAll(this);
 	FEditorDelegates::EndPIE.RemoveAll(this);
@@ -226,11 +224,12 @@ void SStateTreeDebuggerView::StartRecording()
 {
 	if (CanStartRecording())
 	{
+		check(Debugger);
 		Debugger->ClearSelection();
 
 		// We give priority to the Editor actions even if an analysis was active (remote process)
 		// This will stop current analysis and connect to the new live trace.
-		bRecording = Debugger->RequestAnalysisOfNextLiveSession();
+		bRecording = Debugger->RequestAnalysisOfEditorSession();
 	}
 }
 
@@ -243,6 +242,10 @@ void SStateTreeDebuggerView::StopRecording()
 		
 		bRecording = false;
 
+		// Update max duration from current recording until track data gets reset
+		check(Debugger);
+		MaxTrackRecordingDuration = FMath::Max(MaxTrackRecordingDuration, Debugger->GetRecordingDuration());
+
 		// Mark all tracks from the stopped session as stale to have different look.
 		for (const TSharedPtr<RewindDebugger::FRewindDebuggerTrack>& DebugTrack : InstanceOwnerTracks)
 		{
@@ -256,33 +259,59 @@ void SStateTreeDebuggerView::StopRecording()
 
 bool SStateTreeDebuggerView::CanToggleDebuggerAnalysis() const
 {
-	return Debugger && Debugger->IsAnalysisSessionActive(); 
+	check(Debugger);
+	return Debugger->IsAnalysisSessionActive(); 
 }
 
 void SStateTreeDebuggerView::ToggleDebuggerAnalysis() const
 {
-	if (Debugger)
+	check(Debugger);
+	if (Debugger->IsAnalysisSessionPaused())
 	{
-		if (Debugger->IsAnalysisSessionPaused())
-		{
-			Debugger->ResumeSessionAnalysis();
-		}
-		else
-		{
-			Debugger->PauseSessionAnalysis();
-		}
+		Debugger->ResumeSessionAnalysis();
+	}
+	else
+	{
+		Debugger->PauseSessionAnalysis();
 	}
 }
 
 FSlateIcon SStateTreeDebuggerView::GetDebuggerAnalysisIcon() const
 {
-	if (Debugger && Debugger->IsAnalysisSessionActive() && !Debugger->IsAnalysisSessionPaused())
+	check(Debugger);
+	if (Debugger->IsAnalysisSessionActive() && !Debugger->IsAnalysisSessionPaused())
 	{
 		return FSlateIcon(FStateTreeEditorStyle::Get().GetStyleSetName(), "StateTreeEditor.PauseDebuggerAnalysis");
 	}
 	return FSlateIcon(FStateTreeEditorStyle::Get().GetStyleSetName(), "StateTreeEditor.ResumeDebuggerAnalysis");
 }
 
+bool SStateTreeDebuggerView::CanResetTracks() const
+{
+	return !InstanceOwnerTracks.IsEmpty();
+}
+
+void SStateTreeDebuggerView::ResetTracks()
+{
+	InstanceOwnerTracks.Reset();
+
+	check(Debugger);
+	Debugger->ResetEventCollections();
+
+	MaxTrackRecordingDuration = 0;
+
+	// Refresh tree view
+	if (InstancesTreeView)
+	{
+		InstancesTreeView->Refresh();
+	}
+
+	if (InstanceTimelinesTreeView)
+	{
+		InstanceTimelinesTreeView->Refresh();
+	}
+}
+
 void SStateTreeDebuggerView::OnPIEStarted(const bool bIsSimulating)
 {
 	if (UStateTreeEditorSettings::Get().bShouldDebuggerAutoRecordOnPIE)
@@ -294,21 +323,26 @@ void SStateTreeDebuggerView::OnPIEStarted(const bool bIsSimulating)
 void SStateTreeDebuggerView::OnPIEStopped(const bool bIsSimulating)
 {
 	StopRecording();
+
+	check(Debugger);
 	Debugger->StopSessionAnalysis();
 }
 
 void SStateTreeDebuggerView::OnPIEPaused(const bool bIsSimulating) const
 {
+	check(Debugger);
 	Debugger->PauseSessionAnalysis();
 }
 
 void SStateTreeDebuggerView::OnPIEResumed(const bool bIsSimulating) const
 {
+	check(Debugger);
 	Debugger->ResumeSessionAnalysis();
 }
 
 void SStateTreeDebuggerView::OnPIESingleStepped(bool bSimulating) const
 {
+	check(Debugger);
 	Debugger->SyncToCurrentSessionDuration();
 }
 
@@ -319,6 +353,8 @@ void SStateTreeDebuggerView::Construct(const FArguments& InArgs, const UStateTre
 	StateTreeEditorData = Cast<UStateTreeEditorData>(InStateTree.EditorData.Get()); 
 
 	Debugger = InStateTreeViewModel->GetDebugger();
+	check(Debugger);
+
 	IStateTreeModule& StateTreeModule = FModuleManager::GetModuleChecked<IStateTreeModule>("StateTreeModule");
 	bRecording = StateTreeModule.IsTracing();
 
@@ -376,16 +412,22 @@ void SStateTreeDebuggerView::Construct(const FArguments& InArgs, const UStateTre
 		const FStateTreeDebuggerCommands& DebuggerCommands = FStateTreeDebuggerCommands::Get();
 		LeftToolbar.AddToolBarButton(DebuggerCommands.StartRecording);
 		LeftToolbar.AddToolBarButton(DebuggerCommands.StopRecording);
-		LeftToolbar.AddToolBarButton(DebuggerCommands.PreviousFrameWithStateChange);
-		LeftToolbar.AddToolBarButton(DebuggerCommands.PreviousFrameWithEvents);
-		LeftToolbar.AddToolBarButton(DebuggerCommands.NextFrameWithEvents);
-		LeftToolbar.AddToolBarButton(DebuggerCommands.NextFrameWithStateChange);
+
 		LeftToolbar.AddSeparator();
+
 		LeftToolbar.AddToolBarButton(DebuggerCommands.ToggleAnalysis,
 			NAME_None,
 			TAttribute<FText>(),
 			TAttribute<FText>(),
 			TAttribute<FSlateIcon>::CreateSP(this, &SStateTreeDebuggerView::GetDebuggerAnalysisIcon));
+		LeftToolbar.AddToolBarButton(DebuggerCommands.ResetTracks);
+
+		LeftToolbar.AddSeparator();
+
+		LeftToolbar.AddToolBarButton(DebuggerCommands.PreviousFrameWithStateChange);
+		LeftToolbar.AddToolBarButton(DebuggerCommands.PreviousFrameWithEvents);
+		LeftToolbar.AddToolBarButton(DebuggerCommands.NextFrameWithEvents);
+		LeftToolbar.AddToolBarButton(DebuggerCommands.NextFrameWithStateChange);
 	}
 	LeftToolbar.EndSection();
 
@@ -418,7 +460,7 @@ void SStateTreeDebuggerView::Construct(const FArguments& InArgs, const UStateTre
 	// Make sure that the active trace contains statetree specific traces
 	if (TraceDescriptors.Num() == 1 && StateTreeModule.IsTracing())
 	{
-		Debugger->StartSessionAnalysis(TraceDescriptors.Last());
+		Debugger->RequestSessionAnalysis(TraceDescriptors.Last());
 	}
 
 	// Trace selection combo
@@ -430,7 +472,8 @@ void SStateTreeDebuggerView::Construct(const FArguments& InArgs, const UStateTre
 			.ToolTipText(LOCTEXT("SelectTraceSession", "Pick trace session to debug"))
 			.Text_Lambda([Debugger = Debugger]()
 			{
-				return Debugger.IsValid() ? Debugger->GetSelectedTraceDescription() : FText::GetEmpty();
+				check(Debugger);
+				return Debugger->GetSelectedTraceDescription();
 			})
 		];
 
@@ -462,7 +505,7 @@ void SStateTreeDebuggerView::Construct(const FArguments& InArgs, const UStateTre
 		.OnScrolled_Lambda([this](double ScrollOffset) { InstancesTreeView->ScrollTo(ScrollOffset); })
 		.InstanceTracks(&InstanceOwnerTracks)
 		.ViewRange_Lambda([this]() { return ViewRange; })
-		.ClampRange_Lambda([this]() { return TRange<double>(0, Debugger->GetRecordingDuration()); })
+		.ClampRange_Lambda([this]() { return TRange<double>(0, MaxTrackRecordingDuration); })
 		.OnViewRangeChanged_Lambda([this](TRange<double> NewRange) { ViewRange = NewRange; })
 		.ScrubPosition(ScrubTimeAttribute)
 		.OnScrubPositionChanged_Lambda([this](double NewScrubTime, bool bIsScrubbing) { OnTimeLineScrubPositionChanged(NewScrubTime, bIsScrubbing); });
@@ -474,6 +517,7 @@ void SStateTreeDebuggerView::Construct(const FArguments& InArgs, const UStateTre
 			.DefaultLabel(LOCTEXT("FrameDetailsColumnHeader", "Frame Details")))
 			.OnGenerateRow_Lambda([this](const TSharedPtr<UE::StateTreeDebugger::FEventTreeElement>& InElement, const TSharedRef<STableViewBase>& InOwnerTableView)
 			{
+				check(StateTreeViewModel);
 				return SNew(SStateTreeDebuggerTableRow, InOwnerTableView, InElement, StateTreeViewModel.ToSharedRef());
 			})
 			.OnGetChildren_Lambda([](const TSharedPtr<const UE::StateTreeDebugger::FEventTreeElement>& InParent, TArray<TSharedPtr<UE::StateTreeDebugger::FEventTreeElement>>& OutChildren)
@@ -507,6 +551,7 @@ void SStateTreeDebuggerView::Construct(const FArguments& InArgs, const UStateTre
 			{
 				FDetailsViewArgs DetailsViewArgs;
 				DetailsViewArgs.NameAreaSettings = FDetailsViewArgs::HideNameArea;
+				DetailsViewArgs.DefaultsOnlyVisibility = EEditDefaultsOnlyNodeVisibility::Hide;
 
 				UScriptStruct* ScriptStruct = FindObject<UScriptStruct>(nullptr, *TypePath, /*ExactClass*/false);
 				if (ScriptStruct == nullptr)
@@ -622,7 +667,7 @@ void SStateTreeDebuggerView::Construct(const FArguments& InArgs, const UStateTre
 						.ScrubPosition(ScrubTimeAttribute)
 						.ViewRange_Lambda([this]() { return ViewRange; })
 						.OnViewRangeChanged_Lambda([this](TRange<double> NewRange) { ViewRange = NewRange; })
-						.ClampRange_Lambda([this]() { return TRange<double>(0, Debugger->GetRecordingDuration()); })
+						.ClampRange_Lambda([this]() { return TRange<double>(0, MaxTrackRecordingDuration); })
 						.OnScrubPositionChanged_Lambda([this](double NewScrubTime, bool bIsScrubbing) { OnTimeLineScrubPositionChanged(NewScrubTime, bIsScrubbing); })
 					]
 				]
@@ -708,10 +753,16 @@ void SStateTreeDebuggerView::Tick(const FGeometry& AllottedGeometry, const doubl
 {
 	QUICK_SCOPE_CYCLE_COUNTER(STAT_StateTreeDebuggerView_TickView);
 
-	// Stick to most recent data if required
-	if (bAutoScroll && !Debugger->IsAnalysisSessionPaused())
+	check(Debugger);
+	if (Debugger->IsAnalysisSessionActive() && !Debugger->IsAnalysisSessionPaused())
 	{
-		Debugger->SetScrubTime(Debugger->GetRecordingDuration());
+		MaxTrackRecordingDuration = FMath::Max(MaxTrackRecordingDuration, Debugger->GetRecordingDuration());
+
+		// Stick to most recent data if required
+		if (bAutoScroll)
+		{
+			Debugger->SetScrubTime(Debugger->GetRecordingDuration());
+		}
 	}
 	
 	RefreshTracks();
@@ -775,45 +826,64 @@ void SStateTreeDebuggerView::BindDebuggerToolbarCommands(const TSharedRef<FUICom
 		Commands.ToggleAnalysis,
 		FExecuteAction::CreateSP(this, &SStateTreeDebuggerView::ToggleDebuggerAnalysis),
 		FCanExecuteAction::CreateSP(this, &SStateTreeDebuggerView::CanToggleDebuggerAnalysis));
+
+	ToolkitCommands->MapAction(
+		Commands.ResetTracks,
+		FExecuteAction::CreateSP(this, &SStateTreeDebuggerView::ResetTracks),
+		FCanExecuteAction::CreateSP(this, &SStateTreeDebuggerView::CanResetTracks));
+}
+
+bool SStateTreeDebuggerView::CanUseScrubButtons() const
+{
+	check(Debugger);
+	return (bAutoScroll == false || Debugger->IsAnalysisSessionPaused() || !Debugger->IsAnalysisSessionActive());
 }
 
 bool SStateTreeDebuggerView::CanStepBackToPreviousStateWithEvents() const
 {
-	return (bAutoScroll == false || Debugger->IsAnalysisSessionPaused()) && Debugger->CanStepBackToPreviousStateWithEvents();
+	check(Debugger);
+	return CanUseScrubButtons() && Debugger->CanStepBackToPreviousStateWithEvents();
 }
 
 void SStateTreeDebuggerView::StepBackToPreviousStateWithEvents()
 {
+	check(Debugger);
 	Debugger->StepBackToPreviousStateWithEvents();
 }
 
 bool SStateTreeDebuggerView::CanStepForwardToNextStateWithEvents() const
 {
-	return (bAutoScroll == false || Debugger->IsAnalysisSessionPaused()) && Debugger->CanStepForwardToNextStateWithEvents();
+	check(Debugger);
+	return CanUseScrubButtons() && Debugger->CanStepForwardToNextStateWithEvents();
 }
 
 void SStateTreeDebuggerView::StepForwardToNextStateWithEvents()
 {
+	check(Debugger);
 	Debugger->StepForwardToNextStateWithEvents();
 }
 
 bool SStateTreeDebuggerView::CanStepBackToPreviousStateChange() const
 {
-	return (bAutoScroll == false || Debugger->IsAnalysisSessionPaused()) && Debugger->CanStepBackToPreviousStateChange();
+	check(Debugger);
+	return CanUseScrubButtons() && Debugger->CanStepBackToPreviousStateChange();
 }
 
 void SStateTreeDebuggerView::StepBackToPreviousStateChange()
 {
+	check(Debugger);
 	Debugger->StepBackToPreviousStateChange();
 }
 
 bool SStateTreeDebuggerView::CanStepForwardToNextStateChange() const
 {
-	return (bAutoScroll == false || Debugger->IsAnalysisSessionPaused()) && Debugger->CanStepForwardToNextStateChange();
+	check(Debugger);
+	return CanUseScrubButtons() && Debugger->CanStepForwardToNextStateChange();
 }
 
 void SStateTreeDebuggerView::StepForwardToNextStateChange()
 {
+	check(Debugger);
 	Debugger->StepForwardToNextStateChange();
 }
 
@@ -1014,7 +1084,11 @@ UStateTreeState* SStateTreeDebuggerView::FindStateAssociatedToBreakpoint(FStateT
 
 void SStateTreeDebuggerView::OnTimeLineScrubPositionChanged(double Time, bool bIsScrubbing)
 {
-	bAutoScroll = false;
+	check(Debugger);
+	if (Debugger->IsAnalysisSessionActive() && !Debugger->IsAnalysisSessionPaused())
+	{
+		bAutoScroll = false;
+	}
 	Debugger->SetScrubTime(Time);
 }
 
@@ -1135,9 +1209,8 @@ void SStateTreeDebuggerView::OnBreakpointHit(const FStateTreeInstanceDebugId Ins
 	// Extract associated UStateTreeState to focus on it.
 	if (UStateTreeState* AssociatedState = FindStateAssociatedToBreakpoint(Breakpoint))
 	{
-		FStateTreeViewModel* ViewModel = StateTreeViewModel.Get();
-		check(ViewModel);
-		ViewModel->SetSelection(AssociatedState);
+		check(StateTreeViewModel);
+		StateTreeViewModel->SetSelection(AssociatedState);
 	}
 
 	// Find matching event in the tree view and select it
@@ -1164,23 +1237,25 @@ void SStateTreeDebuggerView::OnBreakpointHit(const FStateTreeInstanceDebugId Ins
 
 void SStateTreeDebuggerView::OnNewSession()
 {
-	// Clear tracks
-	InstanceOwnerTracks.Reset();
-
-	// Refresh tree view
-	if (InstancesTreeView)
+	// We clear tracks:
+	//  - analysis is not for an Editor session
+	//  - explicitly set in the settings for Editor sessions
+	//  - if previous analysis was not an Editor session
+	check(Debugger);
+	if (!Debugger->IsAnalyzingEditorSession()
+		|| UStateTreeEditorSettings::Get().bShouldDebuggerResetDataOnNewPIESession
+		|| !Debugger->WasAnalyzingEditorSession())
 	{
-		InstancesTreeView->Refresh();
+		ResetTracks();
 	}
 
-	if (InstanceTimelinesTreeView)
-	{
-		InstanceTimelinesTreeView->Refresh();
-	}
+	// Restore automatic scroll to most recent data.
+	bAutoScroll = true;
 }
 
 void SStateTreeDebuggerView::OnNewInstance(FStateTreeInstanceDebugId InstanceId)
 {
+	check(Debugger);
 	const UE::StateTreeDebugger::FInstanceDescriptor* FoundDescriptor = Debugger->GetInstanceDescriptor(InstanceId);
 	if (!ensureMsgf(FoundDescriptor != nullptr, TEXT("This callback is from the Debugger so we expect to be able to find matching descriptor.")))
 	{
@@ -1239,38 +1314,31 @@ TSharedRef<SWidget> SStateTreeDebuggerView::OnGetDebuggerTracesMenu() const
 {
 	FMenuBuilder MenuBuilder(/*bInShouldCloseWindowAfterMenuSelection*/true, /*InCommandList*/nullptr);
 
-	if (Debugger.IsValid())
-	{
-		TArray<FStateTreeDebugger::FTraceDescriptor> TraceDescriptors;
-		Debugger->GetLiveTraces(TraceDescriptors);
+	TArray<FStateTreeDebugger::FTraceDescriptor> TraceDescriptors;
 
-		for (const FStateTreeDebugger::FTraceDescriptor& TraceDescriptor : TraceDescriptors)
-		{
-			const FText Desc = Debugger->DescribeTrace(TraceDescriptor);
+	check(Debugger);
+	Debugger->GetLiveTraces(TraceDescriptors);
+
+	for (const FStateTreeDebugger::FTraceDescriptor& TraceDescriptor : TraceDescriptors)
+	{
+		const FText Desc = Debugger->DescribeTrace(TraceDescriptor);
 
-			FUIAction ItemAction(FExecuteAction::CreateLambda([WeakDebugger = Debugger, TraceDescriptor]()
+		FUIAction ItemAction(FExecuteAction::CreateSPLambda(Debugger.ToSharedRef(), [Debugger = Debugger, TraceDescriptor]()
 			{
-				if (WeakDebugger)
-				{
-					WeakDebugger->StartSessionAnalysis(TraceDescriptor);
-				}
+				Debugger->RequestSessionAnalysis(TraceDescriptor);
 			}));
-			MenuBuilder.AddMenuEntry(Desc, TAttribute<FText>(), FSlateIcon(), ItemAction);
-		}
+		MenuBuilder.AddMenuEntry(Desc, TAttribute<FText>(), FSlateIcon(), ItemAction);
+	}
 
-		// Failsafe when no match
-		if (TraceDescriptors.Num() == 0)
-		{
-			const FText Desc = LOCTEXT("NoLiveSessions","Can't find live trace sessions");
-			FUIAction ItemAction(FExecuteAction::CreateLambda([WeakDebugger = Debugger]()
+	// Failsafe when no match
+	if (TraceDescriptors.Num() == 0)
+	{
+		const FText Desc = LOCTEXT("NoLiveSessions", "Can't find live trace sessions");
+		FUIAction ItemAction(FExecuteAction::CreateSPLambda(Debugger.ToSharedRef(), [Debugger = Debugger]()
 			{
-				if (WeakDebugger)
-				{
-					WeakDebugger->StartSessionAnalysis(FStateTreeDebugger::FTraceDescriptor());
-				}
+				Debugger->RequestSessionAnalysis(FStateTreeDebugger::FTraceDescriptor());
 			}));
-			MenuBuilder.AddMenuEntry(Desc, TAttribute<FText>(), FSlateIcon(), ItemAction);
-		}
+		MenuBuilder.AddMenuEntry(Desc, TAttribute<FText>(), FSlateIcon(), ItemAction);
 	}
 
 	return MenuBuilder.MakeWidget();
diff --git a/Engine/Plugins/Runtime/StateTree/Source/StateTreeEditorModule/Private/Debugger/SStateTreeDebuggerView.h b/Engine/Plugins/Runtime/StateTree/Source/StateTreeEditorModule/Private/Debugger/SStateTreeDebuggerView.h
--- a/Engine/Plugins/Runtime/StateTree/Source/StateTreeEditorModule/Private/Debugger/SStateTreeDebuggerView.h
+++ b/Engine/Plugins/Runtime/StateTree/Source/StateTreeEditorModule/Private/Debugger/SStateTreeDebuggerView.h
@@ -69,6 +69,10 @@ private:
 	void OnNewInstance(FStateTreeInstanceDebugId InstanceId);
 	void OnSelectedInstanceCleared();
 
+	void BindDebuggerToolbarCommands(const TSharedRef<FUICommandList>& ToolkitCommands);
+
+	bool CanUseScrubButtons() const;
+
 	bool CanStartRecording() const { return !IsRecording(); }
 	void StartRecording();
 
@@ -81,7 +85,8 @@ private:
 	void ToggleDebuggerAnalysis() const;
 	FSlateIcon GetDebuggerAnalysisIcon() const;
 
-	void BindDebuggerToolbarCommands(const TSharedRef<FUICommandList>& ToolkitCommands);
+	bool CanResetTracks() const;
+	void ResetTracks();
 
 	bool CanStepBackToPreviousStateWithEvents() const;
 	void StepBackToPreviousStateWithEvents();
@@ -160,6 +165,9 @@ private:
 	/** Object created from the event data when statetree node was holding an object. */
 	TWeakObjectPtr<> SelectedNodeDataObject;
 
+	/** In case tracks are not reset when a new analysis session is started we keep track of the longest duration to adjust our clamp range. */
+	double MaxTrackRecordingDuration = 0;
+
 	/** Indicates that a live session was started (record button or auto record in PIE) to generate StateTree traces. */
 	bool bRecording = false;
 
diff --git a/Engine/Plugins/Runtime/StateTree/Source/StateTreeEditorModule/Private/Debugger/StateTreeDebuggerCommands.cpp b/Engine/Plugins/Runtime/StateTree/Source/StateTreeEditorModule/Private/Debugger/StateTreeDebuggerCommands.cpp
--- a/Engine/Plugins/Runtime/StateTree/Source/StateTreeEditorModule/Private/Debugger/StateTreeDebuggerCommands.cpp
+++ b/Engine/Plugins/Runtime/StateTree/Source/StateTreeEditorModule/Private/Debugger/StateTreeDebuggerCommands.cpp
@@ -33,6 +33,7 @@ void FStateTreeDebuggerCommands::RegisterCommands()
 	UI_COMMAND(NextFrameWithEvents, "Next Frame", "Step one frame forward", EUserInterfaceActionType::Button, FInputChord(EKeys::Right));
 	UI_COMMAND(NextFrameWithStateChange, "Next State Change", "Jump to next state changed frame", EUserInterfaceActionType::Button, FInputChord(EModifierKey::Control, EKeys::Right));
 	UI_COMMAND(ToggleAnalysis, "Toggle Analysis",  "Pause or resume processing events from the selected trace.", EUserInterfaceActionType::Button, FInputChord());
+	UI_COMMAND(ResetTracks, "Reset Tracks",  "Clear current tracks and associated events.", EUserInterfaceActionType::Button, FInputChord());
 }
 
 #undef LOCTEXT_NAMESPACE
diff --git a/Engine/Plugins/Runtime/StateTree/Source/StateTreeEditorModule/Private/Debugger/StateTreeDebuggerCommands.h b/Engine/Plugins/Runtime/StateTree/Source/StateTreeEditorModule/Private/Debugger/StateTreeDebuggerCommands.h
--- a/Engine/Plugins/Runtime/StateTree/Source/StateTreeEditorModule/Private/Debugger/StateTreeDebuggerCommands.h
+++ b/Engine/Plugins/Runtime/StateTree/Source/StateTreeEditorModule/Private/Debugger/StateTreeDebuggerCommands.h
@@ -26,6 +26,7 @@ public:
 	TSharedPtr<FUICommandInfo> NextFrameWithEvents;
 	TSharedPtr<FUICommandInfo> NextFrameWithStateChange;
 	TSharedPtr<FUICommandInfo> ToggleAnalysis;
+	TSharedPtr<FUICommandInfo> ResetTracks;
 };
 
 #endif // WITH_STATETREE_DEBUGGER
\ No newline at end of file
diff --git a/Engine/Plugins/Runtime/StateTree/Source/StateTreeEditorModule/Private/Debugger/StateTreeDebuggerTrack.cpp b/Engine/Plugins/Runtime/StateTree/Source/StateTreeEditorModule/Private/Debugger/StateTreeDebuggerTrack.cpp
--- a/Engine/Plugins/Runtime/StateTree/Source/StateTreeEditorModule/Private/Debugger/StateTreeDebuggerTrack.cpp
+++ b/Engine/Plugins/Runtime/StateTree/Source/StateTreeEditorModule/Private/Debugger/StateTreeDebuggerTrack.cpp
@@ -66,11 +66,10 @@ bool FStateTreeDebuggerInstanceTrack::UpdateInternal()
 			const FStateTreeTraceActiveStatesEvent& Event = Events[EventIndex].Get<FStateTreeTraceActiveStatesEvent>();
 				
 			FString StatePath;
-			const TConstArrayView<FCompactStateTreeState> States = StateTree->GetStates();
 			for (int32 ActiveStateIndex = 0; ActiveStateIndex < Event.ActiveStates.Num(); ActiveStateIndex++)
 			{
-				const FCompactStateTreeState& State = States[Event.ActiveStates[ActiveStateIndex].Index];
-				StatePath.Appendf(TEXT("%s%s"), ActiveStateIndex == 0 ? TEXT("") : TEXT("."), *State.Name.ToString());
+				const FCompactStateTreeState* State = StateTree->GetStateFromHandle(Event.ActiveStates[ActiveStateIndex]);
+				StatePath.Appendf(TEXT("%s%s"), ActiveStateIndex == 0 ? TEXT("") : TEXT("."), State == nullptr ? TEXT("Invalid") : *State->Name.ToString());
 			}
 
 			UE::StateTreeDebugger::FFrameSpan Span = EventCollection.FrameSpans[SpanIndex];
diff --git a/Engine/Plugins/Runtime/StateTree/Source/StateTreeEditorModule/Private/StateTreeEditorStyle.cpp b/Engine/Plugins/Runtime/StateTree/Source/StateTreeEditorModule/Private/StateTreeEditorStyle.cpp
--- a/Engine/Plugins/Runtime/StateTree/Source/StateTreeEditorModule/Private/StateTreeEditorStyle.cpp
+++ b/Engine/Plugins/Runtime/StateTree/Source/StateTreeEditorModule/Private/StateTreeEditorStyle.cpp
@@ -218,9 +218,10 @@ FStateTreeEditorStyle::FStateTreeEditorStyle()
 		Set("StateTreeEditor.Debugger.InstanceTrack", new IMAGE_BRUSH_SVG("Starship/AssetIcons/AnimInstance_64", CoreStyleConstants::Icon16x16));
 
 		Set("StateTreeEditor.EnableStates", new IMAGE_BRUSH("Icons/Empty_16x", CoreStyleConstants::Icon16x16));
-		Set( "StateTreeEditor.Debugger.Breakpoint.EnabledAndValid", new IMAGE_BRUSH_SVG( "Starship/Blueprints/Breakpoint_Valid", CoreStyleConstants::Icon16x16, FStyleColors::AccentRed));
+		Set("StateTreeEditor.Debugger.Breakpoint.EnabledAndValid", new IMAGE_BRUSH_SVG( "Starship/Blueprints/Breakpoint_Valid", CoreStyleConstants::Icon16x16, FStyleColors::AccentRed));
 		Set("StateTreeEditor.PauseDebuggerAnalysis", new IMAGE_BRUSH("Icons/generic_pause_16x", CoreStyleConstants::Icon16x16));
 		Set("StateTreeEditor.ResumeDebuggerAnalysis", new IMAGE_BRUSH_SVG("Starship/Common/Timeline", CoreStyleConstants::Icon16x16));
+		Set("StateTreeEditor.Debugger.ResetTracks", new IMAGE_BRUSH("Icons/icon_TrackDelete_36x24px", CoreStyleConstants::Icon16x16));
 	}
 
 	{
diff --git a/Engine/Plugins/Runtime/StateTree/Source/StateTreeEditorModule/Public/StateTreeEditorSettings.h b/Engine/Plugins/Runtime/StateTree/Source/StateTreeEditorModule/Public/StateTreeEditorSettings.h
--- a/Engine/Plugins/Runtime/StateTree/Source/StateTreeEditorModule/Public/StateTreeEditorSettings.h
+++ b/Engine/Plugins/Runtime/StateTree/Source/StateTreeEditorModule/Public/StateTreeEditorSettings.h
@@ -28,6 +28,10 @@ public:
 	UPROPERTY(EditAnywhere, config, Category = "Debugger")
 	bool bShouldDebuggerAutoRecordOnPIE = true;
 
+	/** If enabled, debugger will clear previous tracks at the start of each PIE session. */
+	UPROPERTY(EditAnywhere, config, Category = "Debugger")
+	bool bShouldDebuggerResetDataOnNewPIESession = false;
+
 protected:
 #if WITH_EDITOR
 	virtual FText GetSectionText() const override;
diff --git a/Engine/Plugins/Runtime/StateTree/Source/StateTreeModule/Private/Debugger/StateTreeDebugger.cpp b/Engine/Plugins/Runtime/StateTree/Source/StateTreeModule/Private/Debugger/StateTreeDebugger.cpp
--- a/Engine/Plugins/Runtime/StateTree/Source/StateTreeModule/Private/Debugger/StateTreeDebugger.cpp
+++ b/Engine/Plugins/Runtime/StateTree/Source/StateTreeModule/Private/Debugger/StateTreeDebugger.cpp
@@ -112,7 +112,7 @@ void FStateTreeDebugger::StopSessionAnalysis()
 	if (const TraceServices::IAnalysisSession* Session = GetAnalysisSession())
 	{
 		Session->Stop(true);
-		AnalysisSession.Reset();	
+		AnalysisSession.Reset();
 	}
 
 	bSessionAnalysisPaused = false;
@@ -163,15 +163,12 @@ void FStateTreeDebugger::SelectInstance(const FStateTreeInstanceDebugId Instance
 		OnSelectedInstanceCleared.ExecuteIfBound();
 
 		// Update event collection index for newly debugged instance
-		ScrubState.SetEventCollectionIndex(InstanceId.IsValid() ? EventCollections.IndexOfByPredicate([InstanceId = InstanceId](const UE::StateTreeDebugger::FInstanceEventCollection& Entry)
-			{
-				return Entry.InstanceId == InstanceId;
-			})
+		SetScrubStateCollectionIndex(InstanceId.IsValid()
+			? EventCollections.IndexOfByPredicate([InstanceId = InstanceId](const UE::StateTreeDebugger::FInstanceEventCollection& Entry)
+				{
+					return Entry.InstanceId == InstanceId;
+				})
 			: INDEX_NONE);
-
-		OnScrubStateChanged.Execute(ScrubState);
-
-		RefreshActiveStates();
 	}
 }
 
@@ -201,7 +198,7 @@ void FStateTreeDebugger::UpdateInstances()
 	}
 }
 
-bool FStateTreeDebugger::RequestAnalysisOfNextLiveSession()
+bool FStateTreeDebugger::RequestAnalysisOfEditorSession()
 {
 	// Get snapshot of current trace to help identify the next live one
 	LastLiveSessionId = INDEX_NONE;
@@ -239,7 +236,7 @@ bool FStateTreeDebugger::RequestAnalysisOfNextLiveSession()
 				return Descriptor.TraceId == ActiveTraceId;
 			}))
 		{
-			return StartSessionAnalysis(*Descriptor);
+			return RequestSessionAnalysis(*Descriptor);
 		}
 	}
 
@@ -253,7 +250,18 @@ bool FStateTreeDebugger::TryStartNewLiveSessionAnalysis(const float RetryPolling
 
 	if (Traces.Num() && Traces.Last().TraceId != LastLiveSessionId)
 	{
-		return StartSessionAnalysis(Traces.Last());
+		// Intentional call to StartSessionAnalysis instead of RequestSessionAnalysis since we want
+		// to set 'bIsAnalyzingNextEditorSession' before calling OnNewSession delegate.
+		const bool bStarted = StartSessionAnalysis(Traces.Last());
+		if (bStarted)
+		{
+			UpdateAnalysisTransitionType(EAnalysisSourceType::EditorSession);
+
+			SetScrubStateCollectionIndex(INDEX_NONE);
+			OnNewSession.ExecuteIfBound();
+		}
+
+		return bStarted;
 	}
 	
 	RetryLoadNextLiveSessionTimer = RetryPollingDuration;
@@ -326,15 +334,21 @@ bool FStateTreeDebugger::StartSessionAnalysis(const FTraceDescriptor& TraceDescr
 		if (AnalysisSession.IsValid())
 		{
 			ActiveSessionTraceDescriptor = TraceDescriptor;
-			EventCollections.Reset();
-			ScrubState.SetEventCollectionIndex(INDEX_NONE);
-			OnNewSession.ExecuteIfBound();
 		}
 	}
 
 	return ActiveSessionTraceDescriptor.IsValid();
 }
 
+void FStateTreeDebugger::SetScrubStateCollectionIndex(const int32 EventCollectionIndex)
+{
+	ScrubState.SetEventCollectionIndex(EventCollectionIndex);
+
+	OnScrubStateChanged.Execute(ScrubState);
+
+	RefreshActiveStates();
+}
+
 void FStateTreeDebugger::GetLiveTraces(TArray<FTraceDescriptor>& OutTraceDescriptors) const
 {
 	UE::Trace::FStoreClient* StoreClient = GetStoreClient();
@@ -629,6 +643,49 @@ const TraceServices::IAnalysisSession* FStateTreeDebugger::GetAnalysisSession()
 	return AnalysisSession.Get();
 }
 
+bool FStateTreeDebugger::RequestSessionAnalysis(const FTraceDescriptor& TraceDescriptor)
+{
+	if (StartSessionAnalysis(TraceDescriptor))
+	{
+		UpdateAnalysisTransitionType(EAnalysisSourceType::SelectedSession);
+		
+		SetScrubStateCollectionIndex(INDEX_NONE);
+		OnNewSession.ExecuteIfBound();
+		return true;
+	}
+	return false;
+}
+
+void FStateTreeDebugger::UpdateAnalysisTransitionType(const EAnalysisSourceType SourceType)
+{
+	switch (AnalysisTransitionType)
+	{
+	case EAnalysisTransitionType::Unset:
+		AnalysisTransitionType = (SourceType == EAnalysisSourceType::SelectedSession)
+				? EAnalysisTransitionType::NoneToSelected
+				: EAnalysisTransitionType::NoneToEditor;
+		break;
+
+	case EAnalysisTransitionType::NoneToSelected:
+	case EAnalysisTransitionType::EditorToSelected:
+	case EAnalysisTransitionType::SelectedToSelected:
+		AnalysisTransitionType = (SourceType == EAnalysisSourceType::SelectedSession)
+				? EAnalysisTransitionType::SelectedToSelected
+				: EAnalysisTransitionType::SelectedToEditor;
+		break;
+
+	case EAnalysisTransitionType::NoneToEditor:
+	case EAnalysisTransitionType::EditorToEditor:
+	case EAnalysisTransitionType::SelectedToEditor:
+		AnalysisTransitionType = (SourceType == EAnalysisSourceType::SelectedSession)
+				? EAnalysisTransitionType::EditorToSelected
+				: EAnalysisTransitionType::EditorToEditor;
+		break;
+	default:
+		ensureMsgf(false, TEXT("Unhandled transition type."));
+	}
+}
+
 UE::Trace::FStoreClient* FStateTreeDebugger::GetStoreClient() const
 {
 	return StateTreeModule.GetStoreClient();
@@ -772,14 +829,14 @@ bool FStateTreeDebugger::ProcessEvent(const FStateTreeInstanceDebugId InstanceId
 		// Push deferred notification for new instance Id
 		NewInstances.Push(InstanceId);
 
+		ExistingCollection = &EventCollections.Emplace_GetRef(InstanceId);
+
 		// Update the active event collection index when it's newly created for the currently debugged instance.
 		// Otherwise (i.e. EventCollection already exists) it is updated when switching instance (i.e. SelectInstance)
 		if (SelectedInstanceId == InstanceId && ScrubState.GetEventCollectionIndex() == INDEX_NONE)
 		{
-			ScrubState.SetEventCollectionIndex(EventCollections.Num());
+			SetScrubStateCollectionIndex(EventCollections.Num()-1);
 		}
-
-		ExistingCollection = &EventCollections.Emplace_GetRef(InstanceId);
 	}
 
 	check(ExistingCollection);
@@ -839,6 +896,12 @@ const UE::StateTreeDebugger::FInstanceEventCollection& FStateTreeDebugger::GetEv
 	return ExistingCollection != nullptr ? *ExistingCollection : FInstanceEventCollection::Invalid;
 }
 
+void FStateTreeDebugger::ResetEventCollections()
+{
+	EventCollections.Reset();
+	SetScrubStateCollectionIndex(INDEX_NONE);
+}
+
 void FStateTreeDebugger::AddEvents(const double StartTime, const double EndTime, const TraceServices::IFrameProvider& FrameProvider, const IStateTreeTraceProvider& StateTreeTraceProvider)
 {
 	check(StateTreeAsset.IsValid());
diff --git a/Engine/Plugins/Runtime/StateTree/Source/StateTreeModule/Private/Debugger/StateTreeDebuggerTypes.cpp b/Engine/Plugins/Runtime/StateTree/Source/StateTreeModule/Private/Debugger/StateTreeDebuggerTypes.cpp
--- a/Engine/Plugins/Runtime/StateTree/Source/StateTreeModule/Private/Debugger/StateTreeDebuggerTypes.cpp
+++ b/Engine/Plugins/Runtime/StateTree/Source/StateTreeModule/Private/Debugger/StateTreeDebuggerTypes.cpp
@@ -158,6 +158,10 @@ bool FScrubState::HasPreviousActiveStates() const
 	}
 
 	const TArray<FInstanceEventCollection::FActiveStatesChangePair>& ActiveStatesChanges = EventCollections[EventCollectionIndex].ActiveStatesChanges;
+	if (ScrubTimeBoundState == EScrubTimeBoundState::AfterHigherBound && ActiveStatesChanges.Num() > 0)
+	{
+		return true;
+	}
 
 	if (ActiveStatesChanges.IsValidIndex(ActiveStatesIndex) && ActiveStatesChanges[ActiveStatesIndex].SpanIndex < FrameSpanIndex)
 	{
@@ -170,7 +174,11 @@ bool FScrubState::HasPreviousActiveStates() const
 double FScrubState::GotoPreviousActiveStates()
 {
 	const TArray<FInstanceEventCollection::FActiveStatesChangePair>& ActiveStatesChanges = EventCollections[EventCollectionIndex].ActiveStatesChanges;
-	if (ActiveStatesChanges.IsValidIndex(ActiveStatesIndex) && ActiveStatesChanges[ActiveStatesIndex].SpanIndex < FrameSpanIndex)
+	if (ScrubTimeBoundState == EScrubTimeBoundState::AfterHigherBound)
+	{
+		SetActiveStatesIndex(ActiveStatesChanges.Num()-1);
+	}
+	else if (ActiveStatesChanges.IsValidIndex(ActiveStatesIndex) && ActiveStatesChanges[ActiveStatesIndex].SpanIndex < FrameSpanIndex)
 	{
 		SetActiveStatesIndex(ActiveStatesIndex);
 	}
diff --git a/Engine/Plugins/Runtime/StateTree/Source/StateTreeModule/Public/Debugger/StateTreeDebugger.h b/Engine/Plugins/Runtime/StateTree/Source/StateTreeModule/Public/Debugger/StateTreeDebugger.h
--- a/Engine/Plugins/Runtime/StateTree/Source/StateTreeModule/Public/Debugger/StateTreeDebugger.h
+++ b/Engine/Plugins/Runtime/StateTree/Source/StateTreeModule/Public/Debugger/StateTreeDebugger.h
@@ -122,6 +122,9 @@ struct STATETREEMODULE_API FStateTreeDebugger : FTickableGameObject
 	 */
 	const UE::StateTreeDebugger::FInstanceEventCollection& GetEventCollection(FStateTreeInstanceDebugId InstanceId) const;
 
+	/** Clears events from all instances. */
+	void ResetEventCollections();
+
 	/** Returns the recording duration in world recorded time. */
 	double GetRecordingDuration() const { return RecordingDuration; }
 
@@ -139,12 +142,31 @@ struct STATETREEMODULE_API FStateTreeDebugger : FTickableGameObject
 	 * This will replace the current analysis session if any.
 	 * @return True if connection was successfully requested or was able to use active trace, false otherwise.
 	 */
-	bool RequestAnalysisOfNextLiveSession();
+	bool RequestAnalysisOfEditorSession();
+	bool IsAnalyzingEditorSession() const
+	{
+		return AnalysisTransitionType == EAnalysisTransitionType::NoneToEditor
+			|| AnalysisTransitionType == EAnalysisTransitionType::EditorToEditor
+			|| AnalysisTransitionType == EAnalysisTransitionType::SelectedToEditor;
+	}
+
+	bool WasAnalyzingEditorSession() const
+	{
+		return AnalysisTransitionType == EAnalysisTransitionType::EditorToSelected
+			|| AnalysisTransitionType == EAnalysisTransitionType::EditorToEditor;
+	}
 
 	bool IsAnalysisSessionActive() const { return GetAnalysisSession() != nullptr; }
 	bool IsAnalysisSessionPaused() const { return bSessionAnalysisPaused; }
 	const TraceServices::IAnalysisSession* GetAnalysisSession() const;
-	bool StartSessionAnalysis(const FTraceDescriptor& TraceDescriptor);
+
+	/**
+	 * Tries to start an analysis for a given trace descriptor.
+	 * On success this method will execute the OnNewSession delegate.
+	 * @param TraceDescriptor Descriptor of the trace that needs to be analyzed
+	 * @return True if analysis was successfully started, false otherwise.
+	 */
+	bool RequestSessionAnalysis(const FTraceDescriptor& TraceDescriptor);
 	void PauseSessionAnalysis()  { bSessionAnalysisPaused = true; }
 	void ResumeSessionAnalysis() { bSessionAnalysisPaused = false; }
 	void StopSessionAnalysis();
@@ -166,6 +188,28 @@ protected:
 	virtual TStatId GetStatId() const override { RETURN_QUICK_DECLARE_CYCLE_STAT(FStateTreeDebugger, STATGROUP_Tickables); }
 	
 private:
+	enum class EAnalysisSourceType : uint8
+	{
+		// Analysis selected from available sessions
+		SelectedSession,
+
+		// Analysis automatically started from Editor new recording
+		EditorSession
+	};
+
+	enum class EAnalysisTransitionType : uint8
+	{
+		Unset,
+		NoneToSelected,
+		NoneToEditor,
+		EditorToSelected,
+		EditorToEditor,
+		SelectedToSelected,
+		SelectedToEditor,
+	};
+
+	void UpdateAnalysisTransitionType(EAnalysisSourceType SourceType);
+	
 	void ReadTrace(double ScrubTime);
 	void ReadTrace(uint64 FrameIndex);
 	void ReadTrace(
@@ -193,7 +237,11 @@ private:
 	 * @return True if the analysis was successfully started; false otherwise.
 	 */
 	bool TryStartNewLiveSessionAnalysis(float RetryPollingDuration);
-	
+
+	bool StartSessionAnalysis(const FTraceDescriptor& TraceDescriptor);
+
+	void SetScrubStateCollectionIndex(int32 EventCollectionIndex);
+
 	/**
 	 * Recompute index of the span that contains the active states change event and update the active states.
 	 * This method handles unselected instances in which case it will reset the active states and set the span index to INDEX_NONE
@@ -270,6 +318,11 @@ private:
 	 * This can be an external explicit request or after hitting a breakpoint.
 	 */
 	bool bSessionAnalysisPaused = false;
+
+	/**
+	 * Indicates the last transition type between two consecutive analyses to manage track cleanup properly.
+	 */
+	EAnalysisTransitionType AnalysisTransitionType = EAnalysisTransitionType::Unset;
 };
 
 #endif // WITH_STATETREE_DEBUGGER
