[StateTreeDebuggger]  Frame details view update:
- replaced node details view by embedding property values in the tree view
- node exported text (e.g. Task, Conditions, etc) is no longer imported in struct or object to avoid problems with missing types or property types with specific requirements (e.g. Object base property needs to find the specified value). Goal is to always be able to inspect key/value pairs.
- improved hierarchy of event to improve readibility
- started adding some custom row widgets for the different event types
- reworked the toolbar to regroup the different concepts (Simulation controls, Trace recording, Trace analyzing)
#rnx

diff --git a/Engine/Plugins/Runtime/StateTree/Source/StateTreeEditorModule/Private/Debugger/SStateTreeDebuggerView.cpp b/Engine/Plugins/Runtime/StateTree/Source/StateTreeEditorModule/Private/Debugger/SStateTreeDebuggerView.cpp
--- a/Engine/Plugins/Runtime/StateTree/Source/StateTreeEditorModule/Private/Debugger/SStateTreeDebuggerView.cpp
+++ b/Engine/Plugins/Runtime/StateTree/Source/StateTreeEditorModule/Private/Debugger/SStateTreeDebuggerView.cpp
@@ -3,12 +3,12 @@
 #if WITH_STATETREE_DEBUGGER
 
 #include "SStateTreeDebuggerView.h"
+#include "SStateTreeDebuggerViewRow.h"
 #include "Debugger/StateTreeDebugger.h"
 #include "DetailsViewArgs.h"
 #include "Editor.h"
 #include "Factories.h"
 #include "Framework/MultiBox/MultiBoxBuilder.h"
-#include "IStructureDetailsView.h"
 #include "Kismet2/DebuggerCommands.h"
 #include "Modules/ModuleManager.h"
 #include "PropertyEditorModule.h"
@@ -16,6 +16,7 @@
 #include "SStateTreeDebuggerInstanceTree.h"
 #include "SStateTreeDebuggerTimelines.h"
 #include "StateTree.h"
+#include "StateTreeCompiler.h"
 #include "StateTreeDebuggerCommands.h"
 #include "StateTreeDebuggerTrack.h"
 #include "StateTreeEditorData.h"
@@ -34,55 +35,17 @@
 
 namespace UE::StateTreeDebugger
 {
-//----------------------------------------------------------------------//
-// FTraceTextObjectFactory
-//----------------------------------------------------------------------//
-struct FTraceTextObjectFactory : FCustomizableTextObjectFactory
-{
-	UObject* NodeInstanceObject = nullptr;
-	FTraceTextObjectFactory() : FCustomizableTextObjectFactory(GWarn) {}
-
-protected:
-	virtual bool CanCreateClass(UClass* ObjectClass, bool& bOmitSubObjs) const override
-	{
-		return true;
-	}
-
-	virtual void ProcessConstructedObject(UObject* CreatedObject) override
-	{
-		NodeInstanceObject = CreatedObject;
-	}
-};
-
-
-//----------------------------------------------------------------------//
-// FEventTreeElement
-//----------------------------------------------------------------------//
-/** An item in the trace event tree */
-struct FEventTreeElement : TSharedFromThis<FEventTreeElement>
-{
-	explicit FEventTreeElement(const TraceServices::FFrame& Frame, const FStateTreeTraceEventVariantType& Event)
-		: Frame(Frame), Event(Event)
-	{
-	}
-
-	TraceServices::FFrame Frame;
-	FStateTreeTraceEventVariantType Event;
-	TArray<TSharedPtr<FEventTreeElement>> Children;
-};
-
-
 /**
  * Iterates over all tree elements for the frame events
  * @param Elements Container of hierarchical tree element to visit
  * @param InFunc function called at each element, should return true if visiting is continued or false to stop.
  */
-void VisitEventTreeElements(const TConstArrayView<TSharedPtr<FEventTreeElement>> Elements, TFunctionRef<bool(TSharedPtr<FEventTreeElement>& VisitedElement)> InFunc)
+void VisitEventTreeElements(const TConstArrayView<TSharedPtr<FStateTreeDebuggerEventTreeElement>> Elements, TFunctionRef<bool(TSharedPtr<FStateTreeDebuggerEventTreeElement>& VisitedElement)> InFunc)
 {
-	TArray<TSharedPtr<FEventTreeElement>> Stack;
+	TArray<TSharedPtr<FStateTreeDebuggerEventTreeElement>> Stack;
 	bool bContinue = true;
 
-	for (const TSharedPtr<FEventTreeElement>& RootElement : Elements)
+	for (const TSharedPtr<FStateTreeDebuggerEventTreeElement>& RootElement : Elements)
 	{
 		if (RootElement == nullptr)
 		{
@@ -93,7 +56,7 @@ void VisitEventTreeElements(const TConstArrayView<TSharedPtr<FEventTreeElement>>
 
 		while (!Stack.IsEmpty() && bContinue)
 		{
-			TSharedPtr<FEventTreeElement> StackedElement = Stack[0];
+			TSharedPtr<FStateTreeDebuggerEventTreeElement> StackedElement = Stack[0];
 			check(StackedElement);
 
 			Stack.RemoveAt(0);
@@ -102,7 +65,7 @@ void VisitEventTreeElements(const TConstArrayView<TSharedPtr<FEventTreeElement>>
 
 			if (bContinue)
 			{
-				for (const TSharedPtr<FEventTreeElement>& Child : StackedElement->Children)
+				for (const TSharedPtr<FStateTreeDebuggerEventTreeElement>& Child : StackedElement->Children)
 				{
 					if (Child.IsValid())
 					{
@@ -119,73 +82,113 @@ void VisitEventTreeElements(const TConstArrayView<TSharedPtr<FEventTreeElement>>
 	}
 }
 
-} // UE::StateTreeDebugger
-
-
 //----------------------------------------------------------------------//
-// SStateTreeDebuggerTableRow
+// FTraceTextObjectFactory
 //----------------------------------------------------------------------//
-class SStateTreeDebuggerTableRow : public SMultiColumnTableRow<TSharedPtr<UE::StateTreeDebugger::FEventTreeElement>>
+struct FTraceTextObjectFactory : FCustomizableTextObjectFactory
 {
-public:
-	void Construct(const FArguments& InArgs,
-		const TSharedPtr<STableViewBase>& InOwnerTableView,
-		const TSharedPtr<UE::StateTreeDebugger::FEventTreeElement>& InElement,
-		const TSharedRef<FStateTreeViewModel>& InStateTreeViewModel)
+	UObject* NodeInstanceObject = nullptr;
+	FTraceTextObjectFactory() : FCustomizableTextObjectFactory(GWarn) {}
+
+protected:
+	virtual bool CanCreateClass(UClass* ObjectClass, bool& bOmitSubObjs) const override
 	{
-		Item = InElement;
-		StateTreeViewModel = InStateTreeViewModel.ToSharedPtr();
-		SMultiColumnTableRow::Construct(InArgs, InOwnerTableView.ToSharedRef());
+		return true;
 	}
 
-protected:
-	// SMultiColumnTableRow overrides
-	virtual TSharedRef<SWidget> GenerateWidgetForColumn(const FName& ColumnName) override
+	virtual void ProcessConstructedObject(UObject* CreatedObject) override
 	{
-		const TSharedPtr<SHorizontalBox> Contents = SNew(SHorizontalBox);
+		NodeInstanceObject = CreatedObject;
+	}
+};
 
-		Contents->AddSlot()
-			.VAlign(VAlign_Fill)
-			.HAlign(HAlign_Left)
-			.AutoWidth()
-			[
-				SNew(SExpanderArrow, SharedThis(this))
-				.ShouldDrawWires(true)
-				.IndentAmount(32)
-				.BaseIndentLevel(0)
-			];
-		
-		if (ColumnName == FName("Desc"))
+/**
+ * Prototype methode to try import struct/object from exported text.
+ * Currently not used since it might fail if type is not found and for requires every property to be created
+ * but some types (i.e. Object based) will simply set none as the imported value but for the debugger we still
+ * want text to describe the dumped values.
+ */
+void GenerateElementForProperties(const TCHAR* TypeAsText, const TCHAR* ValueAsText, TSharedRef<FStateTreeDebuggerEventTreeElement>& ParentElement)
+{
+	const UStruct* StructTypeToVisit = nullptr;
+	const void* StructValueToVisit = nullptr;
+
+	// UStruct
+	UScriptStruct* ScriptStruct = FindObject<UScriptStruct>(nullptr, TypeAsText, /*ExactClass*/false);
+	if (ScriptStruct == nullptr)
+	{
+		ScriptStruct = LoadObject<UScriptStruct>(nullptr, TypeAsText);
+	}
+
+	FInstancedStruct NodeDataStruct;
+	if (ScriptStruct != nullptr)
+	{
+		NodeDataStruct.InitializeAs(ScriptStruct);
+
+		ScriptStruct->ImportText(ValueAsText, NodeDataStruct.GetMutableMemory(), /*OwnerObject*/nullptr, PPF_None, GLog, ScriptStruct->GetName());
+		StructTypeToVisit = ScriptStruct;
+		StructValueToVisit = NodeDataStruct.GetMemory();
+	}
+
+	// UObject
+	UClass* Class = FindObject<UClass>(nullptr, TypeAsText, /*ExactClass*/false);
+	if (Class == nullptr)
+	{
+		Class = LoadObject<UClass>(nullptr, TypeAsText);
+	}
+
+	TWeakObjectPtr<> NodeDataObject;
+	if (Class != nullptr)
+	{
+		FTraceTextObjectFactory ObjectFactory;
+		if (ObjectFactory.CanCreateObjectsFromText(ValueAsText))
 		{
-			Contents->AddSlot()
-				.Padding(5, 0)
-				.FillWidth(1.f)
-				[
-					SNew(STextBlock)
-					.Font(FCoreStyle::GetDefaultFontStyle("Mono", 9))
-					.Text_Lambda([&Event=Item->Event, this]()
-					{
-						FString EventDescription;
-						check(StateTreeViewModel);
-						if (const UStateTree* StateTree = StateTreeViewModel->GetStateTree())
-						{
-							Visit([&EventDescription, StateTree](auto& TypedEvent)
-								{
-									EventDescription = TypedEvent.ToString(*StateTree);
-								}, Event);
-						}
-						return FText::FromString(EventDescription);
-					})
-				];
-		}
+			ObjectFactory.ProcessBuffer(GetTransientPackage(), RF_Transactional, ValueAsText);
+			if (ObjectFactory.NodeInstanceObject)
+			{
+				NodeDataObject = ObjectFactory.NodeInstanceObject;
+				NodeDataObject->AddToRoot();
 
-		return Contents.ToSharedRef();
+				StructTypeToVisit = ObjectFactory.NodeInstanceObject->GetClass();
+				StructValueToVisit = ObjectFactory.NodeInstanceObject;
+			}
+		}
 	}
 
-	TSharedPtr<FStateTreeViewModel> StateTreeViewModel;
-	TSharedPtr<UE::StateTreeDebugger::FEventTreeElement> Item;
-};
+	if (StructTypeToVisit != nullptr && StructValueToVisit != nullptr)
+	{
+		for (TPropertyValueIterator<const FProperty> PropertyIt(StructTypeToVisit, StructValueToVisit); PropertyIt; ++PropertyIt)
+		{
+			const FProperty* const Property = PropertyIt.Key();
+			check(Property);
+			const EStateTreePropertyUsage Usage = UE::StateTree::Compiler::GetUsageFromMetaData(Property);
 
+			// If the property is set to one of these usages, display it even if it is not edit on instance.
+			// It is a common mistake to forget to set the "eye" on these properties it and wonder why it does not show up.
+			const bool bShouldShowByUsage = Usage == EStateTreePropertyUsage::Input || Usage == EStateTreePropertyUsage::Output || Usage == EStateTreePropertyUsage::Context;
+			const bool bIsEditable = Property->HasAllPropertyFlags(CPF_Edit) && !Property->HasAllPropertyFlags(CPF_DisableEditOnInstance);
+
+			if (bShouldShowByUsage || bIsEditable)
+			{
+				FText Name = Property->GetDisplayNameText();
+				FString ValueString;
+				Property->ExportTextItem_Direct(ValueString, PropertyIt.Value(), nullptr, nullptr, PPF_None);
+
+				// Create new property Log
+				FStateTreeTracePropertyEvent PropertyEvent(/*RecordingWorldTime*/0, FString::Printf(TEXT("%s = %s"), *Name.ToString(), *ValueString));
+
+				// Create Tree element to hold the event
+				const TSharedPtr<FStateTreeDebuggerEventTreeElement> NewChildElement = MakeShareable(new FStateTreeDebuggerEventTreeElement(
+						ParentElement->Frame,
+						FStateTreeTraceEventVariantType(TInPlaceType<FStateTreeTracePropertyEvent>(), PropertyEvent)));
+
+				ParentElement->Children.Add(NewChildElement);
+			}
+		}
+	}
+}
+
+} // UE::StateTreeDebugger
 
 //----------------------------------------------------------------------//
 // SStateTreeDebuggerView
@@ -201,11 +204,6 @@ SStateTreeDebuggerView::SStateTreeDebuggerView()
 
 SStateTreeDebuggerView::~SStateTreeDebuggerView()
 {
-	if (SelectedNodeDataObject.IsValid())
-	{
-		SelectedNodeDataObject->RemoveFromRoot();
-	}
-
 	check(Debugger);
 	Debugger->OnScrubStateChanged.Unbind();
 	Debugger->OnBreakpointHit.Unbind();
@@ -257,33 +255,19 @@ void SStateTreeDebuggerView::StopRecording()
 	}
 }
 
-bool SStateTreeDebuggerView::CanToggleDebuggerAnalysis() const
+bool SStateTreeDebuggerView::CanResumeDebuggerAnalysis() const
 {
 	check(Debugger);
-	return Debugger->IsAnalysisSessionActive(); 
+	return Debugger->IsAnalysisSessionPaused() && Debugger->HasHitBreakpoint(); 
 }
 
-void SStateTreeDebuggerView::ToggleDebuggerAnalysis() const
+void SStateTreeDebuggerView::ResumeDebuggerAnalysis() const
 {
 	check(Debugger);
 	if (Debugger->IsAnalysisSessionPaused())
 	{
 		Debugger->ResumeSessionAnalysis();
 	}
-	else
-	{
-		Debugger->PauseSessionAnalysis();
-	}
-}
-
-FSlateIcon SStateTreeDebuggerView::GetDebuggerAnalysisIcon() const
-{
-	check(Debugger);
-	if (Debugger->IsAnalysisSessionActive() && !Debugger->IsAnalysisSessionPaused())
-	{
-		return FSlateIcon(FStateTreeEditorStyle::Get().GetStyleSetName(), "StateTreeEditor.PauseDebuggerAnalysis");
-	}
-	return FSlateIcon(FStateTreeEditorStyle::Get().GetStyleSetName(), "StateTreeEditor.ResumeDebuggerAnalysis");
 }
 
 bool SStateTreeDebuggerView::CanResetTracks() const
@@ -394,19 +378,23 @@ void SStateTreeDebuggerView::Construct(const FArguments& InArgs, const UStateTre
 		FCanExecuteAction(),
 		FGetActionCheckState::CreateLambda([this] { return GetStateBreakpointCheckState(EStateTreeBreakpointType::OnExit); }),
 		FIsActionButtonVisible::CreateLambda([this] { return CanAddStateBreakpoint(EStateTreeBreakpointType::OnExit) || CanRemoveStateBreakpoint(EStateTreeBreakpointType::OnExit); }));
-	
+
 	// Toolbars
 	FSlimHorizontalToolBarBuilder LeftToolbar(InCommandList, FMultiBoxCustomization::None, /*InExtender*/ nullptr, /*InForceSmallIcons*/ true);
 	LeftToolbar.BeginSection(TEXT("Debugging"));
 	{
+		LeftToolbar.BeginStyleOverride(FName("Toolbar.BackplateLeft"));
 		const FPlayWorldCommands& PlayWorldCommands = FPlayWorldCommands::Get();
 		LeftToolbar.AddToolBarButton(PlayWorldCommands.RepeatLastPlay);
 		LeftToolbar.AddToolBarButton(PlayWorldCommands.PausePlaySession,
 					NAME_None, TAttribute<FText>(), TAttribute<FText>(), FSlateIcon(FAppStyle::GetAppStyleSetName(), "PlayWorld.PausePlaySession.Small"));
 		LeftToolbar.AddToolBarButton(PlayWorldCommands.ResumePlaySession,
-					NAME_None, TAttribute<FText>(), TAttribute<FText>(), FSlateIcon(FAppStyle::GetAppStyleSetName(), "PlayWorld.ResumePlaySession.Small"));	
+					NAME_None, TAttribute<FText>(), TAttribute<FText>(), FSlateIcon(FAppStyle::GetAppStyleSetName(), "PlayWorld.ResumePlaySession.Small"));
+		LeftToolbar.BeginStyleOverride(FName("Toolbar.BackplateRight"));
 		LeftToolbar.AddToolBarButton(PlayWorldCommands.StopPlaySession,
 					NAME_None, TAttribute<FText>(), TAttribute<FText>(), FSlateIcon(FAppStyle::GetAppStyleSetName(), "PlayWorld.StopPlaySession.Small"));
+		LeftToolbar.EndStyleOverride();
+
 		LeftToolbar.AddSeparator();
 
 		const FStateTreeDebuggerCommands& DebuggerCommands = FStateTreeDebuggerCommands::Get();
@@ -415,53 +403,48 @@ void SStateTreeDebuggerView::Construct(const FArguments& InArgs, const UStateTre
 
 		LeftToolbar.AddSeparator();
 
-		LeftToolbar.AddToolBarButton(DebuggerCommands.ToggleAnalysis,
-			NAME_None,
-			TAttribute<FText>(),
-			TAttribute<FText>(),
-			TAttribute<FSlateIcon>::CreateSP(this, &SStateTreeDebuggerView::GetDebuggerAnalysisIcon));
-		LeftToolbar.AddToolBarButton(DebuggerCommands.ResetTracks);
-
-		LeftToolbar.AddSeparator();
+		LeftToolbar.AddToolBarButton(DebuggerCommands.ResumeDebuggerAnalysis);
 
 		LeftToolbar.AddToolBarButton(DebuggerCommands.PreviousFrameWithStateChange);
 		LeftToolbar.AddToolBarButton(DebuggerCommands.PreviousFrameWithEvents);
 		LeftToolbar.AddToolBarButton(DebuggerCommands.NextFrameWithEvents);
 		LeftToolbar.AddToolBarButton(DebuggerCommands.NextFrameWithStateChange);
+
+		LeftToolbar.AddToolBarButton(DebuggerCommands.ResetTracks);
 	}
 	LeftToolbar.EndSection();
 
 	FSlimHorizontalToolBarBuilder RightToolbar(nullptr, FMultiBoxCustomization::None);
-
-	FUIAction AutoScrollToggleButtonAction;
-	AutoScrollToggleButtonAction.GetActionCheckState.BindLambda([this]
-	{
-		return bAutoScroll ? ECheckBoxState::Checked : ECheckBoxState::Unchecked;
-	});
-	AutoScrollToggleButtonAction.ExecuteAction.BindLambda([this]
+	RightToolbar.BeginSection("Auto-Scroll");
 	{
-		bAutoScroll = !bAutoScroll;
-	});
+		FUIAction AutoScrollToggleButtonAction;
+		AutoScrollToggleButtonAction.GetActionCheckState.BindSPLambda(this, [&bAutoScroll=bAutoScroll]
+		{
+			return bAutoScroll ? ECheckBoxState::Checked : ECheckBoxState::Unchecked;
+		});
+		AutoScrollToggleButtonAction.ExecuteAction.BindSPLambda(this, [&bAutoScroll=bAutoScroll]
+		{
+			bAutoScroll = !bAutoScroll;
+		});
 
-	RightToolbar.BeginSection("Auto-Scroll");
-	RightToolbar.AddToolBarButton(
-		AutoScrollToggleButtonAction,
-		NAME_None,
-		TAttribute<FText>(),
-		LOCTEXT("AutoScrollToolTip", "Auto-Scroll"),
-		FSlateIcon(FStateTreeEditorStyle::Get().GetStyleSetName(),"StateTreeEditor.AutoScroll"),
-		EUserInterfaceActionType::ToggleButton);
+		RightToolbar.AddToolBarButton
+		(
+			AutoScrollToggleButtonAction,
+			NAME_None,
+			TAttribute<FText>(),
+			LOCTEXT("AutoScrollToolTip", "Auto-Scroll"),
+			FSlateIcon(FStateTreeEditorStyle::Get().GetStyleSetName(), "StateTreeEditor.AutoScroll"),
+			EUserInterfaceActionType::ToggleButton);
+	}
 	RightToolbar.EndSection();
 
-	// Auto-select session if there is only one available
-	TArray<FStateTreeDebugger::FTraceDescriptor> TraceDescriptors;
-	Debugger->GetLiveTraces(TraceDescriptors);
-
-	// Make sure that the active trace contains statetree specific traces
-	if (TraceDescriptors.Num() == 1 && StateTreeModule.IsTracing())
+	// Place holder toolbar for now but the intent to add more functionalities (e.g. Search)
+	FSlimHorizontalToolBarBuilder FrameDetailsToolbar(nullptr, FMultiBoxCustomization::None);
+	FrameDetailsToolbar.BeginSection("FrameDetails");
 	{
-		Debugger->RequestSessionAnalysis(TraceDescriptors.Last());
+		FrameDetailsToolbar.AddWidget(SNew(STextBlock).Text(FText::FromString(GetNameSafe(StateTree.Get()))));
 	}
+	FrameDetailsToolbar.EndSection();
 
 	// Trace selection combo
 	const TSharedRef<SWidget> TraceSelectionBox = SNew(SComboButton)
@@ -511,111 +494,21 @@ void SStateTreeDebuggerView::Construct(const FArguments& InArgs, const UStateTre
 		.OnScrubPositionChanged_Lambda([this](double NewScrubTime, bool bIsScrubbing) { OnTimeLineScrubPositionChanged(NewScrubTime, bIsScrubbing); });
 
 	// EventsTreeView
-	EventsTreeView = SNew(STreeView<TSharedPtr<UE::StateTreeDebugger::FEventTreeElement>>)
-		.HeaderRow(SNew(SHeaderRow)
-			+SHeaderRow::Column("Desc")
-			.DefaultLabel(LOCTEXT("FrameDetailsColumnHeader", "Frame Details")))
-			.OnGenerateRow_Lambda([this](const TSharedPtr<UE::StateTreeDebugger::FEventTreeElement>& InElement, const TSharedRef<STableViewBase>& InOwnerTableView)
+	EventsTreeView = SNew(STreeView<TSharedPtr<FStateTreeDebuggerEventTreeElement>>)
+			.OnGenerateRow_Lambda([this](const TSharedPtr<FStateTreeDebuggerEventTreeElement>& InElement, const TSharedRef<STableViewBase>& InOwnerTableView)
 			{
 				check(StateTreeViewModel);
-				return SNew(SStateTreeDebuggerTableRow, InOwnerTableView, InElement, StateTreeViewModel.ToSharedRef());
+				return SNew(SStateTreeDebuggerViewRow, InOwnerTableView, InElement, StateTreeViewModel.ToSharedRef());
 			})
-			.OnGetChildren_Lambda([](const TSharedPtr<const UE::StateTreeDebugger::FEventTreeElement>& InParent, TArray<TSharedPtr<UE::StateTreeDebugger::FEventTreeElement>>& OutChildren)
+			.OnGetChildren_Lambda([](const TSharedPtr<const FStateTreeDebuggerEventTreeElement>& InParent, TArray<TSharedPtr<FStateTreeDebuggerEventTreeElement>>& OutChildren)
 			{
-				if (const UE::StateTreeDebugger::FEventTreeElement* Parent = InParent.Get())
+				if (const FStateTreeDebuggerEventTreeElement* Parent = InParent.Get())
 				{
 					OutChildren.Append(Parent->Children);
 				}
 			})
 		.TreeItemsSource(&EventsTreeElements)
 		.ItemHeight(32)
-		.OnSelectionChanged_Lambda([this](const TSharedPtr<UE::StateTreeDebugger::FEventTreeElement>& InSelectedItem, ESelectInfo::Type SelectionType)
-		{
-			if (!InSelectedItem.IsValid())
-			{
-				return;
-			}
-
-			TSharedPtr<SWidget> DetailsView;
-
-			FString TypePath;
-			FString InstanceDataAsText;
-
-			Visit([&TypePath, &InstanceDataAsText](auto& TypedEvent)
-				{
-					TypePath = TypedEvent.GetDataTypePath();
-					InstanceDataAsText = TypedEvent.GetDataAsText();
-				}, InSelectedItem->Event);
-				
-			if (!TypePath.IsEmpty())
-			{
-				FDetailsViewArgs DetailsViewArgs;
-				DetailsViewArgs.NameAreaSettings = FDetailsViewArgs::HideNameArea;
-				DetailsViewArgs.DefaultsOnlyVisibility = EEditDefaultsOnlyNodeVisibility::Hide;
-
-				UScriptStruct* ScriptStruct = FindObject<UScriptStruct>(nullptr, *TypePath, /*ExactClass*/false);
-				if (ScriptStruct == nullptr)
-				{
-					ScriptStruct = LoadObject<UScriptStruct>(nullptr, *TypePath);
-				}
-
-				if (ScriptStruct != nullptr)
-				{
-					SelectedNodeDataStruct.InitializeAs(ScriptStruct);
-
-					ScriptStruct->ImportText(*InstanceDataAsText, SelectedNodeDataStruct.GetMutableMemory(), /*OwnerObject*/nullptr, PPF_None, GLog, ScriptStruct->GetName());
-
-					SelectedNodeDataStruct.GetScriptStruct();
-					FPropertyEditorModule& PropertyEditorModule = FModuleManager::LoadModuleChecked<FPropertyEditorModule>("PropertyEditor");
-					const TSharedPtr<FStructOnScope> InstanceDataStruct = MakeShared<FStructOnScope>(SelectedNodeDataStruct.GetScriptStruct(), const_cast<uint8*>(SelectedNodeDataStruct.GetMemory()));
-
-					FStructureDetailsViewArgs StructureViewArgs;
-					StructureViewArgs.bShowObjects = true;
-					StructureViewArgs.bShowAssets = true;
-					StructureViewArgs.bShowClasses = true;
-					StructureViewArgs.bShowInterfaces = true;
-					const TSharedRef<IStructureDetailsView> StructDetailsView = PropertyEditorModule.CreateStructureDetailView(DetailsViewArgs, StructureViewArgs, InstanceDataStruct);
-
-					DetailsView = StructDetailsView->GetDetailsView()->AsShared();
-				}
-
-				// UObject
-				UClass* Class = FindObject<UClass>(nullptr, *TypePath, /*ExactClass*/false);
-				if (Class == nullptr)
-				{
-					Class = LoadObject<UClass>(nullptr, *TypePath);
-				}
-
-				if (Class != nullptr)
-				{
-					if (SelectedNodeDataObject.IsValid())
-					{
-						SelectedNodeDataObject->RemoveFromRoot();
-					}
-					UE::StateTreeDebugger::FTraceTextObjectFactory ObjectFactory;
-					if (ObjectFactory.CanCreateObjectsFromText(InstanceDataAsText))
-					{
-						ObjectFactory.ProcessBuffer(GetTransientPackage(), RF_Transactional, InstanceDataAsText);
-						SelectedNodeDataObject = ObjectFactory.NodeInstanceObject;
-						SelectedNodeDataObject->AddToRoot();
-
-						FPropertyEditorModule& PropertyEditorModule = FModuleManager::Get().GetModuleChecked<FPropertyEditorModule>("PropertyEditor");
-						const TSharedRef<IDetailsView> ObjectDetailsView = PropertyEditorModule.CreateDetailView(DetailsViewArgs);
-						ObjectDetailsView->SetObject(SelectedNodeDataObject.Get());
-						DetailsView = ObjectDetailsView->AsShared();
-					}
-				}
-			}
-
-			if (DetailsView)
-			{
-				PropertiesBorder->SetContent(DetailsView.ToSharedRef());
-			}
-			else
-			{
-				PropertiesBorder->ClearContent();
-			}
-		})
 		.AllowOverscroll(EAllowOverscroll::No);
 
 	ChildSlot
@@ -626,127 +519,133 @@ void SStateTreeDebuggerView::Construct(const FArguments& InArgs, const UStateTre
 		[
 			SNew(SVerticalBox)
 			+ SVerticalBox::Slot()
-			.AutoHeight()
-			.HAlign(HAlign_Fill)
-			.Padding(0.0f, 0.0f, 0.0f, 4.0f)
-			[
-				SNew(SHorizontalBox)
-				+ SHorizontalBox::Slot()
-				.HAlign(HAlign_Left)
-				.FillWidth(1.0f)
-				[
-					LeftToolbar.MakeWidget()
-				]
-				+SHorizontalBox::Slot()
-				.HAlign(HAlign_Right)
-				.AutoWidth()
-				[
-					RightToolbar.MakeWidget()
-				]
-			]
-			+ SVerticalBox::Slot().AutoHeight()
-			[
-				SNew(SBorder)
-				[
-					SAssignNew(HeaderSplitter, SSplitter)
-					.Orientation(Orient_Horizontal)
-					+ SSplitter::Slot()
-					.Value(0.2f)
-					.MinSize(350)
-					.Resizable(false)
-					[
-						TraceSelectionBox
-					]
-					+ SSplitter::Slot()
-					.Resizable(false)
-					[
-						SNew(SSimpleTimeSlider)
-						.DesiredSize({100, 24})
-						.ClampRangeHighlightSize(0.15f)
-						.ClampRangeHighlightColor(FLinearColor::Red.CopyWithNewOpacity(0.5f))
-						.ScrubPosition(ScrubTimeAttribute)
-						.ViewRange_Lambda([this]() { return ViewRange; })
-						.OnViewRangeChanged_Lambda([this](TRange<double> NewRange) { ViewRange = NewRange; })
-						.ClampRange_Lambda([this]() { return TRange<double>(0, MaxTrackRecordingDuration); })
-						.OnScrubPositionChanged_Lambda([this](double NewScrubTime, bool bIsScrubbing) { OnTimeLineScrubPositionChanged(NewScrubTime, bIsScrubbing); })
-					]
-				]
-			]
-			+ SVerticalBox::Slot()
 			[
 				SNew(SSplitter)
-				.Orientation(Orient_Vertical)
+				.Orientation(Orient_Horizontal)
 				+ SSplitter::Slot()
-				.Value(0.2f)
+				.MinSize(600)
 				[
-					SAssignNew(TreeViewsSplitter, SSplitter)
-					.Orientation(Orient_Horizontal)
-					+ SSplitter::Slot()
-					.Value(0.2f)
-					.MinSize(350)
-					.OnSlotResized_Lambda([this](float Size)
-						{
-							// Sync both header and content
-							TreeViewsSplitter->SlotAt(0).SetSizeValue(Size);
-							HeaderSplitter->SlotAt(0).SetSizeValue(Size);
-						})
-					[
-						SNew(SScrollBox)
-						.Orientation(Orient_Horizontal)
-						+ SScrollBox::Slot()
-						.FillSize(1.0f)
+						SNew(SVerticalBox)
+						+ SVerticalBox::Slot()
+						.AutoHeight()
+						.HAlign(HAlign_Fill)
 						[
-							InstancesTreeView.ToSharedRef()
+							SNew(SHorizontalBox)
+							+ SHorizontalBox::Slot()
+							.HAlign(HAlign_Left)
+							.FillWidth(1.0f)
+							[
+								LeftToolbar.MakeWidget()
+							]
+							+ SHorizontalBox::Slot()
+							.HAlign(HAlign_Right)
+							.AutoWidth()
+							[
+								RightToolbar.MakeWidget()
+							]
 						]
-					]
-					+ SSplitter::Slot()
-					.OnSlotResized_Lambda([this](float Size)
-						{
-							TreeViewsSplitter->SlotAt(1).SetSizeValue(Size);
-							HeaderSplitter->SlotAt(1).SetSizeValue(Size);
-						})
-					[
-						SNew(SOverlay)
-						+ SOverlay::Slot()
+						+ SVerticalBox::Slot()
+						.AutoHeight()
 						[
-							InstanceTimelinesTreeView.ToSharedRef()
+							SAssignNew(HeaderSplitter, SSplitter)
+							.Orientation(Orient_Horizontal)
+							+ SSplitter::Slot()
+							.Value(0.2f)
+							.MinSize(350)
+							.Resizable(false)
+							[
+								TraceSelectionBox
+							]
+							+ SSplitter::Slot()
+							.Resizable(false)
+							[
+								SNew(SSimpleTimeSlider)
+								.DesiredSize({100, 24})
+								.ClampRangeHighlightSize(0.15f)
+								.ClampRangeHighlightColor(FLinearColor::Red.CopyWithNewOpacity(0.5f))
+								.ScrubPosition(ScrubTimeAttribute)
+								.ViewRange_Lambda([this]() { return ViewRange; })
+								.OnViewRangeChanged_Lambda([this](TRange<double> NewRange) { ViewRange = NewRange; })
+								.ClampRange_Lambda([this]() { return TRange<double>(0, MaxTrackRecordingDuration); })
+								.OnScrubPositionChanged_Lambda([this](double NewScrubTime, bool bIsScrubbing) { OnTimeLineScrubPositionChanged(NewScrubTime, bIsScrubbing); })
+							]
 						]
-						+ SOverlay::Slot().HAlign(EHorizontalAlignment::HAlign_Right)
+						+ SVerticalBox::Slot()
+						.FillHeight(1.0f)
 						[
-							ScrollBar.ToSharedRef()
+							SAssignNew(TreeViewsSplitter, SSplitter)
+							.Orientation(Orient_Horizontal)
+							+ SSplitter::Slot()
+							.Value(0.2f)
+							.MinSize(350)
+							.OnSlotResized_Lambda([this](float Size)
+								{
+									// Sync both header and content
+									TreeViewsSplitter->SlotAt(0).SetSizeValue(Size);
+									HeaderSplitter->SlotAt(0).SetSizeValue(Size);
+								})
+							[
+								SNew(SScrollBox)
+								.Orientation(Orient_Horizontal)
+								+ SScrollBox::Slot()
+								.FillSize(1.0f)
+								[
+									InstancesTreeView.ToSharedRef()
+								]
+							]
+							+ SSplitter::Slot()
+							.OnSlotResized_Lambda([this](float Size)
+								{
+									TreeViewsSplitter->SlotAt(1).SetSizeValue(Size);
+									HeaderSplitter->SlotAt(1).SetSizeValue(Size);
+								})
+							[
+								SNew(SOverlay)
+								+ SOverlay::Slot()
+								[
+									InstanceTimelinesTreeView.ToSharedRef()
+								]
+								+ SOverlay::Slot().HAlign(EHorizontalAlignment::HAlign_Right)
+								[
+									ScrollBar.ToSharedRef()
+								]
+							]
 						]
-					]
 				]
 				+ SSplitter::Slot()
+				.MinSize(400)
 				[
-					SNew(SSplitter)
-					.Orientation(Orient_Horizontal)
-					+ SSplitter::Slot()
-					.MinSize(400)
+					SNew(SVerticalBox)
+					+ SVerticalBox::Slot()
+					.AutoHeight()
+					.HAlign(HAlign_Fill)
 					[
-						SNew(SScrollBox)
-						.Orientation(Orient_Horizontal)
-						+ SScrollBox::Slot()
-						.FillSize(1.0f)
-						[
-							EventsTreeView.ToSharedRef()
-						]
+						FrameDetailsToolbar.MakeWidget()
 					]
-					+ SSplitter::Slot()
-					.MinSize(400)
+					+ SVerticalBox::Slot()
+					.FillHeight(1.0f)
 					[
 						SNew(SScrollBox)
 						.Orientation(Orient_Horizontal)
 						+ SScrollBox::Slot()
 						.FillSize(1.0f)
 						[
-							SAssignNew(PropertiesBorder, SBorder)
+							EventsTreeView.ToSharedRef()
 						]
 					]
 				]
 			]
 		]
 	];
+	
+	// Auto-select session if there is only one available and that we are trac
+	// Do that after creating all our widgets in case we receive a callback
+	TArray<FStateTreeDebugger::FTraceDescriptor> TraceDescriptors;
+	Debugger->GetLiveTraces(TraceDescriptors);
+	if (TraceDescriptors.Num() == 1 && StateTreeModule.IsTracing())
+	{
+		Debugger->RequestSessionAnalysis(TraceDescriptors.Last());
+	}
 }
 
 void SStateTreeDebuggerView::Tick(const FGeometry& AllottedGeometry, const double InCurrentTime, const float InDeltaTime)
@@ -823,9 +722,9 @@ void SStateTreeDebuggerView::BindDebuggerToolbarCommands(const TSharedRef<FUICom
 		FCanExecuteAction::CreateSP(this, &SStateTreeDebuggerView::CanStepForwardToNextStateChange));
 
 	ToolkitCommands->MapAction(
-		Commands.ToggleAnalysis,
-		FExecuteAction::CreateSP(this, &SStateTreeDebuggerView::ToggleDebuggerAnalysis),
-		FCanExecuteAction::CreateSP(this, &SStateTreeDebuggerView::CanToggleDebuggerAnalysis));
+		Commands.ResumeDebuggerAnalysis,
+		FExecuteAction::CreateSP(this, &SStateTreeDebuggerView::ResumeDebuggerAnalysis),
+		FCanExecuteAction::CreateSP(this, &SStateTreeDebuggerView::CanResumeDebuggerAnalysis));
 
 	ToolkitCommands->MapAction(
 		Commands.ResetTracks,
@@ -1117,13 +1016,7 @@ void SStateTreeDebuggerView::OnDebuggerScrubStateChanged(const UE::StateTreeDebu
 	check(Spans.Num());
 	check(StateTree.IsValid());
 
-	struct FParentInfo
-	{
-		TSharedPtr<UE::StateTreeDebugger::FEventTreeElement> Element;
-		EStateTreeUpdatePhase Phase;
-	};
-	TArray<FParentInfo, TInlineAllocator<8>> PhaseStack;
-	TArray<FStateTreeStateHandle, TInlineAllocator<8>> StateStack;
+	TArray<TSharedPtr<FStateTreeDebuggerEventTreeElement>, TInlineAllocator<8>> ScopeStack;
 
 	const int32 SpanIdx = ScrubState.GetFrameSpanIndex();
 	const int32 FirstEventIdx = Spans[SpanIdx].EventIdx;
@@ -1133,65 +1026,177 @@ void SStateTreeDebuggerView::OnDebuggerScrubStateChanged(const UE::StateTreeDebu
 	for (int32 EventIdx = FirstEventIdx; EventIdx < MaxEventIdx; EventIdx++)
 	{
 		const FStateTreeTraceEventVariantType& Event = Events[EventIdx];
+		FString CustomDescription;
+		bool bShouldAddToScopeStack = false;
+		bool bShouldPopScopeStack = false;
 
 		if (const FStateTreeTraceStateEvent* StateEvent = Event.TryGet<FStateTreeTraceStateEvent>())
 		{
-			if (StateEvent->EventType == EStateTreeTraceEventType::Push)
+			if (StateEvent->EventType == EStateTreeTraceEventType::OnEntering
+			 	|| StateEvent->EventType == EStateTreeTraceEventType::OnExiting)
 			{
-				StateStack.Push(StateEvent->GetStateHandle());
-				continue;
+				bShouldAddToScopeStack = true;
 			}
-			
-			if (StateEvent->EventType == EStateTreeTraceEventType::Pop)
+			else if (StateEvent->EventType == EStateTreeTraceEventType::OnEntered
+				|| StateEvent->EventType == EStateTreeTraceEventType::OnExited)
 			{
-				StateStack.Pop();
-				continue;
+				bShouldPopScopeStack = true;
 			}
-			// Otherwise (other EventType) we add as a child
 		}
 		else if (const FStateTreeTracePhaseEvent* PhaseEvent = Event.TryGet<FStateTreeTracePhaseEvent>())
 		{
 			if (PhaseEvent->EventType == EStateTreeTraceEventType::Push)
 			{
-				// Create log event to describe the phase
-				FString NodeDesc;
-				const FStateTreeStateHandle StateHandle = StateStack.IsEmpty() ? FStateTreeStateHandle::Invalid : StateStack.Top();
-				if (StateHandle.IsValid())
+				if (PhaseEvent->Phase != EStateTreeUpdatePhase::Unset)
 				{
-					const FCompactStateTreeState* CompactState = StateTree->GetStateFromHandle(StateHandle);
-					NodeDesc = FString::Printf(TEXT(" '%s'"), CompactState != nullptr ? *CompactState->Name.ToString() : *StateHandle.Describe());
+					CustomDescription = UEnum::GetDisplayValueAsText(PhaseEvent->Phase).ToString();
 				}
 
-				FStateTreeTraceLogEvent LogEvent(/*RecordingWorldTime*/0, FString::Printf(TEXT("%s%s"), *UEnum::GetDisplayValueAsText(PhaseEvent->Phase).ToString(), *NodeDesc));
-
-				// Create Tree element to hold the event
-				const TSharedPtr<UE::StateTreeDebugger::FEventTreeElement> NewElement = MakeShareable(
-					new UE::StateTreeDebugger::FEventTreeElement(Frame, FStateTreeTraceEventVariantType(TInPlaceType<FStateTreeTraceLogEvent>(), LogEvent)));
-
-				// Push tree element to the proper hierarchy level
-				TArray<TSharedPtr<UE::StateTreeDebugger::FEventTreeElement>>& TreeElements = (PhaseStack.IsEmpty()) ? EventsTreeElements : PhaseStack.Top().Element->Children;
-				PhaseStack.Push({TreeElements.Add_GetRef(NewElement), PhaseEvent->Phase});
+				if (PhaseEvent->StateHandle.IsValid())
+				{
+					const FCompactStateTreeState* CompactState = StateTree->GetStateFromHandle(PhaseEvent->StateHandle);
+					CustomDescription += FString::Printf(TEXT("%s"),
+							CompactState != nullptr ? *CompactState->Name.ToString() : *PhaseEvent->StateHandle.Describe());
+				}
+				bShouldAddToScopeStack = true;
 			}
 			else if (PhaseEvent->EventType == EStateTreeTraceEventType::Pop)
 			{
-				PhaseStack.Pop();
+				bShouldPopScopeStack = true;
 			}
-			
-			// We don't want to create a child for phase events since they were already used to create a custom element.  
+		}
+
+		if (bShouldPopScopeStack)
+		{
+			// Pop scope and remove associated element if empty
+			TSharedPtr<FStateTreeDebuggerEventTreeElement> Scope = ScopeStack.Pop();
+			if (Scope->Children.IsEmpty())
+			{
+				TArray<TSharedPtr<FStateTreeDebuggerEventTreeElement>>& TreeElements = ScopeStack.IsEmpty() ? EventsTreeElements : ScopeStack.Top()->Children;
+				TreeElements.Remove(Scope);
+			}
+			// We don't want to create a child when a scope is popped.
 			continue;
 		}
 
-		TArray<TSharedPtr<UE::StateTreeDebugger::FEventTreeElement>>& TreeElements = PhaseStack.IsEmpty() ? EventsTreeElements : PhaseStack.Top().Element->Children;
-		TreeElements.Add(MakeShareable(new UE::StateTreeDebugger::FEventTreeElement(Spans[SpanIdx].Frame, Event)));
+		const TSharedRef<FStateTreeDebuggerEventTreeElement> NewElement = MakeShareable(new FStateTreeDebuggerEventTreeElement(Frame, Event));
+		NewElement->Description = CustomDescription;
+
+		TArray<TSharedPtr<FStateTreeDebuggerEventTreeElement>>& TreeElements = ScopeStack.IsEmpty() ? EventsTreeElements : ScopeStack.Top()->Children;
+		const TSharedPtr<FStateTreeDebuggerEventTreeElement>& ElementPtr = TreeElements.Add_GetRef(NewElement);
+
+		if (bShouldAddToScopeStack)
+		{
+			ScopeStack.Push(ElementPtr);
+		}
+
+		GenerateElementsForProperties(Event, NewElement);
 	}
 }
 
-void SStateTreeDebuggerView::ExpandAll(const TArray<TSharedPtr<UE::StateTreeDebugger::FEventTreeElement>>& Items)
+void SStateTreeDebuggerView::GenerateElementsForProperties(const FStateTreeTraceEventVariantType& Event, const TSharedRef<FStateTreeDebuggerEventTreeElement>& ParentElement)
 {
-	for (const TSharedPtr<UE::StateTreeDebugger::FEventTreeElement>& Item : Items)
+	FString TypePath;
+	FString InstanceDataAsString;
+
+	Visit([&TypePath, &InstanceDataAsString](auto& TypedEvent)
+		{
+			TypePath = TypedEvent.GetDataTypePath();
+			InstanceDataAsString = TypedEvent.GetDataAsText();
+		}, Event);
+
+	if (!InstanceDataAsString.IsEmpty())
 	{
-		EventsTreeView->SetItemExpansion(Item, true);
-		ExpandAll(Item->Children);
+		auto CreatePropertyElement = [ParentElement](const FStringView Line, const int32 NestedCount = 0)
+			{
+				constexpr int32 Indent = 4;
+				FString ConvertedString = FString::Printf(TEXT("%*s"), NestedCount*Indent, TEXT(""));
+				ConvertedString.Append(Line);
+				ConvertedString.ReplaceInline(TEXT("="), TEXT(" = "));
+				ConvertedString.ReplaceInline(TEXT("\""), TEXT(""));
+
+				// Create new property event
+				FStateTreeTracePropertyEvent PropertyEvent(/*RecordingWorldTime*/0, *ConvertedString);
+
+				// Create Tree element to hold the event
+				const TSharedPtr<FStateTreeDebuggerEventTreeElement> NewChildElement = MakeShareable(new FStateTreeDebuggerEventTreeElement(
+					ParentElement->Frame,
+					FStateTreeTraceEventVariantType(TInPlaceType<FStateTreeTracePropertyEvent>(), PropertyEvent)));
+
+				ParentElement->Children.Add(NewChildElement);
+			};
+
+		// Try to parse Struct for which properties are exported between '(' and ')' 
+		if (InstanceDataAsString.StartsWith("(") && InstanceDataAsString.EndsWith(")"))
+		{
+			const FStringView View(GetData(InstanceDataAsString) + 1, InstanceDataAsString.Len() - 2);
+			const TCHAR* ViewIt = View.GetData();
+			const TCHAR* const ViewEnd = ViewIt + View.Len();
+			const TCHAR* NextToken = ViewIt;
+			int32 NestedCount = 0;
+			
+			for (; ViewIt != ViewEnd; ++ViewIt)
+			{
+				int32 LocalNestedCount = 0;
+				if (*ViewIt == TCHAR('('))
+				{
+					LocalNestedCount++;
+				}
+				else if (*ViewIt == TCHAR(')'))
+				{
+					LocalNestedCount--;
+				}
+				else if (*ViewIt != TCHAR(','))
+				{
+					continue;
+				}
+
+				if (ViewIt != NextToken)
+				{
+					CreatePropertyElement(FStringView(NextToken, UE_PTRDIFF_TO_INT32(ViewIt - NextToken)), NestedCount);
+				}
+				NextToken = ViewIt+1;
+				NestedCount += LocalNestedCount;
+			}
+			CreatePropertyElement(FStringView(NextToken, UE_PTRDIFF_TO_INT32(ViewIt - NextToken)), NestedCount);
+		}
+		else
+		{
+			const TCHAR* Buffer = *InstanceDataAsString;
+			FParse::Next(&Buffer);
+			FString StrLine;
+			while (FParse::Line(&Buffer, StrLine))
+			{
+				const TCHAR* Str = *StrLine;
+				if (!FParse::Command(&Str, TEXT("BEGIN OBJECT"))
+					&& !FParse::Command(&Str, TEXT("END OBJECT")))
+				{
+					CreatePropertyElement(Str);
+				}
+			}
+		}
+	}
+}
+
+void SStateTreeDebuggerView::ExpandAll(const TArray<TSharedPtr<FStateTreeDebuggerEventTreeElement>>& Items)
+{
+	for (const TSharedPtr<FStateTreeDebuggerEventTreeElement>& Item : Items)
+	{
+		bool bExpand = true;
+		if (Item->Children.Num() > 0)
+		{
+			const FStateTreeDebuggerEventTreeElement* FirstChild = Item->Children[0].Get();
+			if (FirstChild && FirstChild->Event.IsType<FStateTreeTracePropertyEvent>())
+			{
+				bExpand = false;
+			}	
+		}
+
+		if (bExpand)
+		{
+			EventsTreeView->SetItemExpansion(Item, true);
+			ExpandAll(Item->Children);
+		}
 	}
 }
 
@@ -1216,8 +1221,8 @@ void SStateTreeDebuggerView::OnBreakpointHit(const FStateTreeInstanceDebugId Ins
 	// Find matching event in the tree view and select it
 	if (EventsTreeView.IsValid())
 	{
-		TSharedPtr<UE::StateTreeDebugger::FEventTreeElement> MatchingElement = nullptr;
-		VisitEventTreeElements(EventsTreeElements, [&MatchingElement, Breakpoint](const TSharedPtr<UE::StateTreeDebugger::FEventTreeElement>& VisitedElement)
+		TSharedPtr<FStateTreeDebuggerEventTreeElement> MatchingElement = nullptr;
+		UE::StateTreeDebugger::VisitEventTreeElements(EventsTreeElements, [&MatchingElement, Breakpoint](const TSharedPtr<FStateTreeDebuggerEventTreeElement>& VisitedElement)
 			{
 				if (Breakpoint.IsMatchingEvent(VisitedElement->Event))
 				{
@@ -1306,8 +1311,6 @@ void SStateTreeDebuggerView::OnSelectedInstanceCleared()
 	{
 		EventsTreeView->RequestTreeRefresh();
 	}
-
-	PropertiesBorder->ClearContent();
 }
 
 TSharedRef<SWidget> SStateTreeDebuggerView::OnGetDebuggerTracesMenu() const
diff --git a/Engine/Plugins/Runtime/StateTree/Source/StateTreeEditorModule/Private/Debugger/SStateTreeDebuggerView.h b/Engine/Plugins/Runtime/StateTree/Source/StateTreeEditorModule/Private/Debugger/SStateTreeDebuggerView.h
--- a/Engine/Plugins/Runtime/StateTree/Source/StateTreeEditorModule/Private/Debugger/SStateTreeDebuggerView.h
+++ b/Engine/Plugins/Runtime/StateTree/Source/StateTreeEditorModule/Private/Debugger/SStateTreeDebuggerView.h
@@ -5,7 +5,6 @@
 
 #if WITH_STATETREE_DEBUGGER
 
-#include "InstancedStruct.h"
 #include "Widgets/Views/STreeView.h"
 #include "Widgets/SCompoundWidget.h"
 
@@ -16,7 +15,6 @@ namespace RewindDebugger
 namespace UE::StateTreeDebugger
 {
 	struct FFrameSpan;
-	struct FEventTreeElement;
 }
 
 enum class EStateTreeBreakpointType : uint8;
@@ -25,6 +23,7 @@ struct FSlateIcon;
 struct FStateTreeDebugger;
 struct FStateTreeInstanceDebugId;
 struct FStateTreeDebuggerBreakpoint;
+struct FStateTreeDebuggerEventTreeElement;
 class FStateTreeEditor;
 class FStateTreeViewModel;
 class FUICommandList;
@@ -81,9 +80,8 @@ private:
 	bool CanStopRecording() const { return IsRecording(); }	
 	void StopRecording();
 
-	bool CanToggleDebuggerAnalysis() const;
-	void ToggleDebuggerAnalysis() const;
-	FSlateIcon GetDebuggerAnalysisIcon() const;
+	bool CanResumeDebuggerAnalysis() const;
+	void ResumeDebuggerAnalysis() const;
 
 	bool CanResetTracks() const;
 	void ResetTracks();
@@ -113,8 +111,10 @@ private:
 	/** Callback used to reflect debugger scrub state in the UI. */
 	void OnDebuggerScrubStateChanged(const UE::StateTreeDebugger::FScrubState& ScrubState);
 
+	static void GenerateElementsForProperties(const FStateTreeTraceEventVariantType& Event, const TSharedRef<FStateTreeDebuggerEventTreeElement>& ParentElement);
+
 	/** Recursively sets tree items as expanded. */
-	void ExpandAll(const TArray<TSharedPtr<UE::StateTreeDebugger::FEventTreeElement>>& Items);
+	void ExpandAll(const TArray<TSharedPtr<FStateTreeDebuggerEventTreeElement>>& Items);
 
 	TSharedPtr<FStateTreeDebugger> Debugger;
 	TSharedPtr<FStateTreeViewModel> StateTreeViewModel;
@@ -145,13 +145,10 @@ private:
 	TSharedPtr<SSplitter> TreeViewsSplitter;
 
 	/** All trace events received for a given instance. */
-	TArray<TSharedPtr<UE::StateTreeDebugger::FEventTreeElement>> EventsTreeElements;
+	TArray<TSharedPtr<FStateTreeDebuggerEventTreeElement>> EventsTreeElements;
 	
 	/** Tree view displaying the frame events of the instance associated to the selected track. */
-	TSharedPtr<STreeView<TSharedPtr<UE::StateTreeDebugger::FEventTreeElement>>> EventsTreeView;
-
-	/** Widget holding the event details when available (i.e. property details of the associated struct or object). */
-	TSharedPtr<SBorder> PropertiesBorder;
+	TSharedPtr<STreeView<TSharedPtr<FStateTreeDebuggerEventTreeElement>>> EventsTreeView;
 
 	/** Attribute provided by the debugger scrub position to control cursor and timelines positions. */
 	TAttribute<double> ScrubTimeAttribute;
@@ -159,12 +156,6 @@ private:
 	/** Range controlled by the timeline widgets and used to adjust cursor position and track content. */
 	TRange<double> ViewRange = TRange<double>(0, 10);
 
-	/** Struct created from the event data when statetree node was holding a struct. */
-	FInstancedStruct SelectedNodeDataStruct;
-
-	/** Object created from the event data when statetree node was holding an object. */
-	TWeakObjectPtr<> SelectedNodeDataObject;
-
 	/** In case tracks are not reset when a new analysis session is started we keep track of the longest duration to adjust our clamp range. */
 	double MaxTrackRecordingDuration = 0;
 
diff --git a/Engine/Plugins/Runtime/StateTree/Source/StateTreeEditorModule/Private/Debugger/SStateTreeDebuggerViewRow.cpp b/Engine/Plugins/Runtime/StateTree/Source/StateTreeEditorModule/Private/Debugger/SStateTreeDebuggerViewRow.cpp
new file mode 100644
--- /dev/null
+++ b/Engine/Plugins/Runtime/StateTree/Source/StateTreeEditorModule/Private/Debugger/SStateTreeDebuggerViewRow.cpp
@@ -0,0 +1,214 @@
+// Copyright Epic Games, Inc. All Rights Reserved.
+
+#pragma once
+
+#if WITH_STATETREE_DEBUGGER
+
+#include "Debugger/SStateTreeDebuggerViewRow.h"
+#include "StateTreeEditorStyle.h"
+#include "StateTreeViewModel.h"
+#include "Widgets/Images/SImage.h"
+#include "Widgets/Text/STextBlock.h"
+
+
+void SStateTreeDebuggerViewRow::Construct(const FArguments& InArgs,
+										  const TSharedPtr<STableViewBase>& InOwnerTableView,
+										  const TSharedPtr<FStateTreeDebuggerEventTreeElement>& InElement,
+										  const TSharedRef<FStateTreeViewModel>& InStateTreeViewModel)
+{
+	Item = InElement;
+	StateTreeViewModel = InStateTreeViewModel.ToSharedPtr();
+	STableRow::Construct(InArgs, InOwnerTableView.ToSharedRef());
+
+	const TSharedPtr<SHorizontalBox> HorizontalBox = SNew(SHorizontalBox);
+	HorizontalBox->SetToolTipText(GetEventTooltip());
+
+	HorizontalBox->AddSlot()
+		.VAlign(VAlign_Center)
+		.HAlign(HAlign_Left)
+		.AutoWidth()
+		[
+			SNew(SExpanderArrow, SharedThis(this))
+			.ShouldDrawWires(false)
+			.IndentAmount(32)
+			.BaseIndentLevel(0)
+		];
+
+	const TSharedPtr<SWidget> EventWidget = GenerateEventWidget();
+	if (EventWidget.IsValid())
+	{
+		HorizontalBox->AddSlot()
+		.VAlign(VAlign_Center)
+		.HAlign(HAlign_Left)
+		.AutoWidth()
+		[
+			EventWidget.ToSharedRef()
+		];
+	}
+
+	HorizontalBox->AddSlot()
+		.Padding(2)
+		.AutoWidth()
+		.VAlign(VAlign_Center)
+		.HAlign(HAlign_Left)
+		[
+			SNew(STextBlock)
+			.TextStyle(&GetEventTextStyle())
+			.Text(GetEventDescription())
+		];
+
+	this->ChildSlot
+		.HAlign(HAlign_Fill)
+		[
+			HorizontalBox.ToSharedRef()
+		];
+}
+
+TSharedPtr<SWidget> SStateTreeDebuggerViewRow::GenerateEventWidget() const
+{
+	const FStateTreeEditorStyle& StyleSet = FStateTreeEditorStyle::Get();
+	if (const FStateTreeTracePhaseEvent* PhaseEvent = Item->Event.TryGet<FStateTreeTracePhaseEvent>())
+	{
+		const FSlateBrush* Image = nullptr;
+		switch (PhaseEvent->Phase)
+		{
+		case EStateTreeUpdatePhase::EnterStates:	Image = StyleSet.GetBrush("StateTreeEditor.Debugger.State.Enter");	break;
+		case EStateTreeUpdatePhase::ExitStates:		Image = StyleSet.GetBrush("StateTreeEditor.Debugger.State.Exit");	break;
+		default:
+			return nullptr;
+		}
+
+		return SNew(SImage).Image(Image);
+	}
+
+	if (const FStateTreeTraceStateEvent* StateEvent = Item->Event.TryGet<FStateTreeTraceStateEvent>())
+	{
+		const FSlateBrush* Image = nullptr;
+
+		switch (StateEvent->EventType)
+		{
+		case EStateTreeTraceEventType::OnStateSelected:		Image = StyleSet.GetBrush("StateTreeEditor.Debugger.State.Selected");	break;
+		case EStateTreeTraceEventType::OnStateCompleted:	Image = StyleSet.GetBrush("StateTreeEditor.Debugger.State.Completed");	break;
+		default:
+			return nullptr;
+		}
+
+		return SNew(SImage).Image(Image);
+	}
+
+	if (const FStateTreeTraceTaskEvent* TaskEvent = Item->Event.TryGet<FStateTreeTraceTaskEvent>())
+	{
+		const FSlateBrush* Image = nullptr;
+
+		switch (TaskEvent->EventType)
+		{
+		case EStateTreeTraceEventType::OnEntered:
+		case EStateTreeTraceEventType::OnTaskCompleted:
+		case EStateTreeTraceEventType::OnTaskTicked:
+			switch (TaskEvent->Status)
+			{
+			case EStateTreeRunStatus::Failed:		Image = StyleSet.GetBrush("StateTreeEditor.Debugger.Task.Failed");		break;
+			case EStateTreeRunStatus::Succeeded:	Image = StyleSet.GetBrush("StateTreeEditor.Debugger.Task.Succeeded");	break;
+			case EStateTreeRunStatus::Stopped:		Image = StyleSet.GetBrush("StateTreeEditor.Debugger.Task.Stopped");		break;
+			default:
+				return nullptr;
+			}
+			break;
+		default:
+			return nullptr;
+		}
+
+		return SNew(SImage).Image(Image);
+	}
+
+	if (Item->Event.IsType<FStateTreeTraceConditionEvent>())
+	{
+		const FStateTreeTraceConditionEvent& ConditionEvent = Item->Event.Get<FStateTreeTraceConditionEvent>();
+		const FSlateBrush* Image = nullptr;
+
+		switch (ConditionEvent.EventType)
+		{
+		case EStateTreeTraceEventType::Passed:			Image = StyleSet.GetBrush("StateTreeEditor.Debugger.Condition.Passed");			break;
+		case EStateTreeTraceEventType::Failed:			Image = StyleSet.GetBrush("StateTreeEditor.Debugger.Condition.Failed");			break;
+		case EStateTreeTraceEventType::OnEvaluating:	Image = StyleSet.GetBrush("StateTreeEditor.Debugger.Condition.OnEvaluating");	break;
+		case EStateTreeTraceEventType::OnTransition:	Image = StyleSet.GetBrush("StateTreeEditor.Debugger.Condition.OnTransition");	break;
+		default:
+			Image = StyleSet.GetBrush("StateTreeEditor.Debugger.Unset");
+		}
+
+		return SNew(SImage).Image(Image);
+	}
+
+	return nullptr;
+}
+
+const FTextBlockStyle& SStateTreeDebuggerViewRow::GetEventTextStyle() const
+{
+	const FStateTreeEditorStyle& StyleSet = FStateTreeEditorStyle::Get();
+
+	if (Item->Event.IsType<FStateTreeTracePhaseEvent>())
+	{
+		return StyleSet.GetWidgetStyle<FTextBlockStyle>("StateTreeDebugger.Element.Bold");
+	}
+
+	 if (Item->Event.IsType<FStateTreeTracePropertyEvent>())
+	{
+		return StyleSet.GetWidgetStyle<FTextBlockStyle>("StateTreeDebugger.Element.Subdued");
+	}
+
+	return StyleSet.GetWidgetStyle<FTextBlockStyle>("StateTreeDebugger.Element.Normal");
+}
+
+FText SStateTreeDebuggerViewRow::GetEventDescription() const
+{
+	FString EventDescription;
+	if (Item->Description.IsEmpty())
+	{
+		check(StateTreeViewModel);
+		if (const UStateTree* StateTree = StateTreeViewModel->GetStateTree())
+		{
+			// Some types have some custom representations so we want to use a more minimal description.
+			if (Item->Event.IsType<FStateTreeTraceStateEvent>()
+				|| Item->Event.IsType<FStateTreeTraceTaskEvent>()
+				|| Item->Event.IsType<FStateTreeTracePropertyEvent>()
+				|| Item->Event.IsType<FStateTreeTraceConditionEvent>()
+				|| Item->Event.IsType<FStateTreeTraceLogEvent>())
+			{
+				Visit([&EventDescription, StateTree](auto& TypedEvent)
+					{
+						EventDescription = TypedEvent.GetValueString(*StateTree);
+					}, Item->Event);
+			}
+			else
+			{
+				Visit([&EventDescription, StateTree](auto& TypedEvent)
+					{
+						EventDescription = TypedEvent.ToFullString(*StateTree);
+					}, Item->Event);
+			}
+		}
+	}
+	else
+	{
+		EventDescription = Item->Description;
+	}
+	
+	return FText::FromString(EventDescription);
+}
+
+FText SStateTreeDebuggerViewRow::GetEventTooltip() const
+{
+	FString Tooltip;
+	check(StateTreeViewModel);
+	if (const UStateTree* StateTree = StateTreeViewModel->GetStateTree())
+	{
+		Visit([&Tooltip, StateTree](auto& TypedEvent)
+			{
+				Tooltip = TypedEvent.ToFullString(*StateTree);
+			}, Item->Event);
+	}
+
+	return FText::FromString(Tooltip);
+}
+
+#endif // WITH_STATETREE_DEBUGGER
diff --git a/Engine/Plugins/Runtime/StateTree/Source/StateTreeEditorModule/Private/Debugger/SStateTreeDebuggerViewRow.h b/Engine/Plugins/Runtime/StateTree/Source/StateTreeEditorModule/Private/Debugger/SStateTreeDebuggerViewRow.h
new file mode 100644
--- /dev/null
+++ b/Engine/Plugins/Runtime/StateTree/Source/StateTreeEditorModule/Private/Debugger/SStateTreeDebuggerViewRow.h
@@ -0,0 +1,51 @@
+// Copyright Epic Games, Inc. All Rights Reserved.
+
+#pragma once
+
+#if WITH_STATETREE_DEBUGGER
+
+#include "Debugger/StateTreeTraceTypes.h"
+#include "Templates/SharedPointer.h"
+#include "TraceServices/Model/Frames.h"
+#include "Widgets/Views/SListView.h"
+#include "Widgets/Views/STableRow.h"
+
+class FStateTreeViewModel;
+
+/** An item in the StateTreeDebugger trace event tree */
+struct FStateTreeDebuggerEventTreeElement : TSharedFromThis<FStateTreeDebuggerEventTreeElement>
+{
+	explicit FStateTreeDebuggerEventTreeElement(const TraceServices::FFrame& Frame, const FStateTreeTraceEventVariantType& Event)
+		: Frame(Frame), Event(Event)
+	{
+	}
+
+	TraceServices::FFrame Frame;
+	FStateTreeTraceEventVariantType Event;
+	TArray<TSharedPtr<FStateTreeDebuggerEventTreeElement>> Children;
+	FString Description;
+};
+
+
+/**
+ * Widget for row inside the StateTreeDebugger TreeView.
+ */
+class SStateTreeDebuggerViewRow : public STableRow<TSharedPtr<FStateTreeDebuggerEventTreeElement>>
+{
+public:
+	void Construct(const FArguments& InArgs,
+				   const TSharedPtr<STableViewBase>& InOwnerTableView,
+				   const TSharedPtr<FStateTreeDebuggerEventTreeElement>& InElement,
+				   const TSharedRef<FStateTreeViewModel>& InStateTreeViewModel);
+
+private:
+	TSharedPtr<SWidget> GenerateEventWidget() const;
+	const FTextBlockStyle& GetEventTextStyle() const;
+	FText GetEventDescription() const;
+	FText GetEventTooltip() const;
+	
+	TSharedPtr<FStateTreeViewModel> StateTreeViewModel;
+	TSharedPtr<FStateTreeDebuggerEventTreeElement> Item;
+};
+
+#endif // WITH_STATETREE_DEBUGGER
\ No newline at end of file
diff --git a/Engine/Plugins/Runtime/StateTree/Source/StateTreeEditorModule/Private/Debugger/StateTreeDebuggerCommands.cpp b/Engine/Plugins/Runtime/StateTree/Source/StateTreeEditorModule/Private/Debugger/StateTreeDebuggerCommands.cpp
--- a/Engine/Plugins/Runtime/StateTree/Source/StateTreeEditorModule/Private/Debugger/StateTreeDebuggerCommands.cpp
+++ b/Engine/Plugins/Runtime/StateTree/Source/StateTreeEditorModule/Private/Debugger/StateTreeDebuggerCommands.cpp
@@ -32,8 +32,8 @@ void FStateTreeDebuggerCommands::RegisterCommands()
 	UI_COMMAND(PreviousFrameWithEvents, "Previous Frame", "Step one frame back", EUserInterfaceActionType::Button, FInputChord(EKeys::Left));
 	UI_COMMAND(NextFrameWithEvents, "Next Frame", "Step one frame forward", EUserInterfaceActionType::Button, FInputChord(EKeys::Right));
 	UI_COMMAND(NextFrameWithStateChange, "Next State Change", "Jump to next state changed frame", EUserInterfaceActionType::Button, FInputChord(EModifierKey::Control, EKeys::Right));
-	UI_COMMAND(ToggleAnalysis, "Toggle Analysis",  "Pause or resume processing events from the selected trace.", EUserInterfaceActionType::Button, FInputChord());
-	UI_COMMAND(ResetTracks, "Reset Tracks",  "Clear current tracks and associated events.", EUserInterfaceActionType::Button, FInputChord());
+	UI_COMMAND(ResumeDebuggerAnalysis, "Resume Analysis", "Resume trace analysis.", EUserInterfaceActionType::Button, FInputChord());
+	UI_COMMAND(ResetTracks, "Reset Tracks", "Clear current tracks and associated events.", EUserInterfaceActionType::Button, FInputChord());
 }
 
 #undef LOCTEXT_NAMESPACE
diff --git a/Engine/Plugins/Runtime/StateTree/Source/StateTreeEditorModule/Private/Debugger/StateTreeDebuggerCommands.h b/Engine/Plugins/Runtime/StateTree/Source/StateTreeEditorModule/Private/Debugger/StateTreeDebuggerCommands.h
--- a/Engine/Plugins/Runtime/StateTree/Source/StateTreeEditorModule/Private/Debugger/StateTreeDebuggerCommands.h
+++ b/Engine/Plugins/Runtime/StateTree/Source/StateTreeEditorModule/Private/Debugger/StateTreeDebuggerCommands.h
@@ -25,7 +25,7 @@ public:
 	TSharedPtr<FUICommandInfo> PreviousFrameWithEvents;
 	TSharedPtr<FUICommandInfo> NextFrameWithEvents;
 	TSharedPtr<FUICommandInfo> NextFrameWithStateChange;
-	TSharedPtr<FUICommandInfo> ToggleAnalysis;
+	TSharedPtr<FUICommandInfo> ResumeDebuggerAnalysis;
 	TSharedPtr<FUICommandInfo> ResetTracks;
 };
 
diff --git a/Engine/Plugins/Runtime/StateTree/Source/StateTreeEditorModule/Private/StateTreeEditorStyle.cpp b/Engine/Plugins/Runtime/StateTree/Source/StateTreeEditorModule/Private/StateTreeEditorStyle.cpp
--- a/Engine/Plugins/Runtime/StateTree/Source/StateTreeEditorModule/Private/StateTreeEditorStyle.cpp
+++ b/Engine/Plugins/Runtime/StateTree/Source/StateTreeEditorModule/Private/StateTreeEditorStyle.cpp
@@ -7,7 +7,6 @@
 #include "Styling/CoreStyle.h"
 #include "Brushes/SlateRoundedBoxBrush.h"
 #include "Styling/SlateTypes.h"
-#include "Interfaces/IPluginManager.h"
 #include "Misc/Paths.h"
 #include "Styling/StyleColors.h"
 #include "Styling/SlateStyleMacros.h"
@@ -54,22 +53,22 @@ FStateTreeEditorStyle::FStateTreeEditorStyle()
 
 	// State
 	{
-		FTextBlockStyle StateIcon = FTextBlockStyle(NormalText)
+		const FTextBlockStyle StateIcon = FTextBlockStyle(NormalText)
 			.SetFont(FAppStyle::Get().GetFontStyle("FontAwesome.12"))
 			.SetColorAndOpacity(FLinearColor(230.0f / 255.0f, 230.0f / 255.0f, 230.0f / 255.0f, 0.5f));
 		Set("StateTree.Icon", StateIcon);
 
-		FTextBlockStyle StateDetailsIcon = FTextBlockStyle(NormalText)
+		const FTextBlockStyle StateDetailsIcon = FTextBlockStyle(NormalText)
 			.SetFont(FAppStyle::Get().GetFontStyle("FontAwesome.10"))
 		    .SetColorAndOpacity(FLinearColor(230.0f / 255.0f, 230.0f / 255.0f, 230.0f / 255.0f, 0.5f));
 		Set("StateTree.DetailsIcon", StateDetailsIcon);
 
-		FTextBlockStyle StateTitle = FTextBlockStyle(NormalText)
+		const FTextBlockStyle StateTitle = FTextBlockStyle(NormalText)
 			.SetFont(CORE_FONT("Fonts/Roboto-Bold", 12))
 			.SetColorAndOpacity(FLinearColor(230.0f / 255.0f, 230.0f / 255.0f, 230.0f / 255.0f, 0.9f));
 		Set("StateTree.State.Title", StateTitle);
 
-		FEditableTextBoxStyle StateTitleEditableText = FEditableTextBoxStyle()
+		const FEditableTextBoxStyle StateTitleEditableText = FEditableTextBoxStyle()
 			.SetTextStyle(NormalText)
 			.SetFont(CORE_FONT("Fonts/Roboto-Bold", 12))
 			.SetBackgroundImageNormal(CORE_BOX_BRUSH("Common/TextBox", FMargin(4.0f / 16.0f)))
@@ -128,6 +127,22 @@ FStateTreeEditorStyle::FStateTreeEditorStyle()
 			.SetFont(FAppStyle::GetFontStyle(TEXT("PropertyWindow.NormalFont"))));
 	}
 
+	// Debugger
+	{
+		Set("StateTreeDebugger.Element.Normal",
+			FTextBlockStyle(NormalText)
+			.SetFont(CORE_FONT("Fonts/Roboto-Regular", 10)));
+
+		Set("StateTreeDebugger.Element.Bold",
+			FTextBlockStyle(NormalText)
+			.SetFont(CORE_FONT("Fonts/Roboto-Bold", 10)));
+
+		Set("StateTreeDebugger.Element.Subdued",
+			FTextBlockStyle(NormalText)
+			.SetFont(CORE_FONT("Fonts/Roboto-Regular", 10))
+			.SetColorAndOpacity(FSlateColor::UseSubduedForeground()));
+	}
+	
 	const FLinearColor SelectionColor = FColor(0, 0, 0, 32);
 	const FTableRowStyle& NormalTableRowStyle = FAppStyle::Get().GetWidgetStyle<FTableRowStyle>("TableView.Row");
 	Set("StateTree.Selection",
@@ -195,6 +210,27 @@ FStateTreeEditorStyle::FStateTreeEditorStyle()
 		Set("StateTreeEditor.DeleteStates", new IMAGE_BRUSH_SVG("Starship/Common/Delete", CoreStyleConstants::Icon16x16));
 		Set("StateTreeEditor.RenameState", new IMAGE_BRUSH_SVG("Starship/Common/Rename", CoreStyleConstants::Icon16x16));
 		Set("StateTreeEditor.AutoScroll", new IMAGE_BRUSH_SVG("Starship/Insights/AutoScrollRight_20", CoreStyleConstants::Icon16x16));
+
+		Set("StateTreeEditor.Debugger.ResetTracks", new IMAGE_BRUSH_SVG("Starship/Common/Delete", CoreStyleConstants::Icon16x16));
+
+		Set("StateTreeEditor.Debugger.State.Enter", new CORE_IMAGE_BRUSH_SVG("Starship/Common/arrow-right", CoreStyleConstants::Icon16x16, FStyleColors::Foreground));
+		Set("StateTreeEditor.Debugger.State.Exit", new CORE_IMAGE_BRUSH_SVG("Starship/Common/arrow-left", CoreStyleConstants::Icon16x16, FStyleColors::Foreground));
+		Set("StateTreeEditor.Debugger.State.Selected", new CORE_IMAGE_BRUSH_SVG("Starship/Common/arrow-right", CoreStyleConstants::Icon16x16, FStyleColors::AccentYellow));
+		Set("StateTreeEditor.Debugger.State.Completed", new CORE_IMAGE_BRUSH_SVG("Starship/Common/check", CoreStyleConstants::Icon16x16, FStyleColors::AccentGreen));
+
+		Set("StateTreeEditor.Debugger.Task.Enter", new CORE_IMAGE_BRUSH_SVG("Starship/Common/arrow-right", CoreStyleConstants::Icon16x16, FStyleColors::Foreground));
+		Set("StateTreeEditor.Debugger.Task.Exit", new CORE_IMAGE_BRUSH_SVG("Starship/Common/arrow-left", CoreStyleConstants::Icon16x16, FStyleColors::Foreground));
+		Set("StateTreeEditor.Debugger.Task.Failed", new CORE_IMAGE_BRUSH_SVG("Starship/Common/close", CoreStyleConstants::Icon16x16, FStyleColors::AccentRed));
+		Set("StateTreeEditor.Debugger.Task.Succeeded", new CORE_IMAGE_BRUSH_SVG("Starship/Common/check", CoreStyleConstants::Icon16x16, FStyleColors::AccentGreen));
+		Set("StateTreeEditor.Debugger.Task.Stopped", new CORE_IMAGE_BRUSH_SVG("Starship/Common/close", CoreStyleConstants::Icon16x16, FStyleColors::AccentRed));
+
+		Set("StateTreeEditor.Debugger.Condition", new CORE_IMAGE_BRUSH_SVG("Starship/Common/help", CoreStyleConstants::Icon16x16, FStyleColors::Foreground));
+		Set("StateTreeEditor.Debugger.Condition.Passed", new CORE_IMAGE_BRUSH_SVG("Starship/Common/check", CoreStyleConstants::Icon16x16, FStyleColors::AccentGreen));
+		Set("StateTreeEditor.Debugger.Condition.Failed", new CORE_IMAGE_BRUSH_SVG("Starship/Common/close", CoreStyleConstants::Icon16x16, FStyleColors::AccentRed));
+		Set("StateTreeEditor.Debugger.Condition.OnEvaluating", new CORE_IMAGE_BRUSH_SVG("Starship/Common/Update", CoreStyleConstants::Icon16x16, FStyleColors::AccentYellow));
+		Set("StateTreeEditor.Debugger.Condition.OnTransition", new CORE_IMAGE_BRUSH_SVG("Starship/Common/help", CoreStyleConstants::Icon16x16, FStyleColors::AccentBlue));
+
+		Set("StateTreeEditor.Debugger.Unset", new CORE_IMAGE_BRUSH_SVG("Starship/Common/help", CoreStyleConstants::Icon16x16, FStyleColors::AccentBlack));
 	}
 
 	{
@@ -219,9 +255,7 @@ FStateTreeEditorStyle::FStateTreeEditorStyle()
 
 		Set("StateTreeEditor.EnableStates", new IMAGE_BRUSH("Icons/Empty_16x", CoreStyleConstants::Icon16x16));
 		Set("StateTreeEditor.Debugger.Breakpoint.EnabledAndValid", new IMAGE_BRUSH_SVG( "Starship/Blueprints/Breakpoint_Valid", CoreStyleConstants::Icon16x16, FStyleColors::AccentRed));
-		Set("StateTreeEditor.PauseDebuggerAnalysis", new IMAGE_BRUSH("Icons/generic_pause_16x", CoreStyleConstants::Icon16x16));
-		Set("StateTreeEditor.ResumeDebuggerAnalysis", new IMAGE_BRUSH_SVG("Starship/Common/Timeline", CoreStyleConstants::Icon16x16));
-		Set("StateTreeEditor.Debugger.ResetTracks", new IMAGE_BRUSH("Icons/icon_TrackDelete_36x24px", CoreStyleConstants::Icon16x16));
+		Set("StateTreeEditor.Debugger.ResumeDebuggerAnalysis", new IMAGE_BRUSH_SVG("Starship/Common/Timeline", CoreStyleConstants::Icon16x16));
 	}
 
 	{
diff --git a/Engine/Plugins/Runtime/StateTree/Source/StateTreeModule/Private/Debugger/StateTreeDebuggerTypes.cpp b/Engine/Plugins/Runtime/StateTree/Source/StateTreeModule/Private/Debugger/StateTreeDebuggerTypes.cpp
--- a/Engine/Plugins/Runtime/StateTree/Source/StateTreeModule/Private/Debugger/StateTreeDebuggerTypes.cpp
+++ b/Engine/Plugins/Runtime/StateTree/Source/StateTreeModule/Private/Debugger/StateTreeDebuggerTypes.cpp
@@ -286,7 +286,7 @@ bool FStateTreeDebuggerBreakpoint::IsMatchingEvent(const FStateTreeTraceEventVar
 		else if (const FStateTreeTransitionIndex* TransitionIndex = ElementIdentifier.TryGet<FStateTreeTransitionIndex>())
 		{
 			const FStateTreeTraceTransitionEvent* TransitionEvent = Event.TryGet<FStateTreeTraceTransitionEvent>();
-			bIsMatching = TransitionEvent != nullptr && TransitionEvent->TransitionIndex == TransitionIndex->Index;
+			bIsMatching = TransitionEvent != nullptr && TransitionEvent->TransitionSource.TransitionIndex == TransitionIndex->Index;
 		}
 	}
 	return bIsMatching;
@@ -297,9 +297,9 @@ EStateTreeTraceEventType FStateTreeDebuggerBreakpoint::GetMatchingEventType(cons
 	switch (BreakpointType)
 	{
 	case EStateTreeBreakpointType::OnEnter:
-		return EStateTreeTraceEventType::OnEnter;
+		return EStateTreeTraceEventType::OnEntering;
 	case EStateTreeBreakpointType::OnExit:
-		return EStateTreeTraceEventType::OnExit;
+		return EStateTreeTraceEventType::OnExiting;
 	case EStateTreeBreakpointType::OnTransition:
 		return EStateTreeTraceEventType::OnTransition;
 	default:
diff --git a/Engine/Plugins/Runtime/StateTree/Source/StateTreeModule/Private/Debugger/StateTreeTrace.cpp b/Engine/Plugins/Runtime/StateTree/Source/StateTreeModule/Private/Debugger/StateTreeTrace.cpp
--- a/Engine/Plugins/Runtime/StateTree/Source/StateTreeModule/Private/Debugger/StateTreeTrace.cpp
+++ b/Engine/Plugins/Runtime/StateTree/Source/StateTreeModule/Private/Debugger/StateTreeTrace.cpp
@@ -39,6 +39,7 @@ UE_TRACE_EVENT_BEGIN(StateTreeDebugger, PhaseEvent)
 	UE_TRACE_EVENT_FIELD(uint32, InstanceId)
 	UE_TRACE_EVENT_FIELD(uint32, InstanceSerial)
 	UE_TRACE_EVENT_FIELD(uint16, Phase)
+	UE_TRACE_EVENT_FIELD(uint16, StateIndex)
 	UE_TRACE_EVENT_FIELD(uint8, EventType)
 UE_TRACE_EVENT_END()
 
@@ -72,7 +73,10 @@ UE_TRACE_EVENT_BEGIN(StateTreeDebugger, TransitionEvent)
 	UE_TRACE_EVENT_FIELD(uint64, Cycle)
 	UE_TRACE_EVENT_FIELD(uint32, InstanceId)
 	UE_TRACE_EVENT_FIELD(uint32, InstanceSerial)
+	UE_TRACE_EVENT_FIELD(uint8, SourceType)
 	UE_TRACE_EVENT_FIELD(uint16, TransitionIndex)
+	UE_TRACE_EVENT_FIELD(uint16, TargetStateIndex)
+	UE_TRACE_EVENT_FIELD(uint8, Priority)
 	UE_TRACE_EVENT_FIELD(uint8, EventType)
 UE_TRACE_EVENT_END()
 
@@ -106,9 +110,14 @@ double GRecordingWorldTime = -1;
 /** Struct to keep track if a given phase was traced or not. */
 struct FPhaseTraceStatusPair
 {
-	explicit FPhaseTraceStatusPair(const EStateTreeUpdatePhase Phase) : Phase(Phase) {}
+	explicit FPhaseTraceStatusPair(const EStateTreeUpdatePhase Phase, const FStateTreeStateHandle StateHandle)
+	: Phase(Phase)
+	, StateHandle(StateHandle)
+	{
+	}
 
 	EStateTreeUpdatePhase Phase = EStateTreeUpdatePhase::Unset;
+	FStateTreeStateHandle StateHandle = FStateTreeStateHandle::Invalid; 
 	bool bTraced = false;
 };
 
@@ -248,6 +257,7 @@ struct FBufferedDataList
 							<< PhaseEvent.InstanceId(InstanceId.Id)
 							<< PhaseEvent.InstanceSerial(InstanceId.SerialNumber)
 							<< PhaseEvent.Phase(static_cast<std::underlying_type_t<EStateTreeUpdatePhase>>(StackEntry.Phase))
+							<< PhaseEvent.StateIndex(StackEntry.StateHandle.Index)
 							<< PhaseEvent.EventType(static_cast<std::underlying_type_t<EStateTreeTraceEventType>>(EStateTreeTraceEventType::Push));
 			
 						StackEntry.bTraced = true;
@@ -270,7 +280,7 @@ thread_local FBufferedDataList GBufferedEvents;
  * Pushed or pops an entry on the Phase stack for a given Instance.
  * Will send the Pop events for phases popped if their associated Push events were sent.
  */
-void ProcessPhaseScopeEvent(const FStateTreeInstanceDebugId InstanceId, const EStateTreeUpdatePhase Phase, const EStateTreeTraceEventType EventType)
+void ProcessPhaseScopeEvent(const FStateTreeInstanceDebugId InstanceId, const EStateTreeUpdatePhase Phase, const EStateTreeTraceEventType EventType, const FStateTreeStateHandle StateHandle)
 {
 	TArray<FPhaseStack>& PhaseStacks = GBufferedEvents.PhaseStacks;
 	int32 ExistingStackIndex = PhaseStacks.IndexOfByPredicate([InstanceId](const FPhaseStack& PhaseStack){ return PhaseStack.InstanceId == InstanceId; });
@@ -283,7 +293,7 @@ void ProcessPhaseScopeEvent(const FStateTreeInstanceDebugId InstanceId, const ES
 		}
 		FPhaseStack& PhaseStack = PhaseStacks[ExistingStackIndex];
 		PhaseStack.InstanceId = InstanceId;
-		PhaseStack.Stack.Push(FPhaseTraceStatusPair(Phase));
+		PhaseStack.Stack.Push(FPhaseTraceStatusPair(Phase, StateHandle));
 	}
 	else if (ensureMsgf(ExistingStackIndex != INDEX_NONE, TEXT("Not expected to pop phases for an instance that never pushed a phase.")))
 	{
@@ -309,6 +319,7 @@ void ProcessPhaseScopeEvent(const FStateTreeInstanceDebugId InstanceId, const ES
 				<< PhaseEvent.InstanceId(InstanceId.Id)
 				<< PhaseEvent.InstanceSerial(InstanceId.SerialNumber)
 				<< PhaseEvent.Phase(static_cast<std::underlying_type_t<EStateTreeUpdatePhase>>(Phase))
+				<< PhaseEvent.StateIndex(StateHandle.Index)
 				<< PhaseEvent.EventType(static_cast<std::underlying_type_t<EStateTreeTraceEventType>>(EStateTreeTraceEventType::Pop));
 			}
 		}
@@ -326,13 +337,21 @@ void TraceBufferedEvents(const FStateTreeInstanceDebugId InstanceId)
 
 void SerializeDataViewToArchive(FBufferArchive Ar, const FStateTreeDataView DataView)
 {
+	constexpr uint32 PortFlags = 
+		PPF_PropertyWindow // limit to properties visible in Editor 
+		| PPF_ExportsNotFullyQualified
+		| PPF_Delimited // property data should be wrapped in quotes
+		| PPF_ExternalEditor // uses authored names instead of internal names and default values are always written out
+		| PPF_SimpleObjectText // object property values should be exported without the package or class information
+		| PPF_ForDiff; // do not emit object path
+
 	if (const UScriptStruct* ScriptStruct = Cast<const UScriptStruct>(DataView.GetStruct()))
 	{
 		TRACE_CPUPROFILER_EVENT_SCOPE(UE::StateTree::ExportStructAsText)
 		FString StructPath = ScriptStruct->GetPathName();
 		FString TextValue;
 
-		ScriptStruct->ExportText(TextValue, DataView.GetMemory(), DataView.GetMemory(), /*OwnerObject*/nullptr, PPF_None, /*ExportRootScope*/nullptr);
+		ScriptStruct->ExportText(TextValue, DataView.GetMemory(), DataView.GetMemory(), /*OwnerObject*/nullptr, PortFlags | PPF_SeparateDefine, /*ExportRootScope*/nullptr);
 
 		Ar << StructPath;
 		Ar << TextValue;
@@ -346,7 +365,7 @@ void SerializeDataViewToArchive(FBufferArchive Ar, const FStateTreeDataView Data
 
 		// Not using on scope FExportObjectInnerContext since it is very costly to build.
 		// Passing a null context will make the export use an already built thread local context.
-		UExporter::ExportToOutputDevice(nullptr, Object, /*Exporter*/nullptr, OutputDevice, TEXT("copy"), 0, PPF_ExportsNotFullyQualified | PPF_Copy | PPF_Delimited, false, Object->GetOuter());
+		UExporter::ExportToOutputDevice(nullptr, Object, /*Exporter*/nullptr, OutputDevice, TEXT("copy"), 0, PortFlags, false, Object->GetOuter());
 
 		Ar << StructPath;
 		Ar << OutputDevice;
@@ -491,7 +510,7 @@ void OutputTaskEventTrace(
 
 void OutputTransitionEventTrace(
 	const FStateTreeInstanceDebugId InstanceId,
-	const FStateTreeIndex16 TransitionIdx,
+	const FStateTreeTransitionSource Source,
 	const EStateTreeTraceEventType EventType
 	)
 {
@@ -504,7 +523,10 @@ void OutputTransitionEventTrace(
 	<< TransitionEvent.Cycle(FPlatformTime::Cycles64())
 	<< TransitionEvent.InstanceId(InstanceId.Id)
 	<< TransitionEvent.InstanceSerial(InstanceId.SerialNumber)
-	<< TransitionEvent.TransitionIndex(TransitionIdx.Get())
+	<< TransitionEvent.SourceType(static_cast<std::underlying_type_t<EStateTreeTransitionSourceType>>(Source.SourceType))
+	<< TransitionEvent.TransitionIndex(Source.TransitionIndex.Get())
+	<< TransitionEvent.TargetStateIndex(Source.TargetState.Index)
+	<< TransitionEvent.Priority(static_cast<std::underlying_type_t<EStateTreeTransitionPriority>>(Source.Priority))
 	<< TransitionEvent.EventType(static_cast<std::underlying_type_t<EStateTreeTraceEventType>>(EventType));
 }
 
diff --git a/Engine/Plugins/Runtime/StateTree/Source/StateTreeModule/Private/Debugger/StateTreeTraceAnalyzer.cpp b/Engine/Plugins/Runtime/StateTree/Source/StateTreeModule/Private/Debugger/StateTreeTraceAnalyzer.cpp
--- a/Engine/Plugins/Runtime/StateTree/Source/StateTreeModule/Private/Debugger/StateTreeTraceAnalyzer.cpp
+++ b/Engine/Plugins/Runtime/StateTree/Source/StateTreeModule/Private/Debugger/StateTreeTraceAnalyzer.cpp
@@ -90,7 +90,8 @@ bool FStateTreeTraceAnalyzer::OnEvent(const uint16 RouteId, EStyle Style, const
 		{
 			const FStateTreeTracePhaseEvent Event(WorldTime,
 				EventData.GetValue<EStateTreeUpdatePhase>("Phase"),
-				EventData.GetValue<EStateTreeTraceEventType>("EventType"));
+				EventData.GetValue<EStateTreeTraceEventType>("EventType"),
+				FStateTreeStateHandle(EventData.GetValue<uint16>("StateIndex")));
 
 			Provider.AppendEvent(FStateTreeInstanceDebugId(EventData.GetValue<uint32>("InstanceId"), EventData.GetValue<uint32>("InstanceSerial")),
 				Context.EventTime.AsSeconds(EventData.GetValue<uint64>("Cycle")),
@@ -158,7 +159,12 @@ bool FStateTreeTraceAnalyzer::OnEvent(const uint16 RouteId, EStyle Style, const
 	case RouteId_Transition:
 		{
 			const FStateTreeTraceTransitionEvent Event(WorldTime,
-				FStateTreeIndex16(EventData.GetValue<uint16>("TransitionIndex")),
+				FStateTreeTransitionSource(
+					EventData.GetValue<EStateTreeTransitionSourceType>("SourceType"),
+					FStateTreeIndex16(EventData.GetValue<uint16>("TransitionIndex")),
+					FStateTreeStateHandle(EventData.GetValue<uint16>("TargetStateIndex")),
+					EventData.GetValue<EStateTreeTransitionPriority>("Priority")
+					),
 				EventData.GetValue<EStateTreeTraceEventType>("EventType"));
 			
 			Provider.AppendEvent(FStateTreeInstanceDebugId(EventData.GetValue<uint32>("InstanceId"), EventData.GetValue<uint32>("InstanceSerial")),
diff --git a/Engine/Plugins/Runtime/StateTree/Source/StateTreeModule/Private/Debugger/StateTreeTraceTypes.cpp b/Engine/Plugins/Runtime/StateTree/Source/StateTreeModule/Private/Debugger/StateTreeTraceTypes.cpp
--- a/Engine/Plugins/Runtime/StateTree/Source/StateTreeModule/Private/Debugger/StateTreeTraceTypes.cpp
+++ b/Engine/Plugins/Runtime/StateTree/Source/StateTreeModule/Private/Debugger/StateTreeTraceTypes.cpp
@@ -11,53 +11,127 @@
 //----------------------------------------------------------------------//
 // FStateTreeTracePhaseEvent
 //----------------------------------------------------------------------//
-FString FStateTreeTracePhaseEvent::ToString(const UStateTree& StateTree) const
+FString FStateTreeTracePhaseEvent::ToFullString(const UStateTree& StateTree) const
 {
-	return FString::Printf(TEXT("%s %s"),
-			*UEnum::GetDisplayValueAsText(EventType).ToString(),
-			*UEnum::GetDisplayValueAsText(Phase).ToString());
+	return GetValueString(StateTree);
+}
+
+FString FStateTreeTracePhaseEvent::GetValueString(const UStateTree& StateTree) const
+{
+	FStringBuilderBase StrBuilder;
+	StrBuilder.Append(UEnum::GetDisplayValueAsText(Phase).ToString());
+
+	const FCompactStateTreeState* CompactState = StateTree.GetStateFromHandle(StateHandle);
+	if (CompactState != nullptr || StateHandle.IsValid())
+	{
+		StrBuilder.Appendf(TEXT(" '%s'"), CompactState != nullptr ? *CompactState->Name.ToString() : *StateHandle.Describe());
+	}
+
+	return StrBuilder.ToString();
+}
+
+FString FStateTreeTracePhaseEvent::GetTypeString(const UStateTree& StateTree) const
+{
+	return *UEnum::GetDisplayValueAsText(EventType).ToString();
 }
 
 
 //----------------------------------------------------------------------//
 // FStateTreeTraceLogEvent
 //----------------------------------------------------------------------//
-FString FStateTreeTraceLogEvent::ToString(const UStateTree& StateTree) const
+FString FStateTreeTraceLogEvent::ToFullString(const UStateTree& StateTree) const
+{
+	return FString::Printf(TEXT("%s: %s"), *GetTypeString(StateTree), *GetValueString(StateTree));
+}
+
+FString FStateTreeTraceLogEvent::GetValueString(const UStateTree& StateTree) const
 {
 	return (*Message);
 }
 
+FString FStateTreeTraceLogEvent::GetTypeString(const UStateTree& StateTree) const
+{
+	return TEXT("Log");
+}
+
+
+//----------------------------------------------------------------------//
+// FStateTreeTracePropertyEvent
+//----------------------------------------------------------------------//
+FString FStateTreeTracePropertyEvent::ToFullString(const UStateTree& StateTree) const
+{
+	return FString::Printf(TEXT("%s: %s"), *GetTypeString(StateTree), *GetValueString(StateTree));
+}
+
+FString FStateTreeTracePropertyEvent::GetValueString(const UStateTree& StateTree) const
+{
+	return (*Message);
+}
+
+FString FStateTreeTracePropertyEvent::GetTypeString(const UStateTree& StateTree) const
+{
+	return TEXT("Property value");
+}
 
 //----------------------------------------------------------------------//
 // FStateTreeTraceTransitionEvent
 //----------------------------------------------------------------------//
-FString FStateTreeTraceTransitionEvent::ToString(const UStateTree& StateTree) const
+FString FStateTreeTraceTransitionEvent::ToFullString(const UStateTree& StateTree) const
+{
+	return FString::Printf(TEXT("%s %s"), *GetTypeString(StateTree), *GetValueString(StateTree));
+}
+
+FString FStateTreeTraceTransitionEvent::GetValueString(const UStateTree& StateTree) const
+{
+	const FCompactStateTreeState* CompactState = StateTree.GetStateFromHandle(TransitionSource.TargetState);
+	FStringBuilderBase StrBuilder;
+	StrBuilder.Appendf(TEXT("go to State '%s'"), CompactState != nullptr ? *CompactState->Name.ToString() : *TransitionSource.TargetState.Describe());
+
+	if (TransitionSource.Priority != EStateTreeTransitionPriority::None)
+	{
+		StrBuilder.Appendf(TEXT(" (Priority: %s)"), *UEnum::GetDisplayValueAsText(TransitionSource.Priority).ToString()); 
+	}
+
+	if (TransitionSource.SourceType == EStateTreeTransitionSourceType::Asset)
+	{
+		if (const FCompactStateTransition* Transition = StateTree.GetTransitionFromIndex(TransitionSource.TransitionIndex))
+		{
+			ensureAlways(Transition->Priority == TransitionSource.Priority);
+			ensureAlways(Transition->State == TransitionSource.TargetState);
+			if (Transition->EventTag.IsValid())
+			{
+				StrBuilder.Appendf(TEXT("\n\t%s"), *Transition->EventTag.ToString()); 
+			}
+		}
+		else
+		{
+			StrBuilder.Appendf(TEXT("Invalid Transition Index %s for '%s'"), *LexToString(TransitionSource.TransitionIndex.Get()), *StateTree.GetFullName());
+		}
+	}
+
+	return StrBuilder.ToString();
+}
+
+FString FStateTreeTraceTransitionEvent::GetTypeString(const UStateTree& StateTree) const
 {
-	if (const FCompactStateTransition* Transition = StateTree.GetTransitionFromIndex(TransitionIndex))
+	if (TransitionSource.SourceType == EStateTreeTransitionSourceType::Asset)
 	{
-		const FCompactStateTreeState* CompactState = StateTree.GetStateFromHandle(Transition->State);
-		FStringBuilderBase StrBuilder;
-		StrBuilder.Appendf(TEXT("%s Goto %s (Priority: %s)"),
-					*UEnum::GetDisplayValueAsText(Transition->Trigger).ToString(),
-					CompactState != nullptr ? *CompactState->Name.ToString() : *Transition->State.Describe(),
-					*UEnum::GetDisplayValueAsText(Transition->Priority).ToString());
-
-		if (Transition->EventTag.IsValid())
+		if (const FCompactStateTransition* Transition = StateTree.GetTransitionFromIndex(TransitionSource.TransitionIndex))
 		{
-			StrBuilder.Appendf(TEXT("\n\t%s"), *Transition->EventTag.ToString()); 
+			return *UEnum::GetDisplayValueAsText(Transition->Trigger).ToString();
 		}
 
-		return StrBuilder.ToString();
+		return FString::Printf(TEXT("Invalid Transition Index %s for '%s'"), *LexToString(TransitionSource.TransitionIndex.Get()), *StateTree.GetFullName());
 	}
 
-	return FString::Printf(TEXT("Invalid Transition Index %s for '%s'"), *LexToString(TransitionIndex.Get()), *StateTree.GetFullName()); 
+	return *UEnum::GetDisplayValueAsText(TransitionSource.SourceType).ToString();
 }
 
 
 //----------------------------------------------------------------------//
 // FStateTreeTraceNodeEvent
 //----------------------------------------------------------------------//
-FString FStateTreeTraceNodeEvent::ToString(const UStateTree& StateTree) const
+FString FStateTreeTraceNodeEvent::ToFullString(const UStateTree& StateTree) const
 {
 	const FConstStructView NodeView = Index.IsValid() ? StateTree.GetNode(Index.Get()) : FConstStructView();
 	const FStateTreeNodeBase* Node = NodeView.GetPtr<const FStateTreeNodeBase>();
@@ -68,6 +142,22 @@ FString FStateTreeTraceNodeEvent::ToString(const UStateTree& StateTree) const
 			NodeView.IsValid() ? *NodeView.GetScriptStruct()->GetName() : TEXT("Invalid Node"));
 }
 
+FString FStateTreeTraceNodeEvent::GetValueString(const UStateTree& StateTree) const
+{
+	const FConstStructView NodeView = Index.IsValid() ? StateTree.GetNode(Index.Get()) : FConstStructView();
+	const FStateTreeNodeBase* Node = NodeView.GetPtr<const FStateTreeNodeBase>();
+
+	return FString::Printf(TEXT("%s '%s'"),
+			*UEnum::GetDisplayValueAsText(EventType).ToString(),
+			Node != nullptr ? *Node->Name.ToString() : *LexToString(Index.Get()));
+}
+
+FString FStateTreeTraceNodeEvent::GetTypeString(const UStateTree& StateTree) const
+{
+	const FConstStructView NodeView = Index.IsValid() ? StateTree.GetNode(Index.Get()) : FConstStructView();
+	return FString::Printf(TEXT("%s"), NodeView.IsValid() ? *NodeView.GetScriptStruct()->GetName() : TEXT("Invalid Node"));
+}
+
 
 //----------------------------------------------------------------------//
 // FStateTreeTraceStateEvent
@@ -77,7 +167,7 @@ FStateTreeStateHandle FStateTreeTraceStateEvent::GetStateHandle() const
 	return FStateTreeStateHandle(Index.Get());
 }
 
-FString FStateTreeTraceStateEvent::ToString(const UStateTree& StateTree) const
+FString FStateTreeTraceStateEvent::ToFullString(const UStateTree& StateTree) const
 {
 	const FStateTreeStateHandle StateHandle(Index.Get());
 	if (const FCompactStateTreeState* CompactState = StateTree.GetStateFromHandle(StateHandle))
@@ -92,9 +182,36 @@ FString FStateTreeTraceStateEvent::ToString(const UStateTree& StateTree) const
 		else
 		{
 			return FString::Printf(TEXT("%s '%s'"),
-					*UEnum::GetDisplayValueAsText(EventType).ToString(),
-					*CompactState->Name.ToString());
+				*UEnum::GetDisplayValueAsText(EventType).ToString(),
+				*CompactState->Name.ToString());
+		}
+	}
+
+	return FString::Printf(TEXT("Invalid State Index %s for '%s'"), *StateHandle.Describe(), *StateTree.GetFullName());
+}
+
+FString FStateTreeTraceStateEvent::GetValueString(const UStateTree& StateTree) const
+{
+	const FStateTreeStateHandle StateHandle(Index.Get());
+	if (const FCompactStateTreeState* CompactState = StateTree.GetStateFromHandle(StateHandle))
+	{
+		return FString::Printf(TEXT("%s"), *CompactState->Name.ToString());
+	}
+
+	return FString::Printf(TEXT("Invalid State Index %s for '%s'"), *StateHandle.Describe(), *StateTree.GetFullName());
+}
+
+FString FStateTreeTraceStateEvent::GetTypeString(const UStateTree& StateTree) const
+{
+	const FStateTreeStateHandle StateHandle(Index.Get());
+	if (const FCompactStateTreeState* CompactState = StateTree.GetStateFromHandle(StateHandle))
+	{
+		if (SelectionBehavior != EStateTreeStateSelectionBehavior::None)
+		{
+			return FString::Printf(TEXT("%s '%s'"), *UEnum::GetDisplayValueAsText(SelectionBehavior).ToString(), *CompactState->Name.ToString());
 		}
+
+		return FString::Printf(TEXT(""));
 	}
 
 	return FString::Printf(TEXT("Invalid State Index %s for '%s'"), *StateHandle.Describe(), *StateTree.GetFullName());
@@ -104,20 +221,45 @@ FString FStateTreeTraceStateEvent::ToString(const UStateTree& StateTree) const
 //----------------------------------------------------------------------//
 // FStateTreeTraceTaskEvent
 //----------------------------------------------------------------------//
-FString FStateTreeTraceTaskEvent::ToString(const UStateTree& StateTree) const
+FString FStateTreeTraceTaskEvent::ToFullString(const UStateTree& StateTree) const
 {
-	return FString::Printf(TEXT("<%s> %s"),	*UEnum::GetDisplayValueAsText(Status).ToString(), *FStateTreeTraceNodeEvent::ToString(StateTree));
+	return FString::Printf(TEXT("%s -> %s"),	*FStateTreeTraceNodeEvent::ToFullString(StateTree), *UEnum::GetDisplayValueAsText(Status).ToString());
+}
+
+FString FStateTreeTraceTaskEvent::GetValueString(const UStateTree& StateTree) const
+{
+	const FConstStructView NodeView = Index.IsValid() ? StateTree.GetNode(Index.Get()) : FConstStructView();
+	const FStateTreeNodeBase* Node = NodeView.GetPtr<const FStateTreeNodeBase>();
+	return FString::Printf(TEXT("%s"), Node != nullptr ? *Node->Name.ToString() : *LexToString(Index.Get()));
+}
+
+FString FStateTreeTraceTaskEvent::GetTypeString(const UStateTree& StateTree) const
+{
+	const FConstStructView NodeView = Index.IsValid() ? StateTree.GetNode(Index.Get()) : FConstStructView();
+	return FString::Printf(TEXT("%s"), NodeView.IsValid() ? *NodeView.GetScriptStruct()->GetName() : TEXT("Invalid Node"));
 }
 
 
 //----------------------------------------------------------------------//
 // FStateTreeTraceConditionEvent
 //----------------------------------------------------------------------//
-FString FStateTreeTraceConditionEvent::ToString(const UStateTree& StateTree) const
+FString FStateTreeTraceConditionEvent::ToFullString(const UStateTree& StateTree) const
+{
+	return FStateTreeTraceNodeEvent::ToFullString(StateTree);
+}
+
+FString FStateTreeTraceConditionEvent::GetValueString(const UStateTree& StateTree) const
 {
-	return FStateTreeTraceNodeEvent::ToString(StateTree);
+	const FConstStructView NodeView = Index.IsValid() ? StateTree.GetNode(Index.Get()) : FConstStructView();
+	const FStateTreeNodeBase* Node = NodeView.GetPtr<const FStateTreeNodeBase>();
+
+	return FString::Printf(TEXT("%s"), Node != nullptr ? *Node->Name.ToString() : *LexToString(Index.Get()));
 }
 
+FString FStateTreeTraceConditionEvent::GetTypeString(const UStateTree& StateTree) const
+{
+	return FStateTreeTraceNodeEvent::GetTypeString(StateTree);
+}
 
 //----------------------------------------------------------------------//
 // FStateTreeTraceActiveStatesEvent
@@ -127,7 +269,16 @@ FStateTreeTraceActiveStatesEvent::FStateTreeTraceActiveStatesEvent(const double
 {
 }
 
-FString FStateTreeTraceActiveStatesEvent::ToString(const UStateTree& StateTree) const
+FString FStateTreeTraceActiveStatesEvent::ToFullString(const UStateTree& StateTree) const
+{
+	if (ActiveStates.Num() > 0)
+	{
+		return FString::Printf(TEXT("%s: %s"), *GetTypeString(StateTree), *GetValueString(StateTree));
+	}
+	return TEXT("No active states");
+}
+
+FString FStateTreeTraceActiveStatesEvent::GetValueString(const UStateTree& StateTree) const
 {
 	FStringBuilderBase StatePath;
 	for (int32 i = 0; i < ActiveStates.Num(); i++)
@@ -136,8 +287,12 @@ FString FStateTreeTraceActiveStatesEvent::ToString(const UStateTree& StateTree)
 		const FCompactStateTreeState& State = StatesView[ActiveStates[i].Index];
 		StatePath.Appendf(TEXT("%s%s"), i == 0 ? TEXT("") : TEXT("."), *State.Name.ToString());
 	}
+	return StatePath.ToString();
+}
 
-	return FString::Printf(TEXT("New active states: '%s'"), *StatePath);
+FString FStateTreeTraceActiveStatesEvent::GetTypeString(const UStateTree& StateTree) const
+{
+	return TEXT("New active states");
 }
 
 #endif // WITH_STATETREE_DEBUGGER
diff --git a/Engine/Plugins/Runtime/StateTree/Source/StateTreeModule/Private/StateTreeExecutionContext.cpp b/Engine/Plugins/Runtime/StateTree/Source/StateTreeModule/Private/StateTreeExecutionContext.cpp
--- a/Engine/Plugins/Runtime/StateTree/Source/StateTreeModule/Private/StateTreeExecutionContext.cpp
+++ b/Engine/Plugins/Runtime/StateTree/Source/StateTreeModule/Private/StateTreeExecutionContext.cpp
@@ -20,13 +20,18 @@
 
 	#define STATETREE_TRACE_SCOPED_PHASE(Phase) \
 		FStateTreeInstanceDebugId ID_NAME = GetInstanceDebugId(); \
-		TRACE_STATETREE_PHASE_EVENT(ID_NAME, Phase, EStateTreeTraceEventType::Push) \
-		ON_SCOPE_EXIT { TRACE_STATETREE_PHASE_EVENT(ID_NAME, Phase, EStateTreeTraceEventType::Pop) }
+		TRACE_STATETREE_PHASE_EVENT(ID_NAME, Phase, EStateTreeTraceEventType::Push, FStateTreeStateHandle::Invalid) \
+		ON_SCOPE_EXIT { TRACE_STATETREE_PHASE_EVENT(ID_NAME, Phase, EStateTreeTraceEventType::Pop, FStateTreeStateHandle::Invalid) }
 
-	#define STATETREE_TRACE_SCOPED_STATE_SELECTION(StateHandle, SelectionBehavior) \
+	#define STATETREE_TRACE_SCOPED_STATE(StateHandle) \
 		FStateTreeInstanceDebugId ID_NAME = GetInstanceDebugId(); \
-		TRACE_STATETREE_STATE_EVENT(ID_NAME, StateHandle, EStateTreeTraceEventType::Push, SelectionBehavior) \
-		ON_SCOPE_EXIT { TRACE_STATETREE_STATE_EVENT(ID_NAME, StateHandle, EStateTreeTraceEventType::Pop, SelectionBehavior); };
+		TRACE_STATETREE_PHASE_EVENT(ID_NAME, EStateTreeUpdatePhase::Unset, EStateTreeTraceEventType::Push, StateHandle) \
+		ON_SCOPE_EXIT { TRACE_STATETREE_PHASE_EVENT(ID_NAME, EStateTreeUpdatePhase::Unset, EStateTreeTraceEventType::Pop, StateHandle) }
+
+	#define STATETREE_TRACE_SCOPED_STATE_PHASE(StateHandle, Phase) \
+		FStateTreeInstanceDebugId ID_NAME = GetInstanceDebugId(); \
+		TRACE_STATETREE_PHASE_EVENT(ID_NAME, Phase, EStateTreeTraceEventType::Push, StateHandle) \
+		ON_SCOPE_EXIT { TRACE_STATETREE_PHASE_EVENT(ID_NAME, Phase, EStateTreeTraceEventType::Pop, StateHandle) }
 
 	#define STATETREE_TRACE_INSTANCE_EVENT(EventType)						TRACE_STATETREE_INSTANCE_EVENT(GetInstanceDebugId(), GetStateTree(), *GetInstanceDescription(), EventType);
 	#define STATETREE_TRACE_ACTIVE_STATES_EVENT(ActiveStates)				TRACE_STATETREE_ACTIVE_STATES_EVENT(GetInstanceDebugId(), ActiveStates);
@@ -34,18 +39,18 @@
 	#define STATETREE_TRACE_STATE_EVENT(StateHandle, EventType)				TRACE_STATETREE_STATE_EVENT(GetInstanceDebugId(), StateHandle, EventType, EStateTreeStateSelectionBehavior::None);
 	#define STATETREE_TRACE_TASK_EVENT(Index, DataView, EventType, Status)	TRACE_STATETREE_TASK_EVENT(GetInstanceDebugId(), FStateTreeIndex16(Index), DataView, EventType, Status);
 	#define STATETREE_TRACE_CONDITION_EVENT(Index, DataViews, EventType)	TRACE_STATETREE_CONDITION_EVENT(GetInstanceDebugId(), FStateTreeIndex16(Index), DataView, EventType);	
-	#define STATETREE_TRACE_TRANSITION_EVENT(Index, EventType)				TRACE_STATETREE_TRANSITION_EVENT(GetInstanceDebugId(), FStateTreeIndex16(Index), EventType);
+	#define STATETREE_TRACE_TRANSITION_EVENT(Source, EventType)				TRACE_STATETREE_TRANSITION_EVENT(GetInstanceDebugId(), Source, EventType);
 #else
 	#define STATETREE_TRACE_SCOPED_PHASE(Phase)
-	#define STATETREE_TRACE_SCOPED_STATE_SELECTION(StateHandle, SelectionBehavior)
+	#define STATETREE_TRACE_SCOPED_STATE(StateHandle)
+	#define STATETREE_TRACE_SCOPED_STATE_PHASE(StateHandle, Phase)
 	#define STATETREE_TRACE_INSTANCE_EVENT(EventType)
 	#define STATETREE_TRACE_ACTIVE_STATES_EVENT(ActiveStates)
 	#define STATETREE_TRACE_LOG_EVENT(Format, ...)
 	#define STATETREE_TRACE_STATE_EVENT(StateHandle, EventType)
-	#define STATETREE_TRACE_STATE_SELECTION_EVENT(StateHandle, EventType, SelectionBehavior)
 	#define STATETREE_TRACE_TASK_EVENT(Index, DataView, EventType, Status)
 	#define STATETREE_TRACE_CONDITION_EVENT(Index, DataView, EventType)
-	#define STATETREE_TRACE_TRANSITION_EVENT(Index, EventType)
+	#define STATETREE_TRACE_TRANSITION_EVENT(Source, EventType)
 #endif // WITH_STATETREE_DEBUGGER
 
 namespace UE::StateTree
@@ -269,6 +274,7 @@ EStateTreeRunStatus FStateTreeExecutionContext::Start()
 			// Need to reacquire the exec state as EnterState may alter the allocation. 
 			Exec = &GetExecState();
 			Exec->LastTickStatus = LastTickStatus;
+			STATETREE_TRACE_ACTIVE_STATES_EVENT(Exec->ActiveStates);
 
 			// Report state completed immediately.
 			if (Exec->LastTickStatus != EStateTreeRunStatus::Running)
@@ -419,8 +425,8 @@ EStateTreeRunStatus FStateTreeExecutionContext::Tick(const float DeltaTime)
 		if (TriggerTransitions())
 		{
 			STATETREE_TRACE_SCOPED_PHASE(EStateTreeUpdatePhase::ApplyTransitions);
-			STATETREE_TRACE_TRANSITION_EVENT(NextTransitionIndex, EStateTreeTraceEventType::OnTransition);
-			NextTransitionIndex = FStateTreeIndex16::Invalid;
+			STATETREE_TRACE_TRANSITION_EVENT(NextTransitionSource, EStateTreeTraceEventType::OnTransition);
+			NextTransitionSource.Reset();
 
 			// We have committed to state change, consume events that were accumulated during the tick above.
 			EventQueue.Reset();
@@ -452,6 +458,7 @@ EStateTreeRunStatus FStateTreeExecutionContext::Tick(const float DeltaTime)
 			// Need to reacquire the exec state as EnterState may alter the allocation. 
 			Exec = &GetExecState();
 			Exec->LastTickStatus = LastTickStatus;
+			STATETREE_TRACE_ACTIVE_STATES_EVENT(Exec->ActiveStates);
 
 			// Consider events so far processed. Events sent during EnterState went into EventQueue, and are processed in next iteration.
 			EventsToProcess.Reset();
@@ -558,7 +565,10 @@ void FStateTreeExecutionContext::RequestTransition(const FStateTreeTransitionReq
 
 	if (bAllowDirectTransitions)
 	{
-		RequestTransition(Request.TargetState, Request.Priority);		
+		if (RequestTransition(Request.TargetState, Request.Priority))
+		{
+			NextTransitionSource = FStateTreeTransitionSource(EStateTreeTransitionSourceType::ExternalRequest, Request.TargetState, Request.Priority);
+		}
 	}
 	else
 	{
@@ -745,6 +755,7 @@ EStateTreeRunStatus FStateTreeExecutionContext::EnterState(const FStateTreeTrans
 	}
 	
 	STATETREE_LOG(Log, TEXT("Enter state '%s' (%d)"), *DebugGetStatePath(Transition.NextActiveStates), Exec.StateChangeCount);
+	STATETREE_TRACE_SCOPED_PHASE(EStateTreeUpdatePhase::EnterStates);
 
 	for (int32 Index = 0; Index < Transition.NextActiveStates.Num() && Result != EStateTreeRunStatus::Failed; Index++)
 	{
@@ -780,7 +791,7 @@ EStateTreeRunStatus FStateTreeExecutionContext::EnterState(const FStateTreeTrans
 
 		if (bIsEnteringState)
 		{
-			STATETREE_TRACE_STATE_EVENT(CurrentHandle, EStateTreeTraceEventType::OnEnter);
+			STATETREE_TRACE_STATE_EVENT(CurrentHandle, EStateTreeTraceEventType::OnEntering);
 			STATETREE_LOG(Log, TEXT("%*sState '%s' %s"), Index*UE::StateTree::DebugIndentSize, TEXT(""),
 				*DebugGetStatePath(Transition.NextActiveStates, Index),
 				*UEnum::GetDisplayValueAsText(CurrentTransition.ChangeType).ToString());
@@ -820,7 +831,7 @@ EStateTreeRunStatus FStateTreeExecutionContext::EnterState(const FStateTreeTrans
 					Status = Task.EnterState(*this, CurrentTransition);
 				}
 
-				STATETREE_TRACE_TASK_EVENT(TaskIndex, DataViews[Task.DataViewIndex.Get()], EStateTreeTraceEventType::OnEnter, Status);
+				STATETREE_TRACE_TASK_EVENT(TaskIndex, DataViews[Task.DataViewIndex.Get()], EStateTreeTraceEventType::OnEntered, Status);
 
 				if (Status != EStateTreeRunStatus::Running)
 				{
@@ -840,9 +851,13 @@ EStateTreeRunStatus FStateTreeExecutionContext::EnterState(const FStateTreeTrans
 				}
 			}
 		}
+
+		if (bIsEnteringState)
+		{
+			STATETREE_TRACE_STATE_EVENT(CurrentHandle, EStateTreeTraceEventType::OnEntered);
+		}
 	}
 
-	STATETREE_TRACE_ACTIVE_STATES_EVENT(Exec.ActiveStates);
 	return Result;
 }
 
@@ -919,6 +934,7 @@ void FStateTreeExecutionContext::ExitState(const FStateTreeTransitionResult& Tra
 
 	// Call in reverse order.
 	STATETREE_LOG(Log, TEXT("Exit state '%s' (%d)"), *DebugGetStatePath(Transition.CurrentActiveStates), Exec.StateChangeCount);
+	STATETREE_TRACE_SCOPED_PHASE(EStateTreeUpdatePhase::ExitStates);
 
 	FStateTreeTransitionResult CurrentTransition = Transition;
 
@@ -938,7 +954,7 @@ void FStateTreeExecutionContext::ExitState(const FStateTreeTransitionResult& Tra
 		CurrentTransition.ChangeType = ExitedStateChangeType[Index];
 
 		STATETREE_LOG(Log, TEXT("%*sState '%s' %s"), Index*UE::StateTree::DebugIndentSize, TEXT(""), *DebugGetStatePath(Transition.CurrentActiveStates, ExitedStateActiveIndex[Index]), *UEnum::GetDisplayValueAsText(CurrentTransition.ChangeType).ToString());
-		STATETREE_TRACE_STATE_EVENT(CurrentHandle, EStateTreeTraceEventType::OnExit);
+		STATETREE_TRACE_STATE_EVENT(CurrentHandle, EStateTreeTraceEventType::OnExiting);
 
 		// Tasks
 		for (int32 TaskIndex = (State.TasksBegin + State.TasksNum) - 1; TaskIndex >= State.TasksBegin; TaskIndex--)
@@ -968,10 +984,12 @@ void FStateTreeExecutionContext::ExitState(const FStateTreeTransitionResult& Tra
 						CSV_SCOPED_TIMING_STAT_EXCLUSIVE(StateTree_Task_ExitState);
 						Task.ExitState(*this, CurrentTransition);
 					}
-					STATETREE_TRACE_TASK_EVENT(TaskIndex, DataViews[Task.DataViewIndex.Get()], EStateTreeTraceEventType::OnExit, Transition.CurrentRunStatus);
+					STATETREE_TRACE_TASK_EVENT(TaskIndex, DataViews[Task.DataViewIndex.Get()], EStateTreeTraceEventType::OnExited, Transition.CurrentRunStatus);
 				}
 			}
 		}
+
+		STATETREE_TRACE_STATE_EVENT(CurrentHandle, EStateTreeTraceEventType::OnExited);
 	}
 }
 
@@ -1089,7 +1107,7 @@ EStateTreeRunStatus FStateTreeExecutionContext::TickEvaluatorsAndGlobalTasks(con
 				StateTree.PropertyBindings.CopyTo(DataViews, Task.BindingsBatch, TaskDataView);
 			}
 
-			STATETREE_TRACE_TASK_EVENT(TaskIndex, TaskDataView, EStateTreeTraceEventType::OnTickingTask, EStateTreeRunStatus::Running);
+			//STATETREE_TRACE_TASK_EVENT(TaskIndex, TaskDataView, EStateTreeTraceEventType::OnTickingTask, EStateTreeRunStatus::Running);
 			EStateTreeRunStatus TaskResult = EStateTreeRunStatus::Unset;
 			{
 				QUICK_SCOPE_CYCLE_COUNTER(StateTree_Task_Tick);
@@ -1176,7 +1194,7 @@ EStateTreeRunStatus FStateTreeExecutionContext::StartEvaluatorsAndGlobalTasks(FS
 			QUICK_SCOPE_CYCLE_COUNTER(StateTree_Task_TreeStart);
 			const EStateTreeRunStatus TaskStatus = Task.EnterState(*this, Transition); 
 
-			STATETREE_TRACE_TASK_EVENT(TaskIndex, DataViews[Task.DataViewIndex.Get()], EStateTreeTraceEventType::OnEnter, TaskStatus);
+			STATETREE_TRACE_TASK_EVENT(TaskIndex, DataViews[Task.DataViewIndex.Get()], EStateTreeTraceEventType::OnEntered, TaskStatus);
 
 			if (TaskStatus != EStateTreeRunStatus::Running)
 			{
@@ -1253,7 +1271,7 @@ void FStateTreeExecutionContext::StopEvaluatorsAndGlobalTasks(const EStateTreeRu
 				QUICK_SCOPE_CYCLE_COUNTER(StateTree_Task_TreeStop);
 				Task.ExitState(*this, Transition);
 			}
-			STATETREE_TRACE_TASK_EVENT(TaskIndex, DataViews[Task.DataViewIndex.Get()], EStateTreeTraceEventType::OnExit, Transition.CurrentRunStatus);
+			STATETREE_TRACE_TASK_EVENT(TaskIndex, DataViews[Task.DataViewIndex.Get()], EStateTreeTraceEventType::OnExited, Transition.CurrentRunStatus);
 		}
 	}
 
@@ -1302,7 +1320,8 @@ EStateTreeRunStatus FStateTreeExecutionContext::TickTasks(const float DeltaTime)
 		const FCompactStateTreeState& State = StateTree.States[CurrentHandle.Index];
 
 		FCurrentlyProcessedStateScope StateScope(*this, CurrentHandle);
-		
+		STATETREE_TRACE_SCOPED_STATE(CurrentHandle);
+
 		STATETREE_CLOG(State.TasksNum > 0, VeryVerbose, TEXT("%*sState '%s'"), Index*UE::StateTree::DebugIndentSize, TEXT(""), *DebugGetStatePath(Exec.ActiveStates, Index));
 
 		if (State.Type == EStateTreeStateType::Linked)
@@ -1344,7 +1363,7 @@ EStateTreeRunStatus FStateTreeExecutionContext::TickTasks(const float DeltaTime)
 				StateTree.PropertyBindings.CopyTo(DataViews, Task.BindingsBatch, TaskDataView);
 			}
 
-			STATETREE_TRACE_TASK_EVENT(TaskIndex, TaskDataView, EStateTreeTraceEventType::OnTickingTask, EStateTreeRunStatus::Running);
+			//STATETREE_TRACE_TASK_EVENT(TaskIndex, TaskDataView, EStateTreeTraceEventType::OnTickingTask, EStateTreeRunStatus::Running);
 			EStateTreeRunStatus TaskResult = EStateTreeRunStatus::Unset;
 			{
 				QUICK_SCOPE_CYCLE_COUNTER(StateTree_Task_Tick);
@@ -1571,7 +1590,10 @@ bool FStateTreeExecutionContext::TriggerTransitions()
 	//
 	for (const FStateTreeTransitionRequest& Request : InstanceData.GetTransitionRequests())
 	{
-		RequestTransition(Request.TargetState, Request.Priority);
+		if (RequestTransition(Request.TargetState, Request.Priority))
+		{
+			NextTransitionSource = FStateTreeTransitionSource(EStateTreeTransitionSourceType::ExternalRequest, Request.TargetState, Request.Priority);
+		}
 	}
 	InstanceData.ResetTransitionRequests();
 	
@@ -1635,6 +1657,7 @@ bool FStateTreeExecutionContext::TriggerTransitions()
 			ensureMsgf(State.bEnabled, TEXT("Should never try to transition to a disabled state. This indicates an unhandled case."));
 
 			FCurrentlyProcessedStateScope StateScope(*this, StateHandle);
+			STATETREE_TRACE_SCOPED_STATE(StateHandle);
 
 			if (State.bHasTransitionTasks)
 			{
@@ -1703,7 +1726,7 @@ bool FStateTreeExecutionContext::TriggerTransitions()
 						// Trigger Delayed Transition when the delay has passed.
 						if (RequestTransition(Transition.State, Transition.Priority))
 						{
-							NextTransitionIndex = FStateTreeIndex16(TransitionIndex);
+							NextTransitionSource = FStateTreeTransitionSource(FStateTreeIndex16(TransitionIndex), Transition.State, Transition.Priority);
 						}
 						continue;
 					}
@@ -1716,7 +1739,7 @@ bool FStateTreeExecutionContext::TriggerTransitions()
 				bool bPassed = false; 
 				if (bShouldTrigger)
 				{
-					STATETREE_TRACE_TRANSITION_EVENT(TransitionIndex, EStateTreeTraceEventType::OnEvaluating);
+					STATETREE_TRACE_TRANSITION_EVENT(FStateTreeTransitionSource(FStateTreeIndex16(TransitionIndex), Transition.State, Transition.Priority), EStateTreeTraceEventType::OnEvaluating);
 					STATETREE_TRACE_SCOPED_PHASE(EStateTreeUpdatePhase::TransitionConditions);
 					bPassed = TestAllConditions(Transition.ConditionsBegin, Transition.ConditionsNum);
 				}
@@ -1753,7 +1776,7 @@ bool FStateTreeExecutionContext::TriggerTransitions()
 
 					if (RequestTransition(Transition.State, Transition.Priority))
 					{
-						NextTransitionIndex = FStateTreeIndex16(TransitionIndex);
+						NextTransitionSource = FStateTreeTransitionSource(FStateTreeIndex16(TransitionIndex), Transition.State, Transition.Priority);
 					}
 				}
 			}
@@ -1802,7 +1825,8 @@ bool FStateTreeExecutionContext::TriggerTransitions()
 			const FCompactStateTreeState& State = StateTree.States[StateHandle.Index];
 
 			FCurrentlyProcessedStateScope StateScope(*this, StateHandle);
-			
+			STATETREE_TRACE_SCOPED_STATE_PHASE(StateHandle, EStateTreeUpdatePhase::TriggerTransitions);
+
 			for (uint8 i = 0; i < State.TransitionsNum; i++)
 			{
 				// All transition conditions must pass
@@ -1819,7 +1843,7 @@ bool FStateTreeExecutionContext::TriggerTransitions()
 				{
 					bool bPassed = false;
 					{
-						STATETREE_TRACE_TRANSITION_EVENT(TransitionIndex, EStateTreeTraceEventType::OnEvaluating);
+						STATETREE_TRACE_TRANSITION_EVENT(FStateTreeTransitionSource(FStateTreeIndex16(TransitionIndex), Transition.State, Transition.Priority), EStateTreeTraceEventType::OnEvaluating);
 						STATETREE_TRACE_SCOPED_PHASE(EStateTreeUpdatePhase::TransitionConditions);
 						bPassed = TestAllConditions(Transition.ConditionsBegin, Transition.ConditionsNum);
 					}
@@ -1830,7 +1854,7 @@ bool FStateTreeExecutionContext::TriggerTransitions()
 						// No priority on completion transitions, use the priority to signal that state is selected.
 						if (RequestTransition(Transition.State, EStateTreeTransitionPriority::Normal))
 						{
-							NextTransitionIndex = FStateTreeIndex16(TransitionIndex);
+							NextTransitionSource = FStateTreeTransitionSource(FStateTreeIndex16(TransitionIndex), Transition.State, Transition.Priority);
 							break;
 						}
 					}
@@ -1846,8 +1870,12 @@ bool FStateTreeExecutionContext::TriggerTransitions()
 		if (NextTransition.Priority == EStateTreeTransitionPriority::None)
 		{
 			STATETREE_LOG(Verbose, TEXT("Could not trigger completion transition, jump back to start."));
+			STATETREE_TRACE_LOG_EVENT(TEXT("Could not trigger completion transition, jump back to start."));
 			FCurrentlyProcessedStateScope StateScope(*this, FStateTreeStateHandle::Root);
-			RequestTransition(FStateTreeStateHandle::Root, EStateTreeTransitionPriority::Normal);
+			if (RequestTransition(FStateTreeStateHandle::Root, EStateTreeTransitionPriority::Normal))
+			{
+				NextTransitionSource = FStateTreeTransitionSource(EStateTreeTransitionSourceType::Internal, FStateTreeStateHandle::Root, EStateTreeTransitionPriority::Normal);
+			}
 		}
 	}
 
@@ -2007,9 +2035,7 @@ bool FStateTreeExecutionContext::SelectStateInternal(const FStateTreeStateHandle
 		return false;
 	}
 
-	// Push State before the scoped phase so the debugger can create the hierarchy parent using the state description.
-	STATETREE_TRACE_SCOPED_STATE_SELECTION(NextState, State.SelectionBehavior);
-	STATETREE_TRACE_SCOPED_PHASE(EStateTreeUpdatePhase::StateSelection);
+	STATETREE_TRACE_SCOPED_STATE_PHASE(NextState, EStateTreeUpdatePhase::StateSelection);
 	
 	// Check that the state can be entered
 	bool bEnterConditionsPassed = false;
@@ -2097,7 +2123,7 @@ bool FStateTreeExecutionContext::SelectStateInternal(const FStateTreeStateHandle
 				bool bTransitionConditionsPassed = false;
 				if (bShouldTrigger)
 				{
-					STATETREE_TRACE_TRANSITION_EVENT(TransitionIndex, EStateTreeTraceEventType::OnEvaluating);
+					STATETREE_TRACE_TRANSITION_EVENT(FStateTreeTransitionSource(FStateTreeIndex16(TransitionIndex), Transition.State, Transition.Priority), EStateTreeTraceEventType::OnEvaluating);
 					STATETREE_TRACE_SCOPED_PHASE(EStateTreeUpdatePhase::TransitionConditions);
 					bTransitionConditionsPassed = bShouldTrigger && TestAllConditions(Transition.ConditionsBegin, Transition.ConditionsNum);
 				}
diff --git a/Engine/Plugins/Runtime/StateTree/Source/StateTreeModule/Public/Debugger/StateTreeDebugger.h b/Engine/Plugins/Runtime/StateTree/Source/StateTreeModule/Public/Debugger/StateTreeDebugger.h
--- a/Engine/Plugins/Runtime/StateTree/Source/StateTreeModule/Public/Debugger/StateTreeDebugger.h
+++ b/Engine/Plugins/Runtime/StateTree/Source/StateTreeModule/Public/Debugger/StateTreeDebugger.h
@@ -110,6 +110,7 @@ struct STATETREEMODULE_API FStateTreeDebugger : FTickableGameObject
 	void ClearBreakpoint(FStateTreeIndex16 NodeIndex, EStateTreeBreakpointType BreakpointType);
 	void ClearAllBreakpoints();
 	int32 NumBreakpoints() const { return Breakpoints.Num(); }
+	bool HasHitBreakpoint() const { return HitBreakpoint.IsSet(); }
 
 	static FText DescribeTrace(const FTraceDescriptor& TraceDescriptor);
 	static FText DescribeInstance(const UE::StateTreeDebugger::FInstanceDescriptor& StateTreeInstanceDesc);
@@ -168,7 +169,11 @@ struct STATETREEMODULE_API FStateTreeDebugger : FTickableGameObject
 	 */
 	bool RequestSessionAnalysis(const FTraceDescriptor& TraceDescriptor);
 	void PauseSessionAnalysis()  { bSessionAnalysisPaused = true; }
-	void ResumeSessionAnalysis() { bSessionAnalysisPaused = false; }
+	void ResumeSessionAnalysis()
+	{
+		bSessionAnalysisPaused = false;
+		HitBreakpoint.Reset();
+	}
 	void StopSessionAnalysis();
 	FTraceDescriptor GetSelectedTraceDescriptor() const { return ActiveSessionTraceDescriptor; }
 	FText GetSelectedTraceDescription() const;
diff --git a/Engine/Plugins/Runtime/StateTree/Source/StateTreeModule/Public/Debugger/StateTreeTrace.h b/Engine/Plugins/Runtime/StateTree/Source/StateTreeModule/Public/Debugger/StateTreeTrace.h
--- a/Engine/Plugins/Runtime/StateTree/Source/StateTreeModule/Public/Debugger/StateTreeTrace.h
+++ b/Engine/Plugins/Runtime/StateTree/Source/StateTreeModule/Public/Debugger/StateTreeTrace.h
@@ -12,6 +12,7 @@ struct FStateTreeActiveStates;
 struct FStateTreeInstanceDebugId;
 struct FStateTreeIndex16;
 struct FStateTreeStateHandle;
+struct FStateTreeTransitionSource;
 enum class EStateTreeStateSelectionBehavior : uint8;
 enum class EStateTreeRunStatus : uint8;
 enum class EStateTreeTraceEventType : uint8;
@@ -23,21 +24,21 @@ namespace UE::StateTreeTrace
 {
 	void RegisterGlobalDelegates();
 	void UnregisterGlobalDelegates();
-	void ProcessPhaseScopeEvent(FStateTreeInstanceDebugId InstanceId, EStateTreeUpdatePhase Phase, EStateTreeTraceEventType EventType);
+	void ProcessPhaseScopeEvent(FStateTreeInstanceDebugId InstanceId, EStateTreeUpdatePhase Phase, EStateTreeTraceEventType EventType, FStateTreeStateHandle StateHandle);
 	void OutputInstanceLifetimeEvent(FStateTreeInstanceDebugId InstanceId, const UStateTree* StateTree, const TCHAR* InstanceName, EStateTreeTraceEventType EventType);
 	void OutputLogEventTrace(FStateTreeInstanceDebugId InstanceId, const TCHAR* Fmt, ...);
 	void OutputStateEventTrace(FStateTreeInstanceDebugId InstanceId, FStateTreeStateHandle StateHandle, EStateTreeTraceEventType EventType, EStateTreeStateSelectionBehavior SelectionBehavior);
 	void OutputTaskEventTrace(FStateTreeInstanceDebugId InstanceId, FStateTreeIndex16 TaskIdx, FStateTreeDataView DataView, EStateTreeTraceEventType EventType, EStateTreeRunStatus Status);
 	void OutputConditionEventTrace(FStateTreeInstanceDebugId InstanceId, FStateTreeIndex16 ConditionIdx, FStateTreeDataView DataView, EStateTreeTraceEventType EventType);
-	void OutputTransitionEventTrace(FStateTreeInstanceDebugId InstanceId, FStateTreeIndex16 TransitionIdx, EStateTreeTraceEventType EventType);
+	void OutputTransitionEventTrace(FStateTreeInstanceDebugId InstanceId, FStateTreeTransitionSource TransitionSource, EStateTreeTraceEventType EventType);
 	void OutputActiveStatesEventTrace(FStateTreeInstanceDebugId InstanceId, const FStateTreeActiveStates& ActiveStates);
 }
 
 #define TRACE_STATETREE_INSTANCE_EVENT(InstanceID, StateTree, InstanceName, EventType) \
 	UE::StateTreeTrace::OutputInstanceLifetimeEvent(InstanceID, StateTree, InstanceName, EventType);
 
-#define TRACE_STATETREE_PHASE_EVENT(InstanceID, Phase, EventType) \
-	UE::StateTreeTrace::ProcessPhaseScopeEvent(InstanceID, Phase, EventType); \
+#define TRACE_STATETREE_PHASE_EVENT(InstanceID, Phase, EventType, StateHandle) \
+	UE::StateTreeTrace::ProcessPhaseScopeEvent(InstanceID, Phase, EventType, StateHandle); \
 
 #define TRACE_STATETREE_LOG_EVENT(InstanceId, Format, ...) \
 	if (UE_TRACE_CHANNELEXPR_IS_ENABLED(StateTreeDebugChannel)) \
@@ -75,7 +76,7 @@ namespace UE::StateTreeTrace
 #else //STATETREE_DEBUG_TRACE_ENABLED
 
 #define TRACE_STATETREE_INSTANCE_EVENT(InstanceID, StateTree, InstanceName, EventType)
-#define TRACE_STATETREE_PHASE_EVENT(InstanceID, Phase, EventType)
+#define TRACE_STATETREE_PHASE_EVENT(InstanceID, Phase, EventType, StateHandle)
 #define TRACE_STATETREE_LOG_EVENT(InstanceId, Format, ...)
 #define TRACE_STATETREE_STATE_EVENT(InstanceId, StateHandle, EventType, SelectionBehavior)
 #define TRACE_STATETREE_TASK_EVENT(InstanceId, TaskIdx, DataView, EventType, Status)
diff --git a/Engine/Plugins/Runtime/StateTree/Source/StateTreeModule/Public/Debugger/StateTreeTraceTypes.h b/Engine/Plugins/Runtime/StateTree/Source/StateTreeModule/Public/Debugger/StateTreeTraceTypes.h
--- a/Engine/Plugins/Runtime/StateTree/Source/StateTreeModule/Public/Debugger/StateTreeTraceTypes.h
+++ b/Engine/Plugins/Runtime/StateTree/Source/StateTreeModule/Public/Debugger/StateTreeTraceTypes.h
@@ -17,19 +17,21 @@ UENUM()
 enum class EStateTreeTraceEventType : uint8
 {
 	Unset,
-	OnEnter				UMETA(DisplayName = "Enter"),		// for State + Task events
-	OnExit				UMETA(DisplayName = "Exit"),		// for State + Task events
-	Push				UMETA(DisplayName = "Push"),		// for State + Phase + Instance events
-	Pop					UMETA(DisplayName = "Pop"),			// for State + Phase + Instance events
-	OnStateSelected		UMETA(DisplayName = "Selected"),	// for State events
-	OnStateCompleted	UMETA(DisplayName = "Completed"),	// for State events
-	OnTickingTask		UMETA(DisplayName = "Tick"),		// for Task events
-	OnTaskCompleted		UMETA(DisplayName = "Completed"),	// for Task events
-	OnTaskTicked		UMETA(DisplayName = "Ticked"),		// for Task events
-	Passed				UMETA(DisplayName = "Passed"),		// for Condition events
-	Failed				UMETA(DisplayName = "Failed"),		// for Condition events
-	OnEvaluating		UMETA(DisplayName = "Evaluating"),	// for Transition events
-	OnTransition		UMETA(DisplayName = "Transition")	// for Transition events
+	OnEntering			UMETA(DisplayName = "Entering"),
+	OnEntered			UMETA(DisplayName = "Entered"),
+	OnExiting			UMETA(DisplayName = "Exiting"),
+	OnExited			UMETA(DisplayName = "Exited"),
+	Push				UMETA(DisplayName = "Push"),
+	Pop					UMETA(DisplayName = "Pop"),
+	OnStateSelected		UMETA(DisplayName = "Selected"),
+	OnStateCompleted	UMETA(DisplayName = "Completed"),
+	OnTickingTask		UMETA(DisplayName = "Tick"),
+	OnTaskCompleted		UMETA(DisplayName = "Completed"),
+	OnTaskTicked		UMETA(DisplayName = "Ticked"),
+	Passed				UMETA(DisplayName = "Passed"),
+	Failed				UMETA(DisplayName = "Failed"),
+	OnEvaluating		UMETA(DisplayName = "Evaluating"),
+	OnTransition		UMETA(DisplayName = "Transition")
 };
 
 #if WITH_STATETREE_DEBUGGER
@@ -51,15 +53,19 @@ struct FStateTreeTraceBaseEvent
 
 struct FStateTreeTracePhaseEvent : FStateTreeTraceBaseEvent
 {
-	explicit FStateTreeTracePhaseEvent(const double RecordingWorldTime, const EStateTreeUpdatePhase Phase, const EStateTreeTraceEventType EventType)
+	explicit FStateTreeTracePhaseEvent(const double RecordingWorldTime, const EStateTreeUpdatePhase Phase, const EStateTreeTraceEventType EventType, const FStateTreeStateHandle StateHandle)
 		: FStateTreeTraceBaseEvent(RecordingWorldTime, EventType)
 		, Phase(Phase)
+		, StateHandle(StateHandle)
 	{
 	}
 
-	STATETREEMODULE_API FString ToString(const UStateTree& StateTree) const;
+	STATETREEMODULE_API FString ToFullString(const UStateTree& StateTree) const;
+	STATETREEMODULE_API FString GetValueString(const UStateTree& StateTree) const;
+	STATETREEMODULE_API FString GetTypeString(const UStateTree& StateTree) const;
 
 	EStateTreeUpdatePhase Phase;
+	FStateTreeStateHandle StateHandle;
 };
 
 struct FStateTreeTraceLogEvent : FStateTreeTraceBaseEvent
@@ -70,22 +76,38 @@ struct FStateTreeTraceLogEvent : FStateTreeTraceBaseEvent
 	{
 	}
 
-	STATETREEMODULE_API FString ToString(const UStateTree& StateTree) const;
+	STATETREEMODULE_API FString ToFullString(const UStateTree& StateTree) const;
+	STATETREEMODULE_API FString GetValueString(const UStateTree& StateTree) const;
+	STATETREEMODULE_API FString GetTypeString(const UStateTree& StateTree) const;
 
 	FString Message;
 };
 
+struct FStateTreeTracePropertyEvent : FStateTreeTraceLogEvent
+{
+	explicit FStateTreeTracePropertyEvent(const double RecordingWorldTime, const FString& Message)
+		: FStateTreeTraceLogEvent(RecordingWorldTime, Message)
+	{
+	}
+
+	STATETREEMODULE_API FString ToFullString(const UStateTree& StateTree) const;
+	STATETREEMODULE_API FString GetValueString(const UStateTree& StateTree) const;
+	STATETREEMODULE_API FString GetTypeString(const UStateTree& StateTree) const;
+};
+
 struct FStateTreeTraceTransitionEvent : FStateTreeTraceBaseEvent
 {
-	explicit FStateTreeTraceTransitionEvent(const double RecordingWorldTime, const FStateTreeIndex16 Index, const EStateTreeTraceEventType EventType)
+	explicit FStateTreeTraceTransitionEvent(const double RecordingWorldTime, const FStateTreeTransitionSource TransitionSource, const EStateTreeTraceEventType EventType)
 		: FStateTreeTraceBaseEvent(RecordingWorldTime, EventType)
-		, TransitionIndex(Index)
+		, TransitionSource(TransitionSource)
 	{
 	}
 
-	STATETREEMODULE_API FString ToString(const UStateTree& StateTree) const;
+	STATETREEMODULE_API FString ToFullString(const UStateTree& StateTree) const;
+	STATETREEMODULE_API FString GetValueString(const UStateTree& StateTree) const;
+	STATETREEMODULE_API FString GetTypeString(const UStateTree& StateTree) const;
 
-	FStateTreeIndex16 TransitionIndex;
+	FStateTreeTransitionSource TransitionSource;
 };
 
 struct FStateTreeTraceNodeEvent : FStateTreeTraceBaseEvent
@@ -96,7 +118,9 @@ struct FStateTreeTraceNodeEvent : FStateTreeTraceBaseEvent
 	{
 	}
 
-	STATETREEMODULE_API FString ToString(const UStateTree& StateTree) const;
+	STATETREEMODULE_API FString ToFullString(const UStateTree& StateTree) const;
+	STATETREEMODULE_API FString GetValueString(const UStateTree& StateTree) const;
+	STATETREEMODULE_API FString GetTypeString(const UStateTree& StateTree) const;
 	
 	FStateTreeIndex16 Index;
 };
@@ -109,7 +133,9 @@ struct FStateTreeTraceStateEvent : FStateTreeTraceNodeEvent
 	{
 	}
 
-	STATETREEMODULE_API FString ToString(const UStateTree& StateTree) const;
+	STATETREEMODULE_API FString ToFullString(const UStateTree& StateTree) const;
+	STATETREEMODULE_API FString GetValueString(const UStateTree& StateTree) const;
+	STATETREEMODULE_API FString GetTypeString(const UStateTree& StateTree) const;
 	STATETREEMODULE_API FStateTreeStateHandle GetStateHandle() const;
 
 	EStateTreeStateSelectionBehavior SelectionBehavior;
@@ -117,7 +143,7 @@ struct FStateTreeTraceStateEvent : FStateTreeTraceNodeEvent
 
 struct FStateTreeTraceTaskEvent : FStateTreeTraceNodeEvent
 {
-	explicit FStateTreeTraceTaskEvent(const double RecordingWorldTime, const FStateTreeIndex16 Index,	const EStateTreeTraceEventType EventType, const EStateTreeRunStatus Status,	const FString& TypePath, const FString& InstanceDataAsText)
+	explicit FStateTreeTraceTaskEvent(const double RecordingWorldTime, const FStateTreeIndex16 Index, const EStateTreeTraceEventType EventType, const EStateTreeRunStatus Status,	const FString& TypePath, const FString& InstanceDataAsText)
 		: FStateTreeTraceNodeEvent(RecordingWorldTime, Index, EventType)
 		, TypePath(TypePath)
 		, InstanceDataAsText(InstanceDataAsText)
@@ -125,7 +151,9 @@ struct FStateTreeTraceTaskEvent : FStateTreeTraceNodeEvent
 	{
 	}
 
-	STATETREEMODULE_API FString ToString(const UStateTree& StateTree) const;
+	STATETREEMODULE_API FString ToFullString(const UStateTree& StateTree) const;
+	STATETREEMODULE_API FString GetValueString(const UStateTree& StateTree) const;
+	STATETREEMODULE_API FString GetTypeString(const UStateTree& StateTree) const;
 
 	FString GetDataTypePath() const { return TypePath; }
 	FString GetDataAsText() const { return InstanceDataAsText; }
@@ -144,7 +172,9 @@ struct FStateTreeTraceConditionEvent : FStateTreeTraceNodeEvent
 	{
 	}
 
-	STATETREEMODULE_API FString ToString(const UStateTree& StateTree) const;
+	STATETREEMODULE_API FString ToFullString(const UStateTree& StateTree) const;
+	STATETREEMODULE_API FString GetValueString(const UStateTree& StateTree) const;
+	STATETREEMODULE_API FString GetTypeString(const UStateTree& StateTree) const;
 
 	FString GetDataTypePath() const { return TypePath; }
 	FString GetDataAsText() const { return InstanceDataAsText; }
@@ -158,7 +188,9 @@ struct FStateTreeTraceActiveStatesEvent : FStateTreeTraceBaseEvent
 	// Intentionally implemented in source file to compile 'TArray<FStateTreeStateHandle>' using only forward declaration.
 	explicit FStateTreeTraceActiveStatesEvent(const double RecordingWorldTime);
 
-	STATETREEMODULE_API FString ToString(const UStateTree& StateTree) const;
+	STATETREEMODULE_API FString ToFullString(const UStateTree& StateTree) const;
+	STATETREEMODULE_API FString GetValueString(const UStateTree& StateTree) const;
+	STATETREEMODULE_API FString GetTypeString(const UStateTree& StateTree) const;
 
 	TArray<FStateTreeStateHandle> ActiveStates;
 };
@@ -166,6 +198,7 @@ struct FStateTreeTraceActiveStatesEvent : FStateTreeTraceBaseEvent
 /** Type aliases for statetree trace events */
 using FStateTreeTraceEventVariantType = TVariant<FStateTreeTracePhaseEvent,
 												FStateTreeTraceLogEvent,
+												FStateTreeTracePropertyEvent,
 												FStateTreeTraceNodeEvent,
 												FStateTreeTraceStateEvent,
 												FStateTreeTraceTaskEvent,
diff --git a/Engine/Plugins/Runtime/StateTree/Source/StateTreeModule/Public/StateTreeExecutionContext.h b/Engine/Plugins/Runtime/StateTree/Source/StateTreeModule/Public/StateTreeExecutionContext.h
--- a/Engine/Plugins/Runtime/StateTree/Source/StateTreeModule/Public/StateTreeExecutionContext.h
+++ b/Engine/Plugins/Runtime/StateTree/Source/StateTreeModule/Public/StateTreeExecutionContext.h
@@ -442,8 +442,8 @@ protected:
 	/** Next transition, used by RequestTransition(). */
 	FStateTreeTransitionResult NextTransition;
 
-	/** Index of the state transition that caused the state change. Will be invalid for framework transitions and external requests. */
-	FStateTreeIndex16 NextTransitionIndex = FStateTreeIndex16::Invalid;
+	/** Structure describing the origin of the state transition that caused the state change. */
+	FStateTreeTransitionSource NextTransitionSource;
 
 	/** Current state we're processing, or invalid if not applicable. */
 	FStateTreeStateHandle CurrentlyProcessedState;
diff --git a/Engine/Plugins/Runtime/StateTree/Source/StateTreeModule/Public/StateTreeExecutionTypes.h b/Engine/Plugins/Runtime/StateTree/Source/StateTreeModule/Public/StateTreeExecutionTypes.h
--- a/Engine/Plugins/Runtime/StateTree/Source/StateTreeModule/Public/StateTreeExecutionTypes.h
+++ b/Engine/Plugins/Runtime/StateTree/Source/StateTreeModule/Public/StateTreeExecutionTypes.h
@@ -25,6 +25,8 @@ enum class EStateTreeUpdatePhase : uint8
 	TransitionConditions	UMETA(DisplayName = "Transition conditions"),
 	StateSelection			UMETA(DisplayName = "Try Enter"),
 	EnterConditions			UMETA(DisplayName = "Enter conditions"),
+	EnterStates				UMETA(DisplayName = "Enter States"),
+	ExitStates				UMETA(DisplayName = "Exit States"),
 };
 
 
@@ -367,6 +369,61 @@ struct STATETREEMODULE_API FStateTreeActiveStates
 };
 
 
+UENUM()
+enum class EStateTreeTransitionSourceType : uint8
+{
+	Unset,
+	Asset,
+	ExternalRequest,
+	Internal
+};
+
+/**
+ * Describes the origin of an applied transition.
+ */
+USTRUCT()
+struct STATETREEMODULE_API FStateTreeTransitionSource
+{
+	GENERATED_BODY()
+
+	FStateTreeTransitionSource() = default;
+
+	explicit FStateTreeTransitionSource(const EStateTreeTransitionSourceType SourceType, const FStateTreeIndex16 TransitionIndex, const FStateTreeStateHandle TargetState, const EStateTreeTransitionPriority Priority)
+	: SourceType(SourceType)
+	, TransitionIndex(TransitionIndex)
+	, TargetState(TargetState)
+	, Priority(Priority)
+	{
+	}
+
+	explicit FStateTreeTransitionSource(const FStateTreeIndex16 TransitionIndex, const FStateTreeStateHandle TargetState, const EStateTreeTransitionPriority Priority)
+	: FStateTreeTransitionSource(EStateTreeTransitionSourceType::Asset, TransitionIndex, TargetState, Priority)
+	{
+	}
+
+	explicit FStateTreeTransitionSource(const EStateTreeTransitionSourceType SourceType, const FStateTreeStateHandle TargetState, const EStateTreeTransitionPriority Priority)
+	: FStateTreeTransitionSource(SourceType, FStateTreeIndex16::Invalid, TargetState, Priority)
+	{
+	}
+
+	void Reset()
+	{
+		*this = {};
+	}
+
+	EStateTreeTransitionSourceType SourceType = EStateTreeTransitionSourceType::Unset;
+
+	/* Index of the transition if from predefined asset transitions, invalid otherwise */
+	FStateTreeIndex16 TransitionIndex;
+
+	/** Transition target state */
+	FStateTreeStateHandle TargetState = FStateTreeStateHandle::Invalid;
+	
+	/** Priority of the transition that caused the state change. */
+	EStateTreeTransitionPriority Priority = EStateTreeTransitionPriority::None;
+};
+
+
 /**
  * Describes a state tree transition. Source is the state where the transition started, Target describes the state where the transition pointed at,
  * and Next describes the selected state. The reason Transition and Next are different is that Transition state can be a selector state,
