Fixed minor formatting issue in StateTreeDebugggerView

diff --git a/Engine/Plugins/Runtime/StateTree/Source/StateTreeEditorModule/Private/Debugger/SStateTreeDebuggerView.cpp b/Engine/Plugins/Runtime/StateTree/Source/StateTreeEditorModule/Private/Debugger/SStateTreeDebuggerView.cpp
--- a/Engine/Plugins/Runtime/StateTree/Source/StateTreeEditorModule/Private/Debugger/SStateTreeDebuggerView.cpp
+++ b/Engine/Plugins/Runtime/StateTree/Source/StateTreeEditorModule/Private/Debugger/SStateTreeDebuggerView.cpp
@@ -1055,8 +1055,16 @@ void SStateTreeDebuggerView::OnDebuggerScrubStateChanged(const UE::StateTreeDebu
 				if (PhaseEvent->StateHandle.IsValid())
 				{
 					const FCompactStateTreeState* CompactState = StateTree->GetStateFromHandle(PhaseEvent->StateHandle);
-					CustomDescription += FString::Printf(TEXT("%s"),
-							CompactState != nullptr ? *CompactState->Name.ToString() : *PhaseEvent->StateHandle.Describe());
+					if (CustomDescription.IsEmpty())
+					{
+						CustomDescription += FString::Printf(TEXT("%s"),
+								CompactState != nullptr ? *CompactState->Name.ToString() : *PhaseEvent->StateHandle.Describe());	
+					}
+					else
+					{
+						CustomDescription += FString::Printf(TEXT(" '%s'"),
+								CompactState != nullptr ? *CompactState->Name.ToString() : *PhaseEvent->StateHandle.Describe());
+					}
 				}
 				bShouldAddToScopeStack = true;
 			}
