[StateTreeDebugger] add traces for Evaluators
#rnx

diff --git a/Engine/Plugins/Runtime/StateTree/Source/StateTreeEditorModule/Private/Debugger/SStateTreeDebuggerViewRow.cpp b/Engine/Plugins/Runtime/StateTree/Source/StateTreeEditorModule/Private/Debugger/SStateTreeDebuggerViewRow.cpp
--- a/Engine/Plugins/Runtime/StateTree/Source/StateTreeEditorModule/Private/Debugger/SStateTreeDebuggerViewRow.cpp
+++ b/Engine/Plugins/Runtime/StateTree/Source/StateTreeEditorModule/Private/Debugger/SStateTreeDebuggerViewRow.cpp
@@ -102,7 +102,7 @@ TSharedPtr<SWidget> SStateTreeDebuggerViewRow::GenerateEventWidget() const
 		{
 		case EStateTreeTraceEventType::OnEntered:
 		case EStateTreeTraceEventType::OnTaskCompleted:
-		case EStateTreeTraceEventType::OnTaskTicked:
+		case EStateTreeTraceEventType::OnTicked:
 			switch (TaskEvent->Status)
 			{
 			case EStateTreeRunStatus::Failed:		Image = StyleSet.GetBrush("StateTreeEditor.Debugger.Task.Failed");		break;
@@ -168,6 +168,7 @@ FText SStateTreeDebuggerViewRow::GetEventDescription() const
 			// Some types have some custom representations so we want to use a more minimal description.
 			if (Item->Event.IsType<FStateTreeTraceStateEvent>()
 				|| Item->Event.IsType<FStateTreeTraceTaskEvent>()
+				|| Item->Event.IsType<FStateTreeTraceEvaluatorEvent>()
 				|| Item->Event.IsType<FStateTreeTracePropertyEvent>()
 				|| Item->Event.IsType<FStateTreeTraceConditionEvent>()
 				|| Item->Event.IsType<FStateTreeTraceLogEvent>())
diff --git a/Engine/Plugins/Runtime/StateTree/Source/StateTreeModule/Private/Debugger/StateTreeTrace.cpp b/Engine/Plugins/Runtime/StateTree/Source/StateTreeModule/Private/Debugger/StateTreeTrace.cpp
--- a/Engine/Plugins/Runtime/StateTree/Source/StateTreeModule/Private/Debugger/StateTreeTrace.cpp
+++ b/Engine/Plugins/Runtime/StateTree/Source/StateTreeModule/Private/Debugger/StateTreeTrace.cpp
@@ -69,6 +69,15 @@ UE_TRACE_EVENT_BEGIN(StateTreeDebugger, TaskEvent)
 	UE_TRACE_EVENT_FIELD(uint8, Status)
 UE_TRACE_EVENT_END()
 
+UE_TRACE_EVENT_BEGIN(StateTreeDebugger, EvaluatorEvent)
+	UE_TRACE_EVENT_FIELD(uint64, Cycle)
+	UE_TRACE_EVENT_FIELD(uint32, InstanceId)
+	UE_TRACE_EVENT_FIELD(uint32, InstanceSerial)
+	UE_TRACE_EVENT_FIELD(uint16, NodeIndex)
+	UE_TRACE_EVENT_FIELD(uint8[], DataView)
+	UE_TRACE_EVENT_FIELD(uint8, EventType)
+UE_TRACE_EVENT_END()
+
 UE_TRACE_EVENT_BEGIN(StateTreeDebugger, TransitionEvent)
 	UE_TRACE_EVENT_FIELD(uint64, Cycle)
 	UE_TRACE_EVENT_FIELD(uint32, InstanceId)
@@ -508,6 +517,26 @@ void OutputTaskEventTrace(
 		<< TaskEvent.Status(static_cast<std::underlying_type_t<EStateTreeRunStatus>>(Status));
 }
 
+void OutputEvaluatorEventTrace(
+	const FStateTreeInstanceDebugId InstanceId,
+	const FStateTreeIndex16 EvaluatorIdx,
+	const FStateTreeDataView DataView,
+	const EStateTreeTraceEventType EventType)
+{
+	FBufferArchive Archive;
+	SerializeDataViewToArchive(Archive, DataView);
+
+	TraceBufferedEvents(InstanceId);
+
+	UE_TRACE_LOG(StateTreeDebugger, EvaluatorEvent, StateTreeDebugChannel)
+		<< EvaluatorEvent.Cycle(FPlatformTime::Cycles64())
+		<< EvaluatorEvent.InstanceId(InstanceId.Id)
+		<< EvaluatorEvent.InstanceSerial(InstanceId.SerialNumber)
+		<< EvaluatorEvent.NodeIndex(EvaluatorIdx.Get())
+		<< EvaluatorEvent.DataView(Archive.GetData(), Archive.Num())
+		<< EvaluatorEvent.EventType(static_cast<std::underlying_type_t<EStateTreeTraceEventType>>(EventType));
+}
+
 void OutputTransitionEventTrace(
 	const FStateTreeInstanceDebugId InstanceId,
 	const FStateTreeTransitionSource Source,
diff --git a/Engine/Plugins/Runtime/StateTree/Source/StateTreeModule/Private/Debugger/StateTreeTraceAnalyzer.cpp b/Engine/Plugins/Runtime/StateTree/Source/StateTreeModule/Private/Debugger/StateTreeTraceAnalyzer.cpp
--- a/Engine/Plugins/Runtime/StateTree/Source/StateTreeModule/Private/Debugger/StateTreeTraceAnalyzer.cpp
+++ b/Engine/Plugins/Runtime/StateTree/Source/StateTreeModule/Private/Debugger/StateTreeTraceAnalyzer.cpp
@@ -25,6 +25,7 @@ void FStateTreeTraceAnalyzer::OnAnalysisBegin(const FOnAnalysisContext& Context)
 	Builder.RouteEvent(RouteId_LogMessage, "StateTreeDebugger", "LogEvent");
 	Builder.RouteEvent(RouteId_State, "StateTreeDebugger", "StateEvent");
 	Builder.RouteEvent(RouteId_Task, "StateTreeDebugger", "TaskEvent");
+	Builder.RouteEvent(RouteId_Evaluator, "StateTreeDebugger", "EvaluatorEvent");
 	Builder.RouteEvent(RouteId_Transition, "StateTreeDebugger", "TransitionEvent");
 	Builder.RouteEvent(RouteId_Condition, "StateTreeDebugger", "ConditionEvent");
 	Builder.RouteEvent(RouteId_ActiveStates, "StateTreeDebugger", "ActiveStatesEvent");
@@ -139,6 +140,23 @@ bool FStateTreeTraceAnalyzer::OnEvent(const uint16 RouteId, EStyle Style, const
 				FStateTreeTraceEventVariantType(TInPlaceType<FStateTreeTraceTaskEvent>(), Event));
 			break;
 		}
+	case RouteId_Evaluator:
+		{
+			FString TypePath, DataAsText;
+			FMemoryReaderView Archive(EventData.GetArrayView<uint8>("DataView"));
+			Archive << TypePath;
+			Archive << DataAsText;
+
+			const FStateTreeTraceEvaluatorEvent Event(WorldTime,
+				FStateTreeIndex16(EventData.GetValue<uint16>("NodeIndex")),
+				EventData.GetValue<EStateTreeTraceEventType>("EventType"),
+				TypePath, DataAsText);
+
+			Provider.AppendEvent(FStateTreeInstanceDebugId(EventData.GetValue<uint32>("InstanceId"), EventData.GetValue<uint32>("InstanceSerial")),
+				Context.EventTime.AsSeconds(EventData.GetValue<uint64>("Cycle")),
+				FStateTreeTraceEventVariantType(TInPlaceType<FStateTreeTraceEvaluatorEvent>(), Event));
+			break;
+		}
 	case RouteId_Condition:
 		{
 			FString TypePath, DataAsText;
diff --git a/Engine/Plugins/Runtime/StateTree/Source/StateTreeModule/Private/Debugger/StateTreeTraceTypes.cpp b/Engine/Plugins/Runtime/StateTree/Source/StateTreeModule/Private/Debugger/StateTreeTraceTypes.cpp
--- a/Engine/Plugins/Runtime/StateTree/Source/StateTreeModule/Private/Debugger/StateTreeTraceTypes.cpp
+++ b/Engine/Plugins/Runtime/StateTree/Source/StateTreeModule/Private/Debugger/StateTreeTraceTypes.cpp
@@ -234,9 +234,29 @@ FString FStateTreeTraceTaskEvent::GetValueString(const UStateTree& StateTree) co
 }
 
 FString FStateTreeTraceTaskEvent::GetTypeString(const UStateTree& StateTree) const
+{
+	return FStateTreeTraceNodeEvent::GetTypeString(StateTree);
+}
+
+
+//----------------------------------------------------------------------//
+// FStateTreeTraceEvaluatorEvent
+//----------------------------------------------------------------------//
+FString FStateTreeTraceEvaluatorEvent::ToFullString(const UStateTree& StateTree) const
+{
+	return FStateTreeTraceNodeEvent::ToFullString(StateTree);
+}
+
+FString FStateTreeTraceEvaluatorEvent::GetValueString(const UStateTree& StateTree) const
 {
 	const FConstStructView NodeView = Index.IsValid() ? StateTree.GetNode(Index.Get()) : FConstStructView();
-	return FString::Printf(TEXT("%s"), NodeView.IsValid() ? *NodeView.GetScriptStruct()->GetName() : TEXT("Invalid Node"));
+	const FStateTreeNodeBase* Node = NodeView.GetPtr<const FStateTreeNodeBase>();
+	return FString::Printf(TEXT("%s"), Node != nullptr ? *Node->Name.ToString() : *LexToString(Index.Get()));
+}
+
+FString FStateTreeTraceEvaluatorEvent::GetTypeString(const UStateTree& StateTree) const
+{
+	return FStateTreeTraceNodeEvent::GetTypeString(StateTree);
 }
 
 
diff --git a/Engine/Plugins/Runtime/StateTree/Source/StateTreeModule/Private/StateTreeExecutionContext.cpp b/Engine/Plugins/Runtime/StateTree/Source/StateTreeModule/Private/StateTreeExecutionContext.cpp
--- a/Engine/Plugins/Runtime/StateTree/Source/StateTreeModule/Private/StateTreeExecutionContext.cpp
+++ b/Engine/Plugins/Runtime/StateTree/Source/StateTreeModule/Private/StateTreeExecutionContext.cpp
@@ -38,6 +38,7 @@
 	#define STATETREE_TRACE_LOG_EVENT(Format, ...)							TRACE_STATETREE_LOG_EVENT(GetInstanceDebugId(), Format, ##__VA_ARGS__)
 	#define STATETREE_TRACE_STATE_EVENT(StateHandle, EventType)				TRACE_STATETREE_STATE_EVENT(GetInstanceDebugId(), StateHandle, EventType, EStateTreeStateSelectionBehavior::None);
 	#define STATETREE_TRACE_TASK_EVENT(Index, DataView, EventType, Status)	TRACE_STATETREE_TASK_EVENT(GetInstanceDebugId(), FStateTreeIndex16(Index), DataView, EventType, Status);
+	#define STATETREE_TRACE_EVALUATOR_EVENT(Index, DataView, EventType)		TRACE_STATETREE_EVALUATOR_EVENT(GetInstanceDebugId(), FStateTreeIndex16(Index), DataView, EventType);
 	#define STATETREE_TRACE_CONDITION_EVENT(Index, DataViews, EventType)	TRACE_STATETREE_CONDITION_EVENT(GetInstanceDebugId(), FStateTreeIndex16(Index), DataView, EventType);	
 	#define STATETREE_TRACE_TRANSITION_EVENT(Source, EventType)				TRACE_STATETREE_TRANSITION_EVENT(GetInstanceDebugId(), Source, EventType);
 #else
@@ -49,6 +50,7 @@
 	#define STATETREE_TRACE_LOG_EVENT(Format, ...)
 	#define STATETREE_TRACE_STATE_EVENT(StateHandle, EventType)
 	#define STATETREE_TRACE_TASK_EVENT(Index, DataView, EventType, Status)
+	#define STATETREE_TRACE_EVALUATOR_EVENT(Index, DataView, EventType)
 	#define STATETREE_TRACE_CONDITION_EVENT(Index, DataView, EventType)
 	#define STATETREE_TRACE_TRANSITION_EVENT(Source, EventType)
 #endif // WITH_STATETREE_DEBUGGER
@@ -1067,6 +1069,8 @@ EStateTreeRunStatus FStateTreeExecutionContext::TickEvaluatorsAndGlobalTasks(con
 		{
 			QUICK_SCOPE_CYCLE_COUNTER(StateTree_Eval_Tick);
 			Eval.Tick(*this, DeltaTime);
+
+			STATETREE_TRACE_EVALUATOR_EVENT(EvalIndex, DataViews[Eval.DataViewIndex.Get()], EStateTreeTraceEventType::OnTicked);
 		}
 	}
 
@@ -1117,7 +1121,7 @@ EStateTreeRunStatus FStateTreeExecutionContext::TickEvaluatorsAndGlobalTasks(con
 			}
 
 			STATETREE_TRACE_TASK_EVENT(TaskIndex, TaskDataView,
-				TaskResult != EStateTreeRunStatus::Running ? EStateTreeTraceEventType::OnTaskCompleted : EStateTreeTraceEventType::OnTaskTicked,
+				TaskResult != EStateTreeRunStatus::Running ? EStateTreeTraceEventType::OnTaskCompleted : EStateTreeTraceEventType::OnTicked,
 				TaskResult);
 
 			// If a global task succeeds or fails, it will stop the whole tree.
@@ -1164,6 +1168,8 @@ EStateTreeRunStatus FStateTreeExecutionContext::StartEvaluatorsAndGlobalTasks(FS
 		{
 			QUICK_SCOPE_CYCLE_COUNTER(StateTree_Eval_TreeStart);
 			Eval.TreeStart(*this);
+
+			STATETREE_TRACE_EVALUATOR_EVENT(EvalIndex, DataViews[Eval.DataViewIndex.Get()], EStateTreeTraceEventType::OnTreeStarted);
 		}
 	}
 
@@ -1282,6 +1288,8 @@ void FStateTreeExecutionContext::StopEvaluatorsAndGlobalTasks(const EStateTreeRu
 		{
 			QUICK_SCOPE_CYCLE_COUNTER(StateTree_Eval_TreeStop);
 			Eval.TreeStop(*this);
+
+			STATETREE_TRACE_EVALUATOR_EVENT(EvalIndex, DataViews[Eval.DataViewIndex.Get()], EStateTreeTraceEventType::OnTreeStopped);
 		}
 	}
 }
@@ -1373,7 +1381,7 @@ EStateTreeRunStatus FStateTreeExecutionContext::TickTasks(const float DeltaTime)
 			}
 
 			STATETREE_TRACE_TASK_EVENT(TaskIndex, TaskDataView,
-				TaskResult != EStateTreeRunStatus::Running ? EStateTreeTraceEventType::OnTaskCompleted : EStateTreeTraceEventType::OnTaskTicked,
+				TaskResult != EStateTreeRunStatus::Running ? EStateTreeTraceEventType::OnTaskCompleted : EStateTreeTraceEventType::OnTicked,
 				TaskResult);
 			
 			// TODO: Add more control over which states can control the failed/succeeded result.
diff --git a/Engine/Plugins/Runtime/StateTree/Source/StateTreeModule/Public/Debugger/StateTreeTrace.h b/Engine/Plugins/Runtime/StateTree/Source/StateTreeModule/Public/Debugger/StateTreeTrace.h
--- a/Engine/Plugins/Runtime/StateTree/Source/StateTreeModule/Public/Debugger/StateTreeTrace.h
+++ b/Engine/Plugins/Runtime/StateTree/Source/StateTreeModule/Public/Debugger/StateTreeTrace.h
@@ -29,6 +29,7 @@ namespace UE::StateTreeTrace
 	void OutputLogEventTrace(FStateTreeInstanceDebugId InstanceId, const TCHAR* Fmt, ...);
 	void OutputStateEventTrace(FStateTreeInstanceDebugId InstanceId, FStateTreeStateHandle StateHandle, EStateTreeTraceEventType EventType, EStateTreeStateSelectionBehavior SelectionBehavior);
 	void OutputTaskEventTrace(FStateTreeInstanceDebugId InstanceId, FStateTreeIndex16 TaskIdx, FStateTreeDataView DataView, EStateTreeTraceEventType EventType, EStateTreeRunStatus Status);
+	void OutputEvaluatorEventTrace(FStateTreeInstanceDebugId InstanceId, FStateTreeIndex16 EvaluatorIdx, FStateTreeDataView DataView, EStateTreeTraceEventType EventType);
 	void OutputConditionEventTrace(FStateTreeInstanceDebugId InstanceId, FStateTreeIndex16 ConditionIdx, FStateTreeDataView DataView, EStateTreeTraceEventType EventType);
 	void OutputTransitionEventTrace(FStateTreeInstanceDebugId InstanceId, FStateTreeTransitionSource TransitionSource, EStateTreeTraceEventType EventType);
 	void OutputActiveStatesEventTrace(FStateTreeInstanceDebugId InstanceId, const FStateTreeActiveStates& ActiveStates);
@@ -58,6 +59,12 @@ namespace UE::StateTreeTrace
 		UE::StateTreeTrace::OutputTaskEventTrace(InstanceId, TaskIdx, DataView, EventType, Status); \
 	}
 
+#define TRACE_STATETREE_EVALUATOR_EVENT(InstanceId, EvaluatorIdx, DataView, EventType) \
+	if (UE_TRACE_CHANNELEXPR_IS_ENABLED(StateTreeDebugChannel)) \
+	{ \
+		UE::StateTreeTrace::OutputEvaluatorEventTrace(InstanceId, EvaluatorIdx, DataView, EventType); \
+	}
+
 #define TRACE_STATETREE_CONDITION_EVENT(InstanceId, ConditionIdx, DataView, EventType) \
 	if (UE_TRACE_CHANNELEXPR_IS_ENABLED(StateTreeDebugChannel)) \
 	{ \
@@ -80,6 +87,7 @@ namespace UE::StateTreeTrace
 #define TRACE_STATETREE_LOG_EVENT(InstanceId, Format, ...)
 #define TRACE_STATETREE_STATE_EVENT(InstanceId, StateHandle, EventType, SelectionBehavior)
 #define TRACE_STATETREE_TASK_EVENT(InstanceId, TaskIdx, DataView, EventType, Status)
+#define TRACE_STATETREE_EVALUATOR_EVENT(InstanceId, EvaluatorIdx, DataView, EventType)
 #define TRACE_STATETREE_CONDITION_EVENT(InstanceId, ConditionIdx, DataView, EventType)
 #define TRACE_STATETREE_TRANSITION_EVENT(InstanceId, TransitionIdx, EventType)
 #define TRACE_STATETREE_ACTIVE_STATES_EVENT(InstanceId, ActivateStates)
diff --git a/Engine/Plugins/Runtime/StateTree/Source/StateTreeModule/Public/Debugger/StateTreeTraceAnalyzer.h b/Engine/Plugins/Runtime/StateTree/Source/StateTreeModule/Public/Debugger/StateTreeTraceAnalyzer.h
--- a/Engine/Plugins/Runtime/StateTree/Source/StateTreeModule/Public/Debugger/StateTreeTraceAnalyzer.h
+++ b/Engine/Plugins/Runtime/StateTree/Source/StateTreeModule/Public/Debugger/StateTreeTraceAnalyzer.h
@@ -27,6 +27,7 @@ private:
 		RouteId_LogMessage,
 		RouteId_State,
 		RouteId_Task,
+		RouteId_Evaluator,
 		RouteId_Transition,
 		RouteId_Condition,
 		RouteId_ActiveStates
diff --git a/Engine/Plugins/Runtime/StateTree/Source/StateTreeModule/Public/Debugger/StateTreeTraceTypes.h b/Engine/Plugins/Runtime/StateTree/Source/StateTreeModule/Public/Debugger/StateTreeTraceTypes.h
--- a/Engine/Plugins/Runtime/StateTree/Source/StateTreeModule/Public/Debugger/StateTreeTraceTypes.h
+++ b/Engine/Plugins/Runtime/StateTree/Source/StateTreeModule/Public/Debugger/StateTreeTraceTypes.h
@@ -25,13 +25,15 @@ enum class EStateTreeTraceEventType : uint8
 	Pop					UMETA(DisplayName = "Pop"),
 	OnStateSelected		UMETA(DisplayName = "Selected"),
 	OnStateCompleted	UMETA(DisplayName = "Completed"),
-	OnTickingTask		UMETA(DisplayName = "Tick"),
+	OnTicking			UMETA(DisplayName = "Tick"),
 	OnTaskCompleted		UMETA(DisplayName = "Completed"),
-	OnTaskTicked		UMETA(DisplayName = "Ticked"),
+	OnTicked			UMETA(DisplayName = "Ticked"),
 	Passed				UMETA(DisplayName = "Passed"),
 	Failed				UMETA(DisplayName = "Failed"),
 	OnEvaluating		UMETA(DisplayName = "Evaluating"),
-	OnTransition		UMETA(DisplayName = "Transition")
+	OnTransition		UMETA(DisplayName = "Transition"),
+	OnTreeStarted		UMETA(DisplayName = "Tree Started"),
+	OnTreeStopped		UMETA(DisplayName = "Tree Stopped")
 };
 
 #if WITH_STATETREE_DEBUGGER
@@ -143,7 +145,7 @@ struct FStateTreeTraceStateEvent : FStateTreeTraceNodeEvent
 
 struct FStateTreeTraceTaskEvent : FStateTreeTraceNodeEvent
 {
-	explicit FStateTreeTraceTaskEvent(const double RecordingWorldTime, const FStateTreeIndex16 Index, const EStateTreeTraceEventType EventType, const EStateTreeRunStatus Status,	const FString& TypePath, const FString& InstanceDataAsText)
+	explicit FStateTreeTraceTaskEvent(const double RecordingWorldTime, const FStateTreeIndex16 Index, const EStateTreeTraceEventType EventType, const EStateTreeRunStatus Status, const FString& TypePath, const FString& InstanceDataAsText)
 		: FStateTreeTraceNodeEvent(RecordingWorldTime, Index, EventType)
 		, TypePath(TypePath)
 		, InstanceDataAsText(InstanceDataAsText)
@@ -163,6 +165,26 @@ struct FStateTreeTraceTaskEvent : FStateTreeTraceNodeEvent
 	EStateTreeRunStatus Status;
 };
 
+struct FStateTreeTraceEvaluatorEvent : FStateTreeTraceNodeEvent
+{
+	explicit FStateTreeTraceEvaluatorEvent(const double RecordingWorldTime, const FStateTreeIndex16 Index, const EStateTreeTraceEventType EventType, const FString& TypePath, const FString& InstanceDataAsText)
+		: FStateTreeTraceNodeEvent(RecordingWorldTime, Index, EventType)
+		, TypePath(TypePath)
+		, InstanceDataAsText(InstanceDataAsText)
+	{
+	}
+
+	STATETREEMODULE_API FString ToFullString(const UStateTree& StateTree) const;
+	STATETREEMODULE_API FString GetValueString(const UStateTree& StateTree) const;
+	STATETREEMODULE_API FString GetTypeString(const UStateTree& StateTree) const;
+
+	FString GetDataTypePath() const { return TypePath; }
+	FString GetDataAsText() const { return InstanceDataAsText; }
+
+	FString TypePath;
+	FString InstanceDataAsText;
+};
+
 struct FStateTreeTraceConditionEvent : FStateTreeTraceNodeEvent
 {
 	explicit FStateTreeTraceConditionEvent(const double RecordingWorldTime, const FStateTreeIndex16 Index, const EStateTreeTraceEventType EventType, const FString& TypePath, const FString& InstanceDataAsText)
@@ -202,6 +224,7 @@ using FStateTreeTraceEventVariantType = TVariant<FStateTreeTracePhaseEvent,
 												FStateTreeTraceNodeEvent,
 												FStateTreeTraceStateEvent,
 												FStateTreeTraceTaskEvent,
+												FStateTreeTraceEvaluatorEvent,
 												FStateTreeTraceTransitionEvent,
 												FStateTreeTraceConditionEvent,
 												FStateTreeTraceActiveStatesEvent>;
diff --git a/Engine/Plugins/Runtime/StateTree/Source/StateTreeModule/Public/StateTreeExecutionTypes.h b/Engine/Plugins/Runtime/StateTree/Source/StateTreeModule/Public/StateTreeExecutionTypes.h
--- a/Engine/Plugins/Runtime/StateTree/Source/StateTreeModule/Public/StateTreeExecutionTypes.h
+++ b/Engine/Plugins/Runtime/StateTree/Source/StateTreeModule/Public/StateTreeExecutionTypes.h
@@ -15,12 +15,12 @@ enum class EStateTreeUpdatePhase : uint8
 	Unset					= 0,
 	StartTree				UMETA(DisplayName = "Start Tree"),
 	StopTree				UMETA(DisplayName = "Stop Tree"),
-	StartGlobalTasks		UMETA(DisplayName = "Start Global Tasks"),
-	StopGlobalTasks			UMETA(DisplayName = "Stop Global Tasks"),
+	StartGlobalTasks		UMETA(DisplayName = "Start Global Tasks & Evaluators"),
+	StopGlobalTasks			UMETA(DisplayName = "Stop Global Tasks & Evaluators"),
 	TickStateTree			UMETA(DisplayName = "Tick State Tree"),
 	ApplyTransitions		UMETA(DisplayName = "Transition"),
 	TriggerTransitions		UMETA(DisplayName = "Trigger Transitions"),
-	TickingGlobalTasks		UMETA(DisplayName = "Tick Global Tasks"),
+	TickingGlobalTasks		UMETA(DisplayName = "Tick Global Tasks & Evaluators"),
 	TickingTasks			UMETA(DisplayName = "Tick Tasks"),
 	TransitionConditions	UMETA(DisplayName = "Transition conditions"),
 	StateSelection			UMETA(DisplayName = "Try Enter"),
