[StateTree]
Updated compiler to ignore disabled states when getting next sibling states.
Fixed transition to enabled states inside a disabled parent state.
#rnx

diff --git a/Engine/Plugins/Runtime/StateTree/Source/StateTreeEditorModule/Private/StateTreeState.cpp b/Engine/Plugins/Runtime/StateTree/Source/StateTreeEditorModule/Private/StateTreeState.cpp
--- a/Engine/Plugins/Runtime/StateTree/Source/StateTreeEditorModule/Private/StateTreeState.cpp
+++ b/Engine/Plugins/Runtime/StateTree/Source/StateTreeEditorModule/Private/StateTreeState.cpp
@@ -347,7 +347,9 @@ const UStateTreeState* UStateTreeState::GetNextSiblingState() const
 		if (Parent->Children[ChildIdx] == this)
 		{
 			const int NextIdx = ChildIdx + 1;
-			if (NextIdx < Parent->Children.Num())
+
+			// Select the next enabled sibling
+			if (NextIdx < Parent->Children.Num() && Parent->Children[NextIdx]->bEnabled)
 			{
 				return Parent->Children[NextIdx];
 			}
@@ -372,9 +374,12 @@ const UStateTreeState* UStateTreeState::GetNextSelectableSiblingState() const
 	
 	for (int32 ChildIdx = StartChildIndex + 1; ChildIdx < Parent->Children.Num(); ChildIdx++)
 	{
-		if (Parent->Children[ChildIdx]->SelectionBehavior != EStateTreeStateSelectionBehavior::None)
+		// Select the next enabled and selectable sibling
+		const UStateTreeState* State =Parent->Children[ChildIdx];
+		if (State->SelectionBehavior != EStateTreeStateSelectionBehavior::None
+			&& State->bEnabled)
 		{
-			return Parent->Children[ChildIdx];
+			return State;
 		}
 	}
 	
diff --git a/Engine/Plugins/Runtime/StateTree/Source/StateTreeModule/Private/StateTreeExecutionContext.cpp b/Engine/Plugins/Runtime/StateTree/Source/StateTreeModule/Private/StateTreeExecutionContext.cpp
--- a/Engine/Plugins/Runtime/StateTree/Source/StateTreeModule/Private/StateTreeExecutionContext.cpp
+++ b/Engine/Plugins/Runtime/StateTree/Source/StateTreeModule/Private/StateTreeExecutionContext.cpp
@@ -765,9 +765,8 @@ EStateTreeRunStatus FStateTreeExecutionContext::EnterState(const FStateTreeTrans
 		const FStateTreeStateHandle PreviousHandle = Transition.CurrentActiveStates.GetStateSafe(Index);
 		const FCompactStateTreeState& State = StateTree.States[CurrentHandle.Index];
 
-		ensureMsgf(State.bEnabled, TEXT("Should never enter disabled state. This indicates an unhandled case."));
-		
-		if (!Exec.ActiveStates.Push(CurrentHandle))
+		// Add only enabled States to the list of active States
+		if (State.bEnabled && !Exec.ActiveStates.Push(CurrentHandle))
 		{
 			STATETREE_LOG(Error, TEXT("%hs: Reached max execution depth when trying to enter state '%s'.  '%s' using StateTree '%s'."),
 				__FUNCTION__, *GetStateStatusString(Exec), *GetNameSafe(&Owner), *GetFullNameSafe(&StateTree));
@@ -786,7 +785,9 @@ EStateTreeRunStatus FStateTreeExecutionContext::EnterState(const FStateTreeTrans
 
 		bOnTargetBranch = bOnTargetBranch || CurrentHandle == Transition.TargetState;
 		const bool bWasActive = PreviousHandle == CurrentHandle;
-		const bool bIsEnteringState = !bWasActive || bOnTargetBranch;
+
+		// Do not enter a disabled State tasks but maintain property bindings
+		const bool bIsEnteringState = (!bWasActive || bOnTargetBranch) && State.bEnabled;
 
 		CurrentTransition.CurrentState = CurrentHandle;
 		CurrentTransition.ChangeType = bWasActive ? EStateTreeStateChangeType::Sustained : EStateTreeStateChangeType::Changed;
@@ -1662,7 +1663,11 @@ bool FStateTreeExecutionContext::TriggerTransitions()
 			const FStateTreeStateHandle StateHandle = Exec.ActiveStates[StateIndex];
 			const FCompactStateTreeState& State = StateTree.States[StateHandle.Index];
 
-			ensureMsgf(State.bEnabled, TEXT("Should never try to transition to a disabled state. This indicates an unhandled case."));
+			// Do not process any transitions from a disabled state
+			if (!State.bEnabled)
+			{
+				continue;
+			}
 
 			FCurrentlyProcessedStateScope StateScope(*this, StateHandle);
 			STATETREE_TRACE_SCOPED_STATE(StateHandle);
