[StateTree] Update compilation status when adding/removing/moving/disabling states

diff --git a/Engine/Plugins/Runtime/StateTree/Source/StateTreeEditorModule/Private/StateTreeEditor.cpp b/Engine/Plugins/Runtime/StateTree/Source/StateTreeEditorModule/Private/StateTreeEditor.cpp
--- a/Engine/Plugins/Runtime/StateTree/Source/StateTreeEditorModule/Private/StateTreeEditor.cpp
+++ b/Engine/Plugins/Runtime/StateTree/Source/StateTreeEditorModule/Private/StateTreeEditor.cpp
@@ -168,6 +168,9 @@ void FStateTreeEditor::InitEditor( const EToolkitMode::Type Mode, const TSharedP
 	StateTreeViewModel->Init(EditorData);
 
 	StateTreeViewModel->GetOnAssetChanged().AddSP(this, &FStateTreeEditor::HandleModelAssetChanged);
+	StateTreeViewModel->GetOnStateAdded().AddSPLambda(this, [this](UStateTreeState* , UStateTreeState*){ UpdateAsset(); });
+	StateTreeViewModel->GetOnStatesRemoved().AddSPLambda(this, [this](const TSet<UStateTreeState*>&){ UpdateAsset(); });
+	StateTreeViewModel->GetOnStatesMoved().AddSPLambda(this, [this](const TSet<UStateTreeState*>&, const TSet<UStateTreeState*>&){ UpdateAsset(); });
 	StateTreeViewModel->GetOnSelectionChanged().AddSP(this, &FStateTreeEditor::HandleModelSelectionChanged);
 
 	FMessageLogModule& MessageLogModule = FModuleManager::LoadModuleChecked<FMessageLogModule>("MessageLog");
diff --git a/Engine/Plugins/Runtime/StateTree/Source/StateTreeEditorModule/Private/StateTreeViewModel.cpp b/Engine/Plugins/Runtime/StateTree/Source/StateTreeEditorModule/Private/StateTreeViewModel.cpp
--- a/Engine/Plugins/Runtime/StateTree/Source/StateTreeEditorModule/Private/StateTreeViewModel.cpp
+++ b/Engine/Plugins/Runtime/StateTree/Source/StateTreeEditorModule/Private/StateTreeViewModel.cpp
@@ -905,6 +905,8 @@ void FStateTreeViewModel::SetSelectedStatesEnabled(const bool bEnable)
 			State->Modify();
 			State->bEnabled = bEnable;
 		}
+
+		OnAssetChanged.Broadcast();
 	}
 }
 
