[StateTreeDebugger] added trace status changed delegate that gets broadcasted by the StateTree module when traces are started or stopped. Debugger view can then keep recording button in sync.
#rnx

diff --git a/Engine/Plugins/Runtime/StateTree/Source/StateTreeEditorModule/Private/Debugger/SStateTreeDebuggerView.cpp b/Engine/Plugins/Runtime/StateTree/Source/StateTreeEditorModule/Private/Debugger/SStateTreeDebuggerView.cpp
--- a/Engine/Plugins/Runtime/StateTree/Source/StateTreeEditorModule/Private/Debugger/SStateTreeDebuggerView.cpp
+++ b/Engine/Plugins/Runtime/StateTree/Source/StateTreeEditorModule/Private/Debugger/SStateTreeDebuggerView.cpp
@@ -19,6 +19,7 @@
 #include "StateTreeCompiler.h"
 #include "StateTreeDebuggerCommands.h"
 #include "StateTreeDebuggerTrack.h"
+#include "StateTreeDelegates.h"
 #include "StateTreeEditorData.h"
 #include "StateTreeEditorSettings.h"
 #include "StateTreeEditorStyle.h"
@@ -204,6 +205,8 @@ SStateTreeDebuggerView::SStateTreeDebuggerView()
 
 SStateTreeDebuggerView::~SStateTreeDebuggerView()
 {
+	UE::StateTree::Delegates::OnTracingStateChanged.RemoveAll(this);
+
 	check(Debugger);
 	Debugger->OnScrubStateChanged.Unbind();
 	Debugger->OnBreakpointHit.Unbind();
@@ -341,6 +344,10 @@ void SStateTreeDebuggerView::Construct(const FArguments& InArgs, const UStateTre
 
 	IStateTreeModule& StateTreeModule = FModuleManager::GetModuleChecked<IStateTreeModule>("StateTreeModule");
 	bRecording = StateTreeModule.IsTracing();
+	UE::StateTree::Delegates::OnTracingStateChanged.AddSPLambda(this, [&bRecording=bRecording](const bool bTracesEnabled)
+		{
+			bRecording = bTracesEnabled;
+		});
 
 	// Bind callbacks to the debugger delegates
 	Debugger->OnNewSession.BindSP(this, &SStateTreeDebuggerView::OnNewSession);
@@ -642,7 +649,7 @@ void SStateTreeDebuggerView::Construct(const FArguments& InArgs, const UStateTre
 	// Do that after creating all our widgets in case we receive a callback
 	TArray<FStateTreeDebugger::FTraceDescriptor> TraceDescriptors;
 	Debugger->GetLiveTraces(TraceDescriptors);
-	if (TraceDescriptors.Num() == 1 && StateTreeModule.IsTracing())
+	if (TraceDescriptors.Num() == 1 && bRecording)
 	{
 		Debugger->RequestSessionAnalysis(TraceDescriptors.Last());
 	}
diff --git a/Engine/Plugins/Runtime/StateTree/Source/StateTreeModule/Private/StateTreeDelegates.cpp b/Engine/Plugins/Runtime/StateTree/Source/StateTreeModule/Private/StateTreeDelegates.cpp
--- a/Engine/Plugins/Runtime/StateTree/Source/StateTreeModule/Private/StateTreeDelegates.cpp
+++ b/Engine/Plugins/Runtime/StateTree/Source/StateTreeModule/Private/StateTreeDelegates.cpp
@@ -13,6 +13,10 @@ FOnStateParametersChanged OnStateParametersChanged;
 FOnBreakpointsChanged OnBreakpointsChanged;
 FOnPostCompile OnPostCompile;
 FOnRequestCompile OnRequestCompile;
-#endif
-	
+#endif // WITH_EDITOR
+
+#if WITH_STATETREE_DEBUGGER
+FOnTracingStateChanged OnTracingStateChanged;
+#endif // WITH_STATETREE_DEBUGGER
+
 }; // UE::StateTree::Delegates
diff --git a/Engine/Plugins/Runtime/StateTree/Source/StateTreeModule/Private/StateTreeModule.cpp b/Engine/Plugins/Runtime/StateTree/Source/StateTreeModule/Private/StateTreeModule.cpp
--- a/Engine/Plugins/Runtime/StateTree/Source/StateTreeModule/Private/StateTreeModule.cpp
+++ b/Engine/Plugins/Runtime/StateTree/Source/StateTreeModule/Private/StateTreeModule.cpp
@@ -1,6 +1,7 @@
 // Copyright Epic Games, Inc. All Rights Reserved.
 
 #include "StateTreeModule.h"
+
 #include "StateTreeTypes.h"
 
 #if WITH_STATETREE_DEBUGGER
@@ -10,6 +11,7 @@
 #include "HAL/IConsoleManager.h"
 #include "Misc/CoreDelegates.h"
 #include "ProfilingDebugging/TraceAuxiliary.h"
+#include "StateTreeDelegates.h"
 #include "StateTreeSettings.h"
 #include "Trace/StoreClient.h"
 #include "Trace/StoreService.h"
@@ -175,6 +177,12 @@ bool FStateTreeModule::StartTraces(int32& OutTraceId)
 	}
 
 	bIsTracing = true;
+	if (UE::StateTree::Delegates::OnTracingStateChanged.IsBound())
+	{
+		UE_LOG(LogStateTree, Log, TEXT("StateTree traces enabled"));
+		UE::StateTree::Delegates::OnTracingStateChanged.Broadcast(bIsTracing);
+	}
+
 	return bAreTracesStarted;
 #else
 	return false;
@@ -217,6 +225,12 @@ void FStateTreeModule::StopTraces()
 	}
 
 	bIsTracing = false;
+
+	if (UE::StateTree::Delegates::OnTracingStateChanged.IsBound())
+	{
+		UE_LOG(LogStateTree, Log, TEXT("StateTree traces disabled"));
+		UE::StateTree::Delegates::OnTracingStateChanged.Broadcast(bIsTracing);
+	}
 #endif // WITH_STATETREE_DEBUGGER
 }
 
diff --git a/Engine/Plugins/Runtime/StateTree/Source/StateTreeModule/Public/StateTreeDelegates.h b/Engine/Plugins/Runtime/StateTree/Source/StateTreeModule/Public/StateTreeDelegates.h
--- a/Engine/Plugins/Runtime/StateTree/Source/StateTreeModule/Public/StateTreeDelegates.h
+++ b/Engine/Plugins/Runtime/StateTree/Source/StateTreeModule/Public/StateTreeDelegates.h
@@ -8,50 +8,59 @@ class UStateTree;
 
 namespace UE::StateTree::Delegates
 {
-	
+
 #if WITH_EDITOR
-	/** Called when linkable name in a StateTree has changed. */
-	DECLARE_MULTICAST_DELEGATE_OneParam(FOnIdentifierChanged, const UStateTree& /*StateTree*/);
-	extern STATETREEMODULE_API FOnIdentifierChanged OnIdentifierChanged;
-
-	/**
-	 * Called when schema of the StateTree EditorData has changed.
-	 * This is used to refresh the asset editor.
-	 * Note that this is NOT called when updating the StateTree schema from the EditorData on successful compilation.
-	 */
-	DECLARE_MULTICAST_DELEGATE_OneParam(FOnSchemaChanged, const UStateTree& /*StateTree*/);
-	extern STATETREEMODULE_API FOnSchemaChanged OnSchemaChanged;
-
-	/**
-	 * Called when parameters of the StateTree EditorData changed.
-	 * This should mainly used by the asset editor to maintain consistency in the UI for manipulations on the EditorData
-	 * until the tree gets compiled.
-	 */
-	DECLARE_MULTICAST_DELEGATE_OneParam(FOnParametersChanged, const UStateTree& /*StateTree*/);
-	extern STATETREEMODULE_API FOnParametersChanged OnParametersChanged;
-
-	/**
-	 * Called when parameters of a StateTree State changed.
-	 * This should mainly used by the asset editor to maintain consistency in the UI for manipulations.
-	 */
-	DECLARE_MULTICAST_DELEGATE_TwoParams(FOnStateParametersChanged, const UStateTree& /*StateTree*/, const FGuid /*StateID*/);
-	extern STATETREEMODULE_API FOnStateParametersChanged OnStateParametersChanged;
-
-	/**
-	 * Called when breakpoints of the StateTree EditorData changed.
-	 * This should mainly used by the asset editor to update the debugger.
-	 */
-	DECLARE_MULTICAST_DELEGATE_OneParam(FOnBreakpointsChanged, const UStateTree& /*StateTree*/);
-	extern STATETREEMODULE_API FOnBreakpointsChanged OnBreakpointsChanged;
-
-	/** Called when compilation succeeds */
-	DECLARE_MULTICAST_DELEGATE_OneParam(FOnPostCompile, const UStateTree& /*StateTree*/);
-	extern STATETREEMODULE_API FOnPostCompile OnPostCompile;
-
-	/** Request StateTree compilation. Works only in editor. */
-	DECLARE_DELEGATE_RetVal_OneParam(bool, FOnRequestCompile, UStateTree& /*StateTreeToCompile*/);
-	extern STATETREEMODULE_API FOnRequestCompile OnRequestCompile;
-#endif
 
-}; // UE::StateTree::Delegates
+/** Called when linkable name in a StateTree has changed. */
+DECLARE_MULTICAST_DELEGATE_OneParam(FOnIdentifierChanged, const UStateTree& /*StateTree*/);
+extern STATETREEMODULE_API FOnIdentifierChanged OnIdentifierChanged;
+
+/**
+ * Called when schema of the StateTree EditorData has changed.
+ * This is used to refresh the asset editor.
+ * Note that this is NOT called when updating the StateTree schema from the EditorData on successful compilation.
+ */
+DECLARE_MULTICAST_DELEGATE_OneParam(FOnSchemaChanged, const UStateTree& /*StateTree*/);
+extern STATETREEMODULE_API FOnSchemaChanged OnSchemaChanged;
+
+/**
+ * Called when parameters of the StateTree EditorData changed.
+ * This should mainly used by the asset editor to maintain consistency in the UI for manipulations on the EditorData
+ * until the tree gets compiled.
+ */
+DECLARE_MULTICAST_DELEGATE_OneParam(FOnParametersChanged, const UStateTree& /*StateTree*/);
+extern STATETREEMODULE_API FOnParametersChanged OnParametersChanged;
+
+/**
+ * Called when parameters of a StateTree State changed.
+ * This should mainly used by the asset editor to maintain consistency in the UI for manipulations.
+ */
+DECLARE_MULTICAST_DELEGATE_TwoParams(FOnStateParametersChanged, const UStateTree& /*StateTree*/, const FGuid /*StateID*/);
+extern STATETREEMODULE_API FOnStateParametersChanged OnStateParametersChanged;
+
+/**
+ * Called when breakpoints of the StateTree EditorData changed.
+ * This should mainly used by the asset editor to update the debugger.
+ */
+DECLARE_MULTICAST_DELEGATE_OneParam(FOnBreakpointsChanged, const UStateTree& /*StateTree*/);
+extern STATETREEMODULE_API FOnBreakpointsChanged OnBreakpointsChanged;
+
+/** Called when compilation succeeds */
+DECLARE_MULTICAST_DELEGATE_OneParam(FOnPostCompile, const UStateTree& /*StateTree*/);
+extern STATETREEMODULE_API FOnPostCompile OnPostCompile;
 
+/** Request StateTree compilation. Works only in editor. */
+DECLARE_DELEGATE_RetVal_OneParam(bool, FOnRequestCompile, UStateTree& /*StateTreeToCompile*/);
+extern STATETREEMODULE_API FOnRequestCompile OnRequestCompile;
+
+#endif // WITH_EDITOR
+
+#if WITH_STATETREE_DEBUGGER
+
+/** Called by the StateTree module when StateTree traces are enabled/disabled. */
+DECLARE_MULTICAST_DELEGATE_OneParam(FOnTracingStateChanged, bool /*bTracesEnabled*/);
+extern STATETREEMODULE_API FOnTracingStateChanged OnTracingStateChanged;
+
+#endif // WITH_STATETREE_DEBUGGER
+
+}; // UE::StateTree::Delegates
