[StateTreeDebugger] allow debugger commands to intercept KeyEvent when debugger tab is the active widget.
Also allow user to use the scrubbing buttons at any time the same way he can manually scrub. It will then disable auto-scrolling.
#rnx

diff --git a/Engine/Plugins/Runtime/StateTree/Source/StateTreeEditorModule/Private/Debugger/SStateTreeDebuggerView.cpp b/Engine/Plugins/Runtime/StateTree/Source/StateTreeEditorModule/Private/Debugger/SStateTreeDebuggerView.cpp
--- a/Engine/Plugins/Runtime/StateTree/Source/StateTreeEditorModule/Private/Debugger/SStateTreeDebuggerView.cpp
+++ b/Engine/Plugins/Runtime/StateTree/Source/StateTreeEditorModule/Private/Debugger/SStateTreeDebuggerView.cpp
@@ -337,7 +337,8 @@ void SStateTreeDebuggerView::Construct(const FArguments& InArgs, const UStateTre
 {
 	StateTreeViewModel = InStateTreeViewModel;
 	StateTree = &InStateTree;
-	StateTreeEditorData = Cast<UStateTreeEditorData>(InStateTree.EditorData.Get()); 
+	StateTreeEditorData = Cast<UStateTreeEditorData>(InStateTree.EditorData.Get());
+	CommandList = InCommandList;
 
 	Debugger = InStateTreeViewModel->GetDebugger();
 	check(Debugger);
@@ -655,6 +656,28 @@ void SStateTreeDebuggerView::Construct(const FArguments& InArgs, const UStateTre
 	}
 }
 
+FReply SStateTreeDebuggerView::OnPreviewKeyDown(const FGeometry& MyGeometry, const FKeyEvent& InKeyEvent)
+{
+	// We consider the key as handled regardless if the action can be executed or not since we don't want
+	// some of them affecting other widgets once the action can no longer be executed
+	// (e.g. can no longer scrub once reaching beginning or end of the timeline)
+	// This is why we test the input manually instead of relying on ProcessCommandBindings.
+	const FInputChord InputChord(
+		InKeyEvent.GetKey(),
+		EModifierKey::FromBools(InKeyEvent.IsControlDown(), InKeyEvent.IsAltDown(), InKeyEvent.IsShiftDown(), InKeyEvent.IsCommandDown()));
+
+	const TSharedPtr<FUICommandInfo> CommandInfo = FInputBindingManager::Get().FindCommandInContext(
+		FStateTreeDebuggerCommands::Get().GetContextName(), InputChord, /*bCheckDefault*/false);
+
+	if (CommandInfo.IsValid())
+	{
+		CommandList->TryExecuteAction(CommandInfo.ToSharedRef());
+		return FReply::Handled();
+	}
+
+	return SCompoundWidget::OnPreviewKeyDown(MyGeometry, InKeyEvent);
+}
+
 void SStateTreeDebuggerView::Tick(const FGeometry& AllottedGeometry, const double InCurrentTime, const float InDeltaTime)
 {
 	QUICK_SCOPE_CYCLE_COUNTER(STAT_StateTreeDebuggerView_TickView);
@@ -741,8 +764,8 @@ void SStateTreeDebuggerView::BindDebuggerToolbarCommands(const TSharedRef<FUICom
 
 bool SStateTreeDebuggerView::CanUseScrubButtons() const
 {
-	check(Debugger);
-	return (bAutoScroll == false || Debugger->IsAnalysisSessionPaused() || !Debugger->IsAnalysisSessionActive());
+	// Nothing preventing use of scrub buttons on the Editor side at the moment.
+	return true;
 }
 
 bool SStateTreeDebuggerView::CanStepBackToPreviousStateWithEvents() const
@@ -755,6 +778,7 @@ void SStateTreeDebuggerView::StepBackToPreviousStateWithEvents()
 {
 	check(Debugger);
 	Debugger->StepBackToPreviousStateWithEvents();
+	bAutoScroll = false;
 }
 
 bool SStateTreeDebuggerView::CanStepForwardToNextStateWithEvents() const
@@ -767,6 +791,7 @@ void SStateTreeDebuggerView::StepForwardToNextStateWithEvents()
 {
 	check(Debugger);
 	Debugger->StepForwardToNextStateWithEvents();
+	bAutoScroll = false;
 }
 
 bool SStateTreeDebuggerView::CanStepBackToPreviousStateChange() const
@@ -779,6 +804,7 @@ void SStateTreeDebuggerView::StepBackToPreviousStateChange()
 {
 	check(Debugger);
 	Debugger->StepBackToPreviousStateChange();
+	bAutoScroll = false;
 }
 
 bool SStateTreeDebuggerView::CanStepForwardToNextStateChange() const
@@ -791,6 +817,7 @@ void SStateTreeDebuggerView::StepForwardToNextStateChange()
 {
 	check(Debugger);
 	Debugger->StepForwardToNextStateChange();
+	bAutoScroll = false;
 }
 
 bool SStateTreeDebuggerView::CanAddStateBreakpoint(const EStateTreeBreakpointType Type) const
diff --git a/Engine/Plugins/Runtime/StateTree/Source/StateTreeEditorModule/Private/Debugger/SStateTreeDebuggerView.h b/Engine/Plugins/Runtime/StateTree/Source/StateTreeEditorModule/Private/Debugger/SStateTreeDebuggerView.h
--- a/Engine/Plugins/Runtime/StateTree/Source/StateTreeEditorModule/Private/Debugger/SStateTreeDebuggerView.h
+++ b/Engine/Plugins/Runtime/StateTree/Source/StateTreeEditorModule/Private/Debugger/SStateTreeDebuggerView.h
@@ -49,6 +49,7 @@ public:
 	void Construct(const FArguments& InArgs, const UStateTree& InStateTree, const TSharedRef<FStateTreeViewModel>& InStateTreeViewModel, const TSharedRef<FUICommandList>& InCommandList);
 
 private:
+	virtual FReply OnPreviewKeyDown(const FGeometry& MyGeometry, const FKeyEvent& InKeyEvent) override;
 	virtual void Tick(const FGeometry& AllottedGeometry, const double InCurrentTime, const float InDeltaTime) override;
 
 	void RefreshTracks();
@@ -120,6 +121,7 @@ private:
 	TSharedPtr<FStateTreeViewModel> StateTreeViewModel;
 	TWeakObjectPtr<const UStateTree> StateTree;
 	TWeakObjectPtr<UStateTreeEditorData> StateTreeEditorData;
+	TSharedPtr<FUICommandList> CommandList;
 
 	/** Tracks for all statetree instance owners producing trace events for the associated state tree asset. */
 	TArray<TSharedPtr<RewindDebugger::FRewindDebuggerTrack>> InstanceOwnerTracks;
