Removed reference to private directory for StateTreeEditorModule
#rnx

diff --git a/Engine/Plugins/Runtime/StateTree/Source/StateTreeEditorModule/StateTreeEditorModule.Build.cs b/Engine/Plugins/Runtime/StateTree/Source/StateTreeEditorModule/StateTreeEditorModule.Build.cs
--- a/Engine/Plugins/Runtime/StateTree/Source/StateTreeEditorModule/StateTreeEditorModule.Build.cs
+++ b/Engine/Plugins/Runtime/StateTree/Source/StateTreeEditorModule/StateTreeEditorModule.Build.cs
@@ -55,10 +55,6 @@ namespace UnrealBuildTool.Rules
 				"MessageLog",
 			});
 
-			PrivateIncludePaths.AddRange(new string[] {
-				System.IO.Path.Combine(GetModuleDirectory("PropertyEditor"), "Private"),
-			});
-		
 			if (Target.Platform == UnrealTargetPlatform.Win64)
 			{
 				PublicDefinitions.Add("WITH_STATETREE_DEBUGGER=1");
