StateTree: Changed bindings to not derefernce soft objects during copy

diff --git a/Engine/Plugins/Runtime/StateTree/Source/StateTreeModule/Private/StateTreePropertyBindings.cpp b/Engine/Plugins/Runtime/StateTree/Source/StateTreeModule/Private/StateTreePropertyBindings.cpp
--- a/Engine/Plugins/Runtime/StateTree/Source/StateTreeModule/Private/StateTreePropertyBindings.cpp
+++ b/Engine/Plugins/Runtime/StateTree/Source/StateTreeModule/Private/StateTreePropertyBindings.cpp
@@ -424,7 +424,17 @@ bool FStateTreePropertyBindings::ResolveCopyType(const FStateTreePropertyPathInd
 		}
 		else if (CastField<FObjectPropertyBase>(TargetProperty))
 		{
-			OutCopy.Type = EStateTreePropertyCopyType::CopyObject;
+			if (SourceProperty->IsA<FSoftObjectProperty>()
+				&& TargetProperty->IsA<FSoftObjectProperty>())
+			{
+				// Use CopyComplex when copying soft object to another soft object so that we do not try to dereference the object (just copies the path).
+				// This handles soft class too.
+				OutCopy.Type = EStateTreePropertyCopyType::CopyComplex;
+			}
+			else
+			{
+				OutCopy.Type = EStateTreePropertyCopyType::CopyObject;
+			}
 			return true;
 		}
 		else if (CastField<FArrayProperty>(TargetProperty) && TargetProperty->HasAnyPropertyFlags(CPF_EditFixedSize))
diff --git a/Engine/Plugins/Runtime/StateTree/Source/StateTreeTestSuite/Private/StateTreeTest.cpp b/Engine/Plugins/Runtime/StateTree/Source/StateTreeTestSuite/Private/StateTreeTest.cpp
--- a/Engine/Plugins/Runtime/StateTree/Source/StateTreeTestSuite/Private/StateTreeTest.cpp
+++ b/Engine/Plugins/Runtime/StateTree/Source/StateTreeTestSuite/Private/StateTreeTest.cpp
@@ -1372,9 +1372,9 @@ struct FStateTreeTest_BindingsCompiler : FAITestBase
 
 		TArray<FStateTreeDataView> SourceViews;
 		SourceViews.SetNum(Bindings.GetSourceStructNum());
-		SourceViews[SourceAIndex] = FStateTreeDataView(TBaseStructure<FStateTreeTest_PropertyCopy>::Get(), (uint8*)&SourceA);
-		SourceViews[SourceBIndex] = FStateTreeDataView(TBaseStructure<FStateTreeTest_PropertyCopy>::Get(), (uint8*)&SourceB);
-		FStateTreeDataView TargetView(TBaseStructure<FStateTreeTest_PropertyCopy>::Get(), (uint8*)&Target);
+		SourceViews[SourceAIndex] = FStateTreeDataView(FStructView::Make(SourceA));
+		SourceViews[SourceBIndex] = FStateTreeDataView(FStructView::Make(SourceB));
+		FStateTreeDataView TargetView(FStructView::Make(Target));
 		
 		const bool bCopyResult = Bindings.CopyTo(SourceViews, FStateTreeIndex16(CopyBatchIndex), TargetView);
 		AITEST_TRUE("CopyTo should succeed", bCopyResult);
@@ -1395,6 +1395,132 @@ struct FStateTreeTest_BindingsCompiler : FAITestBase
 };
 IMPLEMENT_AI_INSTANT_TEST(FStateTreeTest_BindingsCompiler, "System.StateTree.BindingsCompiler");
 
+struct FStateTreeTest_CopyObjects : FAITestBase
+{
+	virtual bool InstantTest() override
+	{
+		FStateTreeCompilerLog Log;
+		FStateTreePropertyBindings Bindings;
+		FStateTreePropertyBindingCompiler BindingCompiler;
+
+		const bool bInitResult = BindingCompiler.Init(Bindings, Log);
+		AITEST_TRUE("Expect init to succeed", bInitResult);
+
+		FStateTreeBindableStructDesc SourceDesc;
+		SourceDesc.Name = FName(TEXT("Source"));
+		SourceDesc.Struct = TBaseStructure<FStateTreeTest_PropertyCopyObjects>::Get();
+		SourceDesc.DataSource = EStateTreeBindableStructSource::Parameter;
+		SourceDesc.ID = FGuid::NewGuid();
+
+		FStateTreeBindableStructDesc TargetADesc;
+		TargetADesc.Name = FName(TEXT("TargetA"));
+		TargetADesc.Struct = TBaseStructure<FStateTreeTest_PropertyCopyObjects>::Get();
+		TargetADesc.DataSource = EStateTreeBindableStructSource::Parameter;
+		TargetADesc.ID = FGuid::NewGuid();
+
+		FStateTreeBindableStructDesc TargetBDesc;
+		TargetBDesc.Name = FName(TEXT("TargetB"));
+		TargetBDesc.Struct = TBaseStructure<FStateTreeTest_PropertyCopyObjects>::Get();
+		TargetBDesc.DataSource = EStateTreeBindableStructSource::Parameter;
+		TargetBDesc.ID = FGuid::NewGuid();
+
+		const int32 SourceIndex = BindingCompiler.AddSourceStruct(SourceDesc);
+
+		auto MakeBinding = [](const FGuid& SourceID, const FString& Source, const FGuid& TargetID, const FString& Target)
+		{
+			FStateTreePropertyPath SourcePath;
+			SourcePath.FromString(Source);
+			SourcePath.SetStructID(SourceID);
+
+			FStateTreePropertyPath TargetPath;
+			TargetPath.FromString(Target);
+			TargetPath.SetStructID(TargetID);
+
+			return FStateTreePropertyPathBinding(SourcePath, TargetPath);
+		};
+
+		TArray<FStateTreePropertyPathBinding> PropertyBindings;
+		// One-to-one copy from source to target A
+		PropertyBindings.Add(MakeBinding(SourceDesc.ID, TEXT("Object"), TargetADesc.ID, TEXT("Object")));
+		PropertyBindings.Add(MakeBinding(SourceDesc.ID, TEXT("SoftObject"), TargetADesc.ID, TEXT("SoftObject")));
+		PropertyBindings.Add(MakeBinding(SourceDesc.ID, TEXT("Class"), TargetADesc.ID, TEXT("Class")));
+		PropertyBindings.Add(MakeBinding(SourceDesc.ID, TEXT("SoftClass"), TargetADesc.ID, TEXT("SoftClass")));
+
+		// Cross copy from source to target B
+		PropertyBindings.Add(MakeBinding(SourceDesc.ID, TEXT("SoftObject"), TargetBDesc.ID, TEXT("Object")));
+		PropertyBindings.Add(MakeBinding(SourceDesc.ID, TEXT("Object"), TargetBDesc.ID, TEXT("SoftObject")));
+		PropertyBindings.Add(MakeBinding(SourceDesc.ID, TEXT("SoftClass"), TargetBDesc.ID, TEXT("Class")));
+		PropertyBindings.Add(MakeBinding(SourceDesc.ID, TEXT("Class"), TargetBDesc.ID, TEXT("SoftClass")));
+		
+		int32 TargetACopyBatchIndex = INDEX_NONE;
+		const bool bCompileBatchResultA = BindingCompiler.CompileBatch(TargetADesc, PropertyBindings, TargetACopyBatchIndex);
+		AITEST_TRUE("CompileBatchResultA should succeed", bCompileBatchResultA);
+		AITEST_NOT_EQUAL("TargetACopyBatchIndex should not be INDEX_NONE", TargetACopyBatchIndex, (int32)INDEX_NONE);
+
+		int32 TargetBCopyBatchIndex = INDEX_NONE;
+		const bool bCompileBatchResultB = BindingCompiler.CompileBatch(TargetBDesc, PropertyBindings, TargetBCopyBatchIndex);
+		AITEST_TRUE("CompileBatchResultB should succeed", bCompileBatchResultB);
+		AITEST_NOT_EQUAL("TargetBCopyBatchIndex should not be INDEX_NONE", TargetBCopyBatchIndex, (int32)INDEX_NONE);
+
+		BindingCompiler.Finalize();
+
+		const bool bResolveResult = Bindings.ResolvePaths();
+		AITEST_TRUE("ResolvePaths should succeed", bResolveResult);
+
+		UStateTreeTest_PropertyObject* ObjectA = NewObject<UStateTreeTest_PropertyObject>();
+		UStateTreeTest_PropertyObject2* ObjectB = NewObject<UStateTreeTest_PropertyObject2>();
+		
+		FStateTreeTest_PropertyCopyObjects Source;
+		Source.Object = ObjectA;
+		Source.SoftObject = ObjectB;
+		Source.Class = UStateTreeTest_PropertyObject::StaticClass();
+		Source.SoftClass = UStateTreeTest_PropertyObject::StaticClass();
+
+		AITEST_TRUE("SourceIndex should be less than max number of source structs.", SourceIndex < Bindings.GetSourceStructNum());
+
+		TArray<FStateTreeDataView> SourceViews;
+		SourceViews.SetNum(Bindings.GetSourceStructNum());
+		SourceViews[SourceIndex] = FStateTreeDataView(FStructView::Make(Source));
+		
+		FStateTreeTest_PropertyCopyObjects TargetA;
+		const bool bCopyResultA = Bindings.CopyTo(SourceViews, FStateTreeIndex16(TargetACopyBatchIndex), FStructView::Make(TargetA));
+		AITEST_TRUE("CopyTo should succeed", bCopyResultA);
+
+		AITEST_TRUE("Expect TargetA.Object == Source.Object", TargetA.Object == Source.Object);
+		AITEST_TRUE("Expect TargetA.SoftObject == Source.SoftObject", TargetA.SoftObject == Source.SoftObject);
+		AITEST_TRUE("Expect TargetA.Class == Source.Class", TargetA.Class == Source.Class);
+		AITEST_TRUE("Expect TargetA.SoftClass == Source.SoftClass", TargetA.SoftClass == Source.SoftClass);
+
+		// Copying to TargetB should not affect TargetA
+		TargetA.Object = nullptr;
+		
+		FStateTreeTest_PropertyCopyObjects TargetB;
+		const bool bCopyResultB = Bindings.CopyTo(SourceViews, FStateTreeIndex16(TargetBCopyBatchIndex), FStructView::Make(TargetB));
+		AITEST_TRUE("CopyTo should succeed", bCopyResultB);
+
+		AITEST_TRUE("Expect TargetB.Object == Source.SoftObject", TSoftObjectPtr<UObject>(TargetB.Object) == Source.SoftObject);
+		AITEST_TRUE("Expect TargetB.SoftObject == Source.Object", TargetB.SoftObject == TSoftObjectPtr<UObject>(Source.Object));
+		AITEST_TRUE("Expect TargetB.Class == Source.SoftClass", TSoftClassPtr<UObject>(TargetB.Class) == Source.SoftClass);
+		AITEST_TRUE("Expect TargetB.SoftClass == Source.Class", TargetB.SoftClass == TSoftClassPtr<UObject>(Source.Class));
+
+		AITEST_TRUE("Expect TargetA.Object == nullptr after copy of TargetB", TargetA.Object == nullptr);
+
+		// Collect ObjectA and ObjectB, soft object paths should still copy ok.
+		ObjectA = nullptr;
+		ObjectB = nullptr;
+		Source.Object = nullptr;
+		CollectGarbage(GARBAGE_COLLECTION_KEEPFLAGS);
+
+		FStateTreeTest_PropertyCopyObjects TargetC;
+		const bool bCopyResultC = Bindings.CopyTo(SourceViews, FStateTreeIndex16(TargetACopyBatchIndex), FStructView::Make(TargetC));
+		AITEST_TRUE("CopyTo should succeed", bCopyResultC);
+		AITEST_TRUE("Expect TargetC.SoftObject == Source.SoftObject after GC", TargetC.SoftObject == Source.SoftObject);
+
+		return true;
+	}
+};
+IMPLEMENT_AI_INSTANT_TEST(FStateTreeTest_CopyObjects, "System.StateTree.CopyObjects");
+
 struct FStateTreeTest_FollowTransitions : FAITestBase
 {
 	virtual bool InstantTest() override
diff --git a/Engine/Plugins/Runtime/StateTree/Source/StateTreeTestSuite/Private/StateTreeTestTypes.h b/Engine/Plugins/Runtime/StateTree/Source/StateTreeTestSuite/Private/StateTreeTestTypes.h
--- a/Engine/Plugins/Runtime/StateTree/Source/StateTreeTestSuite/Private/StateTreeTestTypes.h
+++ b/Engine/Plugins/Runtime/StateTree/Source/StateTreeTestSuite/Private/StateTreeTestTypes.h
@@ -396,6 +396,13 @@ public:
 	TArray<FStateTreeTest_PropertyStruct> ArrayOfStruct;
 };
 
+UCLASS(HideDropdown)
+class UStateTreeTest_PropertyObject2 : public UObject
+{
+	GENERATED_BODY()
+public:
+};
+
 USTRUCT()
 struct FStateTreeTest_PropertyCopy
 {
@@ -407,3 +414,21 @@ struct FStateTreeTest_PropertyCopy
 	UPROPERTY(EditAnywhere, Category = "")
 	TArray<FStateTreeTest_PropertyStruct> Array;
 };
+
+USTRUCT()
+struct FStateTreeTest_PropertyCopyObjects
+{
+	GENERATED_BODY()
+
+	UPROPERTY(EditAnywhere, Category = "")
+	TObjectPtr<UObject> Object;
+
+	UPROPERTY(EditAnywhere, Category = "")
+	TSubclassOf<UObject> Class;
+
+	UPROPERTY(EditAnywhere, Category = "")
+	TSoftObjectPtr<UObject> SoftObject;
+
+	UPROPERTY(EditAnywhere, Category = "")
+	TSoftClassPtr<UObject> SoftClass;
+};
