StateTree: Fix transition validation on global tasks

diff --git a/Engine/Plugins/Runtime/StateTree/Source/StateTreeEditorModule/Private/StateTreeCompiler.cpp b/Engine/Plugins/Runtime/StateTree/Source/StateTreeEditorModule/Private/StateTreeCompiler.cpp
--- a/Engine/Plugins/Runtime/StateTree/Source/StateTreeEditorModule/Private/StateTreeCompiler.cpp
+++ b/Engine/Plugins/Runtime/StateTree/Source/StateTreeEditorModule/Private/StateTreeCompiler.cpp
@@ -940,7 +940,7 @@ bool FStateTreeCompiler::ResolveTransitionState(const UStateTreeState* SourceSta
 		// Warn if goto state points to another subtree.
 		if (const UStateTreeState* TargetState = GetState(Link.ID))
 		{
-			if (TargetState->GetRootState() != SourceState->GetRootState())
+			if (SourceState && TargetState->GetRootState() != SourceState->GetRootState())
 			{
 				Log.Reportf(EMessageSeverity::Warning,
 					TEXT("Target state '%s' is in different subtree. Verify that this is intentional."),
