[StateTreeDebugger] added trace when transition tasks get evaluated

diff --git a/Engine/Plugins/Runtime/StateTree/Source/StateTreeModule/Private/StateTreeExecutionContext.cpp b/Engine/Plugins/Runtime/StateTree/Source/StateTreeModule/Private/StateTreeExecutionContext.cpp
--- a/Engine/Plugins/Runtime/StateTree/Source/StateTreeModule/Private/StateTreeExecutionContext.cpp
+++ b/Engine/Plugins/Runtime/StateTree/Source/StateTreeModule/Private/StateTreeExecutionContext.cpp
@@ -1690,6 +1690,7 @@ bool FStateTreeExecutionContext::TriggerTransitions()
 					if (Task.bShouldAffectTransitions)
 					{
 						STATETREE_LOG(VeryVerbose, TEXT("%*sTriggerTransitions: '%s'"), UE::StateTree::DebugIndentSize, TEXT(""), *Task.Name.ToString());
+						STATETREE_TRACE_TASK_EVENT(TaskIndex, DataViews[Task.DataViewIndex.Get()], EStateTreeTraceEventType::OnEvaluating, EStateTreeRunStatus::Running);
 						check(DataViews[Task.DataViewIndex.Get()].IsValid());
 						Task.TriggerTransitions(*this);
 					}
@@ -1813,6 +1814,7 @@ bool FStateTreeExecutionContext::TriggerTransitions()
 			if (Task.bShouldAffectTransitions)
 			{
 				STATETREE_LOG(VeryVerbose, TEXT("%*sTriggerTransitions: '%s'"), UE::StateTree::DebugIndentSize, TEXT(""), *Task.Name.ToString());
+				STATETREE_TRACE_TASK_EVENT(TaskIndex, DataViews[Task.DataViewIndex.Get()], EStateTreeTraceEventType::OnEvaluating, EStateTreeRunStatus::Running);
 				check(DataViews[Task.DataViewIndex.Get()].IsValid());
 				Task.TriggerTransitions(*this);
 			}
