StateTree: Reset state parameters when state is changed to None

diff --git a/Engine/Plugins/Runtime/StateTree/Source/StateTreeEditorModule/Private/StateTreeState.cpp b/Engine/Plugins/Runtime/StateTree/Source/StateTreeEditorModule/Private/StateTreeState.cpp
--- a/Engine/Plugins/Runtime/StateTree/Source/StateTreeEditorModule/Private/StateTreeState.cpp
+++ b/Engine/Plugins/Runtime/StateTree/Source/StateTreeEditorModule/Private/StateTreeState.cpp
@@ -131,6 +131,7 @@ void UStateTreeState::PostEditChangeChainProperty(FPropertyChangedChainEvent& Pr
 		}
 		else
 		{
+			// Other state types do not have parameters, so reset them.
 			Parameters.Reset();
 		}
 	}
@@ -321,6 +322,11 @@ void UStateTreeState::UpdateParametersFromLinkedSubtree()
 		{
 			Parameters.Parameters.MigrateToNewBagInstance(LinkTargetState->Parameters.Parameters);
 		}
+		else
+		{
+			// No state selected, reset. 
+			Parameters.Parameters.Reset();
+		}
 	}
 }
 
