StateTree: Handle instanced objects when copying State Tree nodes

diff --git a/Engine/Plugins/Runtime/StateTree/Source/StateTreeEditorModule/Private/Customizations/StateTreeEditorNodeDetails.cpp b/Engine/Plugins/Runtime/StateTree/Source/StateTreeEditorModule/Private/Customizations/StateTreeEditorNodeDetails.cpp
--- a/Engine/Plugins/Runtime/StateTree/Source/StateTreeEditorModule/Private/Customizations/StateTreeEditorNodeDetails.cpp
+++ b/Engine/Plugins/Runtime/StateTree/Source/StateTreeEditorModule/Private/Customizations/StateTreeEditorNodeDetails.cpp
@@ -57,6 +57,44 @@ public:
 
 namespace UE::StateTreeEditor::Internal
 {
+	/**
+	 * This function recursively instantiates instanced objects of a given struct.
+	 * It is needed to fixup nodes pasted from clipboard, which seem to give shallow copy.
+	 */
+	void InstantiateStructSubobjects(UObject& OuterObject, FStructView Struct)
+	{
+		// Empty struct, nothing to do.
+		if (!Struct.IsValid())
+		{
+			return;
+		}
+
+		for (TPropertyValueIterator<FProperty> It(Struct.GetScriptStruct(), Struct.GetMemory()); It; ++It)
+		{
+			if (const FObjectProperty* ObjectProperty = CastField<FObjectProperty>(It->Key))
+			{
+				// Duplicate instanced objects.
+				if (ObjectProperty->HasAnyPropertyFlags(CPF_InstancedReference | CPF_PersistentInstance))
+				{
+					if (UObject* Object = ObjectProperty->GetObjectPropertyValue(It->Value))
+					{
+						UObject* DuplicatedObject = DuplicateObject(Object, &OuterObject);
+						ObjectProperty->SetObjectPropertyValue(const_cast<void*>(It->Value), DuplicatedObject);
+					}
+				}
+			}
+			if (const FStructProperty* StructProperty = CastField<FStructProperty>(It->Key))
+			{
+				// If we encounter instanced struct, recursively handle it too.
+				if (StructProperty->Struct == TBaseStructure<FInstancedStruct>::Get())
+				{
+					FInstancedStruct& InstancedStruct = *static_cast<FInstancedStruct*>(const_cast<void*>(It->Value));
+					InstantiateStructSubobjects(OuterObject, InstancedStruct);
+				}
+			}
+		}
+	}
+
 	/** @return text describing the pin type, matches SPinTypeSelector. */
 	FText GetPinTypeText(const FEdGraphPinType& PinType)
 	{
@@ -550,9 +588,6 @@ void FStateTreeEditorNodeDetails::OnPasteNode()
 
 	StructProperty->NotifyPreChange();
 
-	// Make sure we instantiate new objects when setting the value.
-	StructProperty->SetValueFromFormattedString(PastedText, EPropertyValueSetFlags::InstanceObjects);
-
 	// Reset GUIDs on paste
 	TArray<void*> RawNodeData;
 	StructProperty->AccessRawData(RawNodeData);
@@ -564,6 +599,20 @@ void FStateTreeEditorNodeDetails::OnPasteNode()
 			FStateTreeEditorNode* EditorNode = static_cast<FStateTreeEditorNode*>(RawNodeData[Index]);
 			if (EditorNode && OuterObject)
 			{
+				// Copy
+				*EditorNode = TempNode;
+
+				// Ensure unique instance value
+				UE::StateTreeEditor::Internal::InstantiateStructSubobjects(*OuterObject, EditorNode->Node);
+				if (EditorNode->InstanceObject)
+				{
+					EditorNode->InstanceObject = DuplicateObject(EditorNode->InstanceObject, OuterObject);
+				}
+				else
+				{
+					UE::StateTreeEditor::Internal::InstantiateStructSubobjects(*OuterObject, EditorNode->Instance);
+				}
+				
 				if (FStateTreeNodeBase* Node = EditorNode->Node.GetMutablePtr<FStateTreeNodeBase>())
 				{
 					Node->Name = FName(Node->Name.ToString() + TEXT(" Copy"));
