Adding additional error printout to UpdateSegmentsFromValue

UpdateSegmentsFromValue contains an optional field to print out additional info regarding failures. This change populates this field and the corresponding error log to print out the additional information.

diff --git a/Engine/Plugins/Runtime/StateTree/Source/StateTreeModule/Private/StateTree.cpp b/Engine/Plugins/Runtime/StateTree/Source/StateTreeModule/Private/StateTree.cpp
--- a/Engine/Plugins/Runtime/StateTree/Source/StateTreeModule/Private/StateTree.cpp
+++ b/Engine/Plugins/Runtime/StateTree/Source/StateTreeModule/Private/StateTree.cpp
@@ -649,19 +649,20 @@ bool UStateTree::PatchBindings()
 			return false;
 		}
 
+		FString ErrorMsg;
 		for (int32 Index = Batch.BindingsBegin; Index != Batch.BindingsEnd; Index++)
 		{
 			FStateTreePropertyPathBinding& Binding = PropertyPathBindings[Index];
 
-			if (!Binding.GetMutableSourcePath().UpdateSegmentsFromValue(DataViews[Binding.GetCompiledSourceStructIndex().Get()]))
+			if (!Binding.GetMutableSourcePath().UpdateSegmentsFromValue(DataViews[Binding.GetCompiledSourceStructIndex().Get()], &ErrorMsg))
 			{
-				UE_LOG(LogStateTree, Error, TEXT("%hs: Failed to update source instance structs for property binding '%s'."), __FUNCTION__, *Binding.GetTargetPath().ToString());
+				UE_LOG(LogStateTree, Error, TEXT("%hs: Failed to update source instance structs for property binding '%s'. Reason: %s"), __FUNCTION__, *Binding.GetTargetPath().ToString(), *ErrorMsg);
 				return false;
 			}
 
-			if (!Binding.GetMutableTargetPath().UpdateSegmentsFromValue(TargetView))
+			if (!Binding.GetMutableTargetPath().UpdateSegmentsFromValue(TargetView, &ErrorMsg))
 			{
-				UE_LOG(LogStateTree, Error, TEXT("%hs: Failed to update target instance structs for property binding '%s'."), __FUNCTION__, *Binding.GetTargetPath().ToString());
+				UE_LOG(LogStateTree, Error, TEXT("%hs: Failed to update target instance structs for property binding '%s'. Reason: %s"), __FUNCTION__, *Binding.GetTargetPath().ToString(), *ErrorMsg);
 				return false;
 			}
 		}
