[StateTreeDebugger] can now auto start analysis of traces started from console command

diff --git a/Engine/Plugins/Runtime/StateTree/Source/StateTreeModule/Private/Debugger/StateTreeDebugger.cpp b/Engine/Plugins/Runtime/StateTree/Source/StateTreeModule/Private/Debugger/StateTreeDebugger.cpp
--- a/Engine/Plugins/Runtime/StateTree/Source/StateTreeModule/Private/Debugger/StateTreeDebugger.cpp
+++ b/Engine/Plugins/Runtime/StateTree/Source/StateTreeModule/Private/Debugger/StateTreeDebugger.cpp
@@ -8,6 +8,7 @@
 #include "Debugger/StateTreeTraceTypes.h"
 #include "Misc/Paths.h"
 #include "Modules/ModuleManager.h"
+#include "StateTreeDelegates.h"
 #include "StateTreeModule.h"
 #include "Trace/StoreClient.h"
 #include "TraceServices/AnalysisService.h"
@@ -77,10 +78,19 @@ FStateTreeDebugger::FStateTreeDebugger()
 	: StateTreeModule(FModuleManager::GetModuleChecked<IStateTreeModule>("StateTreeModule"))
 	, ScrubState(EventCollections)
 {
+	UE::StateTree::Delegates::OnTracingStateChanged.AddLambda([this](const bool bTracesEnabled)
+		{
+			// StateTree traces got enabled in the current process so let's analyse it if not already analysing something.
+			if (bTracesEnabled && !IsAnalysisSessionActive())
+			{
+				RequestAnalysisOfLatestTrace();
+			}
+		});
 }
 
 FStateTreeDebugger::~FStateTreeDebugger()
 {
+	UE::StateTree::Delegates::OnTracingStateChanged.RemoveAll(this);
 	StopSessionAnalysis();
 }
 
@@ -202,30 +212,19 @@ void FStateTreeDebugger::UpdateInstances()
 bool FStateTreeDebugger::RequestAnalysisOfEditorSession()
 {
 	// Get snapshot of current trace to help identify the next live one
-	LastLiveSessionId = INDEX_NONE;
 	TArray<FTraceDescriptor> TraceDescriptors;
 	GetLiveTraces(TraceDescriptors);
+	LastLiveSessionId = TraceDescriptors.Num() ? TraceDescriptors.Last().TraceId : INDEX_NONE;
 
 	// 0 is the invalid value used for Trace Id
 	constexpr int32 InvalidTraceId = 0;
 	int32 ActiveTraceId = InvalidTraceId;
 
-	// StartTraces returns true if a new connection was created. In this case
-	// we try to start an analysis on that new connection as soon as possible.
-	// Otherwise it might have been able to use an active connection in which
-	// case it was returned in the output parameter.
+	// StartTraces returns true if a new connection was created. In this case we will receive OnTracingStateChanged
+	// and we'll try to start an analysis on that new connection as soon as possible.
+	// Otherwise it might have been able to use an active connection in which case it was returned in the output parameter.
 	if (StateTreeModule.StartTraces(ActiveTraceId))
 	{
-		// Invalidate our current active session
-		ActiveSessionTraceDescriptor = FTraceDescriptor();
-
-		// Stop current analysis if any
-		StopSessionAnalysis();
-
-		LastLiveSessionId = TraceDescriptors.Num() ? TraceDescriptors.Last().TraceId : INDEX_NONE;
-
-		// This won't succeed yet but will schedule our next retry
-		TryStartNewLiveSessionAnalysis(1.0f);
 		return true;
 	}
 
@@ -244,6 +243,19 @@ bool FStateTreeDebugger::RequestAnalysisOfEditorSession()
 	return false;
 }
 
+void FStateTreeDebugger::RequestAnalysisOfLatestTrace()
+{
+	// Invalidate our current active session
+	ActiveSessionTraceDescriptor = FTraceDescriptor();
+
+	// Stop current analysis if any
+	StopSessionAnalysis();
+
+	// This might not succeed immediately but will schedule next retry if necessary
+	TryStartNewLiveSessionAnalysis(1.0f);
+}
+
+
 bool FStateTreeDebugger::TryStartNewLiveSessionAnalysis(const float RetryPollingDuration)
 {
 	TArray<FTraceDescriptor> Traces;
diff --git a/Engine/Plugins/Runtime/StateTree/Source/StateTreeModule/Public/Debugger/StateTreeDebugger.h b/Engine/Plugins/Runtime/StateTree/Source/StateTreeModule/Public/Debugger/StateTreeDebugger.h
--- a/Engine/Plugins/Runtime/StateTree/Source/StateTreeModule/Public/Debugger/StateTreeDebugger.h
+++ b/Engine/Plugins/Runtime/StateTree/Source/StateTreeModule/Public/Debugger/StateTreeDebugger.h
@@ -140,7 +140,6 @@ struct STATETREEMODULE_API FStateTreeDebugger : FTickableGameObject
 
 	/**
 	 * Queue a request to auto start an analysis session on the next available live trace.
-	 * This will replace the current analysis session if any.
 	 * @return True if connection was successfully requested or was able to use active trace, false otherwise.
 	 */
 	bool RequestAnalysisOfEditorSession();
@@ -235,6 +234,12 @@ private:
 
 	void SetActiveStates(const TConstArrayView<FStateTreeStateHandle> NewActiveStates);
 
+	/**
+	 * Request an analysis session on the latest next available live trace.
+	 * This will replace the current analysis session if any.
+	 */
+	void RequestAnalysisOfLatestTrace();
+
 	/**
 	 * Looks for a new live traces to start an analysis session.
 	 * On failure, if 'RetryPollingDuration is > 0', will retry connecting every frame for 'RetryPollingDuration' seconds 
