[StateTree] Added deferred stop and validation for reentrant calls to Start/Stop/Tick.

diff --git a/Engine/Plugins/Runtime/StateTree/Source/StateTreeModule/Private/StateTreeExecutionContext.cpp b/Engine/Plugins/Runtime/StateTree/Source/StateTreeModule/Private/StateTreeExecutionContext.cpp
--- a/Engine/Plugins/Runtime/StateTree/Source/StateTreeModule/Private/StateTreeExecutionContext.cpp
+++ b/Engine/Plugins/Runtime/StateTree/Source/StateTreeModule/Private/StateTreeExecutionContext.cpp
@@ -15,6 +15,10 @@
 #define STATETREE_LOG(Verbosity, Format, ...) UE_VLOG_UELOG(GetOwner(), LogStateTree, Verbosity, TEXT("%s: ") Format, *GetInstanceDescription(), ##__VA_ARGS__)
 #define STATETREE_CLOG(Condition, Verbosity, Format, ...) UE_CVLOG_UELOG((Condition), GetOwner(), LogStateTree, Verbosity, TEXT("%s: ") Format, *GetInstanceDescription(), ##__VA_ARGS__)
 
+#define STATETREE_LOG_AND_TRACE(Verbosity, Format, ...) \
+	UE_VLOG_UELOG(GetOwner(), LogStateTree, Verbosity, TEXT("%s: ") Format, *GetInstanceDescription(), ##__VA_ARGS__); \
+	STATETREE_TRACE_LOG_EVENT(Format, ##__VA_ARGS__)
+
 #if WITH_STATETREE_DEBUGGER
 	#define ID_NAME PREPROCESSOR_JOIN(InstanceId,__LINE__) \
 
@@ -208,8 +212,18 @@ EStateTreeRunStatus FStateTreeExecutionContext::Start()
 		return EStateTreeRunStatus::Failed;
 	}
 
-	// Stop if still running previous state.
-	Stop();
+	if (InstanceData.IsValid())
+	{
+		const FStateTreeExecutionState& Exec = GetExecState();
+		if (!ensureMsgf(Exec.CurrentPhase == EStateTreeUpdatePhase::Unset, TEXT("%hs can't be called while already in %s ('%s' using StateTree '%s')."),
+				__FUNCTION__, *UEnum::GetDisplayValueAsText(Exec.CurrentPhase).ToString(), *GetNameSafe(&Owner), *GetFullNameSafe(&StateTree)))
+		{
+			return EStateTreeRunStatus::Failed;
+		}
+
+		// Stop if still running previous state.
+		Stop();
+	}
 
 	// Initialize instance data. No active states yet, so we'll initialize the evals and global tasks.
 	InstanceData.Reset();
@@ -228,76 +242,97 @@ EStateTreeRunStatus FStateTreeExecutionContext::Start()
 	// Set scoped phase only for properly initialized context with valid Instance data
 	// since we need it to output the InstanceId
 	STATETREE_TRACE_SCOPED_PHASE(EStateTreeUpdatePhase::StartTree);
-	
+
+	FStateTreeExecutionState* Exec = &GetExecState(); // Using pointer as we will need to reacquire the exec later.
+
+	// From this point any calls to Stop should be deferred.
+	Exec->CurrentPhase = EStateTreeUpdatePhase::StartTree;
+
 	// Start evaluators and global tasks. Fail the execution if any global task fails.
 	FStateTreeIndex16 LastInitializedTaskIndex;
 	const EStateTreeRunStatus GlobalTasksRunStatus = StartEvaluatorsAndGlobalTasks(LastInitializedTaskIndex);
-	if (GlobalTasksRunStatus != EStateTreeRunStatus::Running)
+	if (GlobalTasksRunStatus == EStateTreeRunStatus::Running)
 	{
-		StopEvaluatorsAndGlobalTasks(GlobalTasksRunStatus, LastInitializedTaskIndex);
-		STATETREE_LOG(VeryVerbose, TEXT("%hs: Global tasks completed the StateTree %s on start in status '%s'. "),
-			__FUNCTION__, *GetNameSafe(&Owner), *GetFullNameSafe(&StateTree), *UEnum::GetDisplayValueAsText(GlobalTasksRunStatus).ToString());
-		return GlobalTasksRunStatus;
-	}
+		// First tick.
+		// Tasks are not ticked here, since their behavior is that EnterState() (called above) is treated as a tick.  
+		TickEvaluatorsAndGlobalTasks(0.0f, /*bTickGlobalTasks*/false);
 
-	// First tick.
-	// Tasks are not ticked here, since their behavior is that EnterState() (called above) is treated as a tick.  
-	TickEvaluatorsAndGlobalTasks(0.0f, /*bTickGlobalTasks*/false);
+		// Initialize to unset running state.
+		Exec->TreeRunStatus = EStateTreeRunStatus::Running;
+		Exec->ActiveStates.Reset();
+		Exec->LastTickStatus = EStateTreeRunStatus::Unset;
 
-	// Initialize to unset running state.
-	FStateTreeExecutionState* Exec = &GetExecState(); // Using pointer as we will need to reacquire the exec later.
-	Exec->TreeRunStatus = EStateTreeRunStatus::Running;
-	Exec->ActiveStates.Reset();
-	Exec->LastTickStatus = EStateTreeRunStatus::Unset;
-
-	static const FStateTreeStateHandle RootState = FStateTreeStateHandle(0);
+		static const FStateTreeStateHandle RootState = FStateTreeStateHandle(0);
 
-	FStateTreeActiveStates NextActiveStates;
-	FStateTreeActiveStates VisitedStates;
-	if (SelectState(RootState, NextActiveStates, VisitedStates))
-	{
-		if (NextActiveStates.Last().IsCompletionState())
+		FStateTreeActiveStates NextActiveStates;
+		FStateTreeActiveStates VisitedStates;
+		if (SelectState(RootState, NextActiveStates, VisitedStates))
 		{
-			// Transition to a terminal state (succeeded/failed), or default transition failed.
-			STATETREE_LOG(Warning, TEXT("%hs: Tree %s at StateTree start on '%s' using StateTree '%s'."),
-				__FUNCTION__, NextActiveStates.Last() == FStateTreeStateHandle::Succeeded ? TEXT("succeeded") : TEXT("failed"), *GetNameSafe(&Owner), *GetFullNameSafe(&StateTree));
-			Exec->TreeRunStatus = NextActiveStates.Last().ToCompletionStatus();
-		}
-		else
-		{
-			// Enter state tasks can fail/succeed, treat it same as tick.
-			FStateTreeTransitionResult Transition;
-			Transition.TargetState = RootState;
-			Transition.CurrentActiveStates = Exec->ActiveStates;
-			Transition.CurrentRunStatus = Exec->LastTickStatus;
-			Transition.NextActiveStates = NextActiveStates; // Enter state will update Exec.ActiveStates.
-			const EStateTreeRunStatus LastTickStatus = EnterState(Transition);
+			if (NextActiveStates.Last().IsCompletionState())
+			{
+				// Transition to a terminal state (succeeded/failed), or default transition failed.
+				STATETREE_LOG(Warning, TEXT("%hs: Tree %s at StateTree start on '%s' using StateTree '%s'."),
+					__FUNCTION__, NextActiveStates.Last() == FStateTreeStateHandle::Succeeded ? TEXT("succeeded") : TEXT("failed"), *GetNameSafe(&Owner), *GetFullNameSafe(&StateTree));
+				Exec->TreeRunStatus = NextActiveStates.Last().ToCompletionStatus();
+			}
+			else
+			{
+				// Enter state tasks can fail/succeed, treat it same as tick.
+				FStateTreeTransitionResult Transition;
+				Transition.TargetState = RootState;
+				Transition.CurrentActiveStates = Exec->ActiveStates;
+				Transition.CurrentRunStatus = Exec->LastTickStatus;
+				Transition.NextActiveStates = NextActiveStates; // Enter state will update Exec.ActiveStates.
+				const EStateTreeRunStatus LastTickStatus = EnterState(Transition);
 			
-			// Need to reacquire the exec state as EnterState may alter the allocation. 
-			Exec = &GetExecState();
-			Exec->LastTickStatus = LastTickStatus;
-			STATETREE_TRACE_ACTIVE_STATES_EVENT(Exec->ActiveStates);
+				// Need to reacquire the exec state as EnterState may alter the allocation. 
+				Exec = &GetExecState();
+				Exec->LastTickStatus = LastTickStatus;
+				STATETREE_TRACE_ACTIVE_STATES_EVENT(Exec->ActiveStates);
 
-			// Report state completed immediately.
-			if (Exec->LastTickStatus != EStateTreeRunStatus::Running)
-			{
-				StateCompleted();
+				// Report state completed immediately.
+				if (Exec->LastTickStatus != EStateTreeRunStatus::Running)
+				{
+					StateCompleted();
+				}
 			}
 		}
-	}
 
-	if (Exec->ActiveStates.IsEmpty())
+		if (Exec->ActiveStates.IsEmpty())
+		{
+			// Should not happen. This may happen if initial state could not be selected.
+			STATETREE_LOG(Error, TEXT("%hs: Failed to select initial state on '%s' using StateTree '%s'. This should not happen, check that the StateTree logic can always select a state at start."),
+				__FUNCTION__, *GetNameSafe(&Owner), *GetFullNameSafe(&StateTree));
+			Exec->TreeRunStatus = EStateTreeRunStatus::Failed;
+		}
+	}
+	else
 	{
-		// Should not happen. This may happen if initial state could not be selected.
-		STATETREE_LOG(Error, TEXT("%hs: Failed to select initial state on '%s' using StateTree '%s'. This should not happen, check that the StateTree logic can always select a state at start."),
-			__FUNCTION__, *GetNameSafe(&Owner), *GetFullNameSafe(&StateTree));
-		Exec->TreeRunStatus = EStateTreeRunStatus::Failed;
+		StopEvaluatorsAndGlobalTasks(GlobalTasksRunStatus, LastInitializedTaskIndex);
+
+		STATETREE_LOG(VeryVerbose, TEXT("%hs: Global tasks completed the StateTree %s on start in status '%s'."),
+			__FUNCTION__, *GetNameSafe(&Owner), *GetFullNameSafe(&StateTree), *UEnum::GetDisplayValueAsText(GlobalTasksRunStatus).ToString());
+
+		// We are not considered as running yet so we only set the status without requiring a stop.
+		Exec->TreeRunStatus = GlobalTasksRunStatus;
 	}
 
-	return Exec->TreeRunStatus;
+	// Reset phase since we are now safe to stop.
+	Exec->CurrentPhase = EStateTreeUpdatePhase::Unset;
+
+	// Use local for resulting run state since Stop will reset the instance data.
+	EStateTreeRunStatus Result = Exec->TreeRunStatus;
+	
+	if (Exec->RequestedStop != EStateTreeRunStatus::Unset)
+	{
+		STATETREE_LOG_AND_TRACE(VeryVerbose, TEXT("Processing Deferred Stop"));
+		Result = Stop(Exec->RequestedStop);
+	}
+	
+	return Result;
 }
 
-EStateTreeRunStatus FStateTreeExecutionContext::Stop(const EStateTreeRunStatus CompletionStatus)
+EStateTreeRunStatus FStateTreeExecutionContext::Stop(EStateTreeRunStatus CompletionStatus)
 {
 	CSV_SCOPED_TIMING_STAT_EXCLUSIVE(StateTree_Stop);
 
@@ -317,7 +352,27 @@ EStateTreeRunStatus FStateTreeExecutionContext::Stop(const EStateTreeRunStatus C
 	// since we need it to output the InstanceId 
 	STATETREE_TRACE_SCOPED_PHASE(EStateTreeUpdatePhase::StopTree);
 
-	const FStateTreeExecutionState& Exec = GetExecState();
+	// Make sure that we return a valid completion status (i.e. Succeeded, Failed or Stopped)
+	if (CompletionStatus == EStateTreeRunStatus::Unset
+		|| CompletionStatus == EStateTreeRunStatus::Running)
+	{
+		CompletionStatus = EStateTreeRunStatus::Stopped;
+	}
+
+	FStateTreeExecutionState& Exec = GetExecState();
+
+	// A reentrant call to Stop or a call from Start or Tick must be deferred.
+	if (Exec.CurrentPhase != EStateTreeUpdatePhase::Unset)
+	{
+		STATETREE_LOG_AND_TRACE(VeryVerbose, TEXT("Deferring Stop at end of %s"), *UEnum::GetDisplayValueAsText(Exec.CurrentPhase).ToString());
+
+		Exec.RequestedStop = CompletionStatus;
+		return EStateTreeRunStatus::Running;
+	}
+
+	// No need to clear on exit since we reset all the instance data before leaving the function.
+	Exec.CurrentPhase = EStateTreeUpdatePhase::StopTree;
+
 	EStateTreeRunStatus Result = Exec.TreeRunStatus;
 	
 	// Capture events added between ticks.
@@ -377,17 +432,26 @@ EStateTreeRunStatus FStateTreeExecutionContext::Tick(const float DeltaTime)
 
 	FStateTreeEventQueue& EventQueue = InstanceData.GetMutableEventQueue();
 	FStateTreeExecutionState* Exec = &GetExecState();
-	
-	// Capture events added between ticks.
-	EventsToProcess = EventQueue.GetEvents();
-	EventQueue.Reset();
-	
-	// No ticking of the tree is done or stopped.
+
+	// No ticking if the tree is done or stopped.
 	if (Exec->TreeRunStatus != EStateTreeRunStatus::Running)
 	{
 		return Exec->TreeRunStatus;
 	}
 
+	if (!ensureMsgf(Exec->CurrentPhase == EStateTreeUpdatePhase::Unset, TEXT("%hs can't be called while already in %s ('%s' using StateTree '%s')."),
+			__FUNCTION__, *UEnum::GetDisplayValueAsText(Exec->CurrentPhase).ToString(), *GetNameSafe(&Owner), *GetFullNameSafe(&StateTree)))
+	{
+		return EStateTreeRunStatus::Failed;
+	}
+
+	// From this point any calls to Stop should be deferred.
+	Exec->CurrentPhase = EStateTreeUpdatePhase::TickStateTree;
+	
+	// Capture events added between ticks.
+	EventsToProcess = EventQueue.GetEvents();
+	EventQueue.Reset();
+	
 	// Update the delayed transitions.
 	for (FStateTreeTransitionDelayedState& DelayedState : Exec->DelayedTransitions)
 	{
@@ -396,106 +460,107 @@ EStateTreeRunStatus FStateTreeExecutionContext::Tick(const float DeltaTime)
 
 	// Tick global evaluators and tasks.
 	const EStateTreeRunStatus EvalAndGlobalTaskStatus = TickEvaluatorsAndGlobalTasks(DeltaTime);
-	if (EvalAndGlobalTaskStatus != EStateTreeRunStatus::Running)
-	{
-		return Stop(EvalAndGlobalTaskStatus);
-	}
-
-	if (Exec->LastTickStatus == EStateTreeRunStatus::Running)
+	if (EvalAndGlobalTaskStatus == EStateTreeRunStatus::Running)
 	{
-		// Tick tasks on active states.
-		Exec->LastTickStatus = TickTasks(DeltaTime);
-		
-		// Report state completed immediately.
-		if (Exec->LastTickStatus != EStateTreeRunStatus::Running)
+		if (Exec->LastTickStatus == EStateTreeRunStatus::Running)
 		{
-			StateCompleted();
-		}
-	}
+			// Tick tasks on active states.
+			Exec->LastTickStatus = TickTasks(DeltaTime);
 
-	// The state selection is repeated up to MaxIteration time. This allows failed EnterState() to potentially find a new state immediately.
-	// This helps event driven StateTrees to not require another event/tick to find a suitable state.
-	static constexpr int32 MaxIterations = 5;
-	for (int32 Iter = 0; Iter < MaxIterations; Iter++)
-	{
-		// Append events accumulated during the tick, so that transitions can immediately act on them.
-		// We'll consume the events only if they lead to state change below (EnterState is treated the same as Tick),
-		// or let them be processed next frame if no transition.
-		EventsToProcess.Append(EventQueue.GetEvents());
+			// Report state completed immediately.
+			if (Exec->LastTickStatus != EStateTreeRunStatus::Running)
+			{
+				StateCompleted();
+			}
+		}
 
-		// Trigger conditional transitions or state succeed/failed transitions. First tick transition is handled here too.
-		if (TriggerTransitions())
+		// The state selection is repeated up to MaxIteration time. This allows failed EnterState() to potentially find a new state immediately.
+		// This helps event driven StateTrees to not require another event/tick to find a suitable state.
+		static constexpr int32 MaxIterations = 5;
+		for (int32 Iter = 0; Iter < MaxIterations; Iter++)
 		{
-			STATETREE_TRACE_SCOPED_PHASE(EStateTreeUpdatePhase::ApplyTransitions);
-			STATETREE_TRACE_TRANSITION_EVENT(NextTransitionSource, EStateTreeTraceEventType::OnTransition);
-			NextTransitionSource.Reset();
+			// Append events accumulated during the tick, so that transitions can immediately act on them.
+			// We'll consume the events only if they lead to state change below (EnterState is treated the same as Tick),
+			// or let them be processed next frame if no transition.
+			EventsToProcess.Append(EventQueue.GetEvents());
 
-			// We have committed to state change, consume events that were accumulated during the tick above.
-			EventQueue.Reset();
+			// Trigger conditional transitions or state succeed/failed transitions. First tick transition is handled here too.
+			if (TriggerTransitions())
+			{
+				STATETREE_TRACE_SCOPED_PHASE(EStateTreeUpdatePhase::ApplyTransitions);
+				STATETREE_TRACE_TRANSITION_EVENT(NextTransitionSource, EStateTreeTraceEventType::OnTransition);
+				NextTransitionSource.Reset();
 
-			ExitState(NextTransition);
+				// We have committed to state change, consume events that were accumulated during the tick above.
+				EventQueue.Reset();
 
-			// Tree succeeded or failed.
-			if (NextTransition.TargetState.IsCompletionState())
-			{
-				// Transition to a terminal state (succeeded/failed), or default transition failed.
-				Exec->TreeRunStatus = NextTransition.TargetState.ToCompletionStatus();
-				Exec->ActiveStates.Reset();
+				ExitState(NextTransition);
 
-				// Stop evaluators and global tasks.
-				StopEvaluatorsAndGlobalTasks(Exec->TreeRunStatus);
+				// Tree succeeded or failed.
+				if (NextTransition.TargetState.IsCompletionState())
+				{
+					// Transition to a terminal state (succeeded/failed), or default transition failed.
+					Exec->TreeRunStatus = NextTransition.TargetState.ToCompletionStatus();
+					Exec->ActiveStates.Reset();
 
-				return Exec->TreeRunStatus;
-			}
+					// Stop evaluators and global tasks.
+					StopEvaluatorsAndGlobalTasks(Exec->TreeRunStatus);
 
-			// Append and consume the events accumulated during the state exit.
-			EventsToProcess.Append(EventQueue.GetEvents());
-			EventQueue.Reset();
+					break;
+				}
 
-			// Enter state tasks can fail/succeed, treat it same as tick.
-			const EStateTreeRunStatus LastTickStatus = EnterState(NextTransition);
+				// Append and consume the events accumulated during the state exit.
+				EventsToProcess.Append(EventQueue.GetEvents());
+				EventQueue.Reset();
 
-			NextTransition.Reset();
+				// Enter state tasks can fail/succeed, treat it same as tick.
+				const EStateTreeRunStatus LastTickStatus = EnterState(NextTransition);
 
-			// Need to reacquire the exec state as EnterState may alter the allocation. 
-			Exec = &GetExecState();
-			Exec->LastTickStatus = LastTickStatus;
-			STATETREE_TRACE_ACTIVE_STATES_EVENT(Exec->ActiveStates);
+				NextTransition.Reset();
 
-			// Consider events so far processed. Events sent during EnterState went into EventQueue, and are processed in next iteration.
-			EventsToProcess.Reset();
+				// Need to reacquire the exec state as EnterState may alter the allocation. 
+				Exec = &GetExecState();
+				Exec->LastTickStatus = LastTickStatus;
+				STATETREE_TRACE_ACTIVE_STATES_EVENT(Exec->ActiveStates);
 
-			// Report state completed immediately.
-			if (Exec->LastTickStatus != EStateTreeRunStatus::Running)
-			{
-				StateCompleted();
+				// Consider events so far processed. Events sent during EnterState went into EventQueue, and are processed in next iteration.
+				EventsToProcess.Reset();
+
+				// Report state completed immediately.
+				if (Exec->LastTickStatus != EStateTreeRunStatus::Running)
+				{
+					StateCompleted();
+				}
 			}
-		}
 
-		// Stop as soon as have found a running state.
-		if (Exec->LastTickStatus == EStateTreeRunStatus::Running)
-		{
-			break;
+			// Stop as soon as have found a running state.
+			if (Exec->LastTickStatus == EStateTreeRunStatus::Running)
+			{
+				break;
+			}
 		}
 	}
-
-	if (Exec->ActiveStates.IsEmpty())
+	else
 	{
-		// Should not happen. This may happen if a state completion transition could not be selected. 
-		STATETREE_LOG(Error, TEXT("%hs: Failed to select state on '%s' using StateTree '%s'. This should not happen, state completion transition is likely missing."),
-			__FUNCTION__, *GetNameSafe(&Owner), *GetFullNameSafe(&StateTree));
+		// If global tasks succeed or fail, stop the tree.
+		Exec->RequestedStop = EvalAndGlobalTaskStatus;
+	}
 
-		Exec->TreeRunStatus = EStateTreeRunStatus::Failed;
+	EventsToProcess.Reset();
 
-		// Stop evaluators and global tasks.
-		StopEvaluatorsAndGlobalTasks(Exec->TreeRunStatus);
+	// Reset phase since we are now safe to stop.
+	Exec->CurrentPhase = EStateTreeUpdatePhase::Unset;
 
-		return Exec->TreeRunStatus;
+	// Use local for resulting run state since Stop will reset the instance data.
+	EStateTreeRunStatus Result = Exec->TreeRunStatus;
+	
+	if (Exec->RequestedStop != EStateTreeRunStatus::Unset)
+	{
+		STATETREE_LOG_AND_TRACE(VeryVerbose, TEXT("Processing Deferred Stop"));
+		Result = Stop(Exec->RequestedStop);
 	}
 
-	EventsToProcess.Reset();
-
-	return Exec->TreeRunStatus;
+	return Result;
 }
 
 EStateTreeRunStatus FStateTreeExecutionContext::GetStateTreeRunStatus() const
@@ -538,8 +603,7 @@ void FStateTreeExecutionContext::SendEvent(const FGameplayTag Tag, const FConstS
 		return;
 	}
 
-	STATETREE_LOG(Verbose, TEXT("Send Event '%s'"), *Tag.ToString());
-	STATETREE_TRACE_LOG_EVENT(TEXT("Send Event '%s'"), *Tag.ToString());
+	STATETREE_LOG_AND_TRACE(Verbose, TEXT("Send Event '%s'"), *Tag.ToString());
 
 	FStateTreeEventQueue& EventQueue = InstanceData.GetMutableEventQueue();
 	EventQueue.SendEvent(&Owner, Tag, Payload, Origin);
@@ -1534,13 +1598,7 @@ bool FStateTreeExecutionContext::RequestTransition(const FStateTreeStateHandle N
 	
 	if (NextState.IsCompletionState())
 	{
-		NextTransition.CurrentActiveStates = Exec.ActiveStates;
-		NextTransition.CurrentRunStatus = Exec.LastTickStatus;
-		NextTransition.SourceState = CurrentlyProcessedState;
-		NextTransition.TargetState = NextState;
-		NextTransition.NextActiveStates = FStateTreeActiveStates(NextState);
-		NextTransition.Priority = Priority;
-
+		SetupNextTransition(NextState, Priority);
 		STATETREE_LOG(Verbose, TEXT("Transition on state '%s' -[%s]-> state '%s'"),
 			*GetSafeStateName(NextTransition.CurrentActiveStates.Last()), *GetSafeStateName(NextState), *GetSafeStateName(NextTransition.NextActiveStates.Last()));
 
@@ -1549,12 +1607,7 @@ bool FStateTreeExecutionContext::RequestTransition(const FStateTreeStateHandle N
 	if (!NextState.IsValid())
 	{
 		// NotSet is no-operation, but can be used to mask a transition at parent state. Returning unset keeps updating current state.
-		NextTransition.CurrentActiveStates = Exec.ActiveStates;
-		NextTransition.CurrentRunStatus = Exec.LastTickStatus;
-		NextTransition.SourceState = CurrentlyProcessedState;
-		NextTransition.TargetState = FStateTreeStateHandle::Invalid;
-		NextTransition.NextActiveStates.Reset();
-		NextTransition.Priority = Priority;
+		SetupNextTransition(FStateTreeStateHandle::Invalid, Priority);
 		return true;
 	}
 
@@ -1562,12 +1615,8 @@ bool FStateTreeExecutionContext::RequestTransition(const FStateTreeStateHandle N
 	FStateTreeActiveStates VisitedStates;
 	if (SelectState(NextState, NewActiveState, VisitedStates))
 	{
-		NextTransition.CurrentActiveStates = Exec.ActiveStates;
-		NextTransition.CurrentRunStatus = Exec.LastTickStatus;
-		NextTransition.SourceState = CurrentlyProcessedState;
-		NextTransition.TargetState = NextState;
+		SetupNextTransition(NextState, Priority);
 		NextTransition.NextActiveStates = NewActiveState;
-		NextTransition.Priority = Priority;
 
 		STATETREE_LOG(Verbose, TEXT("Transition on state '%s' -[%s]-> state '%s'"),
 			*GetSafeStateName(NextTransition.CurrentActiveStates.Last()), *GetSafeStateName(NextState), *GetSafeStateName(NextTransition.NextActiveStates.Last()));
@@ -1578,6 +1627,25 @@ bool FStateTreeExecutionContext::RequestTransition(const FStateTreeStateHandle N
 	return false;
 }
 
+void FStateTreeExecutionContext::SetupNextTransition(const FStateTreeStateHandle NextState, const EStateTreeTransitionPriority Priority)
+{
+	const FStateTreeExecutionState& Exec = GetExecState();
+
+	NextTransition.CurrentActiveStates = Exec.ActiveStates;
+	NextTransition.CurrentRunStatus = Exec.LastTickStatus;
+	NextTransition.SourceState = CurrentlyProcessedState;
+	NextTransition.TargetState = NextState;
+	if (NextState == FStateTreeStateHandle::Invalid)
+	{
+		NextTransition.NextActiveStates.Reset();
+	}
+	else
+	{
+		NextTransition.NextActiveStates = FStateTreeActiveStates(NextState);
+	}
+	NextTransition.Priority = Priority;
+}
+
 bool FStateTreeExecutionContext::TriggerTransitions()
 {
 	CSV_SCOPED_TIMING_STAT_EXCLUSIVE(StateTree_TriggerTransition);
@@ -1588,8 +1656,7 @@ bool FStateTreeExecutionContext::TriggerTransitions()
 
 	if (EventsToProcess.Num() > 0)
 	{
-		STATETREE_LOG(Verbose, TEXT("Trigger transitions with events [%s]"), *DebugGetEventsAsString());
-		STATETREE_TRACE_LOG_EVENT(TEXT("Trigger transitions with events [%s]"), *DebugGetEventsAsString());
+		STATETREE_LOG_AND_TRACE(Verbose, TEXT("Trigger transitions with events [%s]"), *DebugGetEventsAsString());
 	}
 
 	NextTransition.Reset();
@@ -1824,6 +1891,8 @@ bool FStateTreeExecutionContext::TriggerTransitions()
 	//
 	// Check state completion transitions.
 	//
+	bool bProcessSubTreeCompletion = true;
+
 	if (NextTransition.Priority == EStateTreeTransitionPriority::None
 		&& Exec.LastTickStatus != EStateTreeRunStatus::Running)
 	{
@@ -1884,18 +1953,26 @@ bool FStateTreeExecutionContext::TriggerTransitions()
 
 		if (NextTransition.Priority == EStateTreeTransitionPriority::None)
 		{
-			STATETREE_LOG(Verbose, TEXT("Could not trigger completion transition, jump back to start."));
-			STATETREE_TRACE_LOG_EVENT(TEXT("Could not trigger completion transition, jump back to start."));
+			STATETREE_LOG_AND_TRACE(Verbose, TEXT("Could not trigger completion transition, jump back to root state."));
 			FCurrentlyProcessedStateScope StateScope(*this, FStateTreeStateHandle::Root);
 			if (RequestTransition(FStateTreeStateHandle::Root, EStateTreeTransitionPriority::Normal))
 			{
 				NextTransitionSource = FStateTreeTransitionSource(EStateTreeTransitionSourceType::Internal, FStateTreeStateHandle::Root, EStateTreeTransitionPriority::Normal);
 			}
+			else
+			{
+				STATETREE_LOG_AND_TRACE(Warning, TEXT("Failed to select root state. Stopping the tree with failure."));
+
+				SetupNextTransition(FStateTreeStateHandle::Failed, EStateTreeTransitionPriority::Critical);
+
+				// In this case we don't want to complete subtrees, we want to force the whole tree to stop.
+				bProcessSubTreeCompletion = false;
+			}
 		}
 	}
 
 	// Check if the transition was succeed/failed, if we're on a sub-tree, complete the subtree instead of transition.
-	if (NextTransition.TargetState.IsCompletionState())
+	if (NextTransition.TargetState.IsCompletionState() && bProcessSubTreeCompletion)
 	{
 		const FStateTreeStateHandle ParentLinkedState = GetParentLinkedStateHandle(Exec.ActiveStates, NextTransition.SourceState);
 		if (ParentLinkedState.IsValid())
diff --git a/Engine/Plugins/Runtime/StateTree/Source/StateTreeModule/Public/StateTreeExecutionContext.h b/Engine/Plugins/Runtime/StateTree/Source/StateTreeModule/Public/StateTreeExecutionContext.h
--- a/Engine/Plugins/Runtime/StateTree/Source/StateTreeModule/Public/StateTreeExecutionContext.h
+++ b/Engine/Plugins/Runtime/StateTree/Source/StateTreeModule/Public/StateTreeExecutionContext.h
@@ -340,6 +340,11 @@ protected:
 	 */
 	bool RequestTransition(const FStateTreeStateHandle NextState, const EStateTreeTransitionPriority Priority);
 
+	/**
+	 * Sets up NextTransition based on the provided parameters and the current execution status. 
+	 */
+	void SetupNextTransition(const FStateTreeStateHandle NextState, const EStateTreeTransitionPriority Priority);
+
 	/**
 	 * Triggers transitions based on current run status. CurrentStatus is used to select which transitions events are triggered.
 	 * If CurrentStatus is "Running", "Conditional" transitions pass, "Completed/Failed" will trigger "OnCompleted/OnSucceeded/OnFailed" transitions.
diff --git a/Engine/Plugins/Runtime/StateTree/Source/StateTreeModule/Public/StateTreeExecutionTypes.h b/Engine/Plugins/Runtime/StateTree/Source/StateTreeModule/Public/StateTreeExecutionTypes.h
--- a/Engine/Plugins/Runtime/StateTree/Source/StateTreeModule/Public/StateTreeExecutionTypes.h
+++ b/Engine/Plugins/Runtime/StateTree/Source/StateTreeModule/Public/StateTreeExecutionTypes.h
@@ -555,6 +555,12 @@ struct STATETREEMODULE_API FStateTreeExecutionState
 	/** Running status of the instance */
 	EStateTreeRunStatus TreeRunStatus = EStateTreeRunStatus::Unset;
 
+	/** Completion status stored if Stop was called during the Tick and needed to be deferred. */
+	EStateTreeRunStatus RequestedStop = EStateTreeRunStatus::Unset;
+
+	/** Current update phase used to validate reentrant calls to the main entry points of the execution context (i.e. Start, Stop, Tick). */
+	EStateTreeUpdatePhase CurrentPhase = EStateTreeUpdatePhase::Unset;
+
 	/** Handle of the state that was first to report state completed (success or failure), used to trigger completion transitions. */
 	FStateTreeStateHandle CompletedStateHandle = FStateTreeStateHandle::Invalid;
 
diff --git a/Engine/Plugins/Runtime/StateTree/Source/StateTreeTestSuite/Private/StateTreeTest.cpp b/Engine/Plugins/Runtime/StateTree/Source/StateTreeTestSuite/Private/StateTreeTest.cpp
--- a/Engine/Plugins/Runtime/StateTree/Source/StateTreeTestSuite/Private/StateTreeTest.cpp
+++ b/Engine/Plugins/Runtime/StateTree/Source/StateTreeTestSuite/Private/StateTreeTest.cpp
@@ -2125,6 +2125,164 @@ struct FStateTreeTest_Stop_AlreadyStopped : FStateTreeTest_Stop_ExternalStop
 };
 IMPLEMENT_AI_INSTANT_TEST(FStateTreeTest_Stop_AlreadyStopped, "System.StateTree.Stop.AlreadyStopped");
 
+//
+// The deferred stop tests validates that the tree can be properly stopped if requested in the main entry points (Start, Tick, Stop).  
+//
+struct FStateTreeTest_DeferredStop : FAITestBase
+{
+	UStateTree& SetupTree() const
+	{
+		UStateTree& StateTree = UE::StateTree::Tests::NewStateTree(&GetWorld());
+		UStateTreeEditorData& EditorData = *Cast<UStateTreeEditorData>(StateTree.EditorData);
+
+		UStateTreeState& Root = EditorData.AddSubTree(FName(TEXT("Root")));
+		UStateTreeState& StateA = Root.AddChildState(FName(TEXT("A")));
+		TStateTreeEditorNode<FTestTask_StopTree>& TaskA = StateA.AddTask<FTestTask_StopTree>(TEXT("Task"));
+		TStateTreeEditorNode<FTestTask_StopTree>& GlobalTask = EditorData.AddGlobalTask<FTestTask_StopTree>(TEXT("GlobalTask"));
+
+		StateA.AddTransition(EStateTreeTransitionTrigger::OnStateSucceeded, EStateTreeTransitionType::Succeeded);
+		StateA.AddTransition(EStateTreeTransitionTrigger::OnStateFailed, EStateTreeTransitionType::Failed);
+
+		GlobalTask.GetNode().Phase = GlobalTaskPhase;
+		TaskA.GetNode().Phase = TaskPhase;
+
+		return StateTree;
+	}
+
+	virtual bool RunDerivedTest(FTestStateTreeExecutionContext& Exec) = 0;
+
+	virtual bool InstantTest() override
+	{
+		UStateTree& StateTree = SetupTree();
+
+		FStateTreeCompilerLog Log;
+		FStateTreeCompiler Compiler(Log);
+		const bool bResult = Compiler.Compile(StateTree);
+
+		AITEST_TRUE("StateTree should get compiled", bResult);
+
+		FStateTreeInstanceData InstanceData;
+		FTestStateTreeExecutionContext Exec(StateTree, StateTree, InstanceData);
+		const bool bInitSucceeded = Exec.IsValid();
+		AITEST_TRUE("StateTree should init", bInitSucceeded);
+
+		return RunDerivedTest(Exec);
+	}
+
+protected:
+
+	EStateTreeUpdatePhase GlobalTaskPhase = EStateTreeUpdatePhase::Unset;
+	EStateTreeUpdatePhase TaskPhase = EStateTreeUpdatePhase::Unset;
+};
+
+struct FStateTreeTest_DeferredStop_EnterGlobalTask : FStateTreeTest_DeferredStop
+{
+	FStateTreeTest_DeferredStop_EnterGlobalTask() { GlobalTaskPhase = EStateTreeUpdatePhase::EnterStates; }
+	virtual bool RunDerivedTest(FTestStateTreeExecutionContext& Exec) override
+	{
+		EStateTreeRunStatus Status = EStateTreeRunStatus::Unset;
+
+		Status = Exec.Start();
+		AITEST_EQUAL("Tree should be stopped", Status, EStateTreeRunStatus::Stopped);
+
+		return true;
+	}
+};
+IMPLEMENT_AI_INSTANT_TEST(FStateTreeTest_DeferredStop_EnterGlobalTask, "System.StateTree.DeferredStop.EnterGlobalTask");
+
+struct FStateTreeTest_DeferredStop_TickGlobalTask : FStateTreeTest_DeferredStop
+{
+	FStateTreeTest_DeferredStop_TickGlobalTask() { GlobalTaskPhase = EStateTreeUpdatePhase::TickStateTree; }
+	virtual bool RunDerivedTest(FTestStateTreeExecutionContext& Exec) override
+	{
+		EStateTreeRunStatus Status = EStateTreeRunStatus::Unset;
+
+		Status = Exec.Start();
+		AITEST_EQUAL("Tree should be running", Status, EStateTreeRunStatus::Running);
+
+		Status = Exec.Tick(0.1f);
+		AITEST_EQUAL("Tree should be stopped", Status, EStateTreeRunStatus::Stopped);
+
+		return true;
+	}
+};
+IMPLEMENT_AI_INSTANT_TEST(FStateTreeTest_DeferredStop_TickGlobalTask, "System.StateTree.DeferredStop.TickGlobalTask");
+
+struct FStateTreeTest_DeferredStop_ExitGlobalTask : FStateTreeTest_DeferredStop
+{
+	FStateTreeTest_DeferredStop_ExitGlobalTask() { GlobalTaskPhase = EStateTreeUpdatePhase::ExitStates; }
+	virtual bool RunDerivedTest(FTestStateTreeExecutionContext& Exec) override
+	{
+		EStateTreeRunStatus Status = EStateTreeRunStatus::Unset;
+
+		Status = Exec.Start();
+		AITEST_EQUAL("Tree should be running", Status, EStateTreeRunStatus::Running);
+
+		Status = Exec.Tick(0.1f);
+		AITEST_EQUAL("Tree should be running", Status, EStateTreeRunStatus::Running);
+
+		Status = Exec.Stop();
+		AITEST_EQUAL("Tree should be stopped", Status, EStateTreeRunStatus::Stopped);
+
+		return true;
+	}
+};
+IMPLEMENT_AI_INSTANT_TEST(FStateTreeTest_DeferredStop_ExitGlobalTask, "System.StateTree.DeferredStop.ExitGlobalTask");
+
+struct FStateTreeTest_DeferredStop_EnterTask : FStateTreeTest_DeferredStop
+{
+	FStateTreeTest_DeferredStop_EnterTask() { TaskPhase = EStateTreeUpdatePhase::EnterStates; }
+	virtual bool RunDerivedTest(FTestStateTreeExecutionContext& Exec) override
+	{
+		EStateTreeRunStatus Status = EStateTreeRunStatus::Unset;
+
+		Status = Exec.Start();
+		AITEST_EQUAL("Tree should be running", Status, EStateTreeRunStatus::Stopped);
+
+		return true;
+	}
+};
+IMPLEMENT_AI_INSTANT_TEST(FStateTreeTest_DeferredStop_EnterTask, "System.StateTree.DeferredStop.EnterTask");
+
+struct FStateTreeTest_DeferredStop_TickTask : FStateTreeTest_DeferredStop
+{
+	FStateTreeTest_DeferredStop_TickTask() { TaskPhase = EStateTreeUpdatePhase::TickStateTree; }
+	virtual bool RunDerivedTest(FTestStateTreeExecutionContext& Exec) override
+	{
+		EStateTreeRunStatus Status = EStateTreeRunStatus::Unset;
+
+		Status = Exec.Start();
+		AITEST_EQUAL("Tree should be running", Status, EStateTreeRunStatus::Running);
+
+		Status = Exec.Tick(0.1f);
+		AITEST_EQUAL("Tree should be stopped", Status, EStateTreeRunStatus::Stopped);
+
+		return true;
+	}
+};
+IMPLEMENT_AI_INSTANT_TEST(FStateTreeTest_DeferredStop_TickTask, "System.StateTree.DeferredStop.TickTask");
+
+struct FStateTreeTest_DeferredStop_ExitTask : FStateTreeTest_DeferredStop
+{
+	FStateTreeTest_DeferredStop_ExitTask() { TaskPhase = EStateTreeUpdatePhase::ExitStates; }
+	virtual bool RunDerivedTest(FTestStateTreeExecutionContext& Exec) override
+	{
+		EStateTreeRunStatus Status = EStateTreeRunStatus::Unset;
+
+		Status = Exec.Start();
+		AITEST_EQUAL("Tree should be running", Status, EStateTreeRunStatus::Running);
+
+		Status = Exec.Tick(0.1f);
+		AITEST_EQUAL("Tree should be running", Status, EStateTreeRunStatus::Running);
+
+		Status = Exec.Stop();
+		AITEST_EQUAL("Tree should be stopped", Status, EStateTreeRunStatus::Stopped);
+		
+		return true;
+	}
+};
+IMPLEMENT_AI_INSTANT_TEST(FStateTreeTest_DeferredStop_ExitTask, "System.StateTree.DeferredStop.ExitTask");
+
 
 UE_ENABLE_OPTIMIZATION_SHIP
 
diff --git a/Engine/Plugins/Runtime/StateTree/Source/StateTreeTestSuite/Private/StateTreeTestTypes.h b/Engine/Plugins/Runtime/StateTree/Source/StateTreeTestSuite/Private/StateTreeTestTypes.h
--- a/Engine/Plugins/Runtime/StateTree/Source/StateTreeTestSuite/Private/StateTreeTestTypes.h
+++ b/Engine/Plugins/Runtime/StateTree/Source/StateTreeTestSuite/Private/StateTreeTestTypes.h
@@ -199,6 +199,56 @@ struct FTestTask_PrintValue : public FStateTreeTaskBase
 };
 
 
+USTRUCT()
+struct FTestTask_StopTreeInstanceData
+{
+	GENERATED_BODY()
+};
+
+USTRUCT()
+struct FTestTask_StopTree : public FStateTreeTaskBase
+{
+	GENERATED_BODY()
+
+	using FInstanceDataType = FTestTask_PrintValueInstanceData;
+
+	FTestTask_StopTree() = default;
+	explicit FTestTask_StopTree(const FName InName) { Name = InName; }
+	virtual ~FTestTask_StopTree() override {}
+	
+	virtual const UStruct* GetInstanceDataType() const override { return FInstanceDataType::StaticStruct(); }
+
+	virtual EStateTreeRunStatus EnterState(FStateTreeExecutionContext& Context, const FStateTreeTransitionResult& Transition) const override
+	{
+		if (Phase == EStateTreeUpdatePhase::EnterStates)
+		{
+			return Context.Stop();
+		}
+		return EStateTreeRunStatus::Running;
+	}
+
+	virtual void ExitState(FStateTreeExecutionContext& Context, const FStateTreeTransitionResult& Transition) const override
+	{
+		if (Phase == EStateTreeUpdatePhase::ExitStates)
+		{
+			Context.Stop();
+		}
+	}
+
+	virtual EStateTreeRunStatus Tick(FStateTreeExecutionContext& Context, const float DeltaTime) const override
+	{
+		if (Phase == EStateTreeUpdatePhase::TickStateTree)
+		{
+			return Context.Stop();
+		}
+		return EStateTreeRunStatus::Running;
+	};
+
+	/** Indicates in which phase the call to Stop should be performed. Possible values are EnterStates, ExitStats and TickStateTree */
+	UPROPERTY()
+	EStateTreeUpdatePhase Phase = EStateTreeUpdatePhase::Unset;
+};
+
 USTRUCT()
 struct FTestTask_StandInstanceData
 {
