[StateTreeDebugger] removed ensure that was for local debugging only
#rnx

diff --git a/Engine/Plugins/Runtime/StateTree/Source/StateTreeModule/Private/Debugger/StateTreeDebugger.cpp b/Engine/Plugins/Runtime/StateTree/Source/StateTreeModule/Private/Debugger/StateTreeDebugger.cpp
--- a/Engine/Plugins/Runtime/StateTree/Source/StateTreeModule/Private/Debugger/StateTreeDebugger.cpp
+++ b/Engine/Plugins/Runtime/StateTree/Source/StateTreeModule/Private/Debugger/StateTreeDebugger.cpp
@@ -278,7 +278,6 @@ bool FStateTreeDebugger::TryStartNewLiveSessionAnalysis(const float RetryPolling
 	}
 	
 	RetryLoadNextLiveSessionTimer = RetryPollingDuration;
-	ensure(RetryLoadNextLiveSessionTimer > 0);
 	UE_CLOG(RetryLoadNextLiveSessionTimer > 0, LogStateTree, Log, TEXT("Unable to start analysis for the most recent live session."));
 
 	return false;
