State Tree:
- Added Themes, which change the coloring of both States and Tasks (tasks showing up a bit darker)

diff --git a/Engine/Plugins/Runtime/StateTree/Source/StateTreeEditorModule/Private/Customizations/StateTreeEditorColorDetails.cpp b/Engine/Plugins/Runtime/StateTree/Source/StateTreeEditorModule/Private/Customizations/StateTreeEditorColorDetails.cpp
new file mode 100644
--- /dev/null
+++ b/Engine/Plugins/Runtime/StateTree/Source/StateTreeEditorModule/Private/Customizations/StateTreeEditorColorDetails.cpp
@@ -0,0 +1,195 @@
+// Copyright Epic Games, Inc. All Rights Reserved.
+
+#include "StateTreeEditorColorDetails.h"
+#include "DetailLayoutBuilder.h"
+#include "DetailWidgetRow.h"
+#include "Engine/Engine.h"
+#include "ScopedTransaction.h"
+#include "StateTreeEditorData.h"
+#include "Widgets/Colors/SColorBlock.h"
+#include "Widgets/Colors/SColorPicker.h"
+#include "Widgets/Input/SButton.h"
+#include "Widgets/SStateTreeEditorColorComboBox.h"
+#include "Widgets/Text/STextBlock.h"
+
+#define LOCTEXT_NAMESPACE "StateTreeEditor"
+
+//----------------------------------------------------------------//
+// FStateTreeEditorColorRefDetails
+//----------------------------------------------------------------//
+
+TSharedRef<IPropertyTypeCustomization> FStateTreeEditorColorRefDetails::MakeInstance()
+{
+	return MakeShared<FStateTreeEditorColorRefDetails>();
+}
+
+void FStateTreeEditorColorRefDetails::CustomizeHeader(TSharedRef<IPropertyHandle> InStructPropertyHandle, FDetailWidgetRow& HeaderRow, IPropertyTypeCustomizationUtils& StructCustomizationUtils)
+{
+	UStateTreeEditorData* EditorData = GetEditorData(InStructPropertyHandle);
+	if (!EditorData)
+	{
+		return;
+	}
+
+	HeaderRow
+		.NameContent()
+		[
+			InStructPropertyHandle->CreatePropertyNameWidget()
+		]
+		.ValueContent()
+		.MinDesiredWidth(200.0f)
+		[
+			SNew(SStateTreeEditorColorComboBox, InStructPropertyHandle, EditorData)
+		];
+}
+
+UStateTreeEditorData* FStateTreeEditorColorRefDetails::GetEditorData(const TSharedRef<IPropertyHandle>& PropertyHandle) const
+{
+	TArray<UObject*> Objects;
+	PropertyHandle->GetOuterObjects(Objects);
+
+	for (UObject* Outer : Objects)
+	{
+		if (UStateTreeEditorData* EditorData = Cast<UStateTreeEditorData>(Outer))
+		{
+			return EditorData;
+		}
+
+		if (UStateTreeEditorData* EditorData = Outer->GetTypedOuter<UStateTreeEditorData>())
+		{
+			return EditorData;
+		}
+	}
+
+	return nullptr;
+}
+
+//----------------------------------------------------------------//
+// FStateTreeEditorColorDetails
+//----------------------------------------------------------------//
+
+TSharedRef<IPropertyTypeCustomization> FStateTreeEditorColorDetails::MakeInstance()
+{
+	return MakeShared<FStateTreeEditorColorDetails>();
+}
+
+void FStateTreeEditorColorDetails::CustomizeHeader(TSharedRef<IPropertyHandle> InStructPropertyHandle, FDetailWidgetRow& HeaderRow, IPropertyTypeCustomizationUtils& StructCustomizationUtils)
+{
+	StructPropertyHandle = InStructPropertyHandle;
+
+	ColorPropertyHandle = StructPropertyHandle->GetChildHandle(GET_MEMBER_NAME_CHECKED(FStateTreeEditorColor, Color));
+
+	TSharedPtr<IPropertyHandle> NamePropertyHandle = StructPropertyHandle->GetChildHandle(GET_MEMBER_NAME_CHECKED(FStateTreeEditorColor, DisplayName));
+
+	HeaderRow
+		.WholeRowContent()
+		.HAlign(HAlign_Fill)
+		.MinDesiredWidth(200.f)
+		[
+			SNew(SHorizontalBox)
+			+ SHorizontalBox::Slot()
+			.AutoWidth()
+			.MaxWidth(25)
+			.VAlign(VAlign_Center)
+			[
+				SAssignNew(ColorButtonWidget, SButton)
+				.ButtonStyle(FAppStyle::Get(), "SimpleButton")
+				.OnClicked(this, &FStateTreeEditorColorDetails::OnColorButtonClicked)
+				.ContentPadding(2.f)
+				[
+					SNew(SColorBlock)
+					.Color(this, &FStateTreeEditorColorDetails::GetColor)
+					.Size(FVector2D(16.0))
+				]
+			]
+			+ SHorizontalBox::Slot()
+			.Padding(FMargin(5.f, 0.f, 0.f, 0.f))
+			.VAlign(VAlign_Center)
+			[
+				NamePropertyHandle->CreatePropertyValueWidget(/*bDisplayDefaultPropertyButtons*/false)
+			]
+			+ SHorizontalBox::Slot()
+			.AutoWidth()
+			.HAlign(HAlign_Right)
+			[
+				StructPropertyHandle->CreateDefaultPropertyButtonWidgets()
+			]
+		];
+}
+
+FLinearColor* FStateTreeEditorColorDetails::GetColorPtr() const
+{
+	void* ValueData = nullptr;
+	if (ColorPropertyHandle->GetValueData(ValueData) == FPropertyAccess::Result::Success)
+	{
+		return static_cast<FLinearColor*>(ValueData);
+	}
+	return nullptr;
+}
+
+FLinearColor FStateTreeEditorColorDetails::GetColor() const
+{
+	if (FLinearColor* ColorPtr = GetColorPtr())
+	{
+		return *ColorPtr;
+	}
+	return FLinearColor(ForceInitToZero);
+}
+
+void FStateTreeEditorColorDetails::SetColor(FLinearColor Color)
+{
+	if (FLinearColor* ColorPtr = GetColorPtr())
+	{
+		ColorPropertyHandle->NotifyPreChange();
+		*ColorPtr = Color;
+		ColorPropertyHandle->NotifyPostChange(EPropertyChangeType::ValueSet);
+	}
+}
+
+void FStateTreeEditorColorDetails::OnColorCommitted(FLinearColor Color)
+{
+	SetColor(Color);
+
+	// End Transaction
+	ColorPickerTransaction.Reset();
+}
+
+void FStateTreeEditorColorDetails::OnColorCancelled(FLinearColor Color)
+{
+	SetColor(Color);
+
+	// Cancel Transaction
+	if (ColorPickerTransaction.IsValid())
+	{
+		ColorPickerTransaction->Cancel();
+		ColorPickerTransaction.Reset();
+	}
+}
+
+FReply FStateTreeEditorColorDetails::OnColorButtonClicked()
+{
+	CreateColorPickerWindow();
+	return FReply::Handled();
+}
+
+void FStateTreeEditorColorDetails::CreateColorPickerWindow()
+{
+	FColorPickerArgs PickerArgs;
+
+	// Begin Transaction
+	ColorPickerTransaction = MakeShared<FScopedTransaction>(LOCTEXT("SetStateTreeColorProperty", "Set Color Property"));
+
+	PickerArgs.bOnlyRefreshOnMouseUp = false;
+	PickerArgs.ParentWidget = ColorButtonWidget;
+	PickerArgs.bUseAlpha = false;
+	PickerArgs.bOnlyRefreshOnOk = false;
+	PickerArgs.bClampValue = false; // Linear Color
+	PickerArgs.DisplayGamma = TAttribute<float>::Create(TAttribute<float>::FGetter::CreateUObject(GEngine, &UEngine::GetDisplayGamma));
+	PickerArgs.OnColorCommitted = FOnLinearColorValueChanged::CreateSP(this, &FStateTreeEditorColorDetails::OnColorCommitted);
+	PickerArgs.OnColorPickerCancelled = FOnColorPickerCancelled::CreateSP(this, &FStateTreeEditorColorDetails::OnColorCancelled);
+	PickerArgs.InitialColor = GetColor();
+
+	OpenColorPicker(PickerArgs);
+}
+
+#undef LOCTEXT_NAMESPACE
diff --git a/Engine/Plugins/Runtime/StateTree/Source/StateTreeEditorModule/Private/Customizations/StateTreeEditorColorDetails.h b/Engine/Plugins/Runtime/StateTree/Source/StateTreeEditorModule/Private/Customizations/StateTreeEditorColorDetails.h
new file mode 100644
--- /dev/null
+++ b/Engine/Plugins/Runtime/StateTree/Source/StateTreeEditorModule/Private/Customizations/StateTreeEditorColorDetails.h
@@ -0,0 +1,59 @@
+// Copyright Epic Games, Inc. All Rights Reserved.
+
+#pragma once
+
+#include "IPropertyTypeCustomization.h"
+#include "Math/Color.h"
+
+class UStateTreeEditorData;
+class FReply;
+class FScopedTransaction;
+class SWidget;
+
+/** Property Type Customization for FStateTreeEditorColorRef */
+class FStateTreeEditorColorRefDetails : public IPropertyTypeCustomization
+{
+public:
+	/** Makes a new instance of this detail layout class for a specific detail view requesting it */
+	static TSharedRef<IPropertyTypeCustomization> MakeInstance();
+
+	/** IPropertyTypeCustomization interface */
+	virtual void CustomizeHeader(TSharedRef<IPropertyHandle> InStructPropertyHandle, FDetailWidgetRow& HeaderRow, IPropertyTypeCustomizationUtils& StructCustomizationUtils) override;
+	virtual void CustomizeChildren(TSharedRef<IPropertyHandle> InStructPropertyHandle, IDetailChildrenBuilder& StructBuilder, IPropertyTypeCustomizationUtils& StructCustomizationUtils) override {}
+
+private:
+	UStateTreeEditorData* GetEditorData(const TSharedRef<IPropertyHandle>& PropertyHandle) const;
+};
+
+/** Property Type Customization for FStateTreeEditorColor */
+class FStateTreeEditorColorDetails : public IPropertyTypeCustomization
+{
+public:
+	/** Makes a new instance of this detail layout class for a specific detail view requesting it */
+	static TSharedRef<IPropertyTypeCustomization> MakeInstance();
+
+	/** IPropertyTypeCustomization interface */
+	virtual void CustomizeHeader(TSharedRef<IPropertyHandle> InStructPropertyHandle, FDetailWidgetRow& HeaderRow, IPropertyTypeCustomizationUtils& StructCustomizationUtils) override;
+	virtual void CustomizeChildren(TSharedRef<IPropertyHandle> InStructPropertyHandle, IDetailChildrenBuilder& StructBuilder, IPropertyTypeCustomizationUtils& StructCustomizationUtils) override {}
+
+private:
+	FLinearColor* GetColorPtr() const;
+
+	FLinearColor GetColor() const;
+	void SetColor(FLinearColor Color);
+
+	void OnColorCommitted(FLinearColor Color);
+
+	void OnColorCancelled(FLinearColor Color);
+
+	FReply OnColorButtonClicked();
+
+	void CreateColorPickerWindow();
+
+	/** Current Transaction taking place when selecting a color from the picker */
+	TSharedPtr<FScopedTransaction> ColorPickerTransaction;
+
+	TSharedPtr<IPropertyHandle> StructPropertyHandle;
+	TSharedPtr<IPropertyHandle> ColorPropertyHandle;
+	TSharedPtr<SWidget> ColorButtonWidget;
+};
diff --git a/Engine/Plugins/Runtime/StateTree/Source/StateTreeEditorModule/Private/Customizations/StateTreeEditorDataDetails.cpp b/Engine/Plugins/Runtime/StateTree/Source/StateTreeEditorModule/Private/Customizations/StateTreeEditorDataDetails.cpp
--- a/Engine/Plugins/Runtime/StateTree/Source/StateTreeEditorModule/Private/Customizations/StateTreeEditorDataDetails.cpp
+++ b/Engine/Plugins/Runtime/StateTree/Source/StateTreeEditorModule/Private/Customizations/StateTreeEditorDataDetails.cpp
@@ -47,6 +47,10 @@ void FStateTreeEditorDataDetails::CustomizeDetails(IDetailLayoutBuilder& DetailB
 	IDetailCategoryBuilder& ContextDataCategory = DetailBuilder.EditCategory(TEXT("Context"), LOCTEXT("EditorDataDetailsContext", "Context"));
 	ContextDataCategory.SetSortOrder(1);
 
+	// Theme category
+	IDetailCategoryBuilder& ThemeCategory = DetailBuilder.EditCategory(TEXT("Theme"));
+	ThemeCategory.InitiallyCollapsed(true);
+
 	if (Schema != nullptr)
 	{
 		for (const FStateTreeExternalDataDesc& ContextData : Schema->GetContextDataDescs())
diff --git a/Engine/Plugins/Runtime/StateTree/Source/StateTreeEditorModule/Private/Customizations/Widgets/SStateTreeEditorColorComboBox.cpp b/Engine/Plugins/Runtime/StateTree/Source/StateTreeEditorModule/Private/Customizations/Widgets/SStateTreeEditorColorComboBox.cpp
new file mode 100644
--- /dev/null
+++ b/Engine/Plugins/Runtime/StateTree/Source/StateTreeEditorModule/Private/Customizations/Widgets/SStateTreeEditorColorComboBox.cpp
@@ -0,0 +1,155 @@
+// Copyright Epic Games, Inc. All Rights Reserved.
+
+#include "SStateTreeEditorColorComboBox.h"
+#include "DetailLayoutBuilder.h"
+#include "GuidStructCustomization.h"
+#include "StateTreeEditorData.h"
+#include "Widgets/Colors/SColorBlock.h"
+#include "Widgets/Input/SComboBox.h"
+#include "Widgets/Layout/SScaleBox.h"
+#include "Widgets/Text/STextBlock.h"
+
+#define LOCTEXT_NAMESPACE "StateTreeEditor"
+
+void SStateTreeEditorColorComboBox::Construct(const FArguments& InArgs, TSharedPtr<IPropertyHandle> InColorRefHandle, UStateTreeEditorData* InEditorData)
+{
+	WeakEditorData = InEditorData;
+	ColorRefHandle = InColorRefHandle;
+	ColorIDHandle = InColorRefHandle->GetChildHandle(GET_MEMBER_NAME_CHECKED(FStateTreeEditorColorRef, ID));
+
+	ChildSlot
+	[
+		SAssignNew(ColorComboBox, SComboBox<TSharedPtr<FStateTreeEditorColorRef>>)
+		.OptionsSource(&ColorRefOptions)
+		.OnComboBoxOpening(this, &SStateTreeEditorColorComboBox::RefreshColorOptions)
+		.OnGenerateWidget(this, &SStateTreeEditorColorComboBox::GenerateColorOptionWidget)
+		.OnSelectionChanged(this, &SStateTreeEditorColorComboBox::OnSelectionChanged)
+		[
+			SAssignNew(SelectedColorBorder, SBorder)
+			.BorderImage(FAppStyle::GetBrush("NoBorder"))
+		]
+	];
+
+	UpdatedSelectedColorWidget();
+}
+
+const FStateTreeEditorColor* SStateTreeEditorColorComboBox::FindColorEntry(const FStateTreeEditorColorRef& ColorRef) const
+{
+	if (UStateTreeEditorData* EditorData = WeakEditorData.Get())
+	{
+		return EditorData->FindColor(FStateTreeEditorColorRef(ColorRef));
+	}
+	return nullptr;
+}
+
+FText SStateTreeEditorColorComboBox::GetDisplayName(FStateTreeEditorColorRef ColorRef) const
+{
+	if (const FStateTreeEditorColor* ColorEntry = FindColorEntry(ColorRef))
+	{
+		return FText::FromString(ColorEntry->DisplayName);
+	}
+	return FText::GetEmpty();
+}
+
+FLinearColor SStateTreeEditorColorComboBox::GetColor(FStateTreeEditorColorRef ColorRef) const
+{
+	if (const FStateTreeEditorColor* ColorEntry = FindColorEntry(ColorRef))
+	{
+		return ColorEntry->Color;
+	}
+	return FLinearColor(ForceInitToZero);
+}
+
+TSharedRef<SWidget> SStateTreeEditorColorComboBox::GenerateColorOptionWidget(TSharedPtr<FStateTreeEditorColorRef> ColorRef)
+{
+	return SNew(SBox)
+		.WidthOverride(120.f)
+		.Padding(5.f, 1.f)
+		[
+			GenerateColorWidget(*ColorRef)
+		];
+}
+
+TSharedRef<SWidget> SStateTreeEditorColorComboBox::GenerateColorWidget(const FStateTreeEditorColorRef& ColorRef)
+{
+	return SNew(SHorizontalBox)
+		+ SHorizontalBox::Slot()
+		.AutoWidth()
+		.MaxWidth(25)
+		.VAlign(VAlign_Center)
+		[
+			SNew(SColorBlock)
+			.Color(this, &SStateTreeEditorColorComboBox::GetColor, ColorRef)
+			.Size(FVector2D(16.0))
+		]
+		+ SHorizontalBox::Slot()
+		.Padding(FMargin(5.f, 0.f, 0.f, 0.f))
+		.VAlign(VAlign_Center)
+		[
+			SNew(STextBlock)
+			.Text(this, &SStateTreeEditorColorComboBox::GetDisplayName, ColorRef)
+			.Font(IDetailLayoutBuilder::GetDetailFont())
+		];
+}
+
+void SStateTreeEditorColorComboBox::RefreshColorOptions()
+{
+	ColorRefOptions.Reset();
+
+	UStateTreeEditorData* EditorData = WeakEditorData.Get();
+	if (!EditorData)
+	{
+		return;
+	}
+
+	for (const FStateTreeEditorColor& Color : EditorData->Colors)
+	{
+		TSharedRef<FStateTreeEditorColorRef> ColorRefOption = MakeShared<FStateTreeEditorColorRef>(Color.ColorRef);
+		ColorRefOptions.Add(ColorRefOption);
+		if (SelectedColorRef == Color.ColorRef)
+		{
+			ColorComboBox->SetSelectedItem(ColorRefOption);
+		}
+	}
+}
+
+void SStateTreeEditorColorComboBox::OnSelectionChanged(TSharedPtr<FStateTreeEditorColorRef> SelectedColorRefOption, ESelectInfo::Type SelectInfo)
+{
+	if (SelectedColorRefOption.IsValid())
+	{
+		WriteGuidToProperty(ColorIDHandle, SelectedColorRefOption->ID);
+		UpdatedSelectedColorWidget();
+	}
+}
+
+void SStateTreeEditorColorComboBox::UpdatedSelectedColorWidget()
+{
+	check(SelectedColorBorder.IsValid());
+
+	TArray<const void*> RawData;
+	ColorRefHandle->AccessRawData(RawData);
+
+	if (RawData.IsEmpty())
+	{
+		SelectedColorBorder->ClearContent();
+		return;
+	}
+
+	SelectedColorRef = *static_cast<const FStateTreeEditorColorRef*>(RawData.Pop(/*bAllowShrinking*/false));
+
+	// Make "Multiple Values" content if there's at least one Color ID that differs
+	for (const void* ColorRefRaw : RawData)
+	{
+		if (SelectedColorRef != *static_cast<const FStateTreeEditorColorRef*>(ColorRefRaw))
+		{
+			SelectedColorBorder->SetContent(SNew(STextBlock)
+				.Text(LOCTEXT("MultipleValues", "Multiple Values"))
+				.Font(IDetailLayoutBuilder::GetDetailFont()));
+			return;
+		}
+	}
+
+	SelectedColorBorder->SetContent(GenerateColorWidget(SelectedColorRef));
+}
+
+#undef LOCTEXT_NAMESPACE
diff --git a/Engine/Plugins/Runtime/StateTree/Source/StateTreeEditorModule/Private/Customizations/Widgets/SStateTreeEditorColorComboBox.h b/Engine/Plugins/Runtime/StateTree/Source/StateTreeEditorModule/Private/Customizations/Widgets/SStateTreeEditorColorComboBox.h
new file mode 100644
--- /dev/null
+++ b/Engine/Plugins/Runtime/StateTree/Source/StateTreeEditorModule/Private/Customizations/Widgets/SStateTreeEditorColorComboBox.h
@@ -0,0 +1,53 @@
+// Copyright Epic Games, Inc. All Rights Reserved.
+
+#pragma once
+
+#include "StateTreeEditorTypes.h"
+#include "Containers/Array.h"
+#include "Templates/SharedPointer.h"
+#include "Types/SlateEnums.h"
+#include "UObject/WeakObjectPtr.h"
+#include "Widgets/SCompoundWidget.h"
+
+class IPropertyHandle;
+class SBorder;
+class UStateTreeEditorData;
+struct FStateTreeEditorColor;
+struct FStateTreeEditorColorRef;
+template<typename OptionType> class SComboBox;
+
+class SStateTreeEditorColorComboBox : public SCompoundWidget
+{
+public:
+	SLATE_BEGIN_ARGS(SStateTreeEditorColorComboBox){}
+	SLATE_END_ARGS()
+
+	void Construct(const FArguments& InArgs, TSharedPtr<IPropertyHandle> InColorRefHandle, UStateTreeEditorData* InEditorData);
+
+private:
+	const FStateTreeEditorColor* FindColorEntry(const FStateTreeEditorColorRef& ColorRef) const;
+
+	FText GetDisplayName(FStateTreeEditorColorRef ColorRef) const;
+	FLinearColor GetColor(FStateTreeEditorColorRef ColorRef) const;
+
+	TSharedRef<SWidget> GenerateColorOptionWidget(TSharedPtr<FStateTreeEditorColorRef> ColorRef);
+	TSharedRef<SWidget> GenerateColorWidget(const FStateTreeEditorColorRef& ColorRef);
+
+	void RefreshColorOptions();
+
+	void OnSelectionChanged(TSharedPtr<FStateTreeEditorColorRef> SelectedColorRefOption, ESelectInfo::Type SelectInfo);
+
+	void UpdatedSelectedColorWidget();
+
+	TWeakObjectPtr<UStateTreeEditorData> WeakEditorData;
+
+	TSharedPtr<IPropertyHandle> ColorRefHandle;
+	TSharedPtr<IPropertyHandle> ColorIDHandle;
+
+	TSharedPtr<SBorder> SelectedColorBorder;
+	TSharedPtr<SComboBox<TSharedPtr<FStateTreeEditorColorRef>>> ColorComboBox;
+
+	TArray<TSharedPtr<FStateTreeEditorColorRef>> ColorRefOptions;
+
+	FStateTreeEditorColorRef SelectedColorRef;
+};
diff --git a/Engine/Plugins/Runtime/StateTree/Source/StateTreeEditorModule/Private/SStateTreeViewRow.cpp b/Engine/Plugins/Runtime/StateTree/Source/StateTreeEditorModule/Private/SStateTreeViewRow.cpp
--- a/Engine/Plugins/Runtime/StateTree/Source/StateTreeEditorModule/Private/SStateTreeViewRow.cpp
+++ b/Engine/Plugins/Runtime/StateTree/Source/StateTreeEditorModule/Private/SStateTreeViewRow.cpp
@@ -20,6 +20,20 @@
 
 #define LOCTEXT_NAMESPACE "StateTreeEditor"
 
+namespace UE::StateTree::Editor
+{
+	FLinearColor LerpColorSRGB(const FLinearColor ColorA, FLinearColor ColorB, float T)
+	{
+		const FColor A = ColorA.ToFColorSRGB();
+		const FColor B = ColorB.ToFColorSRGB();
+		return FLinearColor(FColor(
+			static_cast<uint8>(FMath::RoundToInt(static_cast<float>(A.R) * (1.f - T) + static_cast<float>(B.R) * T)),
+			static_cast<uint8>(FMath::RoundToInt(static_cast<float>(A.G) * (1.f - T) + static_cast<float>(B.G) * T)),
+			static_cast<uint8>(FMath::RoundToInt(static_cast<float>(A.B) * (1.f - T) + static_cast<float>(B.B) * T)),
+			static_cast<uint8>(FMath::RoundToInt(static_cast<float>(A.A) * (1.f - T) + static_cast<float>(B.A) * T))));
+	}
+} // UE:StateTree::Editor
+
 void SStateTreeViewRow::Construct(const FArguments& InArgs, const TSharedRef<STableViewBase>& InOwnerTableView, TWeakObjectPtr<UStateTreeState> InState, const TSharedPtr<SScrollBox>& ViewBox, TSharedRef<FStateTreeViewModel> InStateTreeViewModel)
 {
 	StateTreeViewModel = InStateTreeViewModel;
@@ -607,6 +621,7 @@ TSharedRef<SHorizontalBox> SStateTreeViewRow::CreateTasksWidget()
 					SNew(SBorder)
 					.VAlign(VAlign_Center)
 					.BorderImage(FStateTreeEditorStyle::Get().GetBrush("StateTree.Task.Rect"))
+					.BorderBackgroundColor(this, &SStateTreeViewRow::GetTitleColor)
 					.Padding(0)
 					.IsEnabled_Lambda(IsTaskEnabledFunc)
 					[
@@ -658,11 +673,18 @@ void SStateTreeViewRow::RequestRename() const
 
 FSlateColor SStateTreeViewRow::GetTitleColor() const
 {
-	if (const UStateTreeState* State = WeakState.Get())
+	const UStateTreeState* State = WeakState.Get();
+	const UStateTreeEditorData* EditorData = WeakTreeData.Get();
+
+	if (State != nullptr && EditorData != nullptr)
 	{
-		if (IsRootState() || State->Type == EStateTreeStateType::Subtree)
+		if (const FStateTreeEditorColor* FoundColor = EditorData->FindColor(State->ColorRef))
 		{
-			return FLinearColor(FColor(17, 117, 131));
+			if (IsRootState() || State->Type == EStateTreeStateType::Subtree)
+			{
+				return UE::StateTree::Editor::LerpColorSRGB(FoundColor->Color, FColor::Black, 0.25f);
+			}
+			return FoundColor->Color;
 		}
 	}
 
@@ -693,7 +715,8 @@ FSlateColor SStateTreeViewRow::GetSubTreeMarkerColor() const
 	{
 		if (IsRootState() || State->Type == EStateTreeStateType::Subtree)
 		{
-			return FLinearColor(FColor(136, 186, 193));
+			const FSlateColor TitleColor = GetTitleColor();
+			return UE::StateTree::Editor::LerpColorSRGB(TitleColor.GetSpecifiedColor(), FLinearColor::White, 0.2f);
 		}
 	}
 
diff --git a/Engine/Plugins/Runtime/StateTree/Source/StateTreeEditorModule/Private/StateTreeEditorData.cpp b/Engine/Plugins/Runtime/StateTree/Source/StateTreeEditorModule/Private/StateTreeEditorData.cpp
--- a/Engine/Plugins/Runtime/StateTree/Source/StateTreeEditorModule/Private/StateTreeEditorData.cpp
+++ b/Engine/Plugins/Runtime/StateTree/Source/StateTreeEditorModule/Private/StateTreeEditorData.cpp
@@ -16,6 +16,16 @@
 
 #include UE_INLINE_GENERATED_CPP_BY_NAME(StateTreeEditorData)
 
+UStateTreeEditorData::UStateTreeEditorData()
+{
+	FStateTreeEditorColor DefaultColor;
+	DefaultColor.ColorRef = FStateTreeEditorColorRef();
+	DefaultColor.Color = FLinearColor(FColor(31, 151, 167));
+	DefaultColor.DisplayName = TEXT("Default Color");
+
+	Colors.Add(MoveTemp(DefaultColor));
+}
+
 void UStateTreeEditorData::PostInitProperties()
 {
 	Super::PostInitProperties();
diff --git a/Engine/Plugins/Runtime/StateTree/Source/StateTreeEditorModule/Private/StateTreeEditorModule.cpp b/Engine/Plugins/Runtime/StateTree/Source/StateTreeEditorModule/Private/StateTreeEditorModule.cpp
--- a/Engine/Plugins/Runtime/StateTree/Source/StateTreeEditorModule/Private/StateTreeEditorModule.cpp
+++ b/Engine/Plugins/Runtime/StateTree/Source/StateTreeEditorModule/Private/StateTreeEditorModule.cpp
@@ -5,6 +5,7 @@
 #include "Blueprint/StateTreeEvaluatorBlueprintBase.h"
 #include "Blueprint/StateTreeTaskBlueprintBase.h"
 #include "Customizations/StateTreeAnyEnumDetails.h"
+#include "Customizations/StateTreeEditorColorDetails.h"
 #include "Customizations/StateTreeEditorDataDetails.h"
 #include "Customizations/StateTreeEditorNodeDetails.h"
 #include "Customizations/StateTreeReferenceDetails.h"
@@ -88,6 +89,8 @@ void FStateTreeEditorModule::StartupModule()
 	PropertyModule.RegisterCustomPropertyTypeLayout("StateTreeStateParameters", FOnGetPropertyTypeCustomizationInstance::CreateStatic(&FStateTreeStateParametersDetails::MakeInstance));
 	PropertyModule.RegisterCustomPropertyTypeLayout("StateTreeAnyEnum", FOnGetPropertyTypeCustomizationInstance::CreateStatic(&FStateTreeAnyEnumDetails::MakeInstance));
 	PropertyModule.RegisterCustomPropertyTypeLayout("StateTreeReference", FOnGetPropertyTypeCustomizationInstance::CreateStatic(&FStateTreeReferenceDetails::MakeInstance));
+	PropertyModule.RegisterCustomPropertyTypeLayout("StateTreeEditorColorRef", FOnGetPropertyTypeCustomizationInstance::CreateStatic(&FStateTreeEditorColorRefDetails::MakeInstance));
+	PropertyModule.RegisterCustomPropertyTypeLayout("StateTreeEditorColor", FOnGetPropertyTypeCustomizationInstance::CreateStatic(&FStateTreeEditorColorDetails::MakeInstance));
 	PropertyModule.RegisterCustomClassLayout("StateTreeState", FOnGetDetailCustomizationInstance::CreateStatic(&FStateTreeStateDetails::MakeInstance));
 	PropertyModule.RegisterCustomClassLayout("StateTreeEditorData", FOnGetDetailCustomizationInstance::CreateStatic(&FStateTreeEditorDataDetails::MakeInstance));
 
diff --git a/Engine/Plugins/Runtime/StateTree/Source/StateTreeEditorModule/Private/StateTreeEditorStyle.cpp b/Engine/Plugins/Runtime/StateTree/Source/StateTreeEditorModule/Private/StateTreeEditorStyle.cpp
--- a/Engine/Plugins/Runtime/StateTree/Source/StateTreeEditorModule/Private/StateTreeEditorStyle.cpp
+++ b/Engine/Plugins/Runtime/StateTree/Source/StateTreeEditorModule/Private/StateTreeEditorStyle.cpp
@@ -111,7 +111,8 @@ FStateTreeEditorStyle::FStateTreeEditorStyle()
 			.SetColorAndOpacity(FLinearColor(230.0f / 255.0f, 230.0f / 255.0f, 230.0f / 255.0f, 0.85f));
 		Set("StateTree.Task.Title", TaskTitle);
 
-		Set("StateTree.Task.Rect", new FSlateRoundedBoxBrush(FLinearColor(FColor(21, 104, 116)), 0.0f, FLinearColor(FColor(27, 138, 154)), 1.0f));
+		// Tasks to be show up a bit darker than the state
+		Set("StateTree.Task.Rect", new FSlateColorBrush(FLinearColor(FVector3f(0.67f))));
 	}
 	
 	// Details rich text
diff --git a/Engine/Plugins/Runtime/StateTree/Source/StateTreeEditorModule/Private/StateTreeEditorTypes.cpp b/Engine/Plugins/Runtime/StateTree/Source/StateTreeEditorModule/Private/StateTreeEditorTypes.cpp
new file mode 100644
--- /dev/null
+++ b/Engine/Plugins/Runtime/StateTree/Source/StateTreeEditorModule/Private/StateTreeEditorTypes.cpp
@@ -0,0 +1,56 @@
+// Copyright Epic Games, Inc. All Rights Reserved.
+
+#include "StateTreeEditorTypes.h"
+
+bool FStateTreeEditorColor::ExportTextItem(FString& OutValueString, const FStateTreeEditorColor& DefaultValue, UObject* OwnerObject, int32 PortFlags, UObject* ExportRootScope) const
+{
+#if WITH_EDITORONLY_DATA
+	// Only go through this Export path if Copy or Duplicate 
+	if ((PortFlags & (PPF_Copy | PPF_Duplicate)) == 0)
+	{
+		return false;
+	}
+
+	uint32 Count = 0;
+
+	for (FProperty* Property : TFieldRange<FProperty>(FStateTreeEditorColor::StaticStruct()))
+	{
+		if (!Property->ShouldPort(PortFlags) || Property->HasMetaData(TEXT("StructExportTransient")))
+		{
+			continue;
+		}
+
+		for (int32 Index = 0; Index < Property->ArrayDim; Index++)
+		{
+			FString InnerValue;
+			if (Property->ExportText_InContainer(Index, InnerValue, this, &DefaultValue, OwnerObject, PPF_Delimited | PortFlags, ExportRootScope))
+			{
+				OutValueString += Count++ == 0 ? TEXT('(') : TEXT(',');
+
+				if (Property->ArrayDim == 1)
+				{
+					OutValueString += FString::Printf(TEXT("%s="), *Property->GetName());
+				}
+				else
+				{
+					OutValueString += FString::Printf(TEXT("%s[%i]="), *Property->GetName(), Index);
+				}
+
+				OutValueString += InnerValue;
+			}
+		}
+	}
+
+	if (Count > 0)
+	{
+		OutValueString += TEXT(")");
+	}
+	else
+	{
+		OutValueString += TEXT("()");
+	}
+	return true;
+#else
+	return false;
+#endif
+}
diff --git a/Engine/Plugins/Runtime/StateTree/Source/StateTreeEditorModule/Public/StateTreeEditorData.h b/Engine/Plugins/Runtime/StateTree/Source/StateTreeEditorModule/Public/StateTreeEditorData.h
--- a/Engine/Plugins/Runtime/StateTree/Source/StateTreeEditorModule/Public/StateTreeEditorData.h
+++ b/Engine/Plugins/Runtime/StateTree/Source/StateTreeEditorModule/Public/StateTreeEditorData.h
@@ -4,6 +4,7 @@
 
 #include "StateTreeState.h"
 #include "StateTreeEditorPropertyBindings.h"
+#include "StateTreeEditorTypes.h"
 #include "Debugger/StateTreeDebuggerTypes.h"
 #include "StateTreeEditorData.generated.h"
 
@@ -49,6 +50,8 @@ class STATETREEEDITORMODULE_API UStateTreeEditorData : public UObject, public IS
 	GENERATED_BODY()
 	
 public:
+	UStateTreeEditorData();
+
 	virtual void PostInitProperties() override;
 	
 	// IStateTreeEditorPropertyBindingsOwner
@@ -235,6 +238,14 @@ PRAGMA_ENABLE_DEPRECATION_WARNINGS
 
 	// ~StateTree Builder API
 
+	/**
+	 * Attempts to find a Color matching the provided Color Key
+	 */
+	const FStateTreeEditorColor* FindColor(const FStateTreeEditorColorRef& ColorRef) const
+	{
+		return Colors.Find(FStateTreeEditorColor(ColorRef));
+	}
+
 private:
 	void FixObjectInstance(TSet<UObject*>& SeenObjects, UObject& Outer, FStateTreeEditorNode& Node);
 	void FixObjectNodes();
@@ -263,6 +274,10 @@ public:
 	UPROPERTY(meta = (ExcludeFromHash))
 	FStateTreeEditorPropertyBindings EditorBindings;
 
+	/** Color Options to assign to a State */
+	UPROPERTY(EditDefaultsOnly, Category = "Theme")
+	TSet<FStateTreeEditorColor> Colors;
+
 	/** Top level States. */
 	UPROPERTY()
 	TArray<TObjectPtr<UStateTreeState>> SubTrees;
diff --git a/Engine/Plugins/Runtime/StateTree/Source/StateTreeEditorModule/Public/StateTreeEditorTypes.h b/Engine/Plugins/Runtime/StateTree/Source/StateTreeEditorModule/Public/StateTreeEditorTypes.h
new file mode 100644
--- /dev/null
+++ b/Engine/Plugins/Runtime/StateTree/Source/StateTreeEditorModule/Public/StateTreeEditorTypes.h
@@ -0,0 +1,95 @@
+// Copyright Epic Games, Inc. All Rights Reserved.
+
+#pragma once
+
+#include "Containers/UnrealString.h"
+#include "Math/Color.h"
+#include "Misc/Guid.h"
+#include "UObject/Class.h"
+#include "StateTreeEditorTypes.generated.h"
+
+/**
+ * Id Struct to uniquely identify an FStateTreeEditorColor instance.
+ * An existing FStateTreeEditorColor instance can be found via UStateTreeEditorData::FindColor
+ */
+USTRUCT()
+struct FStateTreeEditorColorRef
+{
+	GENERATED_BODY()
+
+	FStateTreeEditorColorRef() = default;
+
+	explicit FStateTreeEditorColorRef(const FGuid& ID)
+		: ID(ID)
+	{
+	}
+
+	bool operator==(const FStateTreeEditorColorRef& Other) const
+	{
+		return ID == Other.ID;
+	}
+
+	friend uint32 GetTypeHash(const FStateTreeEditorColorRef& ColorRef)
+	{
+		return GetTypeHash(ColorRef.ID);
+	}
+
+	UPROPERTY(EditDefaultsOnly, Category = "Theme")
+	FGuid ID;
+};
+
+/**
+ * Struct describing a Color, its display name and a unique identifier to get an instance via UStateTreeEditorData::FindColor
+ */
+USTRUCT()
+struct FStateTreeEditorColor
+{
+	GENERATED_BODY()
+
+	FStateTreeEditorColor()
+		: ColorRef(FGuid::NewGuid())
+	{
+	}
+
+	explicit FStateTreeEditorColor(const FStateTreeEditorColorRef& ColorRef)
+		: ColorRef(ColorRef)
+	{
+	}
+
+	/**
+	 * Export Text Item override where properties marked with meta-data "StructExportTransient" are excluded from the exported string
+	 * This is so that copy/pasting State Tree Color entries don't have the effect of also copying over these properties into a new entry.
+	 * Since it works with meta-data, it's editor-only.
+	 * Side note: The existing "TextExportTransient" / "DuplicateTransient" specifiers apply for uclass properties only
+	 */ 
+	STATETREEEDITORMODULE_API bool ExportTextItem(FString& OutValueString, const FStateTreeEditorColor& DefaultValue, UObject* OwnerObject, int32 PortFlags, UObject* ExportRootScope) const;
+
+	bool operator==(const FStateTreeEditorColor& Other) const
+	{
+		return ColorRef == Other.ColorRef;
+	}
+
+	friend uint32 GetTypeHash(const FStateTreeEditorColor& InColor)
+	{
+		return GetTypeHash(InColor.ColorRef);
+	}
+
+	/** ID unique per State Tree Color Entry. Marked as struct export transient so that copy-pasting this entry does not result in the same repeating ID */
+	UPROPERTY(meta=(StructExportTransient))
+	FStateTreeEditorColorRef ColorRef;
+
+	UPROPERTY(EditDefaultsOnly, Category = "Theme")
+	FString DisplayName;
+
+	UPROPERTY(EditDefaultsOnly, Category = "Theme", meta=(HideAlphaChannel))
+	FLinearColor Color = FLinearColor(0.4f, 0.4f, 0.4f);
+};
+
+template<>
+struct TStructOpsTypeTraits<FStateTreeEditorColor> : TStructOpsTypeTraitsBase2<FStateTreeEditorColor>
+{
+	enum
+	{
+		WithExportTextItem = true,
+	};
+};
diff --git a/Engine/Plugins/Runtime/StateTree/Source/StateTreeEditorModule/Public/StateTreeState.h b/Engine/Plugins/Runtime/StateTree/Source/StateTreeEditorModule/Public/StateTreeState.h
--- a/Engine/Plugins/Runtime/StateTree/Source/StateTreeEditorModule/Public/StateTreeState.h
+++ b/Engine/Plugins/Runtime/StateTree/Source/StateTreeEditorModule/Public/StateTreeState.h
@@ -3,6 +3,7 @@
 #pragma once
 
 #include "StateTreeEditorNode.h"
+#include "StateTreeEditorTypes.h"
 #include "StateTreeState.generated.h"
 
 class UStateTreeState;
@@ -221,6 +222,9 @@ public:
 	UPROPERTY(EditDefaultsOnly, Category = "State", meta = (IgnoreForMemberInitializationTest))
 	FGuid ID;
 
+	UPROPERTY(EditDefaultsOnly, Category = "State", DisplayName = "Color")
+	FStateTreeEditorColorRef ColorRef;
+
 	UPROPERTY(EditDefaultsOnly, Category = "Enter Conditions", meta = (BaseStruct = "/Script/StateTreeModule.StateTreeConditionBase", BaseClass = "/Script/StateTreeModule.StateTreeConditionBlueprintBase"))
 	TArray<FStateTreeEditorNode> EnterConditions;
 
diff --git a/Engine/Plugins/Runtime/StateTree/Source/StateTreeEditorModule/StateTreeEditorModule.Build.cs b/Engine/Plugins/Runtime/StateTree/Source/StateTreeEditorModule/StateTreeEditorModule.Build.cs
--- a/Engine/Plugins/Runtime/StateTree/Source/StateTreeEditorModule/StateTreeEditorModule.Build.cs
+++ b/Engine/Plugins/Runtime/StateTree/Source/StateTreeEditorModule/StateTreeEditorModule.Build.cs
@@ -47,7 +47,9 @@ namespace UnrealBuildTool.Rules
 				"ToolWidgets",
 				"ApplicationCore",
 				"DeveloperSettings",
-				"RewindDebuggerInterface"
+				"RewindDebuggerInterface",
+				"DetailCustomizations",
+				"AppFramework"
 			}
 			);
 
