State Tree:
- Fix crash when on item selected passes in nullptr

#rnx

diff --git a/Engine/Plugins/Runtime/StateTree/Source/StateTreeEditorModule/Private/Customizations/StateTreeEditorNodeDetails.cpp b/Engine/Plugins/Runtime/StateTree/Source/StateTreeEditorModule/Private/Customizations/StateTreeEditorNodeDetails.cpp
--- a/Engine/Plugins/Runtime/StateTree/Source/StateTreeEditorModule/Private/Customizations/StateTreeEditorNodeDetails.cpp
+++ b/Engine/Plugins/Runtime/StateTree/Source/StateTreeEditorModule/Private/Customizations/StateTreeEditorNodeDetails.cpp
@@ -1509,8 +1509,8 @@ void FStateTreeEditorNodeDetails::GetNodeTypeChildren(TSharedPtr<FStateTreeNodeT
 
 void FStateTreeEditorNodeDetails::OnNodeTypeSelected(TSharedPtr<FStateTreeNodeTypeItem> SelectedItem, ESelectInfo::Type Type)
 {
-	// Skip selection set via code.
-	if (Type == ESelectInfo::Direct)
+	// Skip selection set via code, or if Selected Item is invalid
+	if (Type == ESelectInfo::Direct || !SelectedItem.IsValid())
 	{
 		return;
 	}
