[StateTreeDebugger] Fixed task breakpoints by using OnEntered/OnExited events instead of OnEntering/OnExiting since Tasks only traced the former.

diff --git a/Engine/Plugins/Runtime/StateTree/Source/StateTreeModule/Private/Debugger/StateTreeDebuggerTypes.cpp b/Engine/Plugins/Runtime/StateTree/Source/StateTreeModule/Private/Debugger/StateTreeDebuggerTypes.cpp
--- a/Engine/Plugins/Runtime/StateTree/Source/StateTreeModule/Private/Debugger/StateTreeDebuggerTypes.cpp
+++ b/Engine/Plugins/Runtime/StateTree/Source/StateTreeModule/Private/Debugger/StateTreeDebuggerTypes.cpp
@@ -297,9 +297,9 @@ EStateTreeTraceEventType FStateTreeDebuggerBreakpoint::GetMatchingEventType(cons
 	switch (BreakpointType)
 	{
 	case EStateTreeBreakpointType::OnEnter:
-		return EStateTreeTraceEventType::OnEntering;
+		return EStateTreeTraceEventType::OnEntered;
 	case EStateTreeBreakpointType::OnExit:
-		return EStateTreeTraceEventType::OnExiting;
+		return EStateTreeTraceEventType::OnExited;
 	case EStateTreeBreakpointType::OnTransition:
 		return EStateTreeTraceEventType::OnTransition;
 	default:
