[StateTreeDebugger] Fixed condition trace events.
Now sent before reseting the bindings

diff --git a/Engine/Plugins/Runtime/StateTree/Source/StateTreeModule/Private/StateTreeExecutionContext.cpp b/Engine/Plugins/Runtime/StateTree/Source/StateTreeModule/Private/StateTreeExecutionContext.cpp
--- a/Engine/Plugins/Runtime/StateTree/Source/StateTreeModule/Private/StateTreeExecutionContext.cpp
+++ b/Engine/Plugins/Runtime/StateTree/Source/StateTreeModule/Private/StateTreeExecutionContext.cpp
@@ -1493,7 +1493,8 @@ bool FStateTreeExecutionContext::TestAllConditions(const int32 ConditionsOffset,
 	
 	for (int32 Index = 0; Index < ConditionsNum; Index++)
 	{
-		const FStateTreeConditionBase& Cond = StateTree.Nodes[ConditionsOffset + Index].Get<const FStateTreeConditionBase>();
+		const int32 ConditionIndex = ConditionsOffset + Index;
+		const FStateTreeConditionBase& Cond = StateTree.Nodes[ConditionIndex].Get<const FStateTreeConditionBase>();
 		FStateTreeDataView& DataView = DataViews[Cond.DataViewIndex.Get()]; 
 		if (Cond.bInstanceIsObject)
 		{
@@ -1519,6 +1520,7 @@ bool FStateTreeExecutionContext::TestAllConditions(const int32 ConditionsOffset,
 			}
 			
 			bValue = Cond.TestCondition(*this);
+			STATETREE_TRACE_CONDITION_EVENT(ConditionIndex, DataView, bValue ? EStateTreeTraceEventType::Passed : EStateTreeTraceEventType::Failed);
 			
 			// Reset copied properties that might contain object references.
 			if (Cond.BindingsBatch.IsValid())
@@ -1529,10 +1531,9 @@ bool FStateTreeExecutionContext::TestAllConditions(const int32 ConditionsOffset,
 		else
 		{
 			bValue = Cond.EvaluationMode == EStateTreeConditionEvaluationMode::ForcedTrue ? true : /* EStateTreeConditionEvaluationMode::AlwaysFalse */ false;
+			STATETREE_TRACE_CONDITION_EVENT(ConditionIndex, FStateTreeDataView{}, bValue ? EStateTreeTraceEventType::Passed : EStateTreeTraceEventType::Failed);
 		}
 
-		STATETREE_TRACE_CONDITION_EVENT(ConditionsOffset + Index, DataView, bValue ? EStateTreeTraceEventType::Passed : EStateTreeTraceEventType::Failed);
-
 		const int32 DeltaIndent = Cond.DeltaIndent;
 		const int32 OpenParens = FMath::Max(0, DeltaIndent) + 1;	// +1 for the current value that is stored at the empty slot at the top of the value stack.
 		const int32 ClosedParens = FMath::Max(0, -DeltaIndent) + 1;
