[StateTreeDebugger]
Changed Tree Path from being just the package name to the actual object path

#rnx

diff --git a/Engine/Plugins/Runtime/StateTree/Source/StateTreeModule/Private/Debugger/StateTreeTrace.cpp b/Engine/Plugins/Runtime/StateTree/Source/StateTreeModule/Private/Debugger/StateTreeTrace.cpp
--- a/Engine/Plugins/Runtime/StateTree/Source/StateTreeModule/Private/Debugger/StateTreeTrace.cpp
+++ b/Engine/Plugins/Runtime/StateTree/Source/StateTreeModule/Private/Debugger/StateTreeTrace.cpp
@@ -425,7 +425,7 @@ void OutputInstanceLifetimeEvent(
 		TraceBufferedEvents(InstanceId);
 
 		const FString TreeName = StateTree->GetName();
-		const FString TreePath = StateTree->GetPackage()->GetName();
+		const FString TreePath = StateTree->GetPathName();
 
 		UE_TRACE_LOG(StateTreeDebugger, InstanceEvent, StateTreeDebugChannel)
 			<< InstanceEvent.Cycle(FPlatformTime::Cycles64())
diff --git a/Engine/Plugins/Runtime/StateTree/Source/StateTreeModule/Private/Debugger/StateTreeTraceAnalyzer.cpp b/Engine/Plugins/Runtime/StateTree/Source/StateTreeModule/Private/Debugger/StateTreeTraceAnalyzer.cpp
--- a/Engine/Plugins/Runtime/StateTree/Source/StateTreeModule/Private/Debugger/StateTreeTraceAnalyzer.cpp
+++ b/Engine/Plugins/Runtime/StateTree/Source/StateTreeModule/Private/Debugger/StateTreeTraceAnalyzer.cpp
@@ -51,14 +51,13 @@ bool FStateTreeTraceAnalyzer::OnEvent(const uint16 RouteId, EStyle Style, const
 			EventData.GetString("TreeName", ObjectName);
 			EventData.GetString("TreePath", ObjectPathName);
 
-			const FTopLevelAssetPath Path((FName)ObjectPathName, (FName)ObjectName);
 			TWeakObjectPtr<const UStateTree> WeakStateTree;
 			{
 				// This might not work when using a debugger on a client but should be fine in Editor as long as
 				// we are not trying to find the object during GC. We might not currently be in the game thread.  
 				// @todo STDBG: eventually errors should be reported in the UI
 				FGCScopeGuard Guard;
-				WeakStateTree = FindObject<UStateTree>(Path);
+				WeakStateTree = FindObject<UStateTree>(nullptr, *ObjectPathName);
 			}
 
 			if (const UStateTree* StateTree = WeakStateTree.Get())
