[StateTreeDebugger]
- Fixed crash when closing and reopening State Tree Debugger
- Exported FStateTreeTraceActiveStatesEvent constructor

#rnx

diff --git a/Engine/Plugins/Runtime/StateTree/Source/StateTreeModule/Private/Debugger/StateTreeDebugger.cpp b/Engine/Plugins/Runtime/StateTree/Source/StateTreeModule/Private/Debugger/StateTreeDebugger.cpp
--- a/Engine/Plugins/Runtime/StateTree/Source/StateTreeModule/Private/Debugger/StateTreeDebugger.cpp
+++ b/Engine/Plugins/Runtime/StateTree/Source/StateTreeModule/Private/Debugger/StateTreeDebugger.cpp
@@ -78,7 +78,7 @@ FStateTreeDebugger::FStateTreeDebugger()
 	: StateTreeModule(FModuleManager::GetModuleChecked<IStateTreeModule>("StateTreeModule"))
 	, ScrubState(EventCollections)
 {
-	UE::StateTree::Delegates::OnTracingStateChanged.AddLambda([this](const bool bTracesEnabled)
+	TracingStateChangedHandle = UE::StateTree::Delegates::OnTracingStateChanged.AddLambda([this](const bool bTracesEnabled)
 		{
 			// StateTree traces got enabled in the current process so let's analyse it if not already analysing something.
 			if (bTracesEnabled && !IsAnalysisSessionActive())
@@ -90,7 +90,9 @@ FStateTreeDebugger::FStateTreeDebugger()
 
 FStateTreeDebugger::~FStateTreeDebugger()
 {
-	UE::StateTree::Delegates::OnTracingStateChanged.RemoveAll(this);
+	UE::StateTree::Delegates::OnTracingStateChanged.Remove(TracingStateChangedHandle);
+	TracingStateChangedHandle.Reset();
+
 	StopSessionAnalysis();
 }
 
diff --git a/Engine/Plugins/Runtime/StateTree/Source/StateTreeModule/Public/Debugger/StateTreeDebugger.h b/Engine/Plugins/Runtime/StateTree/Source/StateTreeModule/Public/Debugger/StateTreeDebugger.h
--- a/Engine/Plugins/Runtime/StateTree/Source/StateTreeModule/Public/Debugger/StateTreeDebugger.h
+++ b/Engine/Plugins/Runtime/StateTree/Source/StateTreeModule/Public/Debugger/StateTreeDebugger.h
@@ -4,6 +4,7 @@
 
 #if WITH_STATETREE_DEBUGGER
 
+#include "Delegates/IDelegateInstance.h"
 #include "IStateTreeTraceProvider.h"
 #include "StateTree.h"
 #include "StateTreeDebuggerTypes.h"
@@ -333,6 +334,11 @@ private:
 	 * Indicates the last transition type between two consecutive analyses to manage track cleanup properly.
 	 */
 	EAnalysisTransitionType AnalysisTransitionType = EAnalysisTransitionType::Unset;
+
+	/**
+	 * Delegate Handle bound to UE::StateTree::Delegates::OnTracingStateChanged
+	 */
+	FDelegateHandle TracingStateChangedHandle;
 };
 
 #endif // WITH_STATETREE_DEBUGGER
diff --git a/Engine/Plugins/Runtime/StateTree/Source/StateTreeModule/Public/Debugger/StateTreeTraceTypes.h b/Engine/Plugins/Runtime/StateTree/Source/StateTreeModule/Public/Debugger/StateTreeTraceTypes.h
--- a/Engine/Plugins/Runtime/StateTree/Source/StateTreeModule/Public/Debugger/StateTreeTraceTypes.h
+++ b/Engine/Plugins/Runtime/StateTree/Source/StateTreeModule/Public/Debugger/StateTreeTraceTypes.h
@@ -208,7 +208,7 @@ struct FStateTreeTraceConditionEvent : FStateTreeTraceNodeEvent
 struct FStateTreeTraceActiveStatesEvent : FStateTreeTraceBaseEvent
 {
 	// Intentionally implemented in source file to compile 'TArray<FStateTreeStateHandle>' using only forward declaration.
-	explicit FStateTreeTraceActiveStatesEvent(const double RecordingWorldTime);
+	STATETREEMODULE_API explicit FStateTreeTraceActiveStatesEvent(const double RecordingWorldTime);
 
 	STATETREEMODULE_API FString ToFullString(const UStateTree& StateTree) const;
 	STATETREEMODULE_API FString GetValueString(const UStateTree& StateTree) const;
