[StateTree]
Added "Next Selectable State" transition option which is similar to "Next State" but the difference is that if the next state fails to be selected (i.e. enter condition failed), then instead of going back up the hierarchy, it will try to select the next sibling until it finds one that is selectable or if it reached the end of the sibling list.

diff --git a/Engine/Plugins/Runtime/StateTree/Source/StateTreeEditorModule/Private/Customizations/StateTreeStateLinkDetails.cpp b/Engine/Plugins/Runtime/StateTree/Source/StateTreeEditorModule/Private/Customizations/StateTreeStateLinkDetails.cpp
--- a/Engine/Plugins/Runtime/StateTree/Source/StateTreeEditorModule/Private/Customizations/StateTreeStateLinkDetails.cpp
+++ b/Engine/Plugins/Runtime/StateTree/Source/StateTreeEditorModule/Private/Customizations/StateTreeStateLinkDetails.cpp
@@ -172,6 +172,9 @@ void FStateTreeStateLinkDetails::OnStateComboChange(int Idx)
 			case ComboNextState:
 				LinkTypeProperty->SetValue((uint8)EStateTreeTransitionType::NextState);
 				break;
+			case ComboNextSelectableState:
+				LinkTypeProperty->SetValue((uint8)EStateTreeTransitionType::NextSelectableState);
+				break;
 			case ComboNotSet:
 			default:
 				LinkTypeProperty->SetValue((uint8)EStateTreeTransitionType::None);
@@ -196,6 +199,9 @@ TSharedRef<SWidget> FStateTreeStateLinkDetails::OnGetStateContent() const
 		const FUIAction NextItemAction(FExecuteAction::CreateSP(const_cast<FStateTreeStateLinkDetails*>(this), &FStateTreeStateLinkDetails::OnStateComboChange, ComboNextState));
 		MenuBuilder.AddMenuEntry(LOCTEXT("TransitionNextState", "Next State"), LOCTEXT("TransitionNextTooltip", "Goto next sibling State."), FSlateIcon(), NextItemAction);
 
+		const FUIAction NextSelectableItemAction(FExecuteAction::CreateSP(const_cast<FStateTreeStateLinkDetails*>(this), &FStateTreeStateLinkDetails::OnStateComboChange, ComboNextSelectableState));
+		MenuBuilder.AddMenuEntry(LOCTEXT("TransitionNextSelectableState", "Next Selectable State"), LOCTEXT("TransitionNextSelectableTooltip", "Goto next sibling state, whose enter conditions pass."), FSlateIcon(), NextSelectableItemAction);
+
 		const FUIAction SucceededItemAction(FExecuteAction::CreateSP(const_cast<FStateTreeStateLinkDetails*>(this), &FStateTreeStateLinkDetails::OnStateComboChange, ComboSucceeded));
 		MenuBuilder.AddMenuEntry(LOCTEXT("TransitionTreeSucceeded", "Tree Succeeded"), LOCTEXT("TransitionTreeSuccessTooltip", "Complete tree with success."), FSlateIcon(), SucceededItemAction);
 
@@ -232,6 +238,9 @@ FText FStateTreeStateLinkDetails::GetCurrentStateDesc() const
 	case EStateTreeTransitionType::NextState:
 		Result = LOCTEXT("TransitionNextState", "Next State");
 		break;
+	case EStateTreeTransitionType::NextSelectableState:
+		Result = LOCTEXT("TransitionNextSelectableState", "Next Selectable State");
+		break;
 	case EStateTreeTransitionType::Succeeded:
 		Result = LOCTEXT("TransitionTreeSucceeded", "Tree Succeeded");
 		break;
@@ -287,6 +296,7 @@ bool FStateTreeStateLinkDetails::IsValidLink() const
 	// all these are unconditionally valid		
 	case EStateTreeTransitionType::None: // fall through on purpose
 	case EStateTreeTransitionType::NextState: // fall through on purpose
+	case EStateTreeTransitionType::NextSelectableState: // fall through on purpose
 	case EStateTreeTransitionType::Succeeded: // fall through on purpose
 	case EStateTreeTransitionType::Failed:
 		bIsValid = true;
diff --git a/Engine/Plugins/Runtime/StateTree/Source/StateTreeEditorModule/Private/Customizations/StateTreeStateLinkDetails.h b/Engine/Plugins/Runtime/StateTree/Source/StateTreeEditorModule/Private/Customizations/StateTreeStateLinkDetails.h
--- a/Engine/Plugins/Runtime/StateTree/Source/StateTreeEditorModule/Private/Customizations/StateTreeStateLinkDetails.h
+++ b/Engine/Plugins/Runtime/StateTree/Source/StateTreeEditorModule/Private/Customizations/StateTreeStateLinkDetails.h
@@ -33,6 +33,7 @@ private:
 	inline static constexpr int ComboSucceeded = -2;
 	inline static constexpr int ComboFailed = -3;
 	inline static constexpr int ComboNextState = -4;
+	inline static constexpr int ComboNextSelectableState = -5;
 
 	void CacheStates();
 	void CacheStates(const UStateTreeState* State);
diff --git a/Engine/Plugins/Runtime/StateTree/Source/StateTreeEditorModule/Private/Customizations/StateTreeTransitionDetails.cpp b/Engine/Plugins/Runtime/StateTree/Source/StateTreeEditorModule/Private/Customizations/StateTreeTransitionDetails.cpp
--- a/Engine/Plugins/Runtime/StateTree/Source/StateTreeEditorModule/Private/Customizations/StateTreeTransitionDetails.cpp
+++ b/Engine/Plugins/Runtime/StateTree/Source/StateTreeEditorModule/Private/Customizations/StateTreeTransitionDetails.cpp
@@ -221,6 +221,9 @@ FText FStateTreeTransitionDetails::GetDescription() const
 		case EStateTreeTransitionType::NextState:
 			TargetText = LOCTEXT("TransitionNextState", "Next State");
 			break;
+		case EStateTreeTransitionType::NextSelectableState:
+			TargetText = LOCTEXT("TransitionNextSelectableState", "Next Selectable State");
+			break;
 		case EStateTreeTransitionType::GotoState:
 			{
 				TargetText = FText::Format(LOCTEXT("TransitionGotoState", "Go to State {0}"), FText::FromName(State->Name));
diff --git a/Engine/Plugins/Runtime/StateTree/Source/StateTreeEditorModule/Private/SStateTreeViewRow.cpp b/Engine/Plugins/Runtime/StateTree/Source/StateTreeEditorModule/Private/SStateTreeViewRow.cpp
--- a/Engine/Plugins/Runtime/StateTree/Source/StateTreeEditorModule/Private/SStateTreeViewRow.cpp
+++ b/Engine/Plugins/Runtime/StateTree/Source/StateTreeEditorModule/Private/SStateTreeViewRow.cpp
@@ -949,6 +949,9 @@ FText SStateTreeViewRow::GetLinkDescription(const FStateTreeStateLink& Link)
 	case EStateTreeTransitionType::NextState:
 		return LOCTEXT("TransitionNextStateStyled", "[Next]");
 		break;
+	case EStateTreeTransitionType::NextSelectableState:
+		return LOCTEXT("TransitionNextSelectableStateStyled", "[Next Selectable]");
+		break;
 	case EStateTreeTransitionType::GotoState:
 		return FText::FromName(Link.Name);
 		break;
@@ -1091,6 +1094,7 @@ FText SStateTreeViewRow::GetTransitionsIcon(const UStateTreeState& State, const
 				IconType |= IconRightArrow;
 				break;
 			case EStateTreeTransitionType::NextState:
+			case EStateTreeTransitionType::NextSelectableState:
 				IconType |= IconDownArrow;
 				break;
 			case EStateTreeTransitionType::GotoState:
diff --git a/Engine/Plugins/Runtime/StateTree/Source/StateTreeEditorModule/Private/StateTreeCompiler.cpp b/Engine/Plugins/Runtime/StateTree/Source/StateTreeEditorModule/Private/StateTreeCompiler.cpp
--- a/Engine/Plugins/Runtime/StateTree/Source/StateTreeEditorModule/Private/StateTreeCompiler.cpp
+++ b/Engine/Plugins/Runtime/StateTree/Source/StateTreeEditorModule/Private/StateTreeCompiler.cpp
@@ -857,7 +857,12 @@ bool FStateTreeCompiler::CreateStateTransitions()
 			CompactTransition.Priority = Transition.Priority;
 			CompactTransition.EventTag = Transition.EventTag;
 			CompactTransition.bTransitionEnabled = Transition.bTransitionEnabled;
-			
+
+			if (Transition.State.LinkType == EStateTreeTransitionType::NextSelectableState)
+			{
+				CompactTransition.Fallback = EStateTreeSelectionFallback::NextSelectableSibling;
+			}
+
 			if (Transition.bDelayTransition)
 			{
 				CompactTransition.Delay.Set(Transition.DelayDuration, Transition.DelayRandomVariance);
@@ -965,7 +970,7 @@ bool FStateTreeCompiler::ResolveTransitionState(const UStateTreeState* SourceSta
 			return false;
 		}
 	}
-	else if (Link.LinkType == EStateTreeTransitionType::NextState)
+	else if (Link.LinkType == EStateTreeTransitionType::NextState || Link.LinkType == EStateTreeTransitionType::NextSelectableState)
 	{
 		// Find next state.
 		const UStateTreeState* NextState = SourceState ? SourceState->GetNextSelectableSiblingState() : nullptr;
diff --git a/Engine/Plugins/Runtime/StateTree/Source/StateTreeModule/Private/StateTreeExecutionContext.cpp b/Engine/Plugins/Runtime/StateTree/Source/StateTreeModule/Private/StateTreeExecutionContext.cpp
--- a/Engine/Plugins/Runtime/StateTree/Source/StateTreeModule/Private/StateTreeExecutionContext.cpp
+++ b/Engine/Plugins/Runtime/StateTree/Source/StateTreeModule/Private/StateTreeExecutionContext.cpp
@@ -1587,7 +1587,7 @@ FString FStateTreeExecutionContext::DebugGetEventsAsString() const
 	return Result;
 }
 
-bool FStateTreeExecutionContext::RequestTransition(const FStateTreeStateHandle NextState, const EStateTreeTransitionPriority Priority)
+bool FStateTreeExecutionContext::RequestTransition(const FStateTreeStateHandle NextState, const EStateTreeTransitionPriority Priority, const EStateTreeSelectionFallback Fallback)
 {
 	// Skip lower priority transitions.
 	if (NextTransition.Priority >= Priority)
@@ -1614,7 +1614,7 @@ bool FStateTreeExecutionContext::RequestTransition(const FStateTreeStateHandle N
 
 	FStateTreeActiveStates NewActiveState;
 	FStateTreeActiveStates VisitedStates;
-	if (SelectState(NextState, NewActiveState, VisitedStates))
+	if (SelectState(NextState, NewActiveState, VisitedStates, Fallback))
 	{
 		SetupNextTransition(NextState, Priority);
 		NextTransition.NextActiveStates = NewActiveState;
@@ -1806,7 +1806,7 @@ bool FStateTreeExecutionContext::TriggerTransitions()
 							});
 
 						// Trigger Delayed Transition when the delay has passed.
-						if (RequestTransition(Transition.State, Transition.Priority))
+						if (RequestTransition(Transition.State, Transition.Priority, Transition.Fallback))
 						{
 							NextTransitionSource = FStateTreeTransitionSource(FStateTreeIndex16(TransitionIndex), Transition.State, Transition.Priority);
 						}
@@ -1856,7 +1856,7 @@ bool FStateTreeExecutionContext::TriggerTransitions()
 						}
 					}
 
-					if (RequestTransition(Transition.State, Transition.Priority))
+					if (RequestTransition(Transition.State, Transition.Priority, Transition.Fallback))
 					{
 						NextTransitionSource = FStateTreeTransitionSource(FStateTreeIndex16(TransitionIndex), Transition.State, Transition.Priority);
 					}
@@ -1937,7 +1937,7 @@ bool FStateTreeExecutionContext::TriggerTransitions()
 					{
 						// No delay allowed on completion conditions.
 						// No priority on completion transitions, use the priority to signal that state is selected.
-						if (RequestTransition(Transition.State, EStateTreeTransitionPriority::Normal))
+						if (RequestTransition(Transition.State, EStateTreeTransitionPriority::Normal, Transition.Fallback))
 						{
 							NextTransitionSource = FStateTreeTransitionSource(FStateTreeIndex16(TransitionIndex), Transition.State, Transition.Priority);
 							break;
@@ -2045,7 +2045,7 @@ FStateTreeStateHandle FStateTreeExecutionContext::GetParentLinkedStateHandle(con
 	return FStateTreeStateHandle();
 }
 
-bool FStateTreeExecutionContext::SelectState(const FStateTreeStateHandle NextState, FStateTreeActiveStates& OutNewActiveState, FStateTreeActiveStates& VisitedStates)
+bool FStateTreeExecutionContext::SelectState(const FStateTreeStateHandle NextState, FStateTreeActiveStates& OutNewActiveState, FStateTreeActiveStates& VisitedStates, const EStateTreeSelectionFallback Fallback)
 {
 	const FStateTreeExecutionState& Exec = GetExecState();
 
@@ -2100,8 +2100,35 @@ bool FStateTreeExecutionContext::SelectState(const FStateTreeStateHandle NextSta
 			__FUNCTION__, *GetSafeStateName(NextState), *GetStateStatusString(Exec), *GetNameSafe(&Owner), *GetFullNameSafe(&StateTree));
 		return false;
 	}
-	
-	return SelectStateInternal(NextState, OutNewActiveState, VisitedStates);
+
+	if (SelectStateInternal(NextState, OutNewActiveState, VisitedStates))
+	{
+		return true;
+	}
+
+	// Failed to Select Next State, handle fallback here
+	// Return true on the first next sibling that gets selected successfully
+	if (Fallback == EStateTreeSelectionFallback::NextSelectableSibling && NumInBetweenStates >= 2)
+	{
+		// InBetweenStates is in reversed order (i.e. from leaf to root)
+		FStateTreeStateHandle Parent = InBetweenStates[1];
+		if (Parent.IsValid())
+		{
+			const FCompactStateTreeState& ParentState = StateTree.States[Parent.Index];
+
+			uint16 ChildState = StateTree.States[NextState.Index].GetNextSibling();
+			for (; ChildState < ParentState.ChildrenEnd; ChildState = StateTree.States[ChildState].GetNextSibling())
+			{
+				FStateTreeStateHandle ChildStateHandle = FStateTreeStateHandle(ChildState);
+				if (SelectStateInternal(ChildStateHandle, OutNewActiveState, VisitedStates))
+				{
+					return true;
+				}
+			}
+		}
+	}
+
+	return false;
 }
 
 bool FStateTreeExecutionContext::SelectStateInternal(const FStateTreeStateHandle NextState, FStateTreeActiveStates& OutNewActiveState, FStateTreeActiveStates& VisitedStates)
@@ -2226,7 +2253,7 @@ bool FStateTreeExecutionContext::SelectStateInternal(const FStateTreeStateHandle
 					// Using SelectState() instead of SelectStateInternal to treat the transitions the same way as regular transitions,
 					// e.g. it may jump to a completely different branch.
 					FStateTreeActiveStates NewActiveState;
-					if (SelectState(Transition.State, NewActiveState, VisitedStates))
+					if (SelectState(Transition.State, NewActiveState, VisitedStates, Transition.Fallback))
 					{
 						// Selection succeeded
 						OutNewActiveState = NewActiveState;
diff --git a/Engine/Plugins/Runtime/StateTree/Source/StateTreeModule/Public/StateTreeExecutionContext.h b/Engine/Plugins/Runtime/StateTree/Source/StateTreeModule/Public/StateTreeExecutionContext.h
--- a/Engine/Plugins/Runtime/StateTree/Source/StateTreeModule/Public/StateTreeExecutionContext.h
+++ b/Engine/Plugins/Runtime/StateTree/Source/StateTreeModule/Public/StateTreeExecutionContext.h
@@ -338,7 +338,7 @@ protected:
 	/**
 	 * Requests transition to a specified state with specified priority.
 	 */
-	bool RequestTransition(const FStateTreeStateHandle NextState, const EStateTreeTransitionPriority Priority);
+	bool RequestTransition(const FStateTreeStateHandle NextState, const EStateTreeTransitionPriority Priority, const EStateTreeSelectionFallback Fallback = EStateTreeSelectionFallback::None);
 
 	/**
 	 * Sets up NextTransition based on the provided parameters and the current execution status. 
@@ -370,9 +370,10 @@ protected:
 	 * @param NextState The state which we try to select next.
 	 * @param OutNewActiveStates Active states that got selected.
 	 * @param VisitedStates States visited so far during selection (used for detecting selection loops)
+	 * @param Fallback selection behavior to execute if it fails to select the desired state
 	 * @return True if succeeded to select new active states.
 	 */
-	bool SelectState(const FStateTreeStateHandle NextState, FStateTreeActiveStates& OutNewActiveStates, FStateTreeActiveStates& VisitedStates);
+	bool SelectState(const FStateTreeStateHandle NextState, FStateTreeActiveStates& OutNewActiveStates, FStateTreeActiveStates& VisitedStates, const EStateTreeSelectionFallback Fallback = EStateTreeSelectionFallback::None);
 
 	/**
 	 * Used internally to do the recursive part of the SelectState().
diff --git a/Engine/Plugins/Runtime/StateTree/Source/StateTreeModule/Public/StateTreeTypes.h b/Engine/Plugins/Runtime/StateTree/Source/StateTreeModule/Public/StateTreeTypes.h
--- a/Engine/Plugins/Runtime/StateTree/Source/StateTreeModule/Public/StateTreeTypes.h
+++ b/Engine/Plugins/Runtime/StateTree/Source/StateTreeModule/Public/StateTreeTypes.h
@@ -41,6 +41,9 @@ enum class EStateTreeTransitionType : uint8
 	/** Transition to the next sibling state. */
 	NextState,
 
+	/** Transition to the next selectable sibling state */
+	NextSelectableState,
+
 	NotSet UE_DEPRECATED(5.0, "Use None instead."),
 };
 
@@ -262,6 +265,16 @@ protected:
 	uint16 RandomVariance = 0;
 };
 
+/** Fallback behavior indicating what to do after failing to select a state */
+UENUM()
+enum class EStateTreeSelectionFallback : uint8
+{
+	/** No fallback */
+	None,
+
+	/** Find next selectable sibling, if any, and select it */
+	NextSelectableSibling,
+};
 
 /**
  *  Runtime representation of a StateTree transition.
@@ -306,6 +319,10 @@ struct STATETREEMODULE_API FCompactStateTransition
 	UPROPERTY()
 	EStateTreeTransitionPriority Priority = EStateTreeTransitionPriority::Normal;
 
+	/** Fallback of the transition if it fails to select the target state */
+	UPROPERTY()
+	EStateTreeSelectionFallback Fallback = EStateTreeSelectionFallback::None;
+
 	/** Number of conditions to test. */
 	UPROPERTY()
 	uint8 ConditionsNum = 0;
