[StateTree]
Added test for "next selectable" transition type

diff --git a/Engine/Plugins/Runtime/StateTree/Source/StateTreeTestSuite/Private/StateTreeTest.cpp b/Engine/Plugins/Runtime/StateTree/Source/StateTreeTestSuite/Private/StateTreeTest.cpp
--- a/Engine/Plugins/Runtime/StateTree/Source/StateTreeTestSuite/Private/StateTreeTest.cpp
+++ b/Engine/Plugins/Runtime/StateTree/Source/StateTreeTestSuite/Private/StateTreeTest.cpp
@@ -728,6 +728,76 @@ struct FStateTreeTest_TransitionPriorityEnterState : FAITestBase
 };
 IMPLEMENT_AI_INSTANT_TEST(FStateTreeTest_TransitionPriorityEnterState, "System.StateTree.Transition.PriorityEnterState");
 
+struct FStateTreeTest_TransitionNextSelectableState : FAITestBase
+{
+	virtual bool InstantTest() override
+	{
+		UStateTree& StateTree = UE::StateTree::Tests::NewStateTree(&GetWorld());
+		UStateTreeEditorData& EditorData = *Cast<UStateTreeEditorData>(StateTree.EditorData);
+
+		UStateTreeState& Root =	EditorData.AddSubTree(FName(TEXT("Root")));
+		UStateTreeState& State0 = Root.AddChildState(FName(TEXT("State0")));
+		UStateTreeState& State1 = Root.AddChildState(FName(TEXT("State1")));
+		UStateTreeState& State2 = Root.AddChildState(FName(TEXT("State2")));
+
+		auto& EvalA = EditorData.AddEvaluator<FTestEval_A>();
+		EvalA.GetInstanceData().bBoolA = true;
+
+		auto& Task0 = State0.AddTask<FTestTask_Stand>(FName(TEXT("Task0")));
+		State0.AddTransition(EStateTreeTransitionTrigger::OnStateCompleted, EStateTreeTransitionType::NextSelectableState);
+
+		// Add Task 1 with Condition that will always fail
+		auto& Task1 = State1.AddTask<FTestTask_Stand>(FName(TEXT("Task1")));
+		auto& BoolCond1 = State1.AddEnterCondition<FStateTreeCompareBoolCondition>();
+
+		EditorData.AddPropertyBinding(EvalA, TEXT("bBoolA"), BoolCond1, TEXT("bLeft"));
+		BoolCond1.GetInstanceData().bRight = !EvalA.GetInstanceData().bBoolA;
+
+		// Add Task 2 with Condition that will always succeed
+		auto& Task2 = State2.AddTask<FTestTask_Stand>(FName(TEXT("Task2")));
+		auto& BoolCond2 = State2.AddEnterCondition<FStateTreeCompareBoolCondition>();
+		State2.AddTransition(EStateTreeTransitionTrigger::OnStateCompleted, EStateTreeTransitionType::Succeeded);
+
+		EditorData.AddPropertyBinding(EvalA, TEXT("bBoolA"), BoolCond2, TEXT("bLeft"));
+		BoolCond2.GetInstanceData().bRight = EvalA.GetInstanceData().bBoolA;
+
+		FStateTreeCompilerLog Log;
+		FStateTreeCompiler Compiler(Log);
+		const bool bResult = Compiler.Compile(StateTree);
+		AITEST_TRUE("StateTree should get compiled", bResult);
+
+		FStateTreeInstanceData InstanceData;
+		FTestStateTreeExecutionContext Exec(StateTree, StateTree, InstanceData);
+		const bool bInitSucceeded = Exec.IsValid();
+		AITEST_TRUE("StateTree should init", bInitSucceeded);
+
+		const FString TickStr(TEXT("Tick"));
+		const FString EnterStateStr(TEXT("EnterState"));
+		const FString ExitStateStr(TEXT("ExitState"));
+		const FString StateCompletedStr(TEXT("StateCompleted"));
+
+		// Start and enter state
+		Exec.Start();
+		AITEST_TRUE("StateTree Task0 should enter state", Exec.Expect(Task0.GetName(), EnterStateStr));
+		Exec.LogClear();
+
+		// Transition from State0 and tries to select State1. It should fail (Task1) and because transition is set to "Next Selectable", it should now select Task 2 and Enter State
+		Exec.Tick(0.1f);
+		AITEST_TRUE("StateTree Task0 should complete", Exec.Expect(Task0.GetName(), StateCompletedStr));
+		AITEST_FALSE("StateTree Task1 should not enter state", Exec.Expect(Task1.GetName(), EnterStateStr));
+		AITEST_TRUE("StateTree Task2 should enter state", Exec.Expect(Task2.GetName(), EnterStateStr));
+		Exec.LogClear();
+
+		// Complete Task2
+		Exec.Tick(0.1f);
+		AITEST_TRUE("StateTree Task2 should complete", Exec.Expect(Task2.GetName(), StateCompletedStr));
+		Exec.LogClear();
+
+		return true;
+	}
+};
+IMPLEMENT_AI_INSTANT_TEST(FStateTreeTest_TransitionNextSelectableState, "System.StateTree.Transition.NextSelectableState");
+
 struct FStateTreeTest_LastConditionWithIndent : FAITestBase
 {
 	virtual bool InstantTest() override
