StateTree: Fixed invalidating binding to parameter of Object type.

diff --git a/Engine/Plugins/Runtime/StateTree/Source/StateTreeEditorModule/Private/Customizations/StateTreeBindingExtension.cpp b/Engine/Plugins/Runtime/StateTree/Source/StateTreeEditorModule/Private/Customizations/StateTreeBindingExtension.cpp
--- a/Engine/Plugins/Runtime/StateTree/Source/StateTreeEditorModule/Private/Customizations/StateTreeBindingExtension.cpp
+++ b/Engine/Plugins/Runtime/StateTree/Source/StateTreeEditorModule/Private/Customizations/StateTreeBindingExtension.cpp
@@ -72,21 +72,24 @@ UStruct* ResolveLeafValueStructType(FStateTreeDataView ValueView, const TArray<F
 
 	UStruct* Result = nullptr;
 
-	if (const FStructProperty* StructProperty = CastField<FStructProperty>(LastIndirection.GetProperty()))
+	if (LastIndirection.GetContainerAddress())
 	{
-		// Get the type of the instanced struct's value.
-		if (StructProperty->Struct == TBaseStructure<FInstancedStruct>::Get())
+		if (const FStructProperty* StructProperty = CastField<FStructProperty>(LastIndirection.GetProperty()))
 		{
-			const FInstancedStruct& InstancedStruct = *reinterpret_cast<const FInstancedStruct*>(LastIndirection.GetPropertyAddress());
-			Result = const_cast<UScriptStruct*>(InstancedStruct.GetScriptStruct());
+			// Get the type of the instanced struct's value.
+			if (StructProperty->Struct == TBaseStructure<FInstancedStruct>::Get())
+			{
+				const FInstancedStruct& InstancedStruct = *reinterpret_cast<const FInstancedStruct*>(LastIndirection.GetPropertyAddress());
+				Result = const_cast<UScriptStruct*>(InstancedStruct.GetScriptStruct());
+			}
 		}
-	}
-	else if (const FObjectProperty* ObjectProperty = CastField<FObjectProperty>(LastIndirection.GetProperty()))
-	{
-		// Get type of the instanced object.
-		if (const UObject* Object = *reinterpret_cast<UObject* const*>(LastIndirection.GetPropertyAddress()))
+		else if (const FObjectProperty* ObjectProperty = CastField<FObjectProperty>(LastIndirection.GetProperty()))
 		{
-			Result = Object->GetClass();
+			// Get type of the instanced object.
+			if (const UObject* Object = *reinterpret_cast<UObject* const*>(LastIndirection.GetPropertyAddress()))
+			{
+				Result = Object->GetClass();
+			}
 		}
 	}
 
diff --git a/Engine/Plugins/Runtime/StateTree/Source/StateTreeModule/Private/StateTreePropertyBindings.cpp b/Engine/Plugins/Runtime/StateTree/Source/StateTreeModule/Private/StateTreePropertyBindings.cpp
--- a/Engine/Plugins/Runtime/StateTree/Source/StateTreeModule/Private/StateTreePropertyBindings.cpp
+++ b/Engine/Plugins/Runtime/StateTree/Source/StateTreeModule/Private/StateTreePropertyBindings.cpp
@@ -1358,12 +1358,10 @@ bool FStateTreePropertyPath::ResolveIndirectionsWithValue(const FStateTreeDataVi
 
 	const uint8* CurrentAddress = BaseValueView.GetMemory();
 	const UStruct* CurrentStruct = BaseValueView.GetStruct();
-
-	const bool bWithValue = CurrentAddress != nullptr;
 	
 	for (const TEnumerateRef<const FStateTreePropertyPathSegment> Segment : EnumerateRange(Segments))
 	{
-		if (CurrentStruct == nullptr || (bWithValue && CurrentAddress == nullptr))
+		if (CurrentStruct == nullptr)
 		{
 			if (OutError)
 			{
@@ -1375,6 +1373,7 @@ bool FStateTreePropertyPath::ResolveIndirectionsWithValue(const FStateTreeDataVi
 		}
 
 		const FProperty* Property = CurrentStruct->FindPropertyByName(Segment->GetName());
+		const bool bWithValue = CurrentAddress != nullptr;
 
 #if WITH_EDITORONLY_DATA
 		FName RedirectedName;
