StateTree: Fixed binding validation for any enums
- Combined common binding validation methods
- Moved all property binding callbacks to the shared binding data

diff --git a/Engine/Plugins/Runtime/StateTree/Source/StateTreeEditorModule/Private/Customizations/StateTreeBindingExtension.cpp b/Engine/Plugins/Runtime/StateTree/Source/StateTreeEditorModule/Private/Customizations/StateTreeBindingExtension.cpp
--- a/Engine/Plugins/Runtime/StateTree/Source/StateTreeEditorModule/Private/Customizations/StateTreeBindingExtension.cpp
+++ b/Engine/Plugins/Runtime/StateTree/Source/StateTreeEditorModule/Private/Customizations/StateTreeBindingExtension.cpp
@@ -260,13 +260,96 @@ FOnStateTreePropertyBindingChanged STATETREEEDITORMODULE_API OnStateTreeProperty
 
 struct FCachedBindingData : public TSharedFromThis<FCachedBindingData>
 {
-	FCachedBindingData(UObject* InOwnerObject, const FStateTreePropertyPath& InTargetPath, const TSharedPtr<const IPropertyHandle>& InPropertyHandle)
+	FCachedBindingData(UObject* InOwnerObject, const FStateTreePropertyPath& InTargetPath, const TSharedPtr<const IPropertyHandle>& InPropertyHandle, TArrayView<FStateTreeBindableStructDesc> InAccessibleStructs)
 		: WeakOwnerObject(InOwnerObject)
 		, TargetPath(InTargetPath)
 		, PropertyHandle(InPropertyHandle)
+		, AccessibleStructs(InAccessibleStructs)
 	{
 	}
-	
+
+	void AddBinding(const TArray<FBindingChainElement>& InBindingChain)
+	{
+		if (InBindingChain.IsEmpty())
+		{
+			return;
+		}
+		
+		if (!TargetPath.GetStructID().IsValid())
+		{
+			return;
+		}
+
+		UObject* OwnerObject = WeakOwnerObject.Get();
+		if (!OwnerObject)
+		{
+			return;
+		}
+		
+		IStateTreeEditorPropertyBindingsOwner* BindingOwner = Cast<IStateTreeEditorPropertyBindingsOwner>(OwnerObject);
+		if (!BindingOwner)
+		{
+			return;
+		}
+
+		FStateTreeEditorPropertyBindings* EditorBindings = BindingOwner->GetPropertyEditorBindings();
+		if (!EditorBindings)
+		{
+			return;
+		}
+
+		// First item in the binding chain is the index in AccessibleStructs.
+		const int32 SourceStructIndex = InBindingChain[0].ArrayIndex;
+		check(SourceStructIndex >= 0 && SourceStructIndex < AccessibleStructs.Num());
+				
+		TArray<FBindingChainElement> SourceBindingChain = InBindingChain;
+		SourceBindingChain.RemoveAt(0); // remove struct index.
+
+		FStateTreeDataView DataView;
+		BindingOwner->GetDataViewByID(AccessibleStructs[SourceStructIndex].ID, DataView);
+
+		// If SourceBindingChain is empty at this stage, it means that the binding points to the source struct itself.
+		FStateTreePropertyPath SourcePath;
+		UE::StateTree::PropertyBinding::MakeStructPropertyPathFromBindingChain(AccessibleStructs[SourceStructIndex].ID, SourceBindingChain, DataView, SourcePath);
+				
+		OwnerObject->Modify();
+		EditorBindings->AddPropertyBinding(SourcePath, TargetPath);
+
+		UpdateData();
+
+		UE::StateTree::PropertyBinding::OnStateTreePropertyBindingChanged.Broadcast(SourcePath, TargetPath);
+	}
+
+	void RemoveBinding()
+	{
+		UObject* OwnerObject = WeakOwnerObject.Get();
+		if (!OwnerObject)
+		{
+			return;
+		}
+		
+		IStateTreeEditorPropertyBindingsOwner* BindingOwner = Cast<IStateTreeEditorPropertyBindingsOwner>(OwnerObject);
+		if (!BindingOwner)
+		{
+			return;
+		}
+
+		FStateTreeEditorPropertyBindings* EditorBindings = BindingOwner->GetPropertyEditorBindings();
+		if (!EditorBindings)
+		{
+			return;
+		}
+
+		OwnerObject->Modify();
+		EditorBindings->RemovePropertyBindings(TargetPath);
+
+		UpdateData();
+		
+		const FStateTreePropertyPath SourcePath; // Null path
+		UE::StateTree::PropertyBinding::OnStateTreePropertyBindingChanged.Broadcast(SourcePath, TargetPath);
+	}
+
+
 	void UpdateData()
 	{
 		static FName PropertyIcon(TEXT("Kismet.Tabs.Variables"));
@@ -305,9 +388,6 @@ struct FCachedBindingData : public TSharedFromThis<FCachedBindingData>
 			return;
 		}
 
-		TArray<FStateTreeBindableStructDesc> AccessibleStructs;
-		BindingOwner->GetAccessibleStructs(TargetPath.GetStructID(), AccessibleStructs);
-		
 		const UEdGraphSchema_K2* Schema = GetDefault<UEdGraphSchema_K2>();
 		check(Schema);
 
@@ -331,29 +411,27 @@ struct FCachedBindingData : public TSharedFromThis<FCachedBindingData>
 				{
 					TArray<FStateTreePropertyPathIndirection> SourceIndirections;
 					TArray<FStateTreePropertyPathIndirection> TargetIndirections;
-					
-					if (SourcePath->ResolveIndirectionsWithValue(SourceDataView, SourceIndirections)
-						&& TargetPath.ResolveIndirectionsWithValue(TargetDataView, TargetIndirections))
-					{
-						FStateTreePropertyCopy DummyCopy;
-						FStateTreePropertyPathIndirection LastSourceIndirection = !SourceIndirections.IsEmpty() ? SourceIndirections.Last() : FStateTreePropertyPathIndirection(SourceDataView.GetStruct());
-						FStateTreePropertyPathIndirection LastTargetIndirection = !TargetIndirections.IsEmpty() ? TargetIndirections.Last() : FStateTreePropertyPathIndirection(TargetDataView.GetStruct());
-						if (FStateTreePropertyBindings::ResolveCopyType(LastSourceIndirection, LastTargetIndirection, DummyCopy))
-						{
-							bIsValidBinding = true;
-						}
-					}
 
-					if (!SourcePath->IsPathEmpty())
+					// Resolve source and target properties.
+					// Source path can be empty, when the binding binds directly to a context struct/class.
+					// Target path must always point to a valid property (at least one indirection).
+					if (SourcePath->ResolveIndirectionsWithValue(SourceDataView, SourceIndirections)
+						&& TargetPath.ResolveIndirectionsWithValue(TargetDataView, TargetIndirections)
+						&& !TargetIndirections.IsEmpty())
 					{
+						const FStateTreePropertyPathIndirection LastTargetIndirection = TargetIndirections.Last();
 						if (SourceIndirections.Num() > 0)
 						{
+							// Binding to a source property.
 							SourceLeafProperty = SourceIndirections.Last().GetProperty();
+							bIsValidBinding = ArePropertiesCompatible(SourceLeafProperty, LastTargetIndirection.GetProperty(), LastTargetIndirection.GetPropertyAddress());
+						}
+						else
+						{
+							// Binding to a source context struct.
+							SourceStruct = SourceDataView.GetStruct();
+							bIsValidBinding = ArePropertyAndContextStructCompatible(SourceStruct, LastTargetIndirection.GetProperty());
 						}
-					}
-					else
-					{
-						SourceStruct = SourceDataView.GetStruct();
 					}
 				}
 
@@ -431,6 +509,154 @@ struct FCachedBindingData : public TSharedFromThis<FCachedBindingData>
 		bIsDataCached = true;
 	}
 
+	bool CanBindToContextStruct(const UStruct* InStruct)
+	{
+		ConditionallyUpdateData();
+
+		// Do not allow to bind directly StateTree nodes
+		// @todo: find a way to more specifically call out the context structs, e.g. pass the property path to the callback.
+		if (InStruct != nullptr)
+		{
+			const bool bIsStateTreeNode = AccessibleStructs.ContainsByPredicate([InStruct](const FStateTreeBindableStructDesc& AccessibleStruct)
+			{
+				return (AccessibleStruct.DataSource != EStateTreeBindableStructSource::Context && AccessibleStruct.DataSource != EStateTreeBindableStructSource::Parameter)
+					&& AccessibleStruct.Struct == InStruct;
+			});
+
+			if (bIsStateTreeNode)
+			{
+				return false;
+			}
+		}
+
+		return ArePropertyAndContextStructCompatible(InStruct, PropertyHandle->GetProperty());
+	}
+			
+	bool CanBindToProperty(const FProperty* SourceProperty)
+	{
+		ConditionallyUpdateData();
+
+		// Special case for binding widget calling OnCanBindProperty with Args.Property (i.e. self).
+		if (PropertyHandle->GetProperty() == SourceProperty)
+		{
+			return true;
+		}
+
+		void* TargetValueAddress = nullptr;
+		if (PropertyHandle->GetValueData(TargetValueAddress) == FPropertyAccess::Success)
+		{
+			return ArePropertiesCompatible(SourceProperty, PropertyHandle->GetProperty(), TargetValueAddress);
+		}
+		
+		return false;
+	}
+
+	static bool ArePropertyAndContextStructCompatible(const UStruct* SourceStruct, const FProperty* TargetProperty)
+	{
+		check(SourceStruct);
+		
+		if (const FStructProperty* TargetStructProperty = CastField<FStructProperty>(TargetProperty))
+		{
+			return TargetStructProperty->Struct == SourceStruct;
+		}
+		if (const FObjectProperty* TargetObjectProperty = CastField<FObjectProperty>(TargetProperty))
+		{
+			return SourceStruct != nullptr && SourceStruct->IsChildOf(TargetObjectProperty->PropertyClass);
+		}
+		
+		return false;
+	}
+
+	static bool ArePropertiesCompatible(const FProperty* SourceProperty, const FProperty* TargetProperty, const void* TargetPropertyValue)
+	{
+		// @TODO: Refactor FStateTreePropertyBindings::ResolveCopyType() so that we can use it directly here.
+		
+		bool bCanBind = false;
+
+		const FStructProperty* TargetStructProperty = CastField<FStructProperty>(TargetProperty);
+		
+		// AnyEnums need special handling.
+		// It is a struct property but we want to treat it as an enum. We need to do this here, instead of 
+		// FStateTreePropertyBindingCompiler::GetPropertyCompatibility() because the treatment depends on the value too.
+		// Note: AnyEnums will need special handling before they can be used for binding.
+		if (TargetStructProperty && TargetStructProperty->Struct == FStateTreeAnyEnum::StaticStruct())
+		{
+			// If the AnyEnum has AllowAnyBinding, allow to bind to any enum.
+			const bool bAllowAnyBinding = TargetProperty->HasMetaData(UE::StateTree::PropertyBinding::AllowAnyBindingName);
+
+			check(TargetPropertyValue);
+			const FStateTreeAnyEnum* TargetAnyEnum = static_cast<const FStateTreeAnyEnum*>(TargetPropertyValue);
+
+			// If the enum class is not specified, allow to bind to any enum, if the class is specified allow only that enum.
+			if (const FByteProperty* SourceByteProperty = CastField<FByteProperty>(SourceProperty))
+			{
+				if (UEnum* Enum = SourceByteProperty->GetIntPropertyEnum())
+				{
+					bCanBind = bAllowAnyBinding || TargetAnyEnum->Enum == Enum;
+				}
+			}
+			else if (const FEnumProperty* SourceEnumProperty = CastField<FEnumProperty>(SourceProperty))
+			{
+				bCanBind = bAllowAnyBinding || TargetAnyEnum->Enum == SourceEnumProperty->GetEnum();
+			}
+		}
+		else if (TargetStructProperty && TargetStructProperty->Struct == FStateTreeStructRef::StaticStruct())
+		{
+			FString BaseStructName;
+			const UScriptStruct* TargetStructRefBaseStruct = UE::StateTree::Compiler::GetBaseStructFromMetaData(TargetProperty, BaseStructName);
+
+			if (const FStructProperty* SourceStructProperty = CastField<FStructProperty>(SourceProperty))
+			{
+				if (SourceStructProperty->Struct == TBaseStructure<FStateTreeStructRef>::Get())
+				{
+					FString SourceBaseStructName;
+					const UScriptStruct* SourceStructRefBaseStruct = UE::StateTree::Compiler::GetBaseStructFromMetaData(SourceStructProperty, SourceBaseStructName);
+					bCanBind = SourceStructRefBaseStruct && SourceStructRefBaseStruct->IsChildOf(TargetStructRefBaseStruct);
+				}
+				else
+				{
+					bCanBind = SourceStructProperty->Struct && SourceStructProperty->Struct->IsChildOf(TargetStructRefBaseStruct);
+				}
+			}
+		}
+		else
+		{
+			// Note: We support type promotion here
+			bCanBind = FStateTreePropertyBindings::GetPropertyCompatibility(SourceProperty, TargetProperty) != EStateTreePropertyAccessCompatibility::Incompatible;
+		}
+
+		return bCanBind;
+	}
+
+	UStruct* ResolveIndirection(TArray<FBindingChainElement> InBindingChain)
+	{
+		UObject* OwnerObject = WeakOwnerObject.Get();
+		if (!OwnerObject)
+		{
+			return nullptr;
+		}
+		
+		IStateTreeEditorPropertyBindingsOwner* BindingOwner = Cast<IStateTreeEditorPropertyBindingsOwner>(OwnerObject);
+		if (!BindingOwner)
+		{
+			return nullptr;
+		}
+
+		const int32 SourceStructIndex = InBindingChain[0].ArrayIndex;
+		check(SourceStructIndex >= 0 && SourceStructIndex < AccessibleStructs.Num());
+		
+		TArray<FBindingChainElement> SourceBindingChain = InBindingChain;
+		SourceBindingChain.RemoveAt(0);
+
+		FStateTreeDataView DataView;
+		if (BindingOwner->GetDataViewByID(AccessibleStructs[SourceStructIndex].ID, DataView))
+		{
+			return UE::StateTree::PropertyBinding::ResolveLeafValueStructType(DataView, InBindingChain);
+		}
+
+		return nullptr;
+	}
+
 	FText GetText()
 	{
 		ConditionallyUpdateData();
@@ -468,11 +694,13 @@ private:
 	TWeakObjectPtr<UObject> WeakOwnerObject = nullptr;
 	FStateTreePropertyPath TargetPath;
 	TSharedPtr<const IPropertyHandle> PropertyHandle;
+	TArray<FStateTreeBindableStructDesc> AccessibleStructs;
 
 	FText Text;
 	FText TooltipText;
 	FLinearColor Color = FLinearColor::White;
 	const FSlateBrush* Image = nullptr;
+	
 	bool bIsDataCached = false;
 };
 
@@ -517,7 +745,7 @@ void FStateTreeBindingExtension::ExtendWidgetRow(FDetailWidgetRow& InWidgetRow,
 	FStateTreeEditorPropertyBindings* EditorBindings = nullptr;
 
 	// Array of structs we can bind to.
-	TArray<FBindingContextStruct> Context;
+	TArray<FBindingContextStruct> BindingContextStructs;
 	TArray<FStateTreeBindableStructDesc> AccessibleStructs;
 
 	// The struct and property where we're binding.
@@ -544,7 +772,7 @@ void FStateTreeBindingExtension::ExtendWidgetRow(FDetailWidgetRow& InWidgetRow,
 			{
 				const UStruct* Struct = StructDesc.Struct;
 
-				FBindingContextStruct& ContextStruct = Context.AddDefaulted_GetRef();
+				FBindingContextStruct& ContextStruct = BindingContextStructs.AddDefaulted_GetRef();
 				ContextStruct.DisplayText = FText::FromString(StructDesc.Name.ToString());
 				ContextStruct.Struct = const_cast<UStruct*>(Struct);
 				ContextStruct.Section = UE::StateTree::PropertyBinding::GetSectionNameFromDataSource(StructDesc.DataSource);
@@ -572,120 +800,19 @@ void FStateTreeBindingExtension::ExtendWidgetRow(FDetailWidgetRow& InWidgetRow,
 		
 	}
 
-	FProperty* Property = InPropertyHandle->GetProperty();
-
-	bool bIsStructRef = false;
-	bool bIsAnyEnum = false;
-	if (FStructProperty* StructProperty = CastField<FStructProperty>(Property))
-	{
-		bIsAnyEnum = StructProperty->Struct == FStateTreeAnyEnum::StaticStruct();
-		bIsStructRef = StructProperty->Struct == FStateTreeStructRef::StaticStruct();
-	}
-
-	const UScriptStruct* StructRefBaseStruct = nullptr;
-	if (bIsStructRef)
-	{
-		FString BaseStructName;
-		StructRefBaseStruct = UE::StateTree::Compiler::GetBaseStructFromMetaData(Property, BaseStructName);
-	}
-
-	TSharedPtr<UE::StateTree::PropertyBinding::FCachedBindingData> CachedBindingData = MakeShared<UE::StateTree::PropertyBinding::FCachedBindingData>(OwnerObject, TargetPath, InPropertyHandle);
+	TSharedPtr<UE::StateTree::PropertyBinding::FCachedBindingData> CachedBindingData = MakeShared<UE::StateTree::PropertyBinding::FCachedBindingData>(OwnerObject, TargetPath, InPropertyHandle, AccessibleStructs);
 	
 	FPropertyBindingWidgetArgs Args;
 	Args.Property = InPropertyHandle->GetProperty();
 
-	Args.OnCanBindProperty = FOnCanBindProperty::CreateLambda([EditorBindings, OwnerObject, InPropertyHandle, bIsAnyEnum, bIsStructRef, StructRefBaseStruct](FProperty* InProperty)
+	Args.OnCanBindProperty = FOnCanBindProperty::CreateLambda([CachedBindingData](FProperty* InProperty)
 		{
-			if (!EditorBindings || !OwnerObject)
-			{
-				return false;
-			}
-
-			// Special case for binding widget calling OnCanBindProperty with Args.Property (i.e. self).
-			if (InPropertyHandle->GetProperty() == InProperty)
-			{
-				return true;
-			}
-
-			bool bCanBind = false;
-
-			// AnyEnums need special handling.
-			// It is a struct property but we want to treat it as an enum. We need to do this here, instead of 
-			// FStateTreePropertyBindingCompiler::GetPropertyCompatibility() because the treatment depends on the value too.
-			// Note: AnyEnums will need special handling before they can be used for binding.
-			if (bIsAnyEnum)
-			{
-				// If the AnyEnum has AllowAnyBinding, allow to bind to any enum.
-				const bool bAllowAnyBinding = InPropertyHandle->HasMetaData(UE::StateTree::PropertyBinding::AllowAnyBindingName);
-
-				FStateTreeAnyEnum AnyEnum;
-				UE::StateTree::PropertyHelpers::GetStructValue(InPropertyHandle, AnyEnum);
-
-				// If the enum class is not specified, allow to bind to any enum, if the class is specified allow only that enum.
-				if (const FByteProperty* ByteProperty = CastField<FByteProperty>(InProperty))
-				{
-					if (UEnum* Enum = ByteProperty->GetIntPropertyEnum())
-					{
-						bCanBind = bAllowAnyBinding || AnyEnum.Enum == Enum;
-					}
-				}
-				else if (const FEnumProperty* EnumProperty = CastField<FEnumProperty>(InProperty))
-				{
-					bCanBind = bAllowAnyBinding || AnyEnum.Enum == EnumProperty->GetEnum();
-				}
-			}
-			else if (bIsStructRef && StructRefBaseStruct != nullptr)
-			{
-				if (const FStructProperty* SourceStructProperty = CastField<FStructProperty>(InProperty))
-				{
-					if (SourceStructProperty->Struct == TBaseStructure<FStateTreeStructRef>::Get())
-					{
-						FString SourceBaseStructName;
-						const UScriptStruct* SourceDataRefBaseStruct = UE::StateTree::Compiler::GetBaseStructFromMetaData(SourceStructProperty, SourceBaseStructName);
-						bCanBind = SourceDataRefBaseStruct && SourceDataRefBaseStruct->IsChildOf(StructRefBaseStruct);
-					}
-					else
-					{
-						bCanBind = SourceStructProperty->Struct && SourceStructProperty->Struct->IsChildOf(StructRefBaseStruct);
-					}
-				}
-			}
-			else
-			{
-				// Note: We support type promotion here
-				bCanBind = FStateTreePropertyBindings::GetPropertyCompatibility(InProperty, InPropertyHandle->GetProperty()) != EStateTreePropertyAccessCompatibility::Incompatible;
-			}
-
-			return bCanBind;
+			return CachedBindingData->CanBindToProperty(InProperty);
 		});
 
-	Args.OnCanBindToContextStruct = FOnCanBindToContextStruct::CreateLambda([InPropertyHandle, AccessibleStructs](const UStruct* InStruct)
+	Args.OnCanBindToContextStruct = FOnCanBindToContextStruct::CreateLambda([CachedBindingData](const UStruct* InStruct)
 		{
-			// Do not allow to bind directly StateTree nodes
-			// @todo: find a way to more specifically call out the context structs, e.g. pass the property path to the callback.
-			if (InStruct != nullptr)
-			{
-				bool bIsStateTreeNode = AccessibleStructs.ContainsByPredicate([InStruct](const FStateTreeBindableStructDesc& AccessibleStruct)
-				{
-					return (AccessibleStruct.DataSource != EStateTreeBindableStructSource::Context && AccessibleStruct.DataSource != EStateTreeBindableStructSource::Parameter)
-						&& AccessibleStruct.Struct == InStruct;
-				});
-
-				if (bIsStateTreeNode)
-				{
-					return false;
-				}
-			}
-		
-			if (const FStructProperty* StructProperty = CastField<FStructProperty>(InPropertyHandle->GetProperty()))
-			{
-				return StructProperty->Struct == InStruct;
-			}
-			if (const FObjectProperty* ObjectProperty = CastField<FObjectProperty>(InPropertyHandle->GetProperty()))
-			{
-				return InStruct != nullptr && InStruct->IsChildOf(ObjectProperty->PropertyClass);
-			}
-			return false;
+			return CachedBindingData->CanBindToContextStruct(InStruct);
 		});
 
 	Args.OnCanAcceptPropertyOrChildren = FOnCanBindProperty::CreateLambda([](FProperty* InProperty)
@@ -699,53 +826,14 @@ void FStateTreeBindingExtension::ExtendWidgetRow(FDetailWidgetRow& InWidgetRow,
 			return true;
 		});
 
-	Args.OnAddBinding = FOnAddBinding::CreateLambda([EditorBindings, OwnerObject, BindingOwner, TargetPath, AccessibleStructs, CachedBindingData](FName InPropertyName, const TArray<FBindingChainElement>& InBindingChain)
+	Args.OnAddBinding = FOnAddBinding::CreateLambda([CachedBindingData](FName InPropertyName, const TArray<FBindingChainElement>& InBindingChain)
 		{
-			IPropertyAccessEditor& PropertyAccessEditor = IModularFeatures::Get().GetModularFeature<IPropertyAccessEditor>("PropertyAccessEditor");
-			if (EditorBindings && OwnerObject)
-			{
-				if (TargetPath.GetStructID().IsValid() && InBindingChain.Num() > 0)
-				{
-					// First item in the binding chain is the index in AccessibleStructs.
-					const int32 SourceStructIndex = InBindingChain[0].ArrayIndex;
-						
-					TArray<FBindingChainElement> SourceBindingChain = InBindingChain;
-					SourceBindingChain.RemoveAt(0); // remove struct index.
-
-					check(SourceStructIndex >= 0 && SourceStructIndex < AccessibleStructs.Num());
-
-					FStateTreeDataView DataView;
-					if (BindingOwner)
-					{
-						BindingOwner->GetDataViewByID(AccessibleStructs[SourceStructIndex].ID, DataView);
-					}
-
-					// If SourceBindingChain is empty at this stage, it means that the binding points to the source struct itself.
-					FStateTreePropertyPath SourcePath;
-					UE::StateTree::PropertyBinding::MakeStructPropertyPathFromBindingChain(AccessibleStructs[SourceStructIndex].ID, SourceBindingChain, DataView, SourcePath);
-						
-					OwnerObject->Modify();
-					EditorBindings->AddPropertyBinding(SourcePath, TargetPath);
-
-					UE::StateTree::PropertyBinding::OnStateTreePropertyBindingChanged.Broadcast(SourcePath, TargetPath);
-				}
-			}
-
-			CachedBindingData->UpdateData();
+			CachedBindingData->AddBinding(InBindingChain);
 		});
 
-	Args.OnRemoveBinding = FOnRemoveBinding::CreateLambda([EditorBindings, OwnerObject, TargetPath, CachedBindingData](FName InPropertyName)
+	Args.OnRemoveBinding = FOnRemoveBinding::CreateLambda([CachedBindingData](FName InPropertyName)
 		{
-			if (EditorBindings && OwnerObject)
-			{
-				OwnerObject->Modify();
-				EditorBindings->RemovePropertyBindings(TargetPath);
-
-				const FStateTreePropertyPath SourcePath; // Null path
-				UE::StateTree::PropertyBinding::OnStateTreePropertyBindingChanged.Broadcast(SourcePath, TargetPath);
-			}
-
-			CachedBindingData->UpdateData();
+			CachedBindingData->RemoveBinding();
 		});
 
 	Args.OnCanRemoveBinding = FOnCanRemoveBinding::CreateLambda([EditorBindings, TargetPath](FName InPropertyName)
@@ -775,21 +863,9 @@ void FStateTreeBindingExtension::ExtendWidgetRow(FDetailWidgetRow& InWidgetRow,
 
 	if (BindingOwner)
 	{
-		Args.OnResolveIndirection = FOnResolveIndirection::CreateLambda([BindingOwner, AccessibleStructs](TArray<FBindingChainElement> InBindingChain)
+		Args.OnResolveIndirection = FOnResolveIndirection::CreateLambda([CachedBindingData](TArray<FBindingChainElement> InBindingChain)
 		{
-			const int32 SourceStructIndex = InBindingChain[0].ArrayIndex;
-			TArray<FBindingChainElement> SourceBindingChain = InBindingChain;
-			SourceBindingChain.RemoveAt(0);
-
-			check(SourceStructIndex >= 0 && SourceStructIndex < AccessibleStructs.Num());
-
-			FStateTreeDataView DataView;
-			if (BindingOwner->GetDataViewByID(AccessibleStructs[SourceStructIndex].ID, DataView))
-			{
-				return UE::StateTree::PropertyBinding::ResolveLeafValueStructType(DataView, InBindingChain);
-			}
-
-			return static_cast<UStruct*>(nullptr);
+			return CachedBindingData->ResolveIndirection(InBindingChain);
 		});
 	}
 
@@ -800,7 +876,7 @@ void FStateTreeBindingExtension::ExtendWidgetRow(FDetailWidgetRow& InWidgetRow,
 
 	InWidgetRow.ExtensionContent()
 	[
-		PropertyAccessEditor.MakePropertyBindingWidget(Context, Args)
+		PropertyAccessEditor.MakePropertyBindingWidget(BindingContextStructs, Args)
 	];
 }
 
diff --git a/Engine/Plugins/Runtime/StateTree/Source/StateTreeEditorModule/Private/StateTreePropertyHelpers.h b/Engine/Plugins/Runtime/StateTree/Source/StateTreeEditorModule/Private/StateTreePropertyHelpers.h
--- a/Engine/Plugins/Runtime/StateTree/Source/StateTreeEditorModule/Private/StateTreePropertyHelpers.h
+++ b/Engine/Plugins/Runtime/StateTree/Source/StateTreeEditorModule/Private/StateTreePropertyHelpers.h
@@ -19,7 +19,7 @@ namespace UE::StateTree::PropertyHelpers {
  * @return Requested value as optional, in case of multiple values the optional is unset.
  */
 template<typename T>
-FPropertyAccess::Result GetStructValue(const TSharedPtr<IPropertyHandle>& ValueProperty, T& OutValue)
+FPropertyAccess::Result GetStructValue(const TSharedPtr<const IPropertyHandle>& ValueProperty, T& OutValue)
 {
 	if (!ValueProperty)
 	{
@@ -33,13 +33,13 @@ FPropertyAccess::Result GetStructValue(const TSharedPtr<IPropertyHandle>& ValueP
 	T Value = T();
 	bool bValueSet = false;
 
-	TArray<void*> RawData;
+	TArray<const void*> RawData;
 	ValueProperty->AccessRawData(RawData);
-	for (void* Data : RawData)
+	for (const void* Data : RawData)
 	{
 		if (Data)
 		{
-			const T& CurValue = *reinterpret_cast<T*>(Data);
+			const T& CurValue = *static_cast<const T*>(Data);
 			if (!bValueSet)
 			{
 				bValueSet = true;
