StateTree: Fixed crash when trying to add binding in a tree with no parameters

- Remove check(SourceStruct) from ArePropertyAndContextStructCompatible, which was triggering when adding a binding for a tree that has no parameters (SourceStruct is nullptr for empty params)

diff --git a/Engine/Plugins/Runtime/StateTree/Source/StateTreeEditorModule/Private/Customizations/StateTreeBindingExtension.cpp b/Engine/Plugins/Runtime/StateTree/Source/StateTreeEditorModule/Private/Customizations/StateTreeBindingExtension.cpp
--- a/Engine/Plugins/Runtime/StateTree/Source/StateTreeEditorModule/Private/Customizations/StateTreeBindingExtension.cpp
+++ b/Engine/Plugins/Runtime/StateTree/Source/StateTreeEditorModule/Private/Customizations/StateTreeBindingExtension.cpp
@@ -553,8 +553,6 @@ struct FCachedBindingData : public TSharedFromThis<FCachedBindingData>
 
 	static bool ArePropertyAndContextStructCompatible(const UStruct* SourceStruct, const FProperty* TargetProperty)
 	{
-		check(SourceStruct);
-		
 		if (const FStructProperty* TargetStructProperty = CastField<FStructProperty>(TargetProperty))
 		{
 			return TargetStructProperty->Struct == SourceStruct;
