StateTree: Patch transition IDs on older assets deterministically.

diff --git a/Engine/Plugins/Runtime/StateTree/Source/StateTreeEditorModule/Private/StateTreeState.cpp b/Engine/Plugins/Runtime/StateTree/Source/StateTreeEditorModule/Private/StateTreeState.cpp
--- a/Engine/Plugins/Runtime/StateTree/Source/StateTreeEditorModule/Private/StateTreeState.cpp
+++ b/Engine/Plugins/Runtime/StateTree/Source/StateTreeEditorModule/Private/StateTreeState.cpp
@@ -26,18 +26,6 @@ FStateTreeTransition::FStateTreeTransition(const EStateTreeTransitionTrigger InT
 	State = InState ? InState->GetLinkToState() : FStateTreeStateLink(InType);
 }
 
-void FStateTreeTransition::PostSerialize(const FArchive& Ar)
-{
-#if WITH_EDITORONLY_DATA
-	const int32 CurrentVersion = Ar.CustomVer(FStateTreeCustomVersion::GUID);
-	if (CurrentVersion < FStateTreeCustomVersion::AddedTransitionIds)
-	{
-		ID = FGuid::NewGuid();
-	}
-#endif // WITH_EDITORONLY_DAT
-}
-
-
 //////////////////////////////////////////////////////////////////////////
 // UStateTreeState
 
@@ -311,7 +299,19 @@ void UStateTreeState::PostLoad()
 		}		
 	}
 	PRAGMA_ENABLE_DEPRECATION_WARNINGS
-#endif
+
+	const int32 CurrentVersion = GetLinkerCustomVersion(FStateTreeCustomVersion::GUID);
+	if (CurrentVersion < FStateTreeCustomVersion::AddedTransitionIds)
+	{
+		// Make guids for transitions. These need to be deterministic when upgrading because of cooking.
+		for (int32 Index = 0; Index < Transitions.Num(); Index++)
+		{
+			FStateTreeTransition& Transition = Transitions[Index];
+			Transition.ID = FGuid::NewDeterministicGuid(GetPathName(), Index);
+		}
+	}
+#endif // WITH_EDITORONLY_DATA
+
 }
 
 void UStateTreeState::UpdateParametersFromLinkedSubtree()
diff --git a/Engine/Plugins/Runtime/StateTree/Source/StateTreeEditorModule/Public/StateTreeState.h b/Engine/Plugins/Runtime/StateTree/Source/StateTreeEditorModule/Public/StateTreeState.h
--- a/Engine/Plugins/Runtime/StateTree/Source/StateTreeEditorModule/Public/StateTreeState.h
+++ b/Engine/Plugins/Runtime/StateTree/Source/StateTreeEditorModule/Public/StateTreeState.h
@@ -20,8 +20,6 @@ struct STATETREEEDITORMODULE_API FStateTreeTransition
 	FStateTreeTransition(const EStateTreeTransitionTrigger InTrigger, const EStateTreeTransitionType InType, const UStateTreeState* InState = nullptr);
 	FStateTreeTransition(const EStateTreeTransitionTrigger InTrigger, const FGameplayTag InEventTag, const EStateTreeTransitionType InType, const UStateTreeState* InState = nullptr);
 
-	void PostSerialize(const FArchive& Ar);
-	
 	template<typename T, typename... TArgs>
 	TStateTreeEditorNode<T>& AddCondition(TArgs&&... InArgs)
 	{
@@ -80,14 +78,6 @@ struct STATETREEEDITORMODULE_API FStateTreeTransition
 	bool bTransitionEnabled = true;
 };
 
-template<>
-struct TStructOpsTypeTraits<FStateTreeTransition> : TStructOpsTypeTraitsBase2<FStateTreeTransition>
-{
-	enum
-	{
-		WithPostSerialize = true,
-	};
-};
 
 USTRUCT()
 struct STATETREEEDITORMODULE_API FStateTreeStateParameters
