StateTree: Fixed binding validation and binding display update issues
- Added Modify() for the validation functions so that they behave nice with undo (e.g. a binding may be removed by validation)
- Changed FCachedBindingData::ConditionallyUpdateData() to check if the binding has changed so that the caching reacts to undos and other modifications
- Changed FStateTreePropertyPath::ResolveIndirectionsWithValue() to fall back to the property object type, on null objects (previous logic was preventing some valid cases on empty instances)

diff --git a/Engine/Plugins/Runtime/StateTree/Source/StateTreeEditorModule/Private/Customizations/StateTreeBindingExtension.cpp b/Engine/Plugins/Runtime/StateTree/Source/StateTreeEditorModule/Private/Customizations/StateTreeBindingExtension.cpp
--- a/Engine/Plugins/Runtime/StateTree/Source/StateTreeEditorModule/Private/Customizations/StateTreeBindingExtension.cpp
+++ b/Engine/Plugins/Runtime/StateTree/Source/StateTreeEditorModule/Private/Customizations/StateTreeBindingExtension.cpp
@@ -1,14 +1,12 @@
 // Copyright Epic Games, Inc. All Rights Reserved.
 
 #include "StateTreeBindingExtension.h"
-#include "Blueprint/StateTreeNodeBlueprintBase.h"
 #include "EdGraphSchema_K2.h"
 #include "Features/IModularFeatures.h"
 #include "IPropertyAccessEditor.h"
 #include "StateTreeAnyEnum.h"
 #include "StateTreeCompiler.h"
 #include "StateTreeEditorPropertyBindings.h"
-#include "StateTreePropertyHelpers.h"
 #include "StateTreeNodeBase.h"
 #include "Styling/AppStyle.h"
 #include "UObject/EnumProperty.h"
@@ -496,6 +494,8 @@ struct FCachedBindingData : public TSharedFromThis<FCachedBindingData>
 				Image = FCoreStyle::Get().GetBrush("Icons.ErrorWithColor");
 				Color = FLinearColor::White;
 			}
+
+			CachedSourcePath = *SourcePath;
 		}
 		else
 		{
@@ -504,6 +504,8 @@ struct FCachedBindingData : public TSharedFromThis<FCachedBindingData>
 			TooltipText = FText::Format(LOCTEXT("BindTooltip", "Bind {0} to value from another property."), UE::StateTree::PropertyBinding::GetPropertyTypeText(Property));
 			Image = FAppStyle::GetBrush(PropertyIcon);
 			Color = Schema->GetPinTypeColor(PinType);
+
+			CachedSourcePath.Reset();
 		}
 
 		bIsDataCached = true;
@@ -683,13 +685,43 @@ private:
 
 	void ConditionallyUpdateData()
 	{
-		if (!bIsDataCached)
+		UObject* OwnerObject = WeakOwnerObject.Get();
+		if (!OwnerObject)
+		{
+			return;
+		}
+		
+		IStateTreeEditorPropertyBindingsOwner* BindingOwner = Cast<IStateTreeEditorPropertyBindingsOwner>(OwnerObject);
+		if (!BindingOwner)
+		{
+			return;
+		}
+
+		FStateTreeEditorPropertyBindings* EditorBindings = BindingOwner->GetPropertyEditorBindings();
+		if (!EditorBindings)
+		{
+			return;
+		}
+
+		const FStateTreePropertyPath* CurrentSourcePath = EditorBindings->GetPropertyBindingSource(TargetPath);
+		bool bPathsIdentical = false;
+		if (CurrentSourcePath)
+		{
+			bPathsIdentical = CachedSourcePath == *CurrentSourcePath;
+		}
+		else
+		{
+			bPathsIdentical = CachedSourcePath.IsPathEmpty();
+		}
+
+		if (!bIsDataCached || !bPathsIdentical)
 		{
 			UpdateData();
 		}
 	}
 	
 	TWeakObjectPtr<UObject> WeakOwnerObject = nullptr;
+	FStateTreePropertyPath CachedSourcePath;
 	FStateTreePropertyPath TargetPath;
 	TSharedPtr<const IPropertyHandle> PropertyHandle;
 	TArray<FStateTreeBindableStructDesc> AccessibleStructs;
diff --git a/Engine/Plugins/Runtime/StateTree/Source/StateTreeEditorModule/Private/StateTreeEditor.cpp b/Engine/Plugins/Runtime/StateTree/Source/StateTreeEditorModule/Private/StateTreeEditor.cpp
--- a/Engine/Plugins/Runtime/StateTree/Source/StateTreeEditorModule/Private/StateTreeEditor.cpp
+++ b/Engine/Plugins/Runtime/StateTree/Source/StateTreeEditorModule/Private/StateTreeEditor.cpp
@@ -603,6 +603,8 @@ namespace UE::StateTree::Editor::Internal
 			return;
 		}
 
+		TreeData->Modify();
+
 		// Make sure all state links are valid and update the names if needed.
 
 		// Create ID to state name map.
@@ -617,6 +619,7 @@ namespace UE::StateTree::Editor::Internal
 		// Fix changed names.
 		TreeData->VisitHierarchy([&IDToName](UStateTreeState& State, UStateTreeState* /*ParentState*/)
 		{
+			State.Modify();
 			if (State.Type == EStateTreeStateType::Linked)
 			{
 				FixChangedStateLinkName(State.LinkedSubtree, IDToName);
@@ -639,6 +642,7 @@ namespace UE::StateTree::Editor::Internal
 			return;
 		}
 
+		TreeData->Modify();
 		TreeData->ReparentStates();
 	}
 
@@ -656,6 +660,8 @@ namespace UE::StateTree::Editor::Internal
 			return;
 		}
 
+		TreeData->Modify();
+		
 		// Clear evaluators if not allowed.
 		if (Schema->AllowEvaluators() == false && TreeData->Evaluators.Num() > 0)
 		{
@@ -666,6 +672,8 @@ namespace UE::StateTree::Editor::Internal
 
 		TreeData->VisitHierarchy([&StateTree, Schema](UStateTreeState& State, UStateTreeState* /*ParentState*/)
 		{
+			State.Modify();
+
 			// Clear enter conditions if not allowed.
 			if (Schema->AllowEnterConditions() == false && State.EnterConditions.Num() > 0)
 			{
@@ -711,6 +719,7 @@ namespace UE::StateTree::Editor::Internal
 
 		TMap<FGuid, const FStateTreeDataView> AllStructValues;
 		TreeData->GetAllStructValues(AllStructValues);
+		TreeData->Modify();
 		TreeData->GetPropertyEditorBindings()->RemoveUnusedBindings(AllStructValues);
 	}
 
@@ -722,10 +731,13 @@ namespace UE::StateTree::Editor::Internal
 			return;
 		}
 
+		TreeData->Modify();
+
 		TreeData->VisitHierarchy([](UStateTreeState& State, UStateTreeState* /*ParentState*/)
 		{
 			if (State.Type == EStateTreeStateType::Linked)
 			{
+				State.Modify();
 				State.UpdateParametersFromLinkedSubtree();
 			}
 			return EStateTreeVisitor::Continue;
diff --git a/Engine/Plugins/Runtime/StateTree/Source/StateTreeEditorModule/Private/StateTreeEditorData.cpp b/Engine/Plugins/Runtime/StateTree/Source/StateTreeEditorModule/Private/StateTreeEditorData.cpp
--- a/Engine/Plugins/Runtime/StateTree/Source/StateTreeEditorModule/Private/StateTreeEditorData.cpp
+++ b/Engine/Plugins/Runtime/StateTree/Source/StateTreeEditorModule/Private/StateTreeEditorData.cpp
@@ -197,6 +197,7 @@ void UStateTreeEditorData::PostEditChangeChainProperty(FPropertyChangedChainEven
 			{
 				TMap<FGuid, const FStateTreeDataView> AllStructValues;
 				GetAllStructValues(AllStructValues);
+				Modify();
 				EditorBindings.RemoveUnusedBindings(AllStructValues);
 			}
 		}
diff --git a/Engine/Plugins/Runtime/StateTree/Source/StateTreeModule/Private/StateTreePropertyBindings.cpp b/Engine/Plugins/Runtime/StateTree/Source/StateTreeModule/Private/StateTreePropertyBindings.cpp
--- a/Engine/Plugins/Runtime/StateTree/Source/StateTreeModule/Private/StateTreePropertyBindings.cpp
+++ b/Engine/Plugins/Runtime/StateTree/Source/StateTreeModule/Private/StateTreePropertyBindings.cpp
@@ -1530,6 +1530,9 @@ bool FStateTreePropertyPath::ResolveIndirectionsWithValue(const FStateTreeDataVi
 				{
 					if (StructProperty->Struct == TBaseStructure<FInstancedStruct>::Get())
 					{
+						// The property path is pointing into the instanced struct, it must be present.
+						// @TODO:	We could potentially check the BaseStruct metadata in editor (for similar behavior as objects)
+						//			Omitting for now to have matching functionality in editor and runtime.
 						const FInstancedStruct& InstancedStruct = *reinterpret_cast<const FInstancedStruct*>(CurrentAddress + Offset);
 						if (!InstancedStruct.IsValid())
 						{
@@ -1576,18 +1579,10 @@ bool FStateTreePropertyPath::ResolveIndirectionsWithValue(const FStateTreeDataVi
 				{
 					const UObject* Object = *reinterpret_cast<UObject* const*>(CurrentAddress + Offset);
 					CurrentAddress = reinterpret_cast<const uint8*>(Object);
-					if (Property->HasAnyPropertyFlags(CPF_PersistentInstance | CPF_InstancedReference))
+					
+					// The property path is pointing into the object, if the object is present use it's specific type, otherwise use the type of the pointer.
+					if (Object)
 					{
-						if (!Object)
-						{
-							if (OutError)
-							{
-								*OutError = FString::Printf(TEXT("Expecting valid instanced object value at path '%s'."),
-									*ToString(Segment.GetIndex(), TEXT("<"), TEXT(">")));
-							}
-							OutIndirections.Reset();
-							return false;
-						}
 						CurrentStruct = Object->GetClass();
 						Indirection.InstanceStruct = CurrentStruct;
 						Indirection.AccessType = EStateTreePropertyAccessType::ObjectInstance;
