StateTrees: Made StateTreeLinker handle IInterface types

diff --git a/Engine/Plugins/Runtime/StateTree/Source/StateTreeModule/Private/StateTreeExecutionContext.cpp b/Engine/Plugins/Runtime/StateTree/Source/StateTreeModule/Private/StateTreeExecutionContext.cpp
--- a/Engine/Plugins/Runtime/StateTree/Source/StateTreeModule/Private/StateTreeExecutionContext.cpp
+++ b/Engine/Plugins/Runtime/StateTree/Source/StateTreeModule/Private/StateTreeExecutionContext.cpp
@@ -123,10 +123,17 @@ bool FStateTreeExecutionContext::AreExternalDataViewsValid() const
 	{
 		const FStateTreeDataView& DataView = DataViews[DataDesc.Handle.DataViewIndex.Get()];
 			
+		auto IsAssignmentValid = [](const FStateTreeExternalDataDesc& DataDesc, const FStateTreeDataView& DataView)
+		{	
+			const UClass* DataDescClass = Cast<UClass>(DataDesc.Struct);
+			const UClass* DataViewClass = Cast<UClass>(DataView.GetStruct());
+			return DataView.GetStruct()->IsChildOf(DataDesc.Struct) || (DataViewClass && DataDescClass && DataViewClass->ImplementsInterface(DataDescClass));
+		};
+
 		if (DataDesc.Requirement == EStateTreeExternalDataRequirement::Required)
 		{
 			// Required items must have valid pointer of the expected type.  
-			if (!DataView.IsValid() || !DataView.GetStruct()->IsChildOf(DataDesc.Struct))
+			if (!DataView.IsValid() || !IsAssignmentValid(DataDesc, DataView))
 			{
 				bResult = false;
 				break;
@@ -135,7 +142,7 @@ bool FStateTreeExecutionContext::AreExternalDataViewsValid() const
 		else
 		{
 			// Optional items must have same type if they are set.
-			if (DataView.IsValid() && !DataView.GetStruct()->IsChildOf(DataDesc.Struct))
+			if (DataView.IsValid() && !IsAssignmentValid(DataDesc, DataView))
 			{
 				bResult = false;
 				break;
diff --git a/Engine/Plugins/Runtime/StateTree/Source/StateTreeModule/Public/StateTreeLinker.h b/Engine/Plugins/Runtime/StateTree/Source/StateTreeModule/Public/StateTreeLinker.h
--- a/Engine/Plugins/Runtime/StateTree/Source/StateTreeModule/Public/StateTreeLinker.h
+++ b/Engine/Plugins/Runtime/StateTree/Source/StateTreeModule/Public/StateTreeLinker.h
@@ -4,6 +4,7 @@
 
 #include "StateTreeSchema.h"
 #include "StateTreeExecutionTypes.h"
+#include "Templates/Casts.h"
 #include "StateTreeLinker.generated.h"
 
 UENUM()
@@ -44,15 +45,25 @@ struct FStateTreeLinker
 	}
 
 	/**
-	 * Links reference to an external UObject.
+	 * Links reference to an external UStruct.
 	 * @param Handle Reference to TStateTreeExternalDataHandle<> with USTRUCT type to link to.
 	 */
 	template <typename T>
-	typename TEnableIf<!TIsDerivedFrom<typename T::DataType, UObject>::IsDerived, void>::Type LinkExternalData(T& Handle)
+	typename TEnableIf<!TIsDerivedFrom<typename T::DataType, UObject>::IsDerived && !TIsIInterface<typename T::DataType>::Value, void>::Type LinkExternalData(T& Handle)
 	{
 		LinkExternalData(Handle, T::DataType::StaticStruct(), T::DataRequirement);
 	}
 
+	/**
+	 * Links reference to an external IInterface.
+	 * @param Handle Reference to TStateTreeExternalDataHandle<> with IINTERFACE type to link to.
+	 */
+	template <typename T>
+	typename TEnableIf<TIsIInterface<typename T::DataType>::Value, void>::Type LinkExternalData(T& Handle)
+	{
+		LinkExternalData(Handle, T::DataType::UClassType::StaticClass(), T::DataRequirement);
+	}
+
 	/**
 	 * Links reference to an external Object or Struct.
 	 * This function should only be used when TStateTreeExternalDataHandle<> cannot be used, i.e. the Struct is based on some data.
diff --git a/Engine/Plugins/Runtime/StateTree/Source/StateTreeModule/Public/StateTreeTypes.h b/Engine/Plugins/Runtime/StateTree/Source/StateTreeModule/Public/StateTreeTypes.h
--- a/Engine/Plugins/Runtime/StateTree/Source/StateTreeModule/Public/StateTreeTypes.h
+++ b/Engine/Plugins/Runtime/StateTree/Source/StateTreeModule/Public/StateTreeTypes.h
@@ -688,7 +688,7 @@ struct STATETREEMODULE_API FStateTreeDataView
 	 * USTRUCT() getters (reference & pointer, const & mutable)
 	 */
 	template <typename T>
-	typename TEnableIf<!TIsDerivedFrom<T, UObject>::IsDerived, const T&>::Type Get() const
+	typename TEnableIf<!TIsDerivedFrom<T, UObject>::IsDerived && !TIsIInterface<T>::Value, const T&>::Type Get() const
 	{
 		check(Memory != nullptr);
 		check(Struct != nullptr);
@@ -697,7 +697,7 @@ struct STATETREEMODULE_API FStateTreeDataView
 	}
 
 	template <typename T>
-    typename TEnableIf<!TIsDerivedFrom<T, UObject>::IsDerived, T&>::Type GetMutable() const
+    typename TEnableIf<!TIsDerivedFrom<T, UObject>::IsDerived && !TIsIInterface<T>::Value, T&>::Type GetMutable() const
 	{
 		check(Memory != nullptr);
 		check(Struct != nullptr);
@@ -706,7 +706,7 @@ struct STATETREEMODULE_API FStateTreeDataView
 	}
 
 	template <typename T>
-    typename TEnableIf<!TIsDerivedFrom<T, UObject>::IsDerived, const T*>::Type GetPtr() const
+    typename TEnableIf<!TIsDerivedFrom<T, UObject>::IsDerived && !TIsIInterface<T>::Value, const T*>::Type GetPtr() const
 	{
 		// If Memory is set, expect Struct too. Otherwise, let nulls pass through.
 		check(!Memory || (Memory && Struct));
@@ -715,7 +715,7 @@ struct STATETREEMODULE_API FStateTreeDataView
 	}
 
 	template <typename T>
-    typename TEnableIf<!TIsDerivedFrom<T, UObject>::IsDerived, T*>::Type GetMutablePtr() const
+    typename TEnableIf<!TIsDerivedFrom<T, UObject>::IsDerived && !TIsIInterface<T>::Value, T*>::Type GetMutablePtr() const
 	{
 		// If Memory is set, expect Struct too. Otherwise, let nulls pass through.
 		check(!Memory || (Memory && Struct));
@@ -723,6 +723,41 @@ struct STATETREEMODULE_API FStateTreeDataView
 		return ((T*)Memory);
 	}
 
+	/*
+	 * IInterface() getters (reference & pointer, const & mutable)
+	 */
+	template <typename T>
+	typename TEnableIf<TIsIInterface<T>::Value, const T&>::Type Get() const
+	{
+		check(!Memory || (Memory && Struct));
+		check(Struct->IsChildOf(UObject::StaticClass()) && ((UClass*)Struct)->ImplementsInterface(T::UClassType::StaticClass()));
+		return *(T*)((UObject*)Memory)->GetInterfaceAddress(T::UClassType::StaticClass());
+	}
+
+	template <typename T>
+    typename TEnableIf<TIsIInterface<T>::Value, T&>::Type GetMutable() const
+	{
+		check(!Memory || (Memory && Struct));
+		check(Struct->IsChildOf(UObject::StaticClass()) && ((UClass*)Struct)->ImplementsInterface(T::UClassType::StaticClass()));
+		return *(T*)((UObject*)Memory)->GetInterfaceAddress(T::UClassType::StaticClass());
+	}
+
+	template <typename T>
+    typename TEnableIf<TIsIInterface<T>::Value, const T*>::Type GetPtr() const
+	{
+		check(!Memory || (Memory && Struct));
+		check(Struct->IsChildOf(UObject::StaticClass()) && ((UClass*)Struct)->ImplementsInterface(T::UClassType::StaticClass()));
+		return (T*)((UObject*)Memory)->GetInterfaceAddress(T::UClassType::StaticClass());
+	}
+
+	template <typename T>
+    typename TEnableIf<TIsIInterface<T>::Value, T*>::Type GetMutablePtr() const
+	{
+		check(!Memory || (Memory && Struct));
+		check(Struct->IsChildOf(UObject::StaticClass()) && ((UClass*)Struct)->ImplementsInterface(T::UClassType::StaticClass()));
+		return (T*)((UObject*)Memory)->GetInterfaceAddress(T::UClassType::StaticClass());
+	}
+
 	/** @return Struct describing the data type. */
 	const UStruct* GetStruct() const { return Struct; }
 
