StateTree: Changed how instance data is accessed
- Remove DataView concept
- Added FStateTreeExecutionFrame, which defines an active branch of a specific state tree that is running
- Added FStateTreeDataHandle, which allows frame relative access of instance data (replaces the DataViews)
- Changed linked states to spawn new frames (will later extend to handle linked trees on other assets)
- Bumped StateTree data version number, requires recomplation of StateTrees

diff --git a/Engine/Plugins/Runtime/StateTree/Source/StateTreeEditorModule/Private/StateTreeCompiler.cpp b/Engine/Plugins/Runtime/StateTree/Source/StateTreeEditorModule/Private/StateTreeCompiler.cpp
--- a/Engine/Plugins/Runtime/StateTree/Source/StateTreeEditorModule/Private/StateTreeCompiler.cpp
+++ b/Engine/Plugins/Runtime/StateTree/Source/StateTreeEditorModule/Private/StateTreeCompiler.cpp
@@ -240,21 +240,19 @@ bool FStateTreeCompiler::Compile(UStateTree& InStateTree)
 	// Copy parameters from EditorData	
 	StateTree->Parameters = EditorData->RootParameters.Parameters;
 
+	int32 ContextDataIndex = 0;
+	
 	// Mark parameters as binding source
 	const FStateTreeBindableStructDesc ParametersDesc = {
 			TEXT("Parameters"),
 			StateTree->Parameters.GetPropertyBagStruct(),
+			FStateTreeDataHandle(EStateTreeDataSourceType::ContextData, ContextDataIndex++),
 			EStateTreeBindableStructSource::Parameter,
 			EditorData->RootParameters.ID
 		};
-	const int32 ParametersDataViewIndex = BindingsCompiler.AddSourceStruct(ParametersDesc);
+	BindingsCompiler.AddSourceStruct(ParametersDesc);
 
-	if (const auto Validation = UE::StateTree::Compiler::IsValidIndex8(ParametersDataViewIndex); Validation.DidFail())
-	{
-		Validation.Log(Log, TEXT("ParametersDataViewIndex"), ParametersDesc);
-		return false;
-	}
-	StateTree->ParametersDataViewIndex = FStateTreeIndex8(ParametersDataViewIndex); 
+	StateTree->ParametersDataHandle = ParametersDesc.DataHandle;
 
 	if (!UE::StateTree::Compiler::ValidateNoLevelActorReferences(Log, ParametersDesc, FStateTreeDataView(), FStateTreeDataView(EditorData->RootParameters.Parameters.GetMutableValue())))
 	{
@@ -271,18 +269,21 @@ bool FStateTreeCompiler::Compile(UStateTree& InStateTree)
 			const FStateTreeBindableStructDesc ExtDataDesc = {
 					Desc.Name,
 					Desc.Struct,
+					FStateTreeDataHandle(EStateTreeDataSourceType::ContextData, ContextDataIndex++),
 					EStateTreeBindableStructSource::Context,
 					Desc.ID
 				};
-			const int32 ExternalStructIndex = BindingsCompiler.AddSourceStruct(ExtDataDesc);
-			if (const auto Validation = UE::StateTree::Compiler::IsValidIndex16(ExternalStructIndex); Validation.DidFail())
+			BindingsCompiler.AddSourceStruct(ExtDataDesc);
+			if (const auto Validation = UE::StateTree::Compiler::IsValidIndex16(ContextDataIndex); Validation.DidFail())
 			{
 				Validation.Log(Log, TEXT("ExternalStructIndex"), ParametersDesc);
 				return false;
 			}
-			Desc.Handle.DataViewIndex = FStateTreeIndex16(ExternalStructIndex); 
+			Desc.Handle.DataHandle = ExtDataDesc.DataHandle;
 		} 
 	}
+
+	StateTree->NumContextData = ContextDataIndex;
 	
 	if (!CreateStates())
 	{
@@ -290,19 +291,29 @@ bool FStateTreeCompiler::Compile(UStateTree& InStateTree)
 		return false;
 	}
 
+	// Eval and Global task methods use InstanceStructs.Num() as ID generator.
+	check(InstanceStructs.Num() == 0);
+	
 	if (!CreateEvaluators())
 	{
 		StateTree->ResetCompiled();
 		return false;
 	}
 
-	
 	if (!CreateGlobalTasks())
 	{
 		StateTree->ResetCompiled();
 		return false;
 	}
 
+	const int32 NumGlobalInstanceData = InstanceStructs.Num();
+	if (const auto Validation = UE::StateTree::Compiler::IsValidIndex16(NumGlobalInstanceData); Validation.DidFail())
+	{
+		Validation.Log(Log, TEXT("NumGlobalInstanceData"), ParametersDesc);
+		return false;
+	}
+	StateTree->NumGlobalInstanceData = uint16(NumGlobalInstanceData);
+
 	if (!CreateStateTasksAndParameters())
 	{
 		StateTree->ResetCompiled();
@@ -316,8 +327,8 @@ bool FStateTreeCompiler::Compile(UStateTree& InStateTree)
 	}
 
 	StateTree->Nodes = Nodes;
-	StateTree->DefaultInstanceData.Init(*StateTree, InstanceStructs, InstanceObjects);
-	StateTree->SharedInstanceData.Init(*StateTree, SharedInstanceStructs, SharedInstanceObjects);
+	StateTree->DefaultInstanceData.Init(*StateTree, InstanceStructs);
+	StateTree->SharedInstanceData.Init(*StateTree, SharedInstanceStructs);
 	
 	BindingsCompiler.Finalize();
 
@@ -378,9 +389,6 @@ bool FStateTreeCompiler::CreateStates()
 {
 	check(EditorData);
 	
-	// Create item for the runtime execution state
-	InstanceStructs.Add(FInstancedStruct::Make<FStateTreeExecutionState>());
-
 	// Create main tree (omit subtrees)
 	for (UStateTreeState* SubTree : EditorData->SubTrees)
 	{
@@ -432,7 +440,7 @@ bool FStateTreeCompiler::CreateStateRecursive(UStateTreeState& State, const FSta
 
 	CompactState.Type = State.Type;
 	CompactState.SelectionBehavior = State.SelectionBehavior;
-	
+
 	SourceStates.Add(&State);
 	IDToState.Add(State.ID, StateIdx);
 
@@ -506,7 +514,9 @@ bool FStateTreeCompiler::CreateEvaluators()
 
 	for (FStateTreeEditorNode& EvalNode : EditorData->Evaluators)
 	{
-		if (!CreateEvaluator(EvalNode))
+		const int32 GlobalInstanceIndex = InstanceStructs.Num();
+		const FStateTreeDataHandle EvalDataHandle(EStateTreeDataSourceType::GlobalInstanceData, GlobalInstanceIndex);
+		if (!CreateEvaluator(EvalNode, EvalDataHandle))
 		{
 			return false;
 		}
@@ -545,7 +555,9 @@ bool FStateTreeCompiler::CreateGlobalTasks()
 			continue;
 		}
 
-		if (!CreateTask(nullptr, TaskNode))
+		const int32 GlobalInstanceIndex = InstanceStructs.Num();
+		const FStateTreeDataHandle TaskDataHandle(EStateTreeDataSourceType::GlobalInstanceData, GlobalInstanceIndex);
+		if (!CreateTask(nullptr, TaskNode, TaskDataHandle))
 		{
 			return false;
 		}
@@ -570,105 +582,98 @@ bool FStateTreeCompiler::CreateStateTasksAndParameters()
 {
 	check(StateTree);
 
+	// Index of the first instance data per state. Accumulated depth first.
+	TArray<int32> FirstInstanceDataIndex;
+	FirstInstanceDataIndex.Init(0, StateTree->States.Num());
+	
 	for (int32 i = 0; i < StateTree->States.Num(); i++)
 	{
 		FCompactStateTreeState& CompactState = StateTree->States[i];
-		UStateTreeState* SourceState = SourceStates[i];
-		check(SourceState != nullptr);
+		const FStateTreeStateHandle CompactStateHandle(i);
+		UStateTreeState* State = SourceStates[i];
+		check(State != nullptr);
 
-		FStateTreeCompilerLogStateScope LogStateScope(SourceState, Log);
+		// Carry over instance data count from parent.
+		if (CompactState.Parent.IsValid())
+		{
+			const FCompactStateTreeState& ParentCompactState = StateTree->States[CompactState.Parent.Index];
+			const int32 InstanceDataBegin = FirstInstanceDataIndex[CompactState.Parent.Index] + (int32)ParentCompactState.InstanceDataNum;
+			FirstInstanceDataIndex[i] = InstanceDataBegin;
+		}
+
+		int32 InstanceDataIndex = FirstInstanceDataIndex[i];
+
+		FStateTreeCompilerLogStateScope LogStateScope(State, Log);
 
 		// Create parameters
-		if (SourceState->Type == EStateTreeStateType::Linked || SourceState->Type == EStateTreeStateType::Subtree)
+		if (State->Type == EStateTreeStateType::Linked
+			|| State->Type == EStateTreeStateType::Subtree)
 		{
-			// Both linked and subtree has instance data describing their parameters.
-			// This allows to resolve the binding paths and lets us have bindable parameters when transitioned into a parameterized subtree directly.
+			// Each state has their parameters as instance data.
 			FInstancedStruct& Instance = InstanceStructs.AddDefaulted_GetRef();
+			Instance.InitializeAs<FCompactStateTreeParameters>(State->Parameters.Parameters);
+			FCompactStateTreeParameters& CompactStateTreeParameters = Instance.GetMutable<FCompactStateTreeParameters>(); 
+			
 			const int32 InstanceIndex = InstanceStructs.Num() - 1;
 			if (const auto Validation = UE::StateTree::Compiler::IsValidIndex16(InstanceIndex); Validation.DidFail())
 			{
 				Validation.Log(Log, TEXT("InstanceIndex"));
 				return false;
 			}
-			CompactState.ParameterInstanceIndex = FStateTreeIndex16(InstanceIndex);
-		
-			Instance.InitializeAs<FCompactStateTreeParameters>();
-			FCompactStateTreeParameters& CompactParams = Instance.GetMutable<FCompactStateTreeParameters>();
+			CompactState.ParameterTemplateIndex = FStateTreeIndex16(InstanceIndex);
 
-			CompactParams.Parameters = SourceState->Parameters.Parameters;
-
-			if (SourceState->Type == EStateTreeStateType::Subtree)
+			if (State->Type == EStateTreeStateType::Linked)
 			{
-				// Register a binding source if we have parameters
-				int32 SourceStructIndex = INDEX_NONE;
-
-				if (SourceState->Parameters.Parameters.IsValid())
-				{
-					const FStateTreeBindableStructDesc SubtreeParamsDesc = {
-						SourceState->Name,
-						SourceState->Parameters.Parameters.GetPropertyBagStruct(),
-						EStateTreeBindableStructSource::State,
-						SourceState->Parameters.ID
-					};
-
-					if (!UE::StateTree::Compiler::ValidateNoLevelActorReferences(Log, SubtreeParamsDesc, FStateTreeDataView(), FStateTreeDataView(CompactParams.Parameters.GetMutableValue())))
-					{
-						return false;
-					}
-					
-					SourceStructIndex = BindingsCompiler.AddSourceStruct(SubtreeParamsDesc);
-					
-					if (const auto Validation = UE::StateTree::Compiler::IsValidIndex16(SourceStructIndex); Validation.DidFail())
-					{
-						Validation.Log(Log, TEXT("SourceStructIndex"), SubtreeParamsDesc);
-						return false;
-					}
-				}
-				
-				CompactState.ParameterDataViewIndex = FStateTreeIndex16(SourceStructIndex);
+				CompactState.ParameterDataHandle = FStateTreeDataHandle(EStateTreeDataSourceType::LinkedStateParameterData, InstanceDataIndex++, CompactStateHandle);
 			}
-			else if (SourceState->Type == EStateTreeStateType::Linked)
+			else if (State->Type == EStateTreeStateType::Subtree)
 			{
-				int32 BatchIndex = INDEX_NONE;
+				// Subtrees have special indirection, which is used to access the parameters passed to the state.
+				// If subtree is called via linked state, the linked state's parameters are used, otherwise the instance data.
+				CompactState.ParameterDataHandle = FStateTreeDataHandle(EStateTreeDataSourceType::SubtreeParameterData, InstanceDataIndex++, CompactStateHandle);
+			}
 
-				if (SourceState->Parameters.Parameters.IsValid())
-				{
-					// Binding target
-					FStateTreeBindableStructDesc LinkedParamsDesc = {
-						SourceState->Name,
-						SourceState->Parameters.Parameters.GetPropertyBagStruct(),
-						EStateTreeBindableStructSource::State,
-						SourceState->Parameters.ID
-					};
-
-					if (!UE::StateTree::Compiler::ValidateNoLevelActorReferences(Log, LinkedParamsDesc, FStateTreeDataView(), FStateTreeDataView(CompactParams.Parameters.GetMutableValue())))
-					{
-						return false;
-					}
+			// @todo: We should be able to skip empty parameter data.
 
-					// Check that the bindings for this struct are still all valid.
-					TArray<FStateTreePropertyPathBinding> Bindings;
-					if (!GetAndValidateBindings(LinkedParamsDesc, FStateTreeDataView(SourceState->Parameters.Parameters.GetMutableValue()), Bindings))
-					{
-						return false;
-					}
+			// Binding target
+			FStateTreeBindableStructDesc LinkedParamsDesc = {
+				State->Name,
+				State->Parameters.Parameters.GetPropertyBagStruct(),
+				CompactState.ParameterDataHandle,
+				EStateTreeBindableStructSource::State,
+				State->Parameters.ID
+			};
 
-					if (!BindingsCompiler.CompileBatch(LinkedParamsDesc, Bindings, BatchIndex))
-					{
-						return false;
-					}
-					
-					if (const auto Validation = UE::StateTree::Compiler::IsValidIndex16(BatchIndex); Validation.DidFail())
-					{
-						Validation.Log(Log, TEXT("BatchIndex"), LinkedParamsDesc);
-						return false;
-					}
-				}
+			if (!UE::StateTree::Compiler::ValidateNoLevelActorReferences(Log, LinkedParamsDesc, FStateTreeDataView(), FStateTreeDataView(CompactStateTreeParameters.Parameters.GetMutableValue())))
+			{
+				return false;
+			}
+
+			// Add as binding source.
+			BindingsCompiler.AddSourceStruct(LinkedParamsDesc);
+
+			// Check that the bindings for this struct are still all valid.
+			TArray<FStateTreePropertyPathBinding> Bindings;
+			if (!GetAndValidateBindings(LinkedParamsDesc, FStateTreeDataView(CompactStateTreeParameters.Parameters.GetMutableValue()), Bindings))
+			{
+				return false;
+			}
 
-				CompactParams.BindingsBatch = FStateTreeIndex16(BatchIndex);
+			int32 BatchIndex = INDEX_NONE;
+			if (!BindingsCompiler.CompileBatch(LinkedParamsDesc, Bindings, BatchIndex))
+			{
+				return false;
+			}
+			
+			if (const auto Validation = UE::StateTree::Compiler::IsValidIndex16(BatchIndex); Validation.DidFail())
+			{
+				Validation.Log(Log, TEXT("BatchIndex"), LinkedParamsDesc);
+				return false;
 			}
+
+			CompactState.ParameterBindingsBatch = FStateTreeIndex16(BatchIndex);
 		}
-		
+
 		// Create tasks
 		const int32 TasksBegin = Nodes.Num();
 		if (const auto Validation = UE::StateTree::Compiler::IsValidCount16(TasksBegin); Validation.DidFail())
@@ -678,24 +683,14 @@ bool FStateTreeCompiler::CreateStateTasksAndParameters()
 		}
 		CompactState.TasksBegin = uint16(TasksBegin);
 		
-		int32 TaskInstanceStructNum = 0;
-		int32 TaskInstanceObjectNum = 0;
-
-		// Update instance data num for each state.
-		if (CompactState.Type == EStateTreeStateType::Linked)
-		{
-			// Linked state parameters.
-			TaskInstanceStructNum++;
-		}
-
 		TArrayView<FStateTreeEditorNode> Tasks;
-		if (SourceState->Tasks.Num())
+		if (State->Tasks.Num())
 		{
-			Tasks = SourceState->Tasks;
+			Tasks = State->Tasks;
 		}
-		else if (SourceState->SingleTask.Node.IsValid())
+		else if (State->SingleTask.Node.IsValid())
 		{
-			Tasks = TArrayView<FStateTreeEditorNode>(&SourceState->SingleTask, 1);
+			Tasks = TArrayView<FStateTreeEditorNode>(&State->SingleTask, 1);
 		}
 		
 		bool bStateHasTransitionTasks = false;
@@ -707,46 +702,34 @@ bool FStateTreeCompiler::CreateStateTasksAndParameters()
 				continue;
 			}
 
-			if (!CreateTask(SourceState, TaskNode))
+			const FStateTreeDataHandle TaskDataHandle(EStateTreeDataSourceType::ActiveInstanceData, InstanceDataIndex++, CompactStateHandle);
+			if (!CreateTask(State, TaskNode, TaskDataHandle))
 			{
 				return false;
 			}
 
 			const FStateTreeTaskBase& LastAddedTask = Nodes.Last().Get<FStateTreeTaskBase>();
-			if (LastAddedTask.bInstanceIsObject)
-			{
-				TaskInstanceObjectNum++;
-			}
-			else
-			{
-				TaskInstanceStructNum++;
-			}
 			
 			bStateHasTransitionTasks |= LastAddedTask.bShouldAffectTransitions;
 		}
 
 		CompactState.bHasTransitionTasks = bStateHasTransitionTasks;
-	
+		
 		const int32 TasksNum = Nodes.Num() - TasksBegin;
 		if (const auto Validation = UE::StateTree::Compiler::IsValidCount8(TasksNum); Validation.DidFail())
 		{
 			Validation.Log(Log, TEXT("TasksNum"));
 			return false;
 		}
-		if (const auto Validation = UE::StateTree::Compiler::IsValidCount8(TaskInstanceObjectNum); Validation.DidFail())
-		{
-			Validation.Log(Log, TEXT("TaskInstanceObjectNum"));
-			return false;
-		}
-		if (const auto Validation = UE::StateTree::Compiler::IsValidCount8(TaskInstanceStructNum); Validation.DidFail())
+
+		if (const auto Validation = UE::StateTree::Compiler::IsValidCount8(InstanceDataIndex); Validation.DidFail())
 		{
-			Validation.Log(Log, TEXT("TaskInstanceStructNum"));
+			Validation.Log(Log, TEXT("InstanceDataIndex"));
 			return false;
 		}
 
 		CompactState.TasksNum = uint8(TasksNum);
-		CompactState.TaskInstanceStructNum = uint8(TaskInstanceStructNum);
-		CompactState.TaskInstanceObjectNum = uint8(TaskInstanceObjectNum);
+		CompactState.InstanceDataNum = uint8(InstanceDataIndex);
 	}
 	
 	return true;
@@ -1056,17 +1039,19 @@ bool FStateTreeCompiler::CreateCondition(UStateTreeState& State, const FStateTre
 			Validation.Log(Log, TEXT("InstanceIndex"), StructDesc);
 			return false;
 		}
-		Cond.InstanceIndex = FStateTreeIndex16(InstanceIndex);
-		Cond.bInstanceIsObject = false;
+		Cond.InstanceTemplateIndex = FStateTreeIndex16(InstanceIndex);
+		Cond.InstanceDataHandle = FStateTreeDataHandle(EStateTreeDataSourceType::SharedInstanceData, InstanceIndex);
 		InstanceDataView = FStateTreeDataView(SharedInstanceStructs[InstanceIndex]);
 	}
 	else
 	{
 		// Object Instance
 		check(CondNode.InstanceObject != nullptr);
-		
+
 		UObject* Instance = DuplicateObject(CondNode.InstanceObject, StateTree);
-		const int32 InstanceIndex = SharedInstanceObjects.Add(Instance);
+		FInstancedStruct Wrapper;
+		Wrapper.InitializeAs<FStateTreeInstanceObjectWrapper>(Instance);
+		const int32 InstanceIndex = SharedInstanceStructs.Add(MoveTemp(Wrapper));
 		
 		// Create binding source struct descriptor.
 		StructDesc.Struct = Instance->GetClass();
@@ -1077,18 +1062,20 @@ bool FStateTreeCompiler::CreateCondition(UStateTreeState& State, const FStateTre
 			Validation.Log(Log, TEXT("InstanceIndex"), StructDesc);
 			return false;
 		}
-		Cond.InstanceIndex = FStateTreeIndex16(InstanceIndex);
-		Cond.bInstanceIsObject = true;
+		Cond.InstanceTemplateIndex = FStateTreeIndex16(InstanceIndex);
+		Cond.InstanceDataHandle = FStateTreeDataHandle(EStateTreeDataSourceType::SharedInstanceDataObject, InstanceIndex);
 		InstanceDataView = FStateTreeDataView(Instance);
 	}
 
+	StructDesc.DataHandle = Cond.InstanceDataHandle;
+	
 	if (!CompileAndValidateNode(&State, StructDesc, Node, InstanceDataView))
 	{
 		return false;
 	}
 
 	// Mark the struct as binding source.
-	const int32 SourceStructIndex = BindingsCompiler.AddSourceStruct(StructDesc);
+	BindingsCompiler.AddSourceStruct(StructDesc);
 
 	// Check that the bindings for this struct are still all valid.
 	TArray<FStateTreePropertyPathBinding> Bindings;
@@ -1110,13 +1097,6 @@ bool FStateTreeCompiler::CreateCondition(UStateTreeState& State, const FStateTre
 		return false;
 	}
 	Cond.BindingsBatch = FStateTreeIndex16(BatchIndex);
-
-	if (const auto Validation = UE::StateTree::Compiler::IsValidIndex16(SourceStructIndex); Validation.DidFail())
-	{
-		Validation.Log(Log, TEXT("SourceStructIndex"), StructDesc);
-		return false;
-	}
-	Cond.DataViewIndex = FStateTreeIndex16(SourceStructIndex);
 	
 	return true;
 }
@@ -1184,7 +1164,7 @@ bool FStateTreeCompiler::CompileAndValidateNode(const UStateTreeState* SourceSta
 	return Result != EDataValidationResult::Invalid;
 }
 
-bool FStateTreeCompiler::CreateTask(UStateTreeState* State, const FStateTreeEditorNode& TaskNode)
+bool FStateTreeCompiler::CreateTask(UStateTreeState* State, const FStateTreeEditorNode& TaskNode, const FStateTreeDataHandle TaskDataHandle)
 {
 	if (!TaskNode.Node.IsValid())
 	{
@@ -1228,8 +1208,8 @@ bool FStateTreeCompiler::CreateTask(UStateTreeState* State, const FStateTreeEdit
 			Validation.Log(Log, TEXT("InstanceIndex"), StructDesc);
 			return false;
 		}
-		Task.InstanceIndex = FStateTreeIndex16(InstanceIndex);
-		Task.bInstanceIsObject = false;
+		Task.InstanceTemplateIndex = FStateTreeIndex16(InstanceIndex);
+		Task.InstanceDataHandle = TaskDataHandle;
 		InstanceDataView = FStateTreeDataView(InstanceStructs[InstanceIndex]);
 	}
 	else
@@ -1238,7 +1218,9 @@ bool FStateTreeCompiler::CreateTask(UStateTreeState* State, const FStateTreeEdit
 		check(TaskNode.InstanceObject != nullptr);
 
 		UObject* Instance = DuplicateObject(TaskNode.InstanceObject, StateTree);
-		const int32 InstanceIndex = InstanceObjects.Add(Instance);
+		FInstancedStruct Wrapper;
+		Wrapper.InitializeAs<FStateTreeInstanceObjectWrapper>(Instance);
+		const int32 InstanceIndex = InstanceStructs.Add(MoveTemp(Wrapper));
 
 		// Create binding source struct descriptor.
 		StructDesc.Struct = Instance->GetClass();
@@ -1249,18 +1231,20 @@ bool FStateTreeCompiler::CreateTask(UStateTreeState* State, const FStateTreeEdit
 			Validation.Log(Log, TEXT("InstanceIndex"), StructDesc);
 			return false;
 		}
-		Task.InstanceIndex = FStateTreeIndex16(InstanceIndex);
-		Task.bInstanceIsObject = true;
+		Task.InstanceTemplateIndex = FStateTreeIndex16(InstanceIndex);
+		Task.InstanceDataHandle = TaskDataHandle.ToObjectSource();
 		InstanceDataView = FStateTreeDataView(Instance);
 	}
 
+	StructDesc.DataHandle = Task.InstanceDataHandle;
+
 	if (!CompileAndValidateNode(State, StructDesc, Node,  InstanceDataView))
 	{
 		return false;
 	}
 
 	// Mark the instance as binding source.
-	const int32 SourceStructIndex = BindingsCompiler.AddSourceStruct(StructDesc);
+	BindingsCompiler.AddSourceStruct(StructDesc);
 	
 	// Check that the bindings for this struct are still all valid.
 	TArray<FStateTreePropertyPathBinding> Bindings;
@@ -1282,25 +1266,18 @@ bool FStateTreeCompiler::CreateTask(UStateTreeState* State, const FStateTreeEdit
 		return false;
 	}
 	Task.BindingsBatch = FStateTreeIndex16(BatchIndex);
-
-	if (const auto Validation = UE::StateTree::Compiler::IsValidIndex16(SourceStructIndex); Validation.DidFail())
-	{
-		Validation.Log(Log, TEXT("SourceStructIndex"), StructDesc);
-		return false;
-	}
-	Task.DataViewIndex = FStateTreeIndex16(SourceStructIndex);
 	
 	return true;
 }
 
-bool FStateTreeCompiler::CreateEvaluator(const FStateTreeEditorNode& EvalNode)
+bool FStateTreeCompiler::CreateEvaluator(const FStateTreeEditorNode& EvalNode, const FStateTreeDataHandle EvalDataHandle)
 {
 	// Silently ignore empty nodes.
 	if (!EvalNode.Node.IsValid())
 	{
 		return true;
 	}
-	
+
 	// Create binding source struct descriptor.
 	FStateTreeBindableStructDesc StructDesc;
     StructDesc.ID = EvalNode.ID;
@@ -1338,8 +1315,8 @@ bool FStateTreeCompiler::CreateEvaluator(const FStateTreeEditorNode& EvalNode)
 			Validation.Log(Log, TEXT("InstanceIndex"), StructDesc);
 			return false;
 		}
-		Eval.InstanceIndex = FStateTreeIndex16(InstanceIndex);
-		Eval.bInstanceIsObject = false;
+		Eval.InstanceTemplateIndex = FStateTreeIndex16(InstanceIndex);
+		Eval.InstanceDataHandle = EvalDataHandle;
 		InstanceDataView = FStateTreeDataView(InstanceStructs[InstanceIndex]);
 	}
 	else
@@ -1348,8 +1325,10 @@ bool FStateTreeCompiler::CreateEvaluator(const FStateTreeEditorNode& EvalNode)
 		check(EvalNode.InstanceObject != nullptr);
 
 		UObject* Instance = DuplicateObject(EvalNode.InstanceObject, StateTree);
-		const int32 InstanceIndex = InstanceObjects.Add(Instance);
-
+		FInstancedStruct Wrapper;
+		Wrapper.InitializeAs<FStateTreeInstanceObjectWrapper>(Instance);
+		const int32 InstanceIndex = InstanceStructs.Add(MoveTemp(Wrapper));
+		
 		// Create binding source struct descriptor.
 		StructDesc.Struct = Instance->GetClass();
 		StructDesc.Name = Eval.Name;
@@ -1359,18 +1338,20 @@ bool FStateTreeCompiler::CreateEvaluator(const FStateTreeEditorNode& EvalNode)
 			Validation.Log(Log, TEXT("InstanceIndex"), StructDesc);
 			return false;
 		}
-		Eval.InstanceIndex = FStateTreeIndex16(InstanceIndex);
-		Eval.bInstanceIsObject = true;
+		Eval.InstanceTemplateIndex = FStateTreeIndex16(InstanceIndex);
+		Eval.InstanceDataHandle = EvalDataHandle.ToObjectSource();
 		InstanceDataView = FStateTreeDataView(Instance);
 	}
 
+	StructDesc.DataHandle = Eval.InstanceDataHandle;
+
 	if (!CompileAndValidateNode(nullptr, StructDesc, Node,  InstanceDataView))
 	{
 		return false;
 	}
 
 	// Mark the instance as binding source.
-	const int32 SourceStructIndex = BindingsCompiler.AddSourceStruct(StructDesc);
+	BindingsCompiler.AddSourceStruct(StructDesc);
 
 	// Check that the bindings for this struct are still all valid.
 	TArray<FStateTreePropertyPathBinding> Bindings;
@@ -1393,13 +1374,6 @@ bool FStateTreeCompiler::CreateEvaluator(const FStateTreeEditorNode& EvalNode)
 	}
 	Eval.BindingsBatch = FStateTreeIndex16(BatchIndex);
 
-	if (const auto Validation = UE::StateTree::Compiler::IsValidIndex16(SourceStructIndex); Validation.DidFail())
-	{
-		Validation.Log(Log, TEXT("SourceStructIndex"), StructDesc);
-		return false;
-	}
-	Eval.DataViewIndex = FStateTreeIndex16(SourceStructIndex);
-
 	return true;
 }
 
@@ -1509,43 +1483,19 @@ bool FStateTreeCompiler::ValidateStructRef(const FStateTreeBindableStructDesc& S
 	return true;
 }
 
-FStateTreeDataView FStateTreeCompiler::GetBindingSourceValue(const int32 SourceIndex)
-{
-	for (const FInstancedStruct& Node : Nodes)
-	{
-		if (const FStateTreeNodeBase* NodeBase = Node.GetPtr<FStateTreeNodeBase>())
-		{
-			if (NodeBase->DataViewIndex.Get() == SourceIndex)
-			{
-				if (NodeBase->bInstanceIsObject)
-				{
-					return FStateTreeDataView(InstanceObjects[NodeBase->InstanceIndex.Get()]);
-				}
-				else
-				{
-					return FStateTreeDataView(InstanceStructs[NodeBase->InstanceIndex.Get()]);
-				}
-			}
-		}
-	}
-
-	return {};
-}
 
 bool FStateTreeCompiler::GetAndValidateBindings(const FStateTreeBindableStructDesc& TargetStruct, FStateTreeDataView TargetValue, TArray<FStateTreePropertyPathBinding>& OutBindings) const
 {
 	check(EditorData);
 	
+	OutBindings.Reset();
+
+	// If target struct is not set, nothing to do.
 	if (TargetStruct.Struct == nullptr)
 	{
-		Log.Reportf(EMessageSeverity::Error, TargetStruct,
-				TEXT("The type of binding target %s is invalid."),
-				*TargetStruct.ToString());
-		return false;
+		return true;
 	}
-	
-	OutBindings.Reset();
-	
+
 	for (FStateTreePropertyPathBinding& Binding : EditorData->EditorBindings.GetMutableBindings())
 	{
 		if (Binding.GetTargetPath().GetStructID() != TargetStruct.ID)
@@ -1555,15 +1505,14 @@ bool FStateTreeCompiler::GetAndValidateBindings(const FStateTreeBindableStructDe
 
 		// Source must be one of the source structs we have discovered in the tree.
 		const FGuid SourceStructID = Binding.GetSourcePath().GetStructID();
-		const int32 SourceStructIdx = BindingsCompiler.GetSourceStructIndexByID(SourceStructID);
-		if (SourceStructIdx == INDEX_NONE)
+		const FStateTreeBindableStructDesc* SourceStruct = BindingsCompiler.GetSourceStructDescByID(SourceStructID);
+		if (!SourceStruct)
 		{
 			Log.Reportf(EMessageSeverity::Error, TargetStruct,
 						TEXT("Failed to find binding source property '%s' for target %s."),
 						*Binding.GetSourcePath().ToString(), *UE::StateTree::GetDescAndPathAsString(TargetStruct, Binding.GetTargetPath()));
 			return false;
 		}
-		const FStateTreeBindableStructDesc& SourceStruct = BindingsCompiler.GetSourceStructDesc(SourceStructIdx);
 
 		// Update path instance types from latest data. E.g. binding may have been created for instanced object of type FooB, and changed to FooA.
  		// @todo: not liking how this mutates the Binding.TargetPath, but currently we dont track well the instanced object changes.
@@ -1589,9 +1538,9 @@ bool FStateTreeCompiler::GetAndValidateBindings(const FStateTreeBindableStructDe
 		{
 			Log.Reportf(EMessageSeverity::Error, TargetStruct,
 						TEXT("Property at %s cannot be bound to %s, because the binding source %s is not updated before %s in the tree."),
-						*UE::StateTree::GetDescAndPathAsString(SourceStruct, Binding.GetSourcePath()),
+						*UE::StateTree::GetDescAndPathAsString(*SourceStruct, Binding.GetSourcePath()),
 						*UE::StateTree::GetDescAndPathAsString(TargetStruct, Binding.GetTargetPath()),
-						*SourceStruct.ToString(), *TargetStruct.ToString());
+						*SourceStruct->ToString(), *TargetStruct.ToString());
 			return false;
 		}
 
@@ -1616,31 +1565,38 @@ bool FStateTreeCompiler::GetAndValidateBindings(const FStateTreeBindableStructDe
 			}
 		}
 
+		if (!SourceStruct->DataHandle.IsValid())
+		{
+			Log.Reportf(EMessageSeverity::Error, TargetStruct,
+				TEXT("Malformed source'%s for property binding property '%s'."),
+				*UE::StateTree::GetDescAndPathAsString(*SourceStruct, Binding.GetSourcePath()), *Binding.GetSourcePath().ToString());
+			return false;
+		}
+		
+		FStateTreePropertyPathBinding BindingCopy(Binding);
+		BindingCopy.SetSourceDataHandle(SourceStruct->DataHandle);
+
 		// Special case fo AnyEnum. StateTreeBindingExtension allows AnyEnums to bind to other enum types.
 		// The actual copy will be done via potential type promotion copy, into the value property inside the AnyEnum.
 		// We amend the paths here to point to the 'Value' property.
-		const bool bSourceIsAnyEnum = IsPropertyAnyEnum(SourceStruct, Binding.GetSourcePath());
+		const bool bSourceIsAnyEnum = IsPropertyAnyEnum(*SourceStruct, Binding.GetSourcePath());
 		const bool bTargetIsAnyEnum = IsPropertyAnyEnum(TargetStruct, Binding.GetTargetPath());
 		if (bSourceIsAnyEnum || bTargetIsAnyEnum)
 		{
-			FStateTreePropertyPathBinding ModifiedBinding(Binding);
 			if (bSourceIsAnyEnum)
 			{
-				ModifiedBinding.GetMutableSourcePath().AddPathSegment(GET_MEMBER_NAME_STRING_CHECKED(FStateTreeAnyEnum, Value));
+				BindingCopy.GetMutableSourcePath().AddPathSegment(GET_MEMBER_NAME_STRING_CHECKED(FStateTreeAnyEnum, Value));
 			}
 			if (bTargetIsAnyEnum)
 			{
-				ModifiedBinding.GetMutableTargetPath().AddPathSegment(GET_MEMBER_NAME_STRING_CHECKED(FStateTreeAnyEnum, Value));
+				BindingCopy.GetMutableTargetPath().AddPathSegment(GET_MEMBER_NAME_STRING_CHECKED(FStateTreeAnyEnum, Value));
 			}
-			OutBindings.Add(ModifiedBinding);
-		}
-		else
-		{
-			OutBindings.Add(Binding);
 		}
+		
+		OutBindings.Add(BindingCopy);
 
 		// Check if the bindings is for struct ref and validate the types.
-		if (!ValidateStructRef(SourceStruct, Binding.GetSourcePath(), TargetStruct, Binding.GetTargetPath()))
+		if (!ValidateStructRef(*SourceStruct, Binding.GetSourcePath(), TargetStruct, Binding.GetTargetPath()))
 		{
 			return false;
 		}
diff --git a/Engine/Plugins/Runtime/StateTree/Source/StateTreeEditorModule/Private/StateTreeEditorData.cpp b/Engine/Plugins/Runtime/StateTree/Source/StateTreeEditorModule/Private/StateTreeEditorData.cpp
--- a/Engine/Plugins/Runtime/StateTree/Source/StateTreeEditorModule/Private/StateTreeEditorData.cpp
+++ b/Engine/Plugins/Runtime/StateTree/Source/StateTreeEditorModule/Private/StateTreeEditorData.cpp
@@ -303,7 +303,7 @@ FStateTreeBindableStructDesc UStateTreeEditorData::FindContextData(const UStruct
 	{
 		if (Desc.Struct->IsChildOf(ObjectType))
 		{
-			Candidates.Emplace(Desc.Name, Desc.Struct, EStateTreeBindableStructSource::Context, Desc.ID);
+			Candidates.Emplace(Desc.Name, Desc.Struct, FStateTreeDataHandle(), EStateTreeBindableStructSource::Context, Desc.ID);
 		}
 	}
 
diff --git a/Engine/Plugins/Runtime/StateTree/Source/StateTreeEditorModule/Private/StateTreePropertyBindingCompiler.cpp b/Engine/Plugins/Runtime/StateTree/Source/StateTreeEditorModule/Private/StateTreePropertyBindingCompiler.cpp
--- a/Engine/Plugins/Runtime/StateTree/Source/StateTreeEditorModule/Private/StateTreePropertyBindingCompiler.cpp
+++ b/Engine/Plugins/Runtime/StateTree/Source/StateTreeEditorModule/Private/StateTreePropertyBindingCompiler.cpp
@@ -14,7 +14,9 @@ bool FStateTreePropertyBindingCompiler::Init(FStateTreePropertyBindings& InPrope
 	Log = &InLog;
 	PropertyBindings = &InPropertyBindings;
 	PropertyBindings->Reset();
+
 	SourceStructs.Reset();
+	
 	return true;
 }
 
@@ -40,26 +42,21 @@ bool FStateTreePropertyBindingCompiler::CompileBatch(const FStateTreeBindableStr
 			continue;
 		}
 		// Source must be in the source array
-		const FGuid SourceStructID = Binding.GetSourcePath().GetStructID();
-		const int32 SourceStructIdx = SourceStructs.IndexOfByPredicate([SourceStructID](const FStateTreeBindableStructDesc& Struct)
-			{
-				return (Struct.ID == SourceStructID);
-			});
-		if (SourceStructIdx == INDEX_NONE)
+		const FStateTreeBindableStructDesc* SourceStruct = GetSourceStructDescByID(Binding.GetSourcePath().GetStructID());
+		if (!SourceStruct)
 		{
 			Log->Reportf(EMessageSeverity::Error, TargetStruct,
 				TEXT("Could not find a binding source."));
 			return false;
 		}
-		const FStateTreeBindableStructDesc& SourceStruct = SourceStructs[SourceStructIdx];
 
 		FString Error;
 		TArray<FStateTreePropertyPathIndirection> SourceIndirections;
 		TArray<FStateTreePropertyPathIndirection> TargetIndirections;
 		
-		if (!Binding.GetSourcePath().ResolveIndirections(SourceStruct.Struct, SourceIndirections, &Error))
+		if (!Binding.GetSourcePath().ResolveIndirections(SourceStruct->Struct, SourceIndirections, &Error))
 		{
-			Log->Reportf(EMessageSeverity::Error, TargetStruct, TEXT("Resolving path in %s: %s"), *SourceStruct.ToString(), *Error);
+			Log->Reportf(EMessageSeverity::Error, TargetStruct, TEXT("Resolving path in %s: %s"), *SourceStruct->ToString(), *Error);
 			return false;
 		}
 
@@ -70,25 +67,19 @@ bool FStateTreePropertyBindingCompiler::CompileBatch(const FStateTreeBindableStr
 		}
 
 		FStateTreePropertyCopy DummyCopy;
-		FStateTreePropertyPathIndirection LastSourceIndirection = !SourceIndirections.IsEmpty() ? SourceIndirections.Last() : FStateTreePropertyPathIndirection(SourceStruct.Struct);
+		FStateTreePropertyPathIndirection LastSourceIndirection = !SourceIndirections.IsEmpty() ? SourceIndirections.Last() : FStateTreePropertyPathIndirection(SourceStruct->Struct);
 		FStateTreePropertyPathIndirection LastTargetIndirection = !TargetIndirections.IsEmpty() ? TargetIndirections.Last() : FStateTreePropertyPathIndirection(TargetStruct.Struct);
 		if (!PropertyBindings->ResolveCopyType(LastSourceIndirection, LastTargetIndirection, DummyCopy))
 		{
 			Log->Reportf(EMessageSeverity::Error, TargetStruct,
 			TEXT("Cannot copy properties between %s and %s, properties are incompatible."),
-				*UE::StateTree::GetDescAndPathAsString(SourceStruct, Binding.GetSourcePath()),
+				*UE::StateTree::GetDescAndPathAsString(*SourceStruct, Binding.GetSourcePath()),
 				*UE::StateTree::GetDescAndPathAsString(TargetStruct, Binding.GetTargetPath()));
 			return false;
 		}
 
-		if (const auto Validation = UE::StateTree::Compiler::IsValidIndex16(SourceStructIdx); Validation.DidFail())
-		{
-			Validation.Log(*Log, TEXT("Source Struct Index"), SourceStruct);
-			return false;
-		}
-
 		FSortedBinding& NewBinding = NewBindings.AddDefaulted_GetRef();
-		NewBinding.Binding = FStateTreePropertyPathBinding(FStateTreeIndex16(SourceStructIdx), Binding.GetSourcePath(), Binding.GetTargetPath());
+		NewBinding.Binding = FStateTreePropertyPathBinding(SourceStruct->DataHandle, Binding.GetSourcePath(), Binding.GetTargetPath());
 		NewBinding.TargetIndirections = TargetIndirections;
 	}
 
@@ -150,6 +141,8 @@ int32 FStateTreePropertyBindingCompiler::AddSourceStruct(const FStateTreeBindabl
 			*SourceStruct.ToString(), *ExistingStruct->ToString(), *ExistingStruct->ID.ToString());
 	}
 	
+	UE_CLOG(!SourceStruct.DataHandle.IsValid(), LogStateTree, Error, TEXT("%s does not have a valid data handle."), *SourceStruct.ToString());
+	
 	SourceStructs.Add(SourceStruct);
 	return SourceStructs.Num() - 1;
 }
diff --git a/Engine/Plugins/Runtime/StateTree/Source/StateTreeEditorModule/Public/StateTreeCompiler.h b/Engine/Plugins/Runtime/StateTree/Source/StateTreeEditorModule/Public/StateTreeCompiler.h
--- a/Engine/Plugins/Runtime/StateTree/Source/StateTreeEditorModule/Public/StateTreeCompiler.h
+++ b/Engine/Plugins/Runtime/StateTree/Source/StateTreeEditorModule/Public/StateTreeCompiler.h
@@ -50,8 +50,8 @@ private:
 	
 	bool CreateConditions(UStateTreeState& State, TConstArrayView<FStateTreeEditorNode> Conditions);
 	bool CreateCondition(UStateTreeState& State, const FStateTreeEditorNode& CondNode, const EStateTreeConditionOperand Operand, const int8 DeltaIndent);
-	bool CreateTask(UStateTreeState* State, const FStateTreeEditorNode& TaskNode);
-	bool CreateEvaluator(const FStateTreeEditorNode& EvalNode);
+	bool CreateTask(UStateTreeState* State, const FStateTreeEditorNode& TaskNode, const FStateTreeDataHandle TaskDataHandle);
+	bool CreateEvaluator(const FStateTreeEditorNode& EvalNode, const FStateTreeDataHandle EvalDataHandle);
 	bool GetAndValidateBindings(const FStateTreeBindableStructDesc& TargetStruct, FStateTreeDataView TargetValue, TArray<FStateTreePropertyPathBinding>& OutBindings) const;
 	bool IsPropertyAnyEnum(const FStateTreeBindableStructDesc& Struct, FStateTreePropertyPath Path) const;
 	bool ValidateStructRef(const FStateTreeBindableStructDesc& SourceStruct, FStateTreePropertyPath SourcePath,
@@ -60,8 +60,6 @@ private:
 
 	void InstantiateStructSubobjects(FStructView Struct);
 	
-	FStateTreeDataView GetBindingSourceValue(const int32 SourceIndex);
-	
 	FStateTreeCompilerLog& Log;
 	UStateTree* StateTree = nullptr;
 	UStateTreeEditorData* EditorData = nullptr;
@@ -74,8 +72,6 @@ private:
 	TArray<FInstancedStruct> Nodes;
 	TArray<FInstancedStruct> InstanceStructs;
 	TArray<FInstancedStruct> SharedInstanceStructs;
-	TArray<TObjectPtr<UObject>> InstanceObjects;
-	TArray<TObjectPtr<UObject>> SharedInstanceObjects;
 	
 	FStateTreePropertyBindingCompiler BindingsCompiler;
 };
diff --git a/Engine/Plugins/Runtime/StateTree/Source/StateTreeEditorModule/Public/StateTreePropertyBindingCompiler.h b/Engine/Plugins/Runtime/StateTree/Source/StateTreeEditorModule/Public/StateTreePropertyBindingCompiler.h
--- a/Engine/Plugins/Runtime/StateTree/Source/StateTreeEditorModule/Public/StateTreePropertyBindingCompiler.h
+++ b/Engine/Plugins/Runtime/StateTree/Source/StateTreeEditorModule/Public/StateTreePropertyBindingCompiler.h
@@ -46,14 +46,21 @@ struct STATETREEEDITORMODULE_API FStateTreePropertyBindingCompiler
 	int32 AddSourceStruct(const FStateTreeBindableStructDesc& SourceStruct);
 
 	/** @return Index of a source struct by specified ID, or INDEX_NONE if not found. */
+	UE_DEPRECATED(5.4, "Use GetSourceStructDescByID() instead.")
 	int32 GetSourceStructIndexByID(const FGuid& ID) const;
 
 	/** @return Reference to a source struct based on ID. */
+	UE_DEPRECATED(5.4, "Use GetSourceStructDescByID() instead.")
 	const FStateTreeBindableStructDesc& GetSourceStructDesc(const int32 Index) const
 	{
 		return SourceStructs[Index];
 	}
 
+	const FStateTreeBindableStructDesc* GetSourceStructDescByID(const FGuid& ID) const
+	{
+		return SourceStructs.FindByPredicate([ID](const FStateTreeBindableStructDesc& Structs) { return (Structs.ID == ID); });
+	}
+
 PRAGMA_DISABLE_DEPRECATION_WARNINGS
 	/**
 	 * Resolves a string based property path in specified struct into segments of property names and access types.
diff --git a/Engine/Plugins/Runtime/StateTree/Source/StateTreeModule/Private/StateTree.cpp b/Engine/Plugins/Runtime/StateTree/Source/StateTreeModule/Private/StateTree.cpp
--- a/Engine/Plugins/Runtime/StateTree/Source/StateTreeModule/Private/StateTree.cpp
+++ b/Engine/Plugins/Runtime/StateTree/Source/StateTreeModule/Private/StateTree.cpp
@@ -12,6 +12,7 @@
 #include "Logging/LogScopedVerbosityOverride.h"
 #include "Misc/DataValidation.h"
 #include "StructUtilsDelegates.h"
+#include "Misc/EnumerateRange.h"
 #if WITH_EDITOR
 #include "Engine/UserDefinedStruct.h"
 #endif
@@ -23,7 +24,7 @@ FCustomVersionRegistration GRegisterStateTreeCustomVersion(FStateTreeCustomVersi
 bool UStateTree::IsReadyToRun() const
 {
 	// Valid tree must have at least one state and valid instance data.
-	return States.Num() > 0 && bIsLinked;
+	return States.Num() > 0 && bIsLinked && PropertyBindings.IsValid();
 }
 
 FConstStructView UStateTree::GetNode(const int32 NodeIndex) const
@@ -134,8 +135,6 @@ void UStateTree::ResetCompiled()
 	IDToStateMappings.Reset();
 	IDToNodeMappings.Reset();
 	IDToTransitionMappings.Reset();
-
-	ParametersDataViewIndex = FStateTreeIndex8::Invalid;
 	
 	EvaluatorsBegin = 0;
 	EvaluatorsNum = 0;
@@ -367,7 +366,7 @@ void UStateTree::ResetLinked()
 	bIsLinked = false;
 	ExternalDataDescs.Reset();
 	ExternalDataBaseIndex = 0;
-	NumDataViews = 0;
+	NumContextDataViews = 0;
 
 	FWriteScopeLock WriteLock(PerThreadSharedInstanceDataLock);
 	PerThreadSharedInstanceData.Reset();
@@ -379,16 +378,17 @@ bool UStateTree::Link()
 	// This data will be used to allocate runtime instance on all StateTree users.
 	ResetLinked();
 
+	ExternalDataBaseIndex = NumContextData;
+
 	// Resolves nodes references to other StateTree data
 	FStateTreeLinker Linker(Schema);
-	Linker.SetExternalDataBaseIndex(PropertyBindings.GetSourceStructNum());
+	Linker.SetExternalDataBaseIndex(ExternalDataBaseIndex);
 
 	for (int32 Index = 0; Index < Nodes.Num(); Index++)
 	{
 		FStructView Node = Nodes[Index];
 		if (FStateTreeNodeBase* NodePtr = Node.GetPtr<FStateTreeNodeBase>())
 		{
-			Linker.SetCurrentInstanceDataType(NodePtr->GetInstanceDataType(), NodePtr->DataViewIndex.Get());
 			const bool bLinkSucceeded = NodePtr->Link(Linker);
 			if (!bLinkSucceeded || Linker.GetStatus() == EStateTreeLinkerStatus::Failed)
 			{
@@ -398,9 +398,8 @@ bool UStateTree::Link()
 		}
 	}
 
-	ExternalDataBaseIndex = PropertyBindings.GetSourceStructNum();
 	ExternalDataDescs = Linker.GetExternalDataDescs();
-	NumDataViews = ExternalDataBaseIndex + ExternalDataDescs.Num();
+	NumContextDataViews = ExternalDataBaseIndex + ExternalDataDescs.Num();
 
 	if (States.Num() > 0 && Nodes.Num() > 0)
 	{
@@ -419,29 +418,29 @@ bool UStateTree::Link()
 				return false;
 			}
 		}
+	}
 
-		if (!DefaultInstanceData.AreAllInstancesValid())
-		{
-			UE_LOG(LogStateTree, Error, TEXT("%s: State Tree asset was not properly loaded (missing instance data). See log for loading failures, or recompile the StateTree asset."), *GetFullName());
-			return false;
-		}
-
-		if (!SharedInstanceData.AreAllInstancesValid())
-		{
-			UE_LOG(LogStateTree, Error, TEXT("%s: State Tree asset was not properly loaded (missing shared instance data). See log for loading failures, or recompile the StateTree asset."), *GetFullName());
-			return false;
-		}
+	if (!DefaultInstanceData.AreAllInstancesValid())
+	{
+		UE_LOG(LogStateTree, Error, TEXT("%s: State Tree asset was not properly loaded (missing instance data). See log for loading failures, or recompile the StateTree asset."), *GetFullName());
+		return false;
+	}
 
-		if (!PatchBindings())
-		{
-			return false;
-		}
+	if (!SharedInstanceData.AreAllInstancesValid())
+	{
+		UE_LOG(LogStateTree, Error, TEXT("%s: State Tree asset was not properly loaded (missing shared instance data). See log for loading failures, or recompile the StateTree asset."), *GetFullName());
+		return false;
+	}
+	
+	if (!PatchBindings())
+	{
+		return false;
+	}
 
-		// Resolves property paths used by bindings a store property pointers
-		if (!PropertyBindings.ResolvePaths())
-		{
-			return false;
-		}
+	// Resolves property paths used by bindings a store property pointers
+	if (!PropertyBindings.ResolvePaths())
+	{
+		return false;
 	}
 
 	// Link succeeded, setup tree to be ready to run
@@ -456,6 +455,22 @@ bool UStateTree::PatchBindings()
 	const TArrayView<FStateTreePropertyCopyBatch> CopyBatches = PropertyBindings.CopyBatches;
 	const TArrayView<FStateTreePropertyPathBinding> PropertyPathBindings = PropertyBindings.PropertyPathBindings;
 
+	// Make mapping from data handle to source struct.
+	TMap<FStateTreeDataHandle, int32> SourceStructByHandle;
+	for (TConstEnumerateRef<FStateTreeBindableStructDesc> SourceStruct : EnumerateRange(SourceStructs))
+	{
+		SourceStructByHandle.Add(SourceStruct->DataHandle, SourceStruct.GetIndex());
+	}
+
+	auto GetSourceStructByHandle = [&SourceStructByHandle, &SourceStructs](const FStateTreeDataHandle DataHandle) -> FStateTreeBindableStructDesc*
+	{
+		if (int32* Index = SourceStructByHandle.Find(DataHandle))
+		{
+			return &SourceStructs[*Index];
+		}
+		return nullptr;
+	};
+	
 	// Reconcile out of date classes.
 	for (FStateTreeBindableStructDesc& SourceStruct : SourceStructs)
 	{
@@ -499,83 +514,94 @@ bool UStateTree::PatchBindings()
 	}
 
 	// Update property bag structs before resolving binding.
-	if (ParametersDataViewIndex.IsValid() && SourceStructs.IsValidIndex(ParametersDataViewIndex.Get()))
+	if (FStateTreeBindableStructDesc* RootParamsDesc = GetSourceStructByHandle(ParametersDataHandle))
 	{
-		SourceStructs[ParametersDataViewIndex.Get()].Struct = Parameters.GetPropertyBagStruct();
+		RootParamsDesc->Struct = Parameters.GetPropertyBagStruct();
 	}
 
+	// Refresh state parameter descs and bindings batches.
 	for (const FCompactStateTreeState& State : States)
 	{
-		if (State.Type == EStateTreeStateType::Subtree)
+		if (State.Type == EStateTreeStateType::Subtree
+			||State.Type == EStateTreeStateType::Linked)
 		{
-			if (State.ParameterInstanceIndex.IsValid() == false)
+			if (!State.ParameterTemplateIndex.IsValid())
 			{
 				UE_LOG(LogStateTree, Error, TEXT("%s: Data for state '%s' is malformed. Please recompile the StateTree asset."), *GetFullName(), *State.Name.ToString());
 				return false;
 			}
 
 			// Subtree is a bind source, update bag struct.
-			if (State.ParameterDataViewIndex.IsValid())
+			const FCompactStateTreeParameters& Params = DefaultInstanceData.GetMutableStruct(State.ParameterTemplateIndex.Get()).Get<FCompactStateTreeParameters>();
+			FStateTreeBindableStructDesc* Desc = GetSourceStructByHandle(State.ParameterDataHandle);
+			if (!Desc)
 			{
-				const FCompactStateTreeParameters& Params = DefaultInstanceData.GetMutableStruct(State.ParameterInstanceIndex.Get()).Get<FCompactStateTreeParameters>();
-				FStateTreeBindableStructDesc& Desc = SourceStructs[State.ParameterDataViewIndex.Get()];
-				Desc.Struct = Params.Parameters.GetPropertyBagStruct();
+				UE_LOG(LogStateTree, Error, TEXT("%s: Data for state '%s' is malformed. Please recompile the StateTree asset."), *GetFullName(), *State.Name.ToString());
+				return false;
+			}
+			Desc->Struct = Params.Parameters.GetPropertyBagStruct();
+
+			if (State.ParameterBindingsBatch.IsValid())
+			{
+				FStateTreePropertyCopyBatch& Batch = CopyBatches[State.ParameterBindingsBatch.Get()];
+				Batch.TargetStruct.Struct = Params.Parameters.GetPropertyBagStruct();
 			}
 		}
-		else if (State.Type == EStateTreeStateType::Linked && State.LinkedState.IsValid())
+	}
+
+	// Check linked state property bags consistency
+	for (const FCompactStateTreeState& State : States)
+	{
+		if (State.Type == EStateTreeStateType::Linked && State.LinkedState.IsValid())
 		{
 			const FCompactStateTreeState& LinkedState = States[State.LinkedState.Index];
 
-			if (State.ParameterInstanceIndex.IsValid() == false
-				|| LinkedState.ParameterInstanceIndex.IsValid() == false)
+			if (State.ParameterTemplateIndex.IsValid() == false
+				|| LinkedState.ParameterTemplateIndex.IsValid() == false)
 			{
 				UE_LOG(LogStateTree, Error, TEXT("%s: Data for state '%s' is malformed. Please recompile the StateTree asset."), *GetFullName(), *State.Name.ToString());
 				return false;
 			}
 
-			const FCompactStateTreeParameters& Params = DefaultInstanceData.GetMutableStruct(State.ParameterInstanceIndex.Get()).Get<FCompactStateTreeParameters>();
-
 			// Check that the bag in linked state matches.
-			const FCompactStateTreeParameters& LinkedStateParams = DefaultInstanceData.GetMutableStruct(LinkedState.ParameterInstanceIndex.Get()).Get<FCompactStateTreeParameters>();
+			const FCompactStateTreeParameters& Params = DefaultInstanceData.GetMutableStruct(State.ParameterTemplateIndex.Get()).Get<FCompactStateTreeParameters>();
+			const FCompactStateTreeParameters& LinkedStateParams = DefaultInstanceData.GetMutableStruct(LinkedState.ParameterTemplateIndex.Get()).Get<FCompactStateTreeParameters>();
 
 			if (LinkedStateParams.Parameters.GetPropertyBagStruct() != Params.Parameters.GetPropertyBagStruct())
 			{
 				UE_LOG(LogStateTree, Error, TEXT("%s: The parameters on state '%s' does not match the linked state parameters in state '%s'. Please recompile the StateTree asset."), *GetFullName(), *State.Name.ToString(), *LinkedState.Name.ToString());
 				return false;
 			}
-
-			if (Params.BindingsBatch.IsValid())
-			{
-				FStateTreePropertyCopyBatch& Batch = CopyBatches[Params.BindingsBatch.Get()];
-				Batch.TargetStruct.Struct = Params.Parameters.GetPropertyBagStruct();
-			}
 		}
 	}
 
-	TArray<FStateTreeDataView> DataViews;
-	DataViews.SetNum(NumDataViews);
+
+	TMap<FStateTreeDataHandle, FStateTreeDataView> DataViews;
+	TMap<FStateTreeIndex16, FStateTreeDataView> BindingBatchDataView;
 
 	// Tree parameters
-	DataViews[ParametersDataViewIndex.Get()] = Parameters.GetMutableValue();
+	DataViews.Add(ParametersDataHandle, Parameters.GetMutableValue());
 
 	// Setup data views for context data. Since the external data is passed at runtime, we can only provide the type.
 	for (const FStateTreeExternalDataDesc& DataDesc : ContextDataDescs)
 	{
-		DataViews[DataDesc.Handle.DataViewIndex.Get()] = FStateTreeDataView(DataDesc.Struct, nullptr);
+		DataViews.Add(DataDesc.Handle.DataHandle, FStateTreeDataView(DataDesc.Struct, nullptr));
 	}
 	
 	// Setup data views for state parameters.
 	for (FCompactStateTreeState& State : States)
 	{
-		// Subtree state parameters can be a binding source, but not a target.
-		// Linked state parameters can only by a binding target, so not handled here.
-		if (State.Type == EStateTreeStateType::Subtree)
+		if (State.Type == EStateTreeStateType::Subtree
+			|| State.Type == EStateTreeStateType::Linked)
 		{
-			if (State.ParameterDataViewIndex.IsValid())
+			if (State.ParameterDataHandle.IsValid())
 			{
-				FCompactStateTreeParameters& Params = DefaultInstanceData.GetMutableStruct(State.ParameterInstanceIndex.Get()).Get<FCompactStateTreeParameters>();
-				ensureMsgf(Params.Parameters.IsValid(), TEXT("Expecting valid state parameter value when parameter view is set."));
-				DataViews[State.ParameterDataViewIndex.Get()] = Params.Parameters.GetMutableValue();
+				FCompactStateTreeParameters& Params = DefaultInstanceData.GetMutableStruct(State.ParameterTemplateIndex.Get()).Get<FCompactStateTreeParameters>();
+				DataViews.Add(State.ParameterDataHandle, Params.Parameters.GetMutableValue());
+				if (State.ParameterBindingsBatch.IsValid())
+				{
+					BindingBatchDataView.Add(State.ParameterBindingsBatch, Params.Parameters.GetMutableValue());
+				}
 			}
 		}
 	}
@@ -583,66 +609,52 @@ bool UStateTree::PatchBindings()
 	// Setup data views for all nodes.
 	for (FConstStructView NodeView : Nodes)
 	{
-		if (const FStateTreeConditionBase* Condition = NodeView.GetPtr<const FStateTreeConditionBase>())
+		const FStateTreeNodeBase& Node = NodeView.Get<const FStateTreeNodeBase>();
+		
+		FStateTreeInstanceData* SourceInstanceData = &DefaultInstanceData;
+		if (NodeView.GetPtr<const FStateTreeConditionBase>())
 		{
-			if (Condition->bInstanceIsObject)
-			{
-				DataViews[Condition->DataViewIndex.Get()] = SharedInstanceData.GetMutableObject(Condition->InstanceIndex.Get());
-			}
-			else
-			{
-				DataViews[Condition->DataViewIndex.Get()] = SharedInstanceData.GetMutableStruct(Condition->InstanceIndex.Get());
-			}
+			// Conditions are stored in shared instance data.
+			SourceInstanceData = &SharedInstanceData;
 		}
-		else
+
+		FStateTreeDataView NodeDataView = Node.InstanceDataHandle.IsObjectSource()
+						? FStateTreeDataView(SourceInstanceData->GetMutableObject(Node.InstanceTemplateIndex.Get()))
+						: FStateTreeDataView(SourceInstanceData->GetMutableStruct(Node.InstanceTemplateIndex.Get()));
+
+		DataViews.Add(Node.InstanceDataHandle, NodeDataView);
+
+		if (Node.BindingsBatch.IsValid())
 		{
-			const FStateTreeNodeBase& Node = NodeView.Get<const FStateTreeNodeBase>();
-			if (Node.bInstanceIsObject)
-			{
-				DataViews[Node.DataViewIndex.Get()] = DefaultInstanceData.GetMutableObject(Node.InstanceIndex.Get());
-			}
-			else
-			{
-				DataViews[Node.DataViewIndex.Get()] = DefaultInstanceData.GetMutableStruct(Node.InstanceIndex.Get());
-			}
+			BindingBatchDataView.Add(Node.BindingsBatch, NodeDataView);
 		}
 	}
-
-	for (int32 BatchIndex = 0; BatchIndex < CopyBatches.Num(); ++BatchIndex)
+	
+	auto GetDataSourceView = [&DataViews](const FStateTreeDataHandle Handle) -> FStateTreeDataView
 	{
-		const FStateTreePropertyCopyBatch& Batch = CopyBatches[BatchIndex];
-
-		// Find data view for the binding target.
-		FStateTreeDataView TargetView;
-
-		// Target can be one of the nodes.
-		for (FConstStructView NodeView : GetNodes())
+		if (const FStateTreeDataView* ViewPtr = DataViews.Find(Handle))
 		{
-			const FStateTreeNodeBase& Node = NodeView.Get<const FStateTreeNodeBase>();
-			if (Node.BindingsBatch.AsInt32() == BatchIndex)
-			{
-				TargetView = DataViews[Node.DataViewIndex.Get()];
-				break;
-			}
+			return *ViewPtr;
 		}
+		return FStateTreeDataView();
+	};
 
-		// Target can be a linked state parameters (subtree state params can only be a binding source).
-		if (!TargetView.IsValid())
+	auto GetBindingBatchDataView = [&BindingBatchDataView](const FStateTreeIndex16 Index) -> FStateTreeDataView
+	{
+		if (const FStateTreeDataView* ViewPtr = BindingBatchDataView.Find(Index))
 		{
-			for (const FCompactStateTreeState& State : GetStates())
-			{
-				if (State.Type == EStateTreeStateType::Linked)
-				{
-					const FCompactStateTreeParameters& Params = DefaultInstanceData.GetStruct(State.ParameterInstanceIndex.Get()).Get<const FCompactStateTreeParameters>();
-					if (Params.BindingsBatch.Get() == BatchIndex)
-					{
-						const FCompactStateTreeState& LinkedState = States[State.LinkedState.Index];
-						TargetView = DataViews[LinkedState.ParameterDataViewIndex.Get()];
-					}
-				}
-			}
+			return *ViewPtr;
 		}
+		return FStateTreeDataView();
+	};
 
+
+	for (int32 BatchIndex = 0; BatchIndex < CopyBatches.Num(); ++BatchIndex)
+	{
+		const FStateTreePropertyCopyBatch& Batch = CopyBatches[BatchIndex];
+
+		// Find data view for the binding target.
+		FStateTreeDataView TargetView = GetBindingBatchDataView(FStateTreeIndex16(BatchIndex));
 		if (!TargetView.IsValid())
 		{
 			UE_LOG(LogStateTree, Error, TEXT("%hs: Invalid target struct when trying to bind to '%s'."), __FUNCTION__, *Batch.TargetStruct.Name.ToString());
@@ -653,8 +665,9 @@ bool UStateTree::PatchBindings()
 		for (int32 Index = Batch.BindingsBegin; Index != Batch.BindingsEnd; Index++)
 		{
 			FStateTreePropertyPathBinding& Binding = PropertyPathBindings[Index];
-
-			if (!Binding.GetMutableSourcePath().UpdateSegmentsFromValue(DataViews[Binding.GetCompiledSourceStructIndex().Get()], &ErrorMsg))
+			FStateTreeDataView SourceView = GetDataSourceView(Binding.GetSourceDataHandle());
+			
+			if (!Binding.GetMutableSourcePath().UpdateSegmentsFromValue(SourceView, &ErrorMsg))
 			{
 				UE_LOG(LogStateTree, Error, TEXT("%hs: Failed to update source instance structs for property binding '%s'. Reason: %s"), __FUNCTION__, *Binding.GetTargetPath().ToString(), *ErrorMsg);
 				return false;
@@ -754,24 +767,28 @@ TArray<FStateTreeMemoryUsage> UStateTree::CalculateEstimatedMemoryUsage() const
 		{
 			const int32 LinkedUsageIndex = GetUsageIndexForState(CompactState.LinkedState);
 			StateLinks.Emplace(ParentUsageIndex, LinkedUsageIndex);
-
+		}
+		
+		if (CompactState.Type == EStateTreeStateType::Linked
+			|| CompactState.Type == EStateTreeStateType::Subtree)
+		{
 			MemUsage.NodeCount++;
-			MemUsage.AddUsage(DefaultInstanceData.GetStruct(CompactState.ParameterInstanceIndex.Get()));
+			MemUsage.AddUsage(DefaultInstanceData.GetStruct(CompactState.ParameterTemplateIndex.Get()));
 		}
 		
 		for (int32 TaskIndex = CompactState.TasksBegin; TaskIndex < (CompactState.TasksBegin + CompactState.TasksNum); TaskIndex++)
 		{
 			if (const FStateTreeTaskBase* Task = Nodes[TaskIndex].GetPtr<const FStateTreeTaskBase>())
 			{
-				if (Task->bInstanceIsObject)
+				if (Task->InstanceDataHandle.IsObjectSource())
 				{
 					MemUsage.NodeCount++;
-					MemUsage.AddUsage(DefaultInstanceData.GetObject(Task->InstanceIndex.Get()));
+					MemUsage.AddUsage(DefaultInstanceData.GetObject(Task->InstanceTemplateIndex.Get()));
 				}
 				else
 				{
 					MemUsage.NodeCount++;
-					MemUsage.AddUsage(DefaultInstanceData.GetStruct(Task->InstanceIndex.Get()));
+					MemUsage.AddUsage(DefaultInstanceData.GetStruct(Task->InstanceTemplateIndex.Get()));
 				}
 			}
 		}
@@ -820,13 +837,13 @@ TArray<FStateTreeMemoryUsage> UStateTree::CalculateEstimatedMemoryUsage() const
 	for (int32 EvalIndex = EvaluatorsBegin; EvalIndex < (EvaluatorsBegin + EvaluatorsNum); EvalIndex++)
 	{
 		const FStateTreeEvaluatorBase& Eval = Nodes[EvalIndex].Get<const FStateTreeEvaluatorBase>();
-		if (Eval.bInstanceIsObject == false)
+		if (Eval.InstanceDataHandle.IsObjectSource())
 		{
-			EvalMemUsage.AddUsage(DefaultInstanceData.GetStruct(Eval.InstanceIndex.Get()));
+			EvalMemUsage.AddUsage(DefaultInstanceData.GetObject(Eval.InstanceTemplateIndex.Get()));
 		}
 		else
 		{
-			EvalMemUsage.AddUsage(DefaultInstanceData.GetObject(Eval.InstanceIndex.Get()));
+			EvalMemUsage.AddUsage(DefaultInstanceData.GetStruct(Eval.InstanceTemplateIndex.Get()));
 		}
 		EvalMemUsage.NodeCount++;
 	}
@@ -836,13 +853,13 @@ TArray<FStateTreeMemoryUsage> UStateTree::CalculateEstimatedMemoryUsage() const
 	for (int32 TaskIndex = GlobalTasksBegin; TaskIndex < (GlobalTasksBegin + GlobalTasksNum); TaskIndex++)
 	{
 		const FStateTreeTaskBase& Task = Nodes[TaskIndex].Get<const FStateTreeTaskBase>();
-		if (Task.bInstanceIsObject == false)
+		if (Task.InstanceDataHandle.IsObjectSource())
 		{
-			GlobalTaskMemUsage.AddUsage(DefaultInstanceData.GetStruct(Task.InstanceIndex.Get()));
+			GlobalTaskMemUsage.AddUsage(DefaultInstanceData.GetObject(Task.InstanceTemplateIndex.Get()));
 		}
 		else
 		{
-			GlobalTaskMemUsage.AddUsage(DefaultInstanceData.GetObject(Task.InstanceIndex.Get()));
+			GlobalTaskMemUsage.AddUsage(DefaultInstanceData.GetStruct(Task.InstanceTemplateIndex.Get()));
 		}
 		GlobalTaskMemUsage.NodeCount++;
 	}
@@ -890,7 +907,7 @@ TArray<FStateTreeMemoryUsage> UStateTree::CalculateEstimatedMemoryUsage() const
 	if (SharedInstanceData.IsValid())
 	{
 		FStateTreeMemoryUsage& SharedMemUsage = MemoryUsages[SharedMemUsageIndex];
-		SharedMemUsage.NodeCount = SharedInstanceData.GetNumItems();
+		SharedMemUsage.NodeCount = SharedInstanceData.Num();
 		SharedMemUsage.EstimatedMemoryUsage = SharedInstanceData.GetEstimatedMemoryUsage();
 	}
 
