StateTree: Changed how instance data is accessed

diff --git a/Engine/Plugins/Runtime/StateTree/Source/StateTreeModule/Private/StateTreeExecutionContext.cpp b/Engine/Plugins/Runtime/StateTree/Source/StateTreeModule/Private/StateTreeExecutionContext.cpp
--- a/Engine/Plugins/Runtime/StateTree/Source/StateTreeModule/Private/StateTreeExecutionContext.cpp
+++ b/Engine/Plugins/Runtime/StateTree/Source/StateTreeModule/Private/StateTreeExecutionContext.cpp
@@ -43,7 +43,7 @@
 	#define STATETREE_TRACE_STATE_EVENT(StateHandle, EventType)				TRACE_STATETREE_STATE_EVENT(GetInstanceDebugId(), StateHandle, EventType, EStateTreeStateSelectionBehavior::None);
 	#define STATETREE_TRACE_TASK_EVENT(Index, DataView, EventType, Status)	TRACE_STATETREE_TASK_EVENT(GetInstanceDebugId(), FStateTreeIndex16(Index), DataView, EventType, Status);
 	#define STATETREE_TRACE_EVALUATOR_EVENT(Index, DataView, EventType)		TRACE_STATETREE_EVALUATOR_EVENT(GetInstanceDebugId(), FStateTreeIndex16(Index), DataView, EventType);
-	#define STATETREE_TRACE_CONDITION_EVENT(Index, DataViews, EventType)	TRACE_STATETREE_CONDITION_EVENT(GetInstanceDebugId(), FStateTreeIndex16(Index), DataView, EventType);	
+	#define STATETREE_TRACE_CONDITION_EVENT(Index, DataView, EventType)		TRACE_STATETREE_CONDITION_EVENT(GetInstanceDebugId(), FStateTreeIndex16(Index), DataView, EventType);	
 	#define STATETREE_TRACE_TRANSITION_EVENT(Source, EventType)				TRACE_STATETREE_TRANSITION_EVENT(GetInstanceDebugId(), Source, EventType);
 #else
 	#define STATETREE_TRACE_SCOPED_PHASE(Phase)
@@ -66,22 +66,24 @@ namespace UE::StateTree
 
 FStateTreeExecutionContext::FStateTreeExecutionContext(UObject& InOwner, const UStateTree& InStateTree, FStateTreeInstanceData& InInstanceData)
 	: Owner(InOwner)
-	, StateTree(InStateTree)
+	, RootStateTree(InStateTree)
 	, InstanceData(InInstanceData)
 {
 	if (InStateTree.IsReadyToRun())
 	{
 		// Initialize data views for all possible items.
-		DataViews.SetNum(StateTree.GetNumDataViews());
+		ContextDataViews.SetNum(RootStateTree.GetNumContextDataViews());
+		
 		// Set data views associated to the parameters using the default values
 		SetDefaultParameters();
 
-		SharedInstanceData = StateTree.GetSharedInstanceData();
+		InstanceDataStorage = &InstanceData.GetMutableStorage();
+		check(InstanceDataStorage);
 	}
 	else
 	{
 		STATETREE_LOG(Warning, TEXT("%hs: StateTree asset is not valid ('%s' using StateTree '%s')"),
-			__FUNCTION__, *GetNameSafe(&Owner), *GetFullNameSafe(&StateTree));
+			__FUNCTION__, *GetNameSafe(&Owner), *GetFullNameSafe(&RootStateTree));
 	}
 }
 
@@ -91,23 +93,23 @@ FStateTreeExecutionContext::~FStateTreeExecutionContext()
 
 void FStateTreeExecutionContext::SetDefaultParameters()
 {
-	if (DataViews.IsValidIndex(StateTree.ParametersDataViewIndex.Get()))
+	if (ContextDataViews.IsValidIndex(RootStateTree.ParametersDataHandle.GetIndex()))
 	{
 		// @todo: Handle constness correctly.
-		const FConstStructView ConstParameters = StateTree.GetDefaultParameters().GetValue();
-		DataViews[StateTree.ParametersDataViewIndex.Get()] = FStateTreeDataView(ConstParameters.GetScriptStruct(), const_cast<uint8*>(ConstParameters.GetMemory()));	
+		const FConstStructView ConstParameters = RootStateTree.GetDefaultParameters().GetValue();
+		ContextDataViews[RootStateTree.ParametersDataHandle.GetIndex()] = FStateTreeDataView(ConstParameters.GetScriptStruct(), const_cast<uint8*>(ConstParameters.GetMemory()));	
 	}
 }
 
 void FStateTreeExecutionContext::SetParameters(const FInstancedPropertyBag& Parameters)
 {
-	if (ensureMsgf(StateTree.GetDefaultParameters().GetPropertyBagStruct() == Parameters.GetPropertyBagStruct(),
+	if (ensureMsgf(RootStateTree.GetDefaultParameters().GetPropertyBagStruct() == Parameters.GetPropertyBagStruct(),
 		TEXT("Parameters must be of the same struct type. Make sure to migrate the provided parameters to the same type as the StateTree default parameters."))
-		&& DataViews.IsValidIndex(StateTree.ParametersDataViewIndex.Get()))
+		&& ContextDataViews.IsValidIndex(RootStateTree.ParametersDataHandle.GetIndex()))
 	{
 		// @todo: Handle constness correctly.
 		const FConstStructView ConstParameters = Parameters.GetValue();
-		DataViews[StateTree.ParametersDataViewIndex.Get()] = FStateTreeDataView(ConstParameters.GetScriptStruct(), const_cast<uint8*>(ConstParameters.GetMemory()));	
+		ContextDataViews[RootStateTree.ParametersDataHandle.GetIndex()] = FStateTreeDataView(ConstParameters.GetScriptStruct(), const_cast<uint8*>(ConstParameters.GetMemory()));	
 	}
 }
 
@@ -119,9 +121,9 @@ bool FStateTreeExecutionContext::AreExternalDataViewsValid() const
 	}
 	
 	bool bResult = true;
-	for (const FStateTreeExternalDataDesc& DataDesc : StateTree.ExternalDataDescs)
+	for (const FStateTreeExternalDataDesc& DataDesc : RootStateTree.ExternalDataDescs)
 	{
-		const FStateTreeDataView& DataView = DataViews[DataDesc.Handle.DataViewIndex.Get()];
+		const FStateTreeDataView& DataView = ContextDataViews[DataDesc.Handle.DataHandle.GetIndex()];
 			
 		auto IsAssignmentValid = [](const FStateTreeExternalDataDesc& DataDesc, const FStateTreeDataView& DataView)
 		{	
@@ -150,9 +152,9 @@ bool FStateTreeExecutionContext::AreExternalDataViewsValid() const
 		}
 	}
 
-	for (const FStateTreeExternalDataDesc& DataDesc : StateTree.GetContextDataDescs())
+	for (const FStateTreeExternalDataDesc& DataDesc : RootStateTree.GetContextDataDescs())
 	{
-		const FStateTreeDataView& DataView = DataViews[DataDesc.Handle.DataViewIndex.Get()];
+		const FStateTreeDataView& DataView = ContextDataViews[DataDesc.Handle.DataHandle.GetIndex()];
 
 		// Items must have valid pointer of the expected type.  
 		if (!DataView.IsValid() || !DataView.GetStruct()->IsChildOf(DataDesc.Struct))
@@ -164,50 +166,6 @@ bool FStateTreeExecutionContext::AreExternalDataViewsValid() const
 	return bResult;
 }
 
-void FStateTreeExecutionContext::UpdateLinkedStateParameters(const FCompactStateTreeState& State, const int32 ParameterInstanceIndex)
-{
-	const FStateTreeDataView StateParamsInstance = InstanceData.GetMutableStruct(ParameterInstanceIndex);
-	FCompactStateTreeParameters& StateParams = StateParamsInstance.GetMutable<FCompactStateTreeParameters>();
-
-	// Update parameters if the state has any.
-	if (StateParams.Parameters.IsValid())
-	{
-		// Parameters property bag
-		const FStateTreeDataView ParametersView(StateParams.Parameters.GetMutableValue());
-		if (StateParams.BindingsBatch.IsValid())
-		{
-			StateTree.PropertyBindings.CopyTo(DataViews, StateParams.BindingsBatch, ParametersView);
-		}
-
-		// Set the parameters as the input parameters for the linked state.
-		check(State.LinkedState.IsValid());
-		const FCompactStateTreeState& LinkedState = StateTree.States[State.LinkedState.Index];
-		check(LinkedState.ParameterDataViewIndex.IsValid());
-		DataViews[LinkedState.ParameterDataViewIndex.Get()] = ParametersView;
-	}
-}
-
-void FStateTreeExecutionContext::UpdateSubtreeStateParameters(const FCompactStateTreeState& State)
-{
-	check(State.ParameterInstanceIndex.IsValid());
-
-	// Update parameters if the state has any.
-	if (State.ParameterDataViewIndex.IsValid())
-	{
-		// Usually the subtree parameter view is set by the linked state. If it's not (i.e. transitioned into a parametrized subtree), we'll set the view default params.
-		if (DataViews[State.ParameterDataViewIndex.Get()].IsValid())
-		{
-			return;
-		}
-
-		// Set view to default parameters.
-		const FConstStructView ParamInstance = StateTree.DefaultInstanceData.GetStruct(State.ParameterInstanceIndex.Get()); // These are used as const, so get them from the tree initial values.
-		const FStateTreeDataView ParamInstanceView(ParamInstance.GetScriptStruct(), const_cast<uint8*>(ParamInstance.GetMemory())); 
-		FCompactStateTreeParameters& Params = ParamInstanceView.GetMutable<FCompactStateTreeParameters>();
-		DataViews[State.ParameterDataViewIndex.Get()] = FStateTreeDataView(Params.Parameters.GetMutableValue());
-	}
-}
-
 EStateTreeRunStatus FStateTreeExecutionContext::Start()
 {
 	CSV_SCOPED_TIMING_STAT_EXCLUSIVE(StateTree_Start);
@@ -215,7 +173,7 @@ EStateTreeRunStatus FStateTreeExecutionContext::Start()
 	if (!IsValid())
 	{
 		STATETREE_LOG(Warning, TEXT("%hs: StateTree context is not initialized properly ('%s' using StateTree '%s')"),
-			__FUNCTION__, *GetNameSafe(&Owner), *GetFullNameSafe(&StateTree));
+			__FUNCTION__, *GetNameSafe(&Owner), *GetFullNameSafe(&RootStateTree));
 		return EStateTreeRunStatus::Failed;
 	}
 
@@ -223,23 +181,35 @@ EStateTreeRunStatus FStateTreeExecutionContext::Start()
 	{
 		const FStateTreeExecutionState& Exec = GetExecState();
 		if (!ensureMsgf(Exec.CurrentPhase == EStateTreeUpdatePhase::Unset, TEXT("%hs can't be called while already in %s ('%s' using StateTree '%s')."),
-				__FUNCTION__, *UEnum::GetDisplayValueAsText(Exec.CurrentPhase).ToString(), *GetNameSafe(&Owner), *GetFullNameSafe(&StateTree)))
+				__FUNCTION__, *UEnum::GetDisplayValueAsText(Exec.CurrentPhase).ToString(), *GetNameSafe(&Owner), *GetFullNameSafe(&RootStateTree)))
 		{
 			return EStateTreeRunStatus::Failed;
 		}
 
 		// Stop if still running previous state.
-		Stop();
+		if (Exec.TreeRunStatus == EStateTreeRunStatus::Running)
+		{
+			Stop();
+		}
 	}
 
 	// Initialize instance data. No active states yet, so we'll initialize the evals and global tasks.
 	InstanceData.Reset();
-	constexpr FStateTreeActiveStates Empty;
-	UpdateInstanceData(Empty, Empty);
+
+	FStateTreeExecutionState& Exec = GetExecState();
+
+	// Initialize for the init frame.
+	FStateTreeExecutionFrame& InitFrame = Exec.ActiveFrames.AddDefaulted_GetRef();
+	InitFrame.StateTree = &RootStateTree;
+	InitFrame.RootState = FStateTreeStateHandle::Root;
+	InitFrame.ActiveStates = {};
+	InitFrame.bIsGlobalFrame = true;
+	
+	UpdateInstanceData({}, Exec.ActiveFrames);
 	if (!InstanceData.IsValid())
 	{
 		STATETREE_LOG(Warning, TEXT("%hs: Failed to initialize instance data on '%s' using StateTree '%s'. Try to recompile the StateTree asset."),
-			__FUNCTION__, *GetNameSafe(&Owner), *GetFullNameSafe(&StateTree));
+			__FUNCTION__, *GetNameSafe(&Owner), *GetFullNameSafe(&RootStateTree));
 		return EStateTreeRunStatus::Failed;
 	}
 
@@ -250,10 +220,8 @@ EStateTreeRunStatus FStateTreeExecutionContext::Start()
 	// since we need it to output the InstanceId
 	STATETREE_TRACE_SCOPED_PHASE(EStateTreeUpdatePhase::StartTree);
 
-	FStateTreeExecutionState* Exec = &GetExecState(); // Using pointer as we will need to reacquire the exec later.
-
 	// From this point any calls to Stop should be deferred.
-	Exec->CurrentPhase = EStateTreeUpdatePhase::StartTree;
+	Exec.CurrentPhase = EStateTreeUpdatePhase::StartTree;
 
 	// Start evaluators and global tasks. Fail the execution if any global task fails.
 	FStateTreeIndex16 LastInitializedTaskIndex;
@@ -265,52 +233,47 @@ EStateTreeRunStatus FStateTreeExecutionContext::Start()
 		TickEvaluatorsAndGlobalTasks(0.0f, /*bTickGlobalTasks*/false);
 
 		// Initialize to unset running state.
-		Exec->TreeRunStatus = EStateTreeRunStatus::Running;
-		Exec->ActiveStates.Reset();
-		Exec->LastTickStatus = EStateTreeRunStatus::Unset;
+		Exec.TreeRunStatus = EStateTreeRunStatus::Running;
+		Exec.LastTickStatus = EStateTreeRunStatus::Unset;
 
 		static const FStateTreeStateHandle RootState = FStateTreeStateHandle(0);
 
-		FStateTreeActiveStates NextActiveStates;
-		FStateTreeActiveStates VisitedStates;
-		if (SelectState(RootState, NextActiveStates, VisitedStates))
+		TArray<FStateTreeExecutionFrame, TFixedAllocator<MaxExecutionFrames>> NextActiveFrames;
+		if (SelectState(InitFrame, RootState, NextActiveFrames))
 		{
-			if (NextActiveStates.Last().IsCompletionState())
+			check(!NextActiveFrames.IsEmpty());
+			if (NextActiveFrames.Last().ActiveStates.Last().IsCompletionState())
 			{
-				// Transition to a terminal state (succeeded/failed), or default transition failed.
+				// Transition to a terminal state (succeeded/failed).
 				STATETREE_LOG(Warning, TEXT("%hs: Tree %s at StateTree start on '%s' using StateTree '%s'."),
-					__FUNCTION__, NextActiveStates.Last() == FStateTreeStateHandle::Succeeded ? TEXT("succeeded") : TEXT("failed"), *GetNameSafe(&Owner), *GetFullNameSafe(&StateTree));
-				Exec->TreeRunStatus = NextActiveStates.Last().ToCompletionStatus();
+					__FUNCTION__, NextActiveFrames.Last().ActiveStates.Last() == FStateTreeStateHandle::Succeeded ? TEXT("succeeded") : TEXT("failed"), *GetNameSafe(&Owner), *GetFullNameSafe(&RootStateTree));
+				Exec.TreeRunStatus = NextActiveFrames.Last().ActiveStates.Last().ToCompletionStatus();
 			}
 			else
 			{
 				// Enter state tasks can fail/succeed, treat it same as tick.
 				FStateTreeTransitionResult Transition;
 				Transition.TargetState = RootState;
-				Transition.CurrentActiveStates = Exec->ActiveStates;
-				Transition.CurrentRunStatus = Exec->LastTickStatus;
-				Transition.NextActiveStates = NextActiveStates; // Enter state will update Exec.ActiveStates.
+				Transition.CurrentRunStatus = Exec.LastTickStatus;
+				Transition.NextActiveFrames = NextActiveFrames; // Enter state will update Exec.ActiveFrames.
 				const EStateTreeRunStatus LastTickStatus = EnterState(Transition);
 			
-				// Need to reacquire the exec state as EnterState may alter the allocation. 
-				Exec = &GetExecState();
-				Exec->LastTickStatus = LastTickStatus;
-				STATETREE_TRACE_ACTIVE_STATES_EVENT(Exec->ActiveStates);
+				Exec.LastTickStatus = LastTickStatus;
 
 				// Report state completed immediately.
-				if (Exec->LastTickStatus != EStateTreeRunStatus::Running)
+				if (Exec.LastTickStatus != EStateTreeRunStatus::Running)
 				{
 					StateCompleted();
 				}
 			}
 		}
 
-		if (Exec->ActiveStates.IsEmpty())
+		if (Exec.LastTickStatus == EStateTreeRunStatus::Unset)
 		{
 			// Should not happen. This may happen if initial state could not be selected.
 			STATETREE_LOG(Error, TEXT("%hs: Failed to select initial state on '%s' using StateTree '%s'. This should not happen, check that the StateTree logic can always select a state at start."),
-				__FUNCTION__, *GetNameSafe(&Owner), *GetFullNameSafe(&StateTree));
-			Exec->TreeRunStatus = EStateTreeRunStatus::Failed;
+				__FUNCTION__, *GetNameSafe(&Owner), *GetFullNameSafe(&RootStateTree));
+			Exec.TreeRunStatus = EStateTreeRunStatus::Failed;
 		}
 	}
 	else
@@ -318,22 +281,25 @@ EStateTreeRunStatus FStateTreeExecutionContext::Start()
 		StopEvaluatorsAndGlobalTasks(GlobalTasksRunStatus, LastInitializedTaskIndex);
 
 		STATETREE_LOG(VeryVerbose, TEXT("%hs: Global tasks completed the StateTree %s on start in status '%s'."),
-			__FUNCTION__, *GetNameSafe(&Owner), *GetFullNameSafe(&StateTree), *UEnum::GetDisplayValueAsText(GlobalTasksRunStatus).ToString());
+			__FUNCTION__, *GetNameSafe(&Owner), *GetFullNameSafe(&RootStateTree), *UEnum::GetDisplayValueAsText(GlobalTasksRunStatus).ToString());
 
+		// No active states or global tasks anymore, reset frames.
+		Exec.ActiveFrames.Reset();
+		
 		// We are not considered as running yet so we only set the status without requiring a stop.
-		Exec->TreeRunStatus = GlobalTasksRunStatus;
+		Exec.TreeRunStatus = GlobalTasksRunStatus;
 	}
 
 	// Reset phase since we are now safe to stop.
-	Exec->CurrentPhase = EStateTreeUpdatePhase::Unset;
+	Exec.CurrentPhase = EStateTreeUpdatePhase::Unset;
 
 	// Use local for resulting run state since Stop will reset the instance data.
-	EStateTreeRunStatus Result = Exec->TreeRunStatus;
+	EStateTreeRunStatus Result = Exec.TreeRunStatus;
 	
-	if (Exec->RequestedStop != EStateTreeRunStatus::Unset)
+	if (Exec.RequestedStop != EStateTreeRunStatus::Unset)
 	{
 		STATETREE_LOG_AND_TRACE(VeryVerbose, TEXT("Processing Deferred Stop"));
-		Result = Stop(Exec->RequestedStop);
+		Result = Stop(Exec.RequestedStop);
 	}
 	
 	return Result;
@@ -346,7 +312,7 @@ EStateTreeRunStatus FStateTreeExecutionContext::Stop(EStateTreeRunStatus Complet
 	if (!IsValid())
 	{
 		STATETREE_LOG(Warning, TEXT("%hs: StateTree context is not initialized properly ('%s' using StateTree '%s')"),
-			__FUNCTION__, *GetNameSafe(&Owner), *GetFullNameSafe(&StateTree));
+			__FUNCTION__, *GetNameSafe(&Owner), *GetFullNameSafe(&RootStateTree));
 		return EStateTreeRunStatus::Failed;
 	}
 	
@@ -393,11 +359,13 @@ EStateTreeRunStatus FStateTreeExecutionContext::Stop(EStateTreeRunStatus Complet
 		// Transition to Succeeded state.
 		FStateTreeTransitionResult Transition;
 		Transition.TargetState = FStateTreeStateHandle::FromCompletionStatus(CompletionStatus);
-		Transition.CurrentActiveStates = Exec.ActiveStates;
 		Transition.CurrentRunStatus = CompletionStatus;
-		Transition.NextActiveStates = FStateTreeActiveStates(Transition.TargetState);
-
-		if (!Exec.ActiveStates.IsEmpty())
+		FStateTreeExecutionFrame& NewFrame = Transition.NextActiveFrames.AddDefaulted_GetRef();
+		NewFrame.StateTree = &RootStateTree;
+		NewFrame.RootState = FStateTreeStateHandle::Root;
+		NewFrame.ActiveStates = {};
+		
+		if (Exec.TreeRunStatus == EStateTreeRunStatus::Running)
 		{
 			ExitState(Transition);
 		}
@@ -405,6 +373,9 @@ EStateTreeRunStatus FStateTreeExecutionContext::Stop(EStateTreeRunStatus Complet
 		// Stop evaluators and global tasks.
 		StopEvaluatorsAndGlobalTasks(CompletionStatus);
 
+		// No active states or global tasks anymore, reset frames.
+		Exec.ActiveFrames.Reset();
+
 		Result = CompletionStatus;
 	}
 
@@ -426,41 +397,41 @@ EStateTreeRunStatus FStateTreeExecutionContext::Tick(const float DeltaTime)
 	if (!IsValid())
 	{
 		STATETREE_LOG(Warning, TEXT("%hs: StateTree context is not initialized properly ('%s' using StateTree '%s')"),
-			__FUNCTION__, *GetNameSafe(&Owner), *GetFullNameSafe(&StateTree));
+			__FUNCTION__, *GetNameSafe(&Owner), *GetFullNameSafe(&RootStateTree));
 		return EStateTreeRunStatus::Failed;
 	}
 
 	if (!InstanceData.IsValid())
 	{
 		STATETREE_LOG(Error, TEXT("%hs: Tick called on %s using StateTree %s with invalid instance data. Start() must be called before Tick()."),
-			__FUNCTION__, *GetNameSafe(&Owner), *GetFullNameSafe(&StateTree));
+			__FUNCTION__, *GetNameSafe(&Owner), *GetFullNameSafe(&RootStateTree));
 		return EStateTreeRunStatus::Failed;
 	}
 
 	FStateTreeEventQueue& EventQueue = InstanceData.GetMutableEventQueue();
-	FStateTreeExecutionState* Exec = &GetExecState();
+	FStateTreeExecutionState& Exec = GetExecState();
 
 	// No ticking if the tree is done or stopped.
-	if (Exec->TreeRunStatus != EStateTreeRunStatus::Running)
+	if (Exec.TreeRunStatus != EStateTreeRunStatus::Running)
 	{
-		return Exec->TreeRunStatus;
+		return Exec.TreeRunStatus;
 	}
 
-	if (!ensureMsgf(Exec->CurrentPhase == EStateTreeUpdatePhase::Unset, TEXT("%hs can't be called while already in %s ('%s' using StateTree '%s')."),
-			__FUNCTION__, *UEnum::GetDisplayValueAsText(Exec->CurrentPhase).ToString(), *GetNameSafe(&Owner), *GetFullNameSafe(&StateTree)))
+	if (!ensureMsgf(Exec.CurrentPhase == EStateTreeUpdatePhase::Unset, TEXT("%hs can't be called while already in %s ('%s' using StateTree '%s')."),
+			__FUNCTION__, *UEnum::GetDisplayValueAsText(Exec.CurrentPhase).ToString(), *GetNameSafe(&Owner), *GetFullNameSafe(&RootStateTree)))
 	{
 		return EStateTreeRunStatus::Failed;
 	}
 
 	// From this point any calls to Stop should be deferred.
-	Exec->CurrentPhase = EStateTreeUpdatePhase::TickStateTree;
+	Exec.CurrentPhase = EStateTreeUpdatePhase::TickStateTree;
 	
 	// Capture events added between ticks.
 	EventsToProcess = EventQueue.GetEvents();
 	EventQueue.Reset();
 	
 	// Update the delayed transitions.
-	for (FStateTreeTransitionDelayedState& DelayedState : Exec->DelayedTransitions)
+	for (FStateTreeTransitionDelayedState& DelayedState : Exec.DelayedTransitions)
 	{
 		DelayedState.TimeLeft -= DeltaTime;
 	}
@@ -469,13 +440,13 @@ EStateTreeRunStatus FStateTreeExecutionContext::Tick(const float DeltaTime)
 	const EStateTreeRunStatus EvalAndGlobalTaskStatus = TickEvaluatorsAndGlobalTasks(DeltaTime);
 	if (EvalAndGlobalTaskStatus == EStateTreeRunStatus::Running)
 	{
-		if (Exec->LastTickStatus == EStateTreeRunStatus::Running)
+		if (Exec.LastTickStatus == EStateTreeRunStatus::Running)
 		{
 			// Tick tasks on active states.
-			Exec->LastTickStatus = TickTasks(DeltaTime);
+			Exec.LastTickStatus = TickTasks(DeltaTime);
 
 			// Report state completed immediately.
-			if (Exec->LastTickStatus != EStateTreeRunStatus::Running)
+			if (Exec.LastTickStatus != EStateTreeRunStatus::Running)
 			{
 				StateCompleted();
 			}
@@ -507,11 +478,13 @@ EStateTreeRunStatus FStateTreeExecutionContext::Tick(const float DeltaTime)
 				if (NextTransition.TargetState.IsCompletionState())
 				{
 					// Transition to a terminal state (succeeded/failed), or default transition failed.
-					Exec->TreeRunStatus = NextTransition.TargetState.ToCompletionStatus();
-					Exec->ActiveStates.Reset();
+					Exec.TreeRunStatus = NextTransition.TargetState.ToCompletionStatus();
 
 					// Stop evaluators and global tasks.
-					StopEvaluatorsAndGlobalTasks(Exec->TreeRunStatus);
+					StopEvaluatorsAndGlobalTasks(Exec.TreeRunStatus);
+
+					// No active states or global tasks anymore, reset frames.
+					Exec.ActiveFrames.Reset();
 
 					break;
 				}
@@ -525,23 +498,20 @@ EStateTreeRunStatus FStateTreeExecutionContext::Tick(const float DeltaTime)
 
 				NextTransition.Reset();
 
-				// Need to reacquire the exec state as EnterState may alter the allocation. 
-				Exec = &GetExecState();
-				Exec->LastTickStatus = LastTickStatus;
-				STATETREE_TRACE_ACTIVE_STATES_EVENT(Exec->ActiveStates);
+				Exec.LastTickStatus = LastTickStatus;
 
 				// Consider events so far processed. Events sent during EnterState went into EventQueue, and are processed in next iteration.
 				EventsToProcess.Reset();
 
 				// Report state completed immediately.
-				if (Exec->LastTickStatus != EStateTreeRunStatus::Running)
+				if (Exec.LastTickStatus != EStateTreeRunStatus::Running)
 				{
 					StateCompleted();
 				}
 			}
 
 			// Stop as soon as have found a running state.
-			if (Exec->LastTickStatus == EStateTreeRunStatus::Running)
+			if (Exec.LastTickStatus == EStateTreeRunStatus::Running)
 			{
 				break;
 			}
@@ -550,21 +520,21 @@ EStateTreeRunStatus FStateTreeExecutionContext::Tick(const float DeltaTime)
 	else
 	{
 		// If global tasks succeed or fail, stop the tree.
-		Exec->RequestedStop = EvalAndGlobalTaskStatus;
+		Exec.RequestedStop = EvalAndGlobalTaskStatus;
 	}
 
 	EventsToProcess.Reset();
 
 	// Reset phase since we are now safe to stop.
-	Exec->CurrentPhase = EStateTreeUpdatePhase::Unset;
+	Exec.CurrentPhase = EStateTreeUpdatePhase::Unset;
 
 	// Use local for resulting run state since Stop will reset the instance data.
-	EStateTreeRunStatus Result = Exec->TreeRunStatus;
+	EStateTreeRunStatus Result = Exec.TreeRunStatus;
 	
-	if (Exec->RequestedStop != EStateTreeRunStatus::Unset)
+	if (Exec.RequestedStop != EStateTreeRunStatus::Unset)
 	{
 		STATETREE_LOG_AND_TRACE(VeryVerbose, TEXT("Processing Deferred Stop"));
-		Result = Stop(Exec->RequestedStop);
+		Result = Stop(Exec.RequestedStop);
 	}
 
 	return Result;
@@ -575,7 +545,7 @@ EStateTreeRunStatus FStateTreeExecutionContext::GetStateTreeRunStatus() const
 	if (!IsValid())
 	{
 		STATETREE_LOG(Warning, TEXT("%hs: StateTree context is not initialized properly ('%s' using StateTree '%s')"),
-			__FUNCTION__, *GetNameSafe(&Owner), *GetFullNameSafe(&StateTree));
+			__FUNCTION__, *GetNameSafe(&Owner), *GetFullNameSafe(&RootStateTree));
 		return EStateTreeRunStatus::Failed;
 	}
 
@@ -599,14 +569,14 @@ void FStateTreeExecutionContext::SendEvent(const FGameplayTag Tag, const FConstS
 	if (!IsValid())
 	{
 		STATETREE_LOG(Warning, TEXT("%hs: StateTree context is not initialized properly ('%s' using StateTree '%s')"),
-			__FUNCTION__, *GetNameSafe(&Owner), *GetFullNameSafe(&StateTree));
+			__FUNCTION__, *GetNameSafe(&Owner), *GetFullNameSafe(&RootStateTree));
 		return;
 	}
 
 	if (!InstanceData.IsValid())
 	{
 		STATETREE_LOG(Error, TEXT("%hs: SendEvent called on %s using StateTree %s with invalid instance data. Start() must be called before sending events."),
-			__FUNCTION__, *GetNameSafe(&Owner), *GetFullNameSafe(&StateTree));
+			__FUNCTION__, *GetNameSafe(&Owner), *GetFullNameSafe(&RootStateTree));
 		return;
 	}
 
@@ -623,30 +593,52 @@ void FStateTreeExecutionContext::RequestTransition(const FStateTreeTransitionReq
 	if (!IsValid())
 	{
 		STATETREE_LOG(Warning, TEXT("%hs: StateTree context is not initialized properly ('%s' using StateTree '%s')"),
-			__FUNCTION__, *GetNameSafe(&Owner), *GetFullNameSafe(&StateTree));
+			__FUNCTION__, *GetNameSafe(&Owner), *GetFullNameSafe(&RootStateTree));
 		return;
 	}
 
 	if (!InstanceData.IsValid())
 	{
 		STATETREE_LOG(Error, TEXT("%hs: RequestTransition called on %s using StateTree %s with invalid instance data. Start() must be called before requesting transition."),
-			__FUNCTION__, *GetNameSafe(&Owner), *GetFullNameSafe(&StateTree));
+			__FUNCTION__, *GetNameSafe(&Owner), *GetFullNameSafe(&RootStateTree));
 		return;
 	}
 
-	STATETREE_LOG(Verbose, TEXT("Request transition to '%s' at priority %s"), *GetSafeStateName(Request.TargetState), *UEnum::GetDisplayValueAsText(Request.Priority).ToString());
+	FStateTreeExecutionState& Exec = GetExecState();
 
 	if (bAllowDirectTransitions)
 	{
-		if (RequestTransition(Request.TargetState, Request.Priority))
+		checkf(CurrentlyProcessedFrame, TEXT("Expecting CurrentlyProcessedFrame to be valid when called during TriggerTransitions()."));
+		
+		STATETREE_LOG(Verbose, TEXT("Request transition to '%s' at priority %s"), *GetSafeStateName(*CurrentlyProcessedFrame, Request.TargetState), *UEnum::GetDisplayValueAsText(Request.Priority).ToString());
+
+		if (RequestTransition(*CurrentlyProcessedFrame, Request.TargetState, Request.Priority))
 		{
 			NextTransitionSource = FStateTreeTransitionSource(EStateTreeTransitionSourceType::ExternalRequest, Request.TargetState, Request.Priority);
 		}
 	}
 	else
 	{
+		const FStateTreeExecutionFrame* RootFrame = &Exec.ActiveFrames[0];
+		if (CurrentlyProcessedFrame)
+		{
+			RootFrame = CurrentlyProcessedFrame;
+		}
+
+		if (!RootFrame)
+		{
+			STATETREE_LOG(Warning, TEXT("%hs: RequestTransition called on %s using StateTree %s without active state. Start() must be called before requesting transition."),
+				__FUNCTION__, *GetNameSafe(&Owner), *GetFullNameSafe(&RootStateTree));
+			return;
+		}
+		
+		STATETREE_LOG(Verbose, TEXT("Request transition to '%s' at priority %s"), *GetSafeStateName(*RootFrame, Request.TargetState), *UEnum::GetDisplayValueAsText(Request.Priority).ToString());
+
 		FStateTreeTransitionRequest RequestWithSource = Request;
+		RequestWithSource.SourceStateTree = RootFrame->StateTree;
+		RequestWithSource.SourceRootState = RootFrame->ActiveStates[0];
 		RequestWithSource.SourceState = CurrentlyProcessedState;
+		
 		InstanceData.AddTransitionRequest(&Owner, RequestWithSource);
 	}
 }
@@ -664,273 +656,476 @@ FStateTreeInstanceDebugId FStateTreeExecutionContext::GetInstanceDebugId() const
 }
 #endif // WITH_STATETREE_DEBUGGER
 
-void FStateTreeExecutionContext::UpdateInstanceData(const FStateTreeActiveStates& CurrentActiveStates, const FStateTreeActiveStates& NextActiveStates)
+void FStateTreeExecutionContext::UpdateInstanceData(TConstArrayView<FStateTreeExecutionFrame> CurrentActiveFrames, TArrayView<FStateTreeExecutionFrame> NextActiveFrames)
 {
-	// Find common section of states at start.
-	int32 NumCommon = 0;
-	while (NumCommon < CurrentActiveStates.Num() && NumCommon < NextActiveStates.Num())
+	// Estimate how many new instance data items we might have.
+	int32 EstimatedNumStructs = 0;
+	for (int32 FrameIndex = 0; FrameIndex < NextActiveFrames.Num(); FrameIndex++)
 	{
-		if (CurrentActiveStates[NumCommon] != NextActiveStates[NumCommon])
+		const FStateTreeExecutionFrame& NextFrame = NextActiveFrames[FrameIndex];
+		if (NextFrame.bIsGlobalFrame)
 		{
-			break;
+			EstimatedNumStructs += NextFrame.StateTree->NumGlobalInstanceData;
+		}
+		// States
+		for (int32 StateIndex = 0; StateIndex < NextFrame.ActiveStates.Num(); StateIndex++)
+		{
+			const FStateTreeStateHandle StateHandle = NextFrame.ActiveStates[StateIndex];
+			const FCompactStateTreeState& State = NextFrame.StateTree->States[StateHandle.Index];
+			EstimatedNumStructs += State.InstanceDataNum;
 		}
-		NumCommon++;
 	}
-
-	// @todo: change this so that we only put the newly added structs and objects here.
-	TArray<FConstStructView> InstanceStructs;
-	TArray<const UObject*> InstanceObjects;
 	
-	int32 NumCommonInstanceStructs = 0;
-	int32 NumCommonInstanceObjects = 0;
+	TArray<FConstStructView, TConcurrentLinearArrayAllocator<FDefaultBlockAllocationTag>> InstanceStructs;
+	InstanceStructs.Reserve(EstimatedNumStructs);
 
-	// Exec
-	InstanceStructs.Add(StateTree.DefaultInstanceData.GetStruct(0));
+	// Find next instance data sources and find common/existing section of instance data at start.
+	int32 CurrentGlobalInstanceIndexBase = 0;
+	int32 NumCommonInstanceData = 0;
 
-	// Evaluators
-	for (int32 EvalIndex = StateTree.EvaluatorsBegin; EvalIndex < (StateTree.EvaluatorsBegin + StateTree.EvaluatorsNum); EvalIndex++)
+	const UStruct* NextStateParameterDataStruct = nullptr;
+	FStateTreeIndex16 NextStateParameterDataIndex = FStateTreeIndex16::Invalid; 
+	
+	bool bAreCommon = true;
+	for (int32 FrameIndex = 0; FrameIndex < NextActiveFrames.Num(); FrameIndex++)
 	{
-		const FStateTreeEvaluatorBase& Eval =  StateTree.Nodes[EvalIndex].Get<const FStateTreeEvaluatorBase>();
-		if (Eval.bInstanceIsObject)
+		const bool bIsCurrentFrameValid = CurrentActiveFrames.IsValidIndex(FrameIndex)
+						&& CurrentActiveFrames[FrameIndex].IsSameFrame(NextActiveFrames[FrameIndex]);
+
+		bAreCommon &= bIsCurrentFrameValid;
+
+		const FStateTreeExecutionFrame* CurrentFrame = bIsCurrentFrameValid ? &CurrentActiveFrames[FrameIndex] : nullptr;
+		FStateTreeExecutionFrame& NextFrame = NextActiveFrames[FrameIndex];
+
+		check(NextFrame.StateTree);
+
+		if (NextFrame.bIsGlobalFrame)
 		{
-			InstanceObjects.Add(StateTree.DefaultInstanceData.GetObject(Eval.InstanceIndex.Get()));
+			// Global Evals
+			const int32 BaseIndex = InstanceStructs.Num();
+
+			NextFrame.GlobalInstanceIndexBase = FStateTreeIndex16(BaseIndex);
+			NextFrame.ActiveInstanceIndexBase = FStateTreeIndex16(BaseIndex);
+			CurrentGlobalInstanceIndexBase = BaseIndex;
+			
+			InstanceStructs.AddDefaulted(NextFrame.StateTree->NumGlobalInstanceData);
+			
+			for (int32 EvalIndex = NextFrame.StateTree->EvaluatorsBegin; EvalIndex < (NextFrame.StateTree->EvaluatorsBegin + NextFrame.StateTree->EvaluatorsNum); EvalIndex++)
+			{
+				const FStateTreeEvaluatorBase& Eval =  NextFrame.StateTree->Nodes[EvalIndex].Get<const FStateTreeEvaluatorBase>();
+				const FConstStructView EvalInstanceData = NextFrame.StateTree->DefaultInstanceData.GetStruct(Eval.InstanceTemplateIndex.Get());
+				InstanceStructs[BaseIndex + Eval.InstanceDataHandle.GetIndex()] = EvalInstanceData;
+			}
+
+			// Global tasks
+			for (int32 TaskIndex = NextFrame.StateTree->GlobalTasksBegin; TaskIndex < (NextFrame.StateTree->GlobalTasksBegin + NextFrame.StateTree->GlobalTasksNum); TaskIndex++)
+			{
+				const FStateTreeTaskBase& Task =  NextFrame.StateTree->Nodes[TaskIndex].Get<const FStateTreeTaskBase>();
+				const FConstStructView TaskInstanceData = NextFrame.StateTree->DefaultInstanceData.GetStruct(Task.InstanceTemplateIndex.Get());
+				InstanceStructs[BaseIndex + Task.InstanceDataHandle.GetIndex()] = TaskInstanceData;
+			}
+
+			if (bAreCommon)
+			{
+				NumCommonInstanceData = InstanceStructs.Num();
+			}
 		}
-		else
+
+		// States
+		const int32 BaseIndex = InstanceStructs.Num();
+
+		NextFrame.GlobalInstanceIndexBase = FStateTreeIndex16(CurrentGlobalInstanceIndexBase);
+		NextFrame.ActiveInstanceIndexBase = FStateTreeIndex16(BaseIndex);
+
+		for (int32 StateIndex = 0; StateIndex < NextFrame.ActiveStates.Num(); StateIndex++)
 		{
-			InstanceStructs.Add(StateTree.DefaultInstanceData.GetStruct(Eval.InstanceIndex.Get()));
+			// Check if the next state is still same as current state, GetStateSafe() will return invalid state if passed out of bounds index.
+			bAreCommon = bAreCommon && (CurrentFrame && CurrentFrame->ActiveStates.GetStateSafe(StateIndex) == NextFrame.ActiveStates[StateIndex]);
+
+			const FStateTreeStateHandle StateHandle = NextFrame.ActiveStates[StateIndex];
+			const FCompactStateTreeState& State = NextFrame.StateTree->States[StateHandle.Index];
+
+			InstanceStructs.AddDefaulted(State.InstanceDataNum);
+
+			if (State.Type == EStateTreeStateType::Linked)
+			{
+				// Linked state's instance data is the parameters.
+				check(State.ParameterDataHandle.IsValid());
+				const FConstStructView ParamsInstanceData = NextFrame.StateTree->DefaultInstanceData.GetStruct(State.ParameterTemplateIndex.Get());
+				InstanceStructs[BaseIndex + State.ParameterDataHandle.GetIndex()] = ParamsInstanceData;
+				
+				// Store the index of the parameter data, so that we can point the linked state to it.
+				check(State.ParameterDataHandle.GetSource() == EStateTreeDataSourceType::LinkedStateParameterData);
+				checkf(!NextStateParameterDataIndex.IsValid(), TEXT("NextStateParameterDataIndex not should be set yet when we encounter a linked state."));
+				NextStateParameterDataIndex = FStateTreeIndex16(BaseIndex + State.ParameterDataHandle.GetIndex());
+				NextStateParameterDataStruct = ParamsInstanceData.GetScriptStruct();
+			}
+			else if (State.Type == EStateTreeStateType::Subtree)
+			{
+				check(State.ParameterDataHandle.IsValid());
+				const FConstStructView ParamsInstanceData = NextFrame.StateTree->DefaultInstanceData.GetStruct(State.ParameterTemplateIndex.Get());
+				if (!NextStateParameterDataIndex.IsValid())
+				{
+					// Parameters are not set by a linked state, create instance data.
+					InstanceStructs[BaseIndex + State.ParameterDataHandle.GetIndex()] = ParamsInstanceData;
+					NextFrame.StateParameterDataIndex = FStateTreeIndex16(BaseIndex + State.ParameterDataHandle.GetIndex());
+				}
+				else
+				{
+					// Point to the parameter block set by linked state.
+					check(NextStateParameterDataStruct == ParamsInstanceData.GetScriptStruct());
+					NextFrame.StateParameterDataIndex = NextStateParameterDataIndex;
+					NextStateParameterDataIndex = FStateTreeIndex16(); // Mark as used.
+				}
+			}
+	
+			for (int32 TaskIndex = State.TasksBegin; TaskIndex < (State.TasksBegin + State.TasksNum); TaskIndex++)
+			{
+				const FStateTreeTaskBase& Task = NextFrame.StateTree->Nodes[TaskIndex].Get<const FStateTreeTaskBase>();
+				const FConstStructView TaskInstanceData = NextFrame.StateTree->DefaultInstanceData.GetStruct(Task.InstanceTemplateIndex.Get());
+				InstanceStructs[BaseIndex + Task.InstanceDataHandle.GetIndex()] = TaskInstanceData;
+			}
+
+			if (bAreCommon)
+			{
+				NumCommonInstanceData = InstanceStructs.Num();
+			}
 		}
 	}
-
-	// Global tasks
-	for (int32 TaskIndex = StateTree.GlobalTasksBegin; TaskIndex < (StateTree.GlobalTasksBegin + StateTree.GlobalTasksNum); TaskIndex++)
+	
+	// Common section should match.
+	// @todo: put this behind a define when enough testing has been done.
+	for (int32 Index = 0; Index < NumCommonInstanceData; Index++)
 	{
-		const FStateTreeTaskBase& Task =  StateTree.Nodes[TaskIndex].Get<const FStateTreeTaskBase>();
-		if (Task.bInstanceIsObject)
-		{
-			InstanceObjects.Add(StateTree.DefaultInstanceData.GetObject(Task.InstanceIndex.Get()));
-		}
-		else
+		check(Index < InstanceData.Num());
+
+		FConstStructView ExistingInstanceDataView = InstanceData.GetStruct(Index);
+		FConstStructView NewInstanceDataView = InstanceStructs[Index]; 
+
+		check(NewInstanceDataView.GetScriptStruct() == ExistingInstanceDataView.GetScriptStruct());
+
+		const FStateTreeInstanceObjectWrapper* ExistingWrapper = ExistingInstanceDataView.GetPtr<const FStateTreeInstanceObjectWrapper>();
+		const FStateTreeInstanceObjectWrapper* NewWrapper = ExistingInstanceDataView.GetPtr<const FStateTreeInstanceObjectWrapper>();
+		if (ExistingWrapper && NewWrapper)
 		{
-			InstanceStructs.Add(StateTree.DefaultInstanceData.GetStruct(Task.InstanceIndex.Get()));
+			check(ExistingWrapper->InstanceObject && NewWrapper->InstanceObject);
+			check(ExistingWrapper->InstanceObject->GetClass() == NewWrapper->InstanceObject->GetClass());
 		}
 	}
 
-	// Expect initialized instance data to contain the common instances.
-	if (InstanceData.IsValid())
-	{
-		NumCommonInstanceStructs = InstanceStructs.Num();
-		NumCommonInstanceObjects = InstanceObjects.Num();
-	}
-	
-	// Tasks
-	const int32 FirstTaskStructIndex = InstanceStructs.Num();
-	const int32 FirstTaskObjectIndex = InstanceObjects.Num();
-	
-	for (int32 Index = 0; Index < NextActiveStates.Num(); Index++)
+	// Remove instance data that was not common.
+	InstanceData.ShrinkTo(NumCommonInstanceData);
+
+	// Add new instance data.
+	InstanceData.Append(Owner, MakeArrayView(InstanceStructs.GetData() + NumCommonInstanceData, InstanceStructs.Num() - NumCommonInstanceData));
+}
+
+FStateTreeDataView FStateTreeExecutionContext::GetDataView(const FStateTreeExecutionFrame& CurrentFrame, const FStateTreeDataHandle Handle) const
+{
+	switch (Handle.GetSource())
 	{
-		const FStateTreeStateHandle CurrentHandle = NextActiveStates[Index];
-		const FCompactStateTreeState& State = StateTree.States[CurrentHandle.Index];
+	case EStateTreeDataSourceType::None:
+		return {};
 
-		if (State.Type == EStateTreeStateType::Linked)
-		{
-			check(State.ParameterInstanceIndex.IsValid());
-			InstanceStructs.Add(StateTree.DefaultInstanceData.GetStruct(State.ParameterInstanceIndex.Get()));
-		}
+	case EStateTreeDataSourceType::GlobalInstanceData:
+		return InstanceDataStorage->GetMutableStruct(CurrentFrame.GlobalInstanceIndexBase.Get() + Handle.GetIndex());
+	case EStateTreeDataSourceType::GlobalInstanceDataObject:
+		return InstanceDataStorage->GetMutableObject(CurrentFrame.GlobalInstanceIndexBase.Get() + Handle.GetIndex());
 		
-		for (int32 TaskIndex = State.TasksBegin; TaskIndex < (State.TasksBegin + State.TasksNum); TaskIndex++)
+	case EStateTreeDataSourceType::ActiveInstanceData:
+		return InstanceDataStorage->GetMutableStruct(CurrentFrame.ActiveInstanceIndexBase.Get() + Handle.GetIndex());
+	case EStateTreeDataSourceType::ActiveInstanceDataObject:
+		return InstanceDataStorage->GetMutableObject(CurrentFrame.ActiveInstanceIndexBase.Get() + Handle.GetIndex());
+
+	case EStateTreeDataSourceType::SharedInstanceData:
+		check(CurrentlyProcessedSharedInstanceStorage);
+		return CurrentlyProcessedSharedInstanceStorage->GetMutableStruct(Handle.GetIndex());
+	case EStateTreeDataSourceType::SharedInstanceDataObject:
+		check(CurrentlyProcessedSharedInstanceStorage);
+		return CurrentlyProcessedSharedInstanceStorage->GetMutableStruct(Handle.GetIndex());
+
+	case EStateTreeDataSourceType::ContextData:
+		return ContextDataViews[Handle.GetIndex()];
+
+	case EStateTreeDataSourceType::SubtreeParameterData:
 		{
-			const FStateTreeTaskBase& Task = StateTree.Nodes[TaskIndex].Get<const FStateTreeTaskBase>();
-			if (Task.bInstanceIsObject)
-			{
-				InstanceObjects.Add(StateTree.DefaultInstanceData.GetObject(Task.InstanceIndex.Get()));
-			}
-			else
-			{
-				InstanceStructs.Add(StateTree.DefaultInstanceData.GetStruct(Task.InstanceIndex.Get()));
-			}
+			FCompactStateTreeParameters& Params = InstanceDataStorage->GetMutableStruct(CurrentFrame.StateParameterDataIndex.Get()).Get<FCompactStateTreeParameters>();
+			return Params.Parameters.GetMutableValue();
 		}
-		
-		if (Index < NumCommon)
+	case EStateTreeDataSourceType::LinkedStateParameterData:
 		{
-			NumCommonInstanceStructs = InstanceStructs.Num();
-			NumCommonInstanceObjects = InstanceObjects.Num();
+			FCompactStateTreeParameters& Params = InstanceDataStorage->GetMutableStruct(CurrentFrame.ActiveInstanceIndexBase.Get() + Handle.GetIndex()).Get<FCompactStateTreeParameters>();
+			return Params.Parameters.GetMutableValue();
 		}
+
+	default:
+		checkf(false, TEXT("Unhandle case %s"), *UEnum::GetValueAsString(Handle.GetSource()));
 	}
 
-	// Common section should match.
-	// @todo: put this behind a define when enough testing has been done.
-	for (int32 Index = 0; Index < NumCommonInstanceStructs; Index++)
-	{
-		check(Index < InstanceData.NumStructs());
-		check(InstanceStructs[Index].GetScriptStruct() == InstanceData.GetStruct(Index).GetScriptStruct());
+	return {};
+}
+
+bool FStateTreeExecutionContext::IsHandleSourceValid(const FStateTreeExecutionFrame& CurrentFrame, const FStateTreeDataHandle Handle) const
+{
+	// Checks that the instance data is valid for specific handle types.
+	// 
+	// The CurrentFrame may not be yet properly initialized, for that reason we need to check
+	// that the path to the handle makes sense (it's part of the active states) as well as that
+	// we actually have instance data for the handle (index is valid).
+	// 
+	// The (base) indices can be invalid if the frame/state is not entered yet.
+	// For active instance data we need to check that the frame is initialized for a specific state,
+	// as well as that the instance data is initialized.
+
+	switch (Handle.GetSource())
+	{
+	case EStateTreeDataSourceType::None:
+		return true;
+
+	case EStateTreeDataSourceType::GlobalInstanceData:
+	case EStateTreeDataSourceType::GlobalInstanceDataObject:
+		return CurrentFrame.GlobalInstanceIndexBase.IsValid()
+			&& InstanceDataStorage->IsValidIndex(CurrentFrame.GlobalInstanceIndexBase.Get() + Handle.GetIndex());
+
+	case EStateTreeDataSourceType::ActiveInstanceData:
+	case EStateTreeDataSourceType::ActiveInstanceDataObject:
+		return CurrentFrame.ActiveInstanceIndexBase.IsValid()
+			&& CurrentFrame.ActiveStates.Contains(Handle.GetState())
+			&& InstanceDataStorage->IsValidIndex(CurrentFrame.ActiveInstanceIndexBase.Get() + Handle.GetIndex());
+		
+	case EStateTreeDataSourceType::SharedInstanceData:
+	case EStateTreeDataSourceType::SharedInstanceDataObject:
+		return true;
+
+	case EStateTreeDataSourceType::ContextData:
+		return true;
+
+	case EStateTreeDataSourceType::SubtreeParameterData:
+		return CurrentFrame.StateParameterDataIndex.IsValid()
+			&& InstanceDataStorage->IsValidIndex(CurrentFrame.StateParameterDataIndex.Get());
+
+	case EStateTreeDataSourceType::LinkedStateParameterData:
+		return CurrentFrame.ActiveInstanceIndexBase.IsValid()
+			&& CurrentFrame.ActiveStates.Contains(Handle.GetState())
+			&& InstanceDataStorage->IsValidIndex(CurrentFrame.ActiveInstanceIndexBase.Get() + Handle.GetIndex());
+
+	default:
+		checkf(false, TEXT("Unhandle case %s"), *UEnum::GetValueAsString(Handle.GetSource()));
 	}
-	for (int32 Index = 0; Index < NumCommonInstanceObjects; Index++)
+
+	return false;
+}
+
+
+bool FStateTreeExecutionContext::CopyBatch(const FStateTreeExecutionFrame& CurrentFrame, const FStateTreeDataView TargetView, const FStateTreeIndex16 BindingsBatch) const
+{
+	const FStateTreePropertyCopyBatch& Batch = CurrentFrame.StateTree->PropertyBindings.GetBatch(BindingsBatch);
+	check(TargetView.GetStruct() == Batch.TargetStruct.Struct);
+
+	bool bSucceed = true;
+	for (const FStateTreePropertyCopy& Copy : CurrentFrame.StateTree->PropertyBindings.GetBatchCopies(Batch))
 	{
-		check(Index < InstanceData.NumObjects());
-		check(InstanceObjects[Index] != nullptr
-			&& InstanceData.GetObject(Index) != nullptr
-			&& InstanceObjects[Index]->GetClass() == InstanceData.GetObject(Index)->GetClass());
+		const FStateTreeDataView SourceView = GetDataView(CurrentFrame, Copy.SourceDataHandle);
+		bSucceed &= CurrentFrame.StateTree->PropertyBindings.CopyProperty(Copy, SourceView, TargetView);
 	}
+	return bSucceed;
+}
 
-	// Remove instance data that was not common.
-	InstanceData.ShrinkTo(NumCommonInstanceStructs, NumCommonInstanceObjects);
+bool FStateTreeExecutionContext::CopyBatchWithValidation(const FStateTreeExecutionFrame& CurrentFrame, const FStateTreeDataView TargetView, const FStateTreeIndex16 BindingsBatch) const
+{
+	const FStateTreePropertyCopyBatch& Batch = CurrentFrame.StateTree->PropertyBindings.GetBatch(BindingsBatch);
+	check(TargetView.GetStruct() == Batch.TargetStruct.Struct);
 
-	// Add new instance data.
-	InstanceData.Append(Owner,
-		MakeArrayView(InstanceStructs.GetData() + NumCommonInstanceStructs, InstanceStructs.Num() - NumCommonInstanceStructs),
-		MakeArrayView(InstanceObjects.GetData() + NumCommonInstanceObjects, InstanceObjects.Num() - NumCommonInstanceObjects));
+	bool bSucceed = true;
+	for (const FStateTreePropertyCopy& Copy : CurrentFrame.StateTree->PropertyBindings.GetBatchCopies(Batch))
+	{
+		if (!IsHandleSourceValid(CurrentFrame, Copy.SourceDataHandle))
+		{
+			bSucceed = false;
+			break;
+		}
 
-	FStateTreeExecutionState& Exec = GetExecState();
-	Exec.FirstTaskStructIndex = FStateTreeIndex16(FirstTaskStructIndex);
-	Exec.FirstTaskObjectIndex = FStateTreeIndex16(FirstTaskObjectIndex);
+		const FStateTreeDataView SourceView = GetDataView(CurrentFrame, Copy.SourceDataHandle);
+		bSucceed &= CurrentFrame.StateTree->PropertyBindings.CopyProperty(Copy, SourceView, TargetView);
+	}
+	return bSucceed;
 }
 
-EStateTreeRunStatus FStateTreeExecutionContext::EnterState(const FStateTreeTransitionResult& Transition)
+EStateTreeRunStatus FStateTreeExecutionContext::EnterState(FStateTreeTransitionResult& Transition)
 {
 	CSV_SCOPED_TIMING_STAT_EXCLUSIVE(StateTree_EnterState);
 
-	if (Transition.NextActiveStates.IsEmpty())
+	if (Transition.NextActiveFrames.IsEmpty())
 	{
 		return EStateTreeRunStatus::Failed;
 	}
 
+	FStateTreeExecutionState& Exec = GetExecState();
+
 	// Allocate new tasks.
-	UpdateInstanceData(Transition.CurrentActiveStates, Transition.NextActiveStates);
+	UpdateInstanceData(Exec.ActiveFrames, Transition.NextActiveFrames);
 
-	FStateTreeExecutionState& Exec = GetExecState();
 	Exec.StateChangeCount++;
+	Exec.CompletedFrameIndex = FStateTreeIndex16::Invalid;
 	Exec.CompletedStateHandle = FStateTreeStateHandle::Invalid;
+	Exec.EnterStateFailedFrameIndex = FStateTreeIndex16::Invalid; // This will make all tasks to be accepted.
 	Exec.EnterStateFailedTaskIndex = FStateTreeIndex16::Invalid; // This will make all tasks to be accepted.
-	Exec.ActiveStates.Reset();
-
+	
 	// On target branch means that the state is the target of current transition or child of it.
 	// States which were active before and will remain active, but are not on target branch will not get
 	// EnterState called. That is, a transition is handled as "replan from this state".
 	bool bOnTargetBranch = false;
 	FStateTreeTransitionResult CurrentTransition = Transition;
 	EStateTreeRunStatus Result = EStateTreeRunStatus::Running;
-	int32 InstanceStructIndex = 1; // Exec is at index 0
-	int32 InstanceObjectIndex = 0;
-
-	// Update data views for evaluators and global tasks as UpdateInstanceData() might have changed the location of the instance data.
-	// Evaluators
-	for (int32 EvalIndex = StateTree.EvaluatorsBegin; EvalIndex < (StateTree.EvaluatorsBegin + StateTree.EvaluatorsNum); EvalIndex++)
-	{
-		const FStateTreeEvaluatorBase& Eval = StateTree.Nodes[EvalIndex].Get<const FStateTreeEvaluatorBase>();
-		SetNodeDataView(Eval, InstanceStructIndex, InstanceObjectIndex);
-	}
 
-	// Global tasks
-	for (int32 TaskIndex = StateTree.GlobalTasksBegin; TaskIndex < (StateTree.GlobalTasksBegin + StateTree.GlobalTasksNum); TaskIndex++)
-	{
-		const FStateTreeTaskBase& Task =  StateTree.Nodes[TaskIndex].Get<const FStateTreeTaskBase>();
-		SetNodeDataView(Task, InstanceStructIndex, InstanceObjectIndex);
-	}
-	
-	STATETREE_LOG(Log, TEXT("Enter state '%s' (%d)"), *DebugGetStatePath(Transition.NextActiveStates), Exec.StateChangeCount);
+	STATETREE_LOG(Log, TEXT("Enter state '%s' (%d)"), *DebugGetStatePath(Transition.NextActiveFrames), Exec.StateChangeCount);
 	STATETREE_TRACE_SCOPED_PHASE(EStateTreeUpdatePhase::EnterStates);
 
-	for (int32 Index = 0; Index < Transition.NextActiveStates.Num() && Result != EStateTreeRunStatus::Failed; Index++)
-	{
-		const FStateTreeStateHandle CurrentHandle = Transition.NextActiveStates[Index];
-		const FStateTreeStateHandle PreviousHandle = Transition.CurrentActiveStates.GetStateSafe(Index);
-		const FCompactStateTreeState& State = StateTree.States[CurrentHandle.Index];
-
-		// Add only enabled States to the list of active States
-		if (State.bEnabled && !Exec.ActiveStates.Push(CurrentHandle))
-		{
-			STATETREE_LOG(Error, TEXT("%hs: Reached max execution depth when trying to enter state '%s'.  '%s' using StateTree '%s'."),
-				__FUNCTION__, *GetStateStatusString(Exec), *GetNameSafe(&Owner), *GetFullNameSafe(&StateTree));
-			break;
-		}
-		
-		if (State.Type == EStateTreeStateType::Linked)
-		{
-			UpdateLinkedStateParameters(State, InstanceStructIndex);
-			InstanceStructIndex++;
-		}
-		else if (State.Type == EStateTreeStateType::Subtree)
-		{
-			UpdateSubtreeStateParameters(State);
-		}
-
-		bOnTargetBranch = bOnTargetBranch || CurrentHandle == Transition.TargetState;
-		const bool bWasActive = PreviousHandle == CurrentHandle;
+	// The previous active frames are needed for state enter logic.
+	TArray<FStateTreeExecutionFrame, TConcurrentLinearArrayAllocator<FDefaultBlockAllocationTag>> PreviousActiveFrames;
+	PreviousActiveFrames = Exec.ActiveFrames;
 
-		// Do not enter a disabled State tasks but maintain property bindings
-		const bool bIsEnteringState = (!bWasActive || bOnTargetBranch) && State.bEnabled;
+	// Reset the current active frames, new ones are added one by one.
+	Exec.ActiveFrames.Reset();
 
-		CurrentTransition.CurrentState = CurrentHandle;
-		CurrentTransition.ChangeType = bWasActive ? EStateTreeStateChangeType::Sustained : EStateTreeStateChangeType::Changed;
+	for (int32 FrameIndex = 0; FrameIndex < Transition.NextActiveFrames.Num() && Result != EStateTreeRunStatus::Failed; FrameIndex++)
+	{
+		const FStateTreeExecutionFrame& NextFrame = Transition.NextActiveFrames[FrameIndex];
+		
+		FStateTreeExecutionFrame& CurrentFrame = Exec.ActiveFrames.Add_GetRef(NextFrame);
+		// We'll add new states one by one, so that active states contain only the states which have EnterState called.
+		CurrentFrame.ActiveStates.Reset();
 
-		if (bIsEnteringState)
+		// Get previous active states, they are used to calculate transition type.
+		FStateTreeActiveStates PreviousActiveStates;
+		if (PreviousActiveFrames.IsValidIndex(FrameIndex)
+			&& PreviousActiveFrames[FrameIndex].IsSameFrame(NextFrame))
 		{
-			STATETREE_TRACE_STATE_EVENT(CurrentHandle, EStateTreeTraceEventType::OnEntering);
-			STATETREE_LOG(Log, TEXT("%*sState '%s' %s"), Index*UE::StateTree::DebugIndentSize, TEXT(""),
-				*DebugGetStatePath(Transition.NextActiveStates, Index),
-				*UEnum::GetDisplayValueAsText(CurrentTransition.ChangeType).ToString());
+			PreviousActiveStates = PreviousActiveFrames[FrameIndex].ActiveStates;
 		}
 
-		// Activate tasks on current state.
-		for (int32 TaskIndex = State.TasksBegin; TaskIndex < (State.TasksBegin + State.TasksNum); TaskIndex++)
+		FCurrentlyProcessedFrameScope FrameScope(*this, CurrentFrame, FrameIndex);
+		const UStateTree* CurrentStateTree = NextFrame.StateTree;
+
+		for (int32 Index = 0; Index < NextFrame.ActiveStates.Num() && Result != EStateTreeRunStatus::Failed; Index++)
 		{
-			const FStateTreeTaskBase& Task = StateTree.Nodes[TaskIndex].Get<const FStateTreeTaskBase>();
-			SetNodeDataView(Task, InstanceStructIndex, InstanceObjectIndex);
+			const FStateTreeStateHandle CurrentHandle = NextFrame.ActiveStates[Index];
+			const FStateTreeStateHandle PreviousHandle = PreviousActiveStates.GetStateSafe(Index);
+			const FCompactStateTreeState& State = CurrentStateTree->States[CurrentHandle.Index];
 
-			// Copy bound properties.
-			if (Task.BindingsBatch.IsValid())
+			FCurrentlyProcessedStateScope StateScope(*this, CurrentHandle);
+			
+			// Add only enabled States to the list of active States
+			if (State.bEnabled && !CurrentFrame.ActiveStates.Push(CurrentHandle))
 			{
-				StateTree.PropertyBindings.CopyTo(DataViews, Task.BindingsBatch, DataViews[Task.DataViewIndex.Get()]);
+				STATETREE_LOG(Error, TEXT("%hs: Reached max execution depth when trying to enter state '%s'.  '%s' using StateTree '%s'."),
+					__FUNCTION__, *GetStateStatusString(Exec), *GetNameSafe(&Owner), *GetFullNameSafe(&RootStateTree));
+				break;
 			}
 
-			// Ignore disabled task
-			if (Task.bTaskEnabled == false)
+			if (State.Type == EStateTreeStateType::Linked)
 			{
-				STATETREE_LOG(VeryVerbose, TEXT("%*sSkipped 'EnterState' for disabled Task: '%s'"), UE::StateTree::DebugIndentSize, TEXT(""), *Task.Name.ToString());
-				continue;
+				if (State.ParameterDataHandle.IsValid()
+					&& State.ParameterBindingsBatch.IsValid())
+				{
+					const FStateTreeDataView StateParamsDataView = GetDataView(CurrentFrame, State.ParameterDataHandle);
+					CopyBatch(CurrentFrame, StateParamsDataView, State.ParameterBindingsBatch);
+				}
 			}
 
-			const bool bShouldCallStateChange = CurrentTransition.ChangeType == EStateTreeStateChangeType::Changed
-												|| (CurrentTransition.ChangeType == EStateTreeStateChangeType::Sustained && Task.bShouldStateChangeOnReselect);
+			bOnTargetBranch = bOnTargetBranch || CurrentHandle == Transition.TargetState;
+			const bool bWasActive = PreviousHandle == CurrentHandle;
+
+			// Do not enter a disabled State tasks but maintain property bindings
+			const bool bIsEnteringState = (!bWasActive || bOnTargetBranch) && State.bEnabled;
 
-			if (bIsEnteringState && bShouldCallStateChange)
+			CurrentTransition.CurrentState = CurrentHandle;
+			CurrentTransition.ChangeType = bWasActive ? EStateTreeStateChangeType::Sustained : EStateTreeStateChangeType::Changed;
+
+			if (bIsEnteringState)
+			{
+				STATETREE_TRACE_STATE_EVENT(CurrentHandle, EStateTreeTraceEventType::OnEntering);
+				STATETREE_LOG(Log, TEXT("%*sState '%s' %s"), Index*UE::StateTree::DebugIndentSize, TEXT(""),
+					*DebugGetStatePath(Transition.NextActiveFrames, &NextFrame, Index),
+					*UEnum::GetDisplayValueAsText(CurrentTransition.ChangeType).ToString());
+			}
+
+			// Activate tasks on current state.
+			for (int32 TaskIndex = State.TasksBegin; TaskIndex < (State.TasksBegin + State.TasksNum); TaskIndex++)
 			{
-				STATETREE_LOG(Verbose, TEXT("%*s  Task '%s'"), Index*UE::StateTree::DebugIndentSize, TEXT(""), *Task.Name.ToString());
+				const FStateTreeTaskBase& Task = NextFrame.StateTree->Nodes[TaskIndex].Get<const FStateTreeTaskBase>();
+				const FStateTreeDataView TaskInstanceView = GetDataView(NextFrame, Task.InstanceDataHandle);
 
-				EStateTreeRunStatus Status = EStateTreeRunStatus::Unset;
+				// Copy bound properties.
+				if (Task.BindingsBatch.IsValid())
 				{
-					QUICK_SCOPE_CYCLE_COUNTER(StateTree_Task_EnterState);
-					CSV_SCOPED_TIMING_STAT_EXCLUSIVE(StateTree_Task_EnterState);
-				
-					Status = Task.EnterState(*this, CurrentTransition);
+					CopyBatch(NextFrame, TaskInstanceView, Task.BindingsBatch);
+				}
+
+				// Ignore disabled task
+				if (Task.bTaskEnabled == false)
+				{
+					STATETREE_LOG(VeryVerbose, TEXT("%*sSkipped 'EnterState' for disabled Task: '%s'"), UE::StateTree::DebugIndentSize, TEXT(""), *Task.Name.ToString());
+					continue;
 				}
 
-				STATETREE_TRACE_TASK_EVENT(TaskIndex, DataViews[Task.DataViewIndex.Get()], EStateTreeTraceEventType::OnEntered, Status);
+				const bool bShouldCallStateChange = CurrentTransition.ChangeType == EStateTreeStateChangeType::Changed
+													|| (CurrentTransition.ChangeType == EStateTreeStateChangeType::Sustained && Task.bShouldStateChangeOnReselect);
 
-				if (Status != EStateTreeRunStatus::Running)
+				if (bIsEnteringState && bShouldCallStateChange)
 				{
-					// Store the first state that completed, will be used to decide where to trigger transitions.
-					if (!Exec.CompletedStateHandle.IsValid())
+					STATETREE_LOG(Verbose, TEXT("%*s  Task '%s'"), Index*UE::StateTree::DebugIndentSize, TEXT(""), *Task.Name.ToString());
+
+					EStateTreeRunStatus Status = EStateTreeRunStatus::Unset;
 					{
-						Exec.CompletedStateHandle = CurrentHandle;
+						QUICK_SCOPE_CYCLE_COUNTER(StateTree_Task_EnterState);
+						CSV_SCOPED_TIMING_STAT_EXCLUSIVE(StateTree_Task_EnterState);
+					
+						Status = Task.EnterState(*this, CurrentTransition);
+					}
+
+					STATETREE_TRACE_TASK_EVENT(TaskIndex, TaskInstanceView, EStateTreeTraceEventType::OnEntered, Status);
+
+					if (Status != EStateTreeRunStatus::Running)
+					{
+						// Store the first state that completed, will be used to decide where to trigger transitions.
+						if (!Exec.CompletedStateHandle.IsValid())
+						{
+							Exec.CompletedFrameIndex = FStateTreeIndex16(FrameIndex);
+							Exec.CompletedStateHandle = CurrentHandle;
+						}
+						Result = Status;
+					}
+					
+					if (Status == EStateTreeRunStatus::Failed)
+					{
+						// Store how far in the enter state we got. This will be used to match the StateCompleted() and ExitState() calls.
+						Exec.EnterStateFailedFrameIndex = FStateTreeIndex16(FrameIndex); 
+						Exec.EnterStateFailedTaskIndex = FStateTreeIndex16(TaskIndex);
+						break;
 					}
-					Result = Status;
-				}
-				
-				if (Status == EStateTreeRunStatus::Failed)
-				{
-					// Store how far in the enter state we got. This will be used to match the StateCompleted() and ExitState() calls.
-					Exec.EnterStateFailedTaskIndex = FStateTreeIndex16(TaskIndex);
-					break;
 				}
 			}
+
+			if (bIsEnteringState)
+			{
+				STATETREE_TRACE_STATE_EVENT(CurrentHandle, EStateTreeTraceEventType::OnEntered);
+			}
 		}
+	}
 
-		if (bIsEnteringState)
+#if WITH_STATETREE_DEBUGGER
+	// @todo: implement support for frames in debugger
+	FStateTreeActiveStates AllActiveStates;
+	for (const FStateTreeExecutionFrame& ActiveFrame : Exec.ActiveFrames)
+	{
+		for (const FStateTreeStateHandle StateHandle : ActiveFrame.ActiveStates)
 		{
-			STATETREE_TRACE_STATE_EVENT(CurrentHandle, EStateTreeTraceEventType::OnEntered);
+			AllActiveStates.Push(StateHandle);	
 		}
 	}
+	STATETREE_TRACE_ACTIVE_STATES_EVENT(AllActiveStates);
+#endif // WITH_STATETREE_DEBUGGER
 
 	return Result;
 }
@@ -939,132 +1134,169 @@ void FStateTreeExecutionContext::ExitState(const FStateTreeTransitionResult& Tra
 {
 	CSV_SCOPED_TIMING_STAT_EXCLUSIVE(StateTree_ExitState);
 
-	if (Transition.CurrentActiveStates.IsEmpty())
+	FStateTreeExecutionState& Exec = GetExecState();
+
+	if (Exec.ActiveFrames.IsEmpty())
 	{
 		return;
 	}
 
-	FStateTreeExecutionState& Exec = GetExecState();
-
 	// On target branch means that the state is the target of current transition or child of it.
 	// States which were active before and will remain active, but are not on target branch will not get
 	// EnterState called. That is, a transition is handled as "replan from this state".
 	bool bOnTargetBranch = false;
 
-	FStateTreeStateHandle ExitedStates[FStateTreeActiveStates::MaxStates];
-	EStateTreeStateChangeType ExitedStateChangeType[FStateTreeActiveStates::MaxStates];
-	int32 ExitedStateActiveIndex[FStateTreeActiveStates::MaxStates];
-	int32 NumExitedStates = 0;
-	
-	// Do property copy on all states, propagating the results from last tick.
-	// Collect the states that need to be called, the actual call is done below in reverse order.
-	check(Exec.FirstTaskStructIndex.IsValid() && Exec.FirstTaskObjectIndex.IsValid()); 
-	int32 InstanceStructIndex = Exec.FirstTaskStructIndex.Get();
-	int32 InstanceObjectIndex = Exec.FirstTaskObjectIndex.Get();
-
-	for (int32 Index = 0; Index < Transition.CurrentActiveStates.Num(); Index++)
+	struct FExitStateCall
 	{
-		const FStateTreeStateHandle CurrentHandle = Transition.CurrentActiveStates[Index];
-		const FStateTreeStateHandle NextHandle = Transition.NextActiveStates.GetStateSafe(Index);
-		const FCompactStateTreeState& State = StateTree.States[CurrentHandle.Index];
-
-		if (State.Type == EStateTreeStateType::Linked)
+		FExitStateCall() = default;
+		FExitStateCall(const EStateTreeStateChangeType InChangeType, const bool bInShouldCall)
+			: ChangeType(InChangeType)
+			, bShouldCall(bInShouldCall)
 		{
-			UpdateLinkedStateParameters(State, InstanceStructIndex);
-			InstanceStructIndex++;
-		}
-		else if (State.Type == EStateTreeStateType::Subtree)
-		{
-			UpdateSubtreeStateParameters(State);
 		}
 
-		const bool bRemainsActive = NextHandle == CurrentHandle;
-		bOnTargetBranch = bOnTargetBranch || NextHandle == Transition.TargetState;
-		const EStateTreeStateChangeType ChangeType = bRemainsActive ? EStateTreeStateChangeType::Sustained : EStateTreeStateChangeType::Changed;
+		EStateTreeStateChangeType ChangeType = EStateTreeStateChangeType::None;
+		bool bShouldCall = false; 
+	};
+
+	TArray<FExitStateCall, TConcurrentLinearArrayAllocator<FDefaultBlockAllocationTag>> ExitStateCalls;
+	
+	for (int32 FrameIndex = 0; FrameIndex < Exec.ActiveFrames.Num(); FrameIndex++)
+	{
+		FStateTreeExecutionFrame& CurrentFrame = Exec.ActiveFrames[FrameIndex];
+		const UStateTree* CurrentStateTree = CurrentFrame.StateTree;
+		FStateTreeInstanceStorage& SharedInstanceDataStorage = CurrentStateTree->GetSharedInstanceData()->GetMutableStorage();
+
+		FCurrentlyProcessedFrameScope FrameScope(*this, CurrentFrame, FrameIndex);
 
-		if (!bRemainsActive || bOnTargetBranch)
+		const FStateTreeExecutionFrame* NextFrame = nullptr;
+		if (Transition.NextActiveFrames.IsValidIndex(FrameIndex)
+			&& Transition.NextActiveFrames[FrameIndex].IsSameFrame(CurrentFrame))
 		{
-			// Should call ExitState() on this state.
-			check (NumExitedStates < FStateTreeActiveStates::MaxStates);
-			ExitedStates[NumExitedStates] = CurrentHandle;
-			ExitedStateChangeType[NumExitedStates] = ChangeType;
-			ExitedStateActiveIndex[NumExitedStates] = Index;
-			NumExitedStates++;
+			NextFrame = &Transition.NextActiveFrames[FrameIndex];
 		}
 
-		// Do property copies, ExitState() is called below.
-		for (int32 TaskIndex = State.TasksBegin; TaskIndex < (State.TasksBegin + State.TasksNum); TaskIndex++)
+		for (int32 Index = 0; Index < CurrentFrame.ActiveStates.Num(); Index++)
 		{
-			const FStateTreeTaskBase& Task = StateTree.Nodes[TaskIndex].Get<const FStateTreeTaskBase>();
-			SetNodeDataView(Task, InstanceStructIndex, InstanceObjectIndex);
+			const FStateTreeStateHandle CurrentHandle = CurrentFrame.ActiveStates[Index];
+			const FStateTreeStateHandle NextHandle = NextFrame ? NextFrame->ActiveStates.GetStateSafe(Index) : FStateTreeStateHandle::Invalid;
+			const FCompactStateTreeState& State = CurrentStateTree->States[CurrentHandle.Index];
 
-			// Copy bound properties.
-			if (Task.BindingsBatch.IsValid() && Task.bShouldCopyBoundPropertiesOnExitState)
+			FCurrentlyProcessedStateScope StateScope(*this, CurrentHandle);
+
+			if (State.Type == EStateTreeStateType::Linked)
 			{
-				StateTree.PropertyBindings.CopyTo(DataViews, Task.BindingsBatch, DataViews[Task.DataViewIndex.Get()]);
+				if (State.ParameterDataHandle.IsValid()
+					&& State.ParameterBindingsBatch.IsValid())
+				{
+					const FStateTreeDataView StateParamsDataView = GetDataView(CurrentFrame, State.ParameterDataHandle);
+					CopyBatch(CurrentFrame, StateParamsDataView, State.ParameterBindingsBatch);
+				}
 			}
-		}
-	}
 
-	// Call in reverse order.
-	STATETREE_LOG(Log, TEXT("Exit state '%s' (%d)"), *DebugGetStatePath(Transition.CurrentActiveStates), Exec.StateChangeCount);
-	STATETREE_TRACE_SCOPED_PHASE(EStateTreeUpdatePhase::ExitStates);
+			const bool bRemainsActive = NextHandle == CurrentHandle;
+			bOnTargetBranch = bOnTargetBranch || NextHandle == Transition.TargetState;
+			const EStateTreeStateChangeType ChangeType = bRemainsActive ? EStateTreeStateChangeType::Sustained : EStateTreeStateChangeType::Changed;
+			
+			// Should call ExitState() on this state.
+			const bool bShouldCall = !bRemainsActive || bOnTargetBranch; 
+			ExitStateCalls.Emplace(ChangeType, bShouldCall);
 
-	FStateTreeTransitionResult CurrentTransition = Transition;
+			// Do property copies, ExitState() is called below.
+			for (int32 TaskIndex = State.TasksBegin; TaskIndex < (State.TasksBegin + State.TasksNum); TaskIndex++)
+			{
+				const FStateTreeTaskBase& Task = CurrentStateTree->Nodes[TaskIndex].Get<const FStateTreeTaskBase>();
+				const FStateTreeDataView TaskInstanceView = GetDataView(CurrentFrame, Task.InstanceDataHandle);
 
-	for (int32 Index = NumExitedStates - 1; Index >= 0; Index--)
-	{
-		const FStateTreeStateHandle CurrentHandle = ExitedStates[Index];
-		const FCompactStateTreeState& State = StateTree.States[CurrentHandle.Index];
+				// Copy bound properties.
+				if (Task.BindingsBatch.IsValid() && Task.bShouldCopyBoundPropertiesOnExitState)
+				{
+					CopyBatch(CurrentFrame, TaskInstanceView, Task.BindingsBatch);
+				}
+			}
+		}
+	}
 
-		// Remove any delayed transitions that belong to this state.
-		Exec.DelayedTransitions.RemoveAllSwap(
-			[Begin = State.TransitionsBegin, End = State.TransitionsBegin + State.TransitionsNum](const FStateTreeTransitionDelayedState& DelayedState)
-			{
-				return DelayedState.TransitionIndex.Get() >= Begin && DelayedState.TransitionIndex.Get() < End;
-			});
-		
-		CurrentTransition.CurrentState = CurrentHandle;
-		CurrentTransition.ChangeType = ExitedStateChangeType[Index];
+	// Call in reverse order.
+	STATETREE_LOG(Log, TEXT("Exit state '%s' (%d)"), *DebugGetStatePath(Exec.ActiveFrames), Exec.StateChangeCount);
+	STATETREE_TRACE_SCOPED_PHASE(EStateTreeUpdatePhase::ExitStates);
 
-		STATETREE_LOG(Log, TEXT("%*sState '%s' %s"), Index*UE::StateTree::DebugIndentSize, TEXT(""), *DebugGetStatePath(Transition.CurrentActiveStates, ExitedStateActiveIndex[Index]), *UEnum::GetDisplayValueAsText(CurrentTransition.ChangeType).ToString());
-		STATETREE_TRACE_STATE_EVENT(CurrentHandle, EStateTreeTraceEventType::OnExiting);
+	FStateTreeTransitionResult CurrentTransition = Transition;
+	int32 CallIndex = ExitStateCalls.Num() - 1;
 
-		// Tasks
-		for (int32 TaskIndex = (State.TasksBegin + State.TasksNum) - 1; TaskIndex >= State.TasksBegin; TaskIndex--)
+	for (int32 FrameIndex = Exec.ActiveFrames.Num() - 1; FrameIndex >= 0; FrameIndex--)
+	{
+		FStateTreeExecutionFrame& CurrentFrame = Exec.ActiveFrames[FrameIndex];
+		const UStateTree* CurrentStateTree = CurrentFrame.StateTree;
+
+		FCurrentlyProcessedFrameScope FrameScope(*this, CurrentFrame, FrameIndex);
+
+		for (int32 Index = CurrentFrame.ActiveStates.Num() - 1; Index >= 0; Index--)
 		{
-			// Call task completed only if EnterState() was called.
-			// The task order in the tree (BF) allows us to use the comparison.
-			// Relying here that invalid value of Exec.EnterStateFailedTaskIndex == MAX_uint16.
-			if (TaskIndex <= Exec.EnterStateFailedTaskIndex.Get())
+			const FStateTreeStateHandle CurrentHandle = CurrentFrame.ActiveStates[Index];
+			const FCompactStateTreeState& State = CurrentStateTree->States[CurrentHandle.Index];
+
+			const FExitStateCall& ExitCall = ExitStateCalls[CallIndex--];
+			CurrentTransition.ChangeType = ExitCall.ChangeType;
+
+			STATETREE_LOG(Log, TEXT("%*sState '%s' %s"), Index*UE::StateTree::DebugIndentSize, TEXT(""), *DebugGetStatePath(Exec.ActiveFrames, &CurrentFrame, CurrentHandle.Index), *UEnum::GetDisplayValueAsText(CurrentTransition.ChangeType).ToString());
+
+			STATETREE_TRACE_STATE_EVENT(CurrentHandle, EStateTreeTraceEventType::OnExiting);
+				
+			if (ExitCall.bShouldCall)
 			{
-				const FStateTreeTaskBase& Task = StateTree.Nodes[TaskIndex].Get<const FStateTreeTaskBase>();
+				FCurrentlyProcessedStateScope StateScope(*this, CurrentHandle);
 
-				// Ignore disabled task
-				if (Task.bTaskEnabled == false)
-				{
-					STATETREE_LOG(VeryVerbose, TEXT("%*sSkipped 'ExitState' for disabled Task: '%s'"), UE::StateTree::DebugIndentSize, TEXT(""), *Task.Name.ToString());
-					continue;
-				}
+				// Remove any delayed transitions that belong to this state.
+				Exec.DelayedTransitions.RemoveAllSwap(
+					[StateTree = CurrentFrame.StateTree, Begin = State.TransitionsBegin, End = State.TransitionsBegin + State.TransitionsNum](const FStateTreeTransitionDelayedState& DelayedState)
+					{
+						return  DelayedState.StateTree == StateTree && DelayedState.TransitionIndex.Get() >= Begin && DelayedState.TransitionIndex.Get() < End;
+					});
 
-				const bool bShouldCallStateChange = CurrentTransition.ChangeType == EStateTreeStateChangeType::Changed
-							|| (CurrentTransition.ChangeType == EStateTreeStateChangeType::Sustained && Task.bShouldStateChangeOnReselect);
+				CurrentTransition.CurrentState = CurrentHandle;
 
-				if (bShouldCallStateChange)
+				// Do property copies, ExitState() is called below.
+				for (int32 TaskIndex = (State.TasksBegin + State.TasksNum) - 1; TaskIndex >= State.TasksBegin; TaskIndex--)
 				{
-					STATETREE_LOG(Verbose, TEXT("%*s  Task '%s'"), Index*UE::StateTree::DebugIndentSize, TEXT(""), *Task.Name.ToString());
+
+					// Call task completed only if EnterState() was called.
+					// The task order in the tree (BF) allows us to use the comparison.
+					// Relying here that invalid value of Exec.EnterStateFailedTaskIndex == MAX_uint16.
+					if (TaskIndex <= Exec.EnterStateFailedTaskIndex.Get())
 					{
-						QUICK_SCOPE_CYCLE_COUNTER(StateTree_Task_ExitState);
-						CSV_SCOPED_TIMING_STAT_EXCLUSIVE(StateTree_Task_ExitState);
-						Task.ExitState(*this, CurrentTransition);
+						const FStateTreeTaskBase& Task = CurrentStateTree->Nodes[TaskIndex].Get<const FStateTreeTaskBase>();
+						const FStateTreeDataView TaskInstanceView = GetDataView(CurrentFrame, Task.InstanceDataHandle);
+
+						// Ignore disabled task
+						if (Task.bTaskEnabled == false)
+						{
+							STATETREE_LOG(VeryVerbose, TEXT("%*sSkipped 'ExitState' for disabled Task: '%s'"), UE::StateTree::DebugIndentSize, TEXT(""), *Task.Name.ToString());
+							continue;
+						}
+
+						const bool bShouldCallStateChange = CurrentTransition.ChangeType == EStateTreeStateChangeType::Changed
+									|| (CurrentTransition.ChangeType == EStateTreeStateChangeType::Sustained && Task.bShouldStateChangeOnReselect);
+
+						if (bShouldCallStateChange)
+						{
+
+							STATETREE_LOG(Verbose, TEXT("%*s  Task '%s'"), Index*UE::StateTree::DebugIndentSize, TEXT(""), *Task.Name.ToString());
+							{
+								QUICK_SCOPE_CYCLE_COUNTER(StateTree_Task_ExitState);
+								CSV_SCOPED_TIMING_STAT_EXCLUSIVE(StateTree_Task_ExitState);
+								Task.ExitState(*this, CurrentTransition);
+							}
+							STATETREE_TRACE_TASK_EVENT(TaskIndex, TaskInstanceView, EStateTreeTraceEventType::OnExited, Transition.CurrentRunStatus);
+						}
 					}
-					STATETREE_TRACE_TASK_EVENT(TaskIndex, DataViews[Task.DataViewIndex.Get()], EStateTreeTraceEventType::OnExited, Transition.CurrentRunStatus);
 				}
 			}
-		}
 
-		STATETREE_TRACE_STATE_EVENT(CurrentHandle, EStateTreeTraceEventType::OnExited);
+			STATETREE_TRACE_STATE_EVENT(CurrentHandle, EStateTreeTraceEventType::OnExited);
+		}
 	}
+
 }
 
 void FStateTreeExecutionContext::StateCompleted()
@@ -1073,7 +1305,7 @@ void FStateTreeExecutionContext::StateCompleted()
 
 	const FStateTreeExecutionState& Exec = GetExecState();
 
-	if (Exec.ActiveStates.IsEmpty())
+	if (Exec.ActiveFrames.IsEmpty())
 	{
 		return;
 	}
@@ -1082,35 +1314,47 @@ void FStateTreeExecutionContext::StateCompleted()
 
 	// Call from child towards root to allow to pass results back.
 	// Note: Completed is assumed to be called immediately after tick or enter state, so there's no property copying.
-	for (int32 Index = Exec.ActiveStates.Num() - 1; Index >= 0; Index--)
+
+	for (int32 FrameIndex = Exec.ActiveFrames.Num() - 1; FrameIndex >= 0; FrameIndex--)
 	{
-		const FStateTreeStateHandle CurrentHandle = Exec.ActiveStates[Index];
-		const FCompactStateTreeState& State = StateTree.States[CurrentHandle.Index];
+		const FStateTreeExecutionFrame& CurrentFrame = Exec.ActiveFrames[FrameIndex];
+		const UStateTree* CurrentStateTree = CurrentFrame.StateTree;
 
-		FCurrentlyProcessedStateScope StateScope(*this, CurrentHandle);
-		
-		STATETREE_LOG(Verbose, TEXT("%*sState '%s'"), Index*UE::StateTree::DebugIndentSize, TEXT(""), *DebugGetStatePath(Exec.ActiveStates, Index));
-		STATETREE_TRACE_STATE_EVENT(CurrentHandle, EStateTreeTraceEventType::OnStateCompleted);
+		FCurrentlyProcessedFrameScope FrameScope(*this, CurrentFrame, FrameIndex);
 
-		// Notify Tasks
-		for (int32 TaskIndex = (State.TasksBegin + State.TasksNum) - 1; TaskIndex >= State.TasksBegin; TaskIndex--)
+		if (FrameIndex <= Exec.EnterStateFailedFrameIndex.Get())
 		{
-			// Call task completed only if EnterState() was called.
-			// The task order in the tree (BF) allows us to use the comparison.
-			// Relying here that invalid value of Exec.EnterStateFailedTaskIndex == MAX_uint16.
-			if (TaskIndex <= Exec.EnterStateFailedTaskIndex.Get())
+			for (int32 Index = CurrentFrame.ActiveStates.Num() - 1; Index >= 0; Index--)
 			{
-				const FStateTreeTaskBase& Task = StateTree.Nodes[TaskIndex].Get<const FStateTreeTaskBase>();
+				const FStateTreeStateHandle CurrentHandle = CurrentFrame.ActiveStates[Index];
+				const FCompactStateTreeState& State = CurrentStateTree->States[CurrentHandle.Index];
 
-				// Ignore disabled task
-				if (Task.bTaskEnabled == false)
+				FCurrentlyProcessedStateScope StateScope(*this, CurrentHandle);
+				
+				STATETREE_LOG(Verbose, TEXT("%*sState '%s'"), Index*UE::StateTree::DebugIndentSize, TEXT(""), *DebugGetStatePath(Exec.ActiveFrames, &CurrentFrame, Index));
+				STATETREE_TRACE_STATE_EVENT(CurrentHandle, EStateTreeTraceEventType::OnStateCompleted);
+
+				// Notify Tasks
+				for (int32 TaskIndex = (State.TasksBegin + State.TasksNum) - 1; TaskIndex >= State.TasksBegin; TaskIndex--)
 				{
-					STATETREE_LOG(VeryVerbose, TEXT("%*sSkipped 'StateCompleted' for disabled Task: '%s'"), UE::StateTree::DebugIndentSize, TEXT(""), *Task.Name.ToString());
-					continue;
+					// Call task completed only if EnterState() was called.
+					// The task order in the tree (BF) allows us to use the comparison.
+					// Relying here that invalid value of Exec.EnterStateFailedTaskIndex == MAX_uint16.
+					if (TaskIndex <= Exec.EnterStateFailedTaskIndex.Get())
+					{
+						const FStateTreeTaskBase& Task = CurrentStateTree->Nodes[TaskIndex].Get<const FStateTreeTaskBase>();
+
+						// Ignore disabled task
+						if (Task.bTaskEnabled == false)
+						{
+							STATETREE_LOG(VeryVerbose, TEXT("%*sSkipped 'StateCompleted' for disabled Task: '%s'"), UE::StateTree::DebugIndentSize, TEXT(""), *Task.Name.ToString());
+							continue;
+						}
+						
+						STATETREE_LOG(Verbose, TEXT("%*s  Task '%s'"), Index*UE::StateTree::DebugIndentSize, TEXT(""), *Task.Name.ToString());
+						Task.StateCompleted(*this, Exec.LastTickStatus, CurrentFrame.ActiveStates);
+					}
 				}
-				
-				STATETREE_LOG(Verbose, TEXT("%*s  Task '%s'"), Index*UE::StateTree::DebugIndentSize, TEXT(""), *Task.Name.ToString());
-				Task.StateCompleted(*this, Exec.LastTickStatus, Exec.ActiveStates);
 			}
 		}
 	}
@@ -1121,90 +1365,97 @@ EStateTreeRunStatus FStateTreeExecutionContext::TickEvaluatorsAndGlobalTasks(con
 	CSV_SCOPED_TIMING_STAT_EXCLUSIVE(StateTree_TickEvaluators);
 	STATETREE_TRACE_SCOPED_PHASE(EStateTreeUpdatePhase::TickingGlobalTasks);
 
-	STATETREE_CLOG(StateTree.EvaluatorsNum > 0, VeryVerbose, TEXT("Ticking Evaluators"));
-
-	// Tick evaluators
-	int32 InstanceStructIndex = 1; // Exec is at index 0
-	int32 InstanceObjectIndex = 0;
-	
-	for (int32 EvalIndex = StateTree.EvaluatorsBegin; EvalIndex < (StateTree.EvaluatorsBegin + StateTree.EvaluatorsNum); EvalIndex++)
-	{
-		const FStateTreeEvaluatorBase& Eval = StateTree.Nodes[EvalIndex].Get<const FStateTreeEvaluatorBase>();
-		SetNodeDataView(Eval, InstanceStructIndex, InstanceObjectIndex);
-
-		// Copy bound properties.
-		if (Eval.BindingsBatch.IsValid())
-		{
-			StateTree.PropertyBindings.CopyTo(DataViews, Eval.BindingsBatch, DataViews[Eval.DataViewIndex.Get()]);
-		}
-		STATETREE_LOG(VeryVerbose, TEXT("  Tick: '%s'"), *Eval.Name.ToString());
-		{
-			QUICK_SCOPE_CYCLE_COUNTER(StateTree_Eval_Tick);
-			Eval.Tick(*this, DeltaTime);
-
-			STATETREE_TRACE_EVALUATOR_EVENT(EvalIndex, DataViews[Eval.DataViewIndex.Get()], EStateTreeTraceEventType::OnTicked);
-		}
-	}
+	STATETREE_LOG(VeryVerbose, TEXT("Ticking Evaluators & Global Tasks"));
 
+	FStateTreeExecutionState& Exec = GetExecState();
 
 	EStateTreeRunStatus Result = EStateTreeRunStatus::Running;
 
-	if (bTickGlobalTasks)
+	for (int32 FrameIndex = 0; FrameIndex < Exec.ActiveFrames.Num(); FrameIndex++)
 	{
-		// Used to stop ticking tasks after one fails, but we still want to keep updating the data views so that property binding works properly.
-		bool bShouldTickTasks = true;
-		const bool bHasEvents = !EventsToProcess.IsEmpty();
-
-		for (int32 TaskIndex = StateTree.GlobalTasksBegin; TaskIndex < (StateTree.GlobalTasksBegin + StateTree.GlobalTasksNum); TaskIndex++)
+		FStateTreeExecutionFrame& CurrentFrame = Exec.ActiveFrames[FrameIndex];
+		if (CurrentFrame.bIsGlobalFrame)
 		{
-			const FStateTreeTaskBase& Task =  StateTree.Nodes[TaskIndex].Get<const FStateTreeTaskBase>();
-			SetNodeDataView(Task, InstanceStructIndex, InstanceObjectIndex);
+			FCurrentlyProcessedFrameScope FrameScope(*this, CurrentFrame, FrameIndex);
 
-			// Ignore disabled task
-			if (Task.bTaskEnabled == false)
-			{
-				STATETREE_LOG(VeryVerbose, TEXT("%*sSkipped 'Tick' for disabled Task: '%s'"), UE::StateTree::DebugIndentSize, TEXT(""), *Task.Name.ToString());
-				continue;
-			}
+			const UStateTree* CurrentStateTree = CurrentFrame.StateTree;
 
-			const bool bNeedsTick = bShouldTickTasks && (Task.bShouldCallTick || (bHasEvents && Task.bShouldCallTickOnlyOnEvents));
-			STATETREE_LOG(VeryVerbose, TEXT("  Tick: '%s' %s"), *Task.Name.ToString(), !bNeedsTick ? TEXT("[not ticked]") : TEXT(""));
-			if (!bNeedsTick)
+			// Tick evaluators
+			for (int32 EvalIndex = CurrentStateTree->EvaluatorsBegin; EvalIndex < (CurrentStateTree->EvaluatorsBegin + CurrentStateTree->EvaluatorsNum); EvalIndex++)
 			{
-				continue;
-			}
+				const FStateTreeEvaluatorBase& Eval = CurrentStateTree->Nodes[EvalIndex].Get<const FStateTreeEvaluatorBase>();
+				const FStateTreeDataView EvalInstanceView = GetDataView(CurrentFrame, Eval.InstanceDataHandle);
 
-			const FStateTreeDataView TaskDataView = DataViews[Task.DataViewIndex.Get()];
+				// Copy bound properties.
+				if (Eval.BindingsBatch.IsValid())
+				{
+					CopyBatch(CurrentFrame, EvalInstanceView, Eval.BindingsBatch);
+				}
+				STATETREE_LOG(VeryVerbose, TEXT("  Tick: '%s'"), *Eval.Name.ToString());
+				{
+					QUICK_SCOPE_CYCLE_COUNTER(StateTree_Eval_Tick);
+					Eval.Tick(*this, DeltaTime);
 
-			// Copy bound properties.
-			// Only copy properties when the task is actually ticked, and copy properties at tick is requested.
-			if (Task.BindingsBatch.IsValid() && Task.bShouldCopyBoundPropertiesOnTick)
-			{
-				StateTree.PropertyBindings.CopyTo(DataViews, Task.BindingsBatch, TaskDataView);
+					STATETREE_TRACE_EVALUATOR_EVENT(EvalIndex, EvalInstanceView, EStateTreeTraceEventType::OnTicked);
+				}
 			}
 
-			//STATETREE_TRACE_TASK_EVENT(TaskIndex, TaskDataView, EStateTreeTraceEventType::OnTickingTask, EStateTreeRunStatus::Running);
-			EStateTreeRunStatus TaskResult = EStateTreeRunStatus::Unset;
+			if (bTickGlobalTasks)
 			{
-				QUICK_SCOPE_CYCLE_COUNTER(StateTree_Task_Tick);
-				CSV_SCOPED_TIMING_STAT_EXCLUSIVE(StateTree_Task_Tick);
+				// Used to stop ticking tasks after one fails, but we still want to keep updating the data views so that property binding works properly.
+				bool bShouldTickTasks = true;
+				const bool bHasEvents = !EventsToProcess.IsEmpty();
 
-				TaskResult = Task.Tick(*this, DeltaTime);
-			}
+				for (int32 TaskIndex = CurrentStateTree->GlobalTasksBegin; TaskIndex < (CurrentStateTree->GlobalTasksBegin + CurrentStateTree->GlobalTasksNum); TaskIndex++)
+				{
+					const FStateTreeTaskBase& Task = CurrentStateTree->Nodes[TaskIndex].Get<const FStateTreeTaskBase>();
+					const FStateTreeDataView TaskDataView = GetDataView(CurrentFrame, Task.InstanceDataHandle);
 
-			STATETREE_TRACE_TASK_EVENT(TaskIndex, TaskDataView,
-				TaskResult != EStateTreeRunStatus::Running ? EStateTreeTraceEventType::OnTaskCompleted : EStateTreeTraceEventType::OnTicked,
-				TaskResult);
+					// Ignore disabled task
+					if (Task.bTaskEnabled == false)
+					{
+						STATETREE_LOG(VeryVerbose, TEXT("%*sSkipped 'Tick' for disabled Task: '%s'"), UE::StateTree::DebugIndentSize, TEXT(""), *Task.Name.ToString());
+						continue;
+					}
 
-			// If a global task succeeds or fails, it will stop the whole tree.
-			if (TaskResult != EStateTreeRunStatus::Running)
-			{
-				Result = TaskResult;
-			}
-				
-			if (TaskResult == EStateTreeRunStatus::Failed)
-			{
-				bShouldTickTasks = false;
+					const bool bNeedsTick = bShouldTickTasks && (Task.bShouldCallTick || (bHasEvents && Task.bShouldCallTickOnlyOnEvents));
+					STATETREE_LOG(VeryVerbose, TEXT("  Tick: '%s' %s"), *Task.Name.ToString(), !bNeedsTick ? TEXT("[not ticked]") : TEXT(""));
+					if (!bNeedsTick)
+					{
+						continue;
+					}
+
+					// Copy bound properties.
+					// Only copy properties when the task is actually ticked, and copy properties at tick is requested.
+					if (Task.BindingsBatch.IsValid() && Task.bShouldCopyBoundPropertiesOnTick)
+					{
+						CopyBatch(CurrentFrame, TaskDataView, Task.BindingsBatch);
+					}
+
+					//STATETREE_TRACE_TASK_EVENT(TaskIndex, TaskDataView, EStateTreeTraceEventType::OnTickingTask, EStateTreeRunStatus::Running);
+					EStateTreeRunStatus TaskResult = EStateTreeRunStatus::Unset;
+					{
+						QUICK_SCOPE_CYCLE_COUNTER(StateTree_Task_Tick);
+						CSV_SCOPED_TIMING_STAT_EXCLUSIVE(StateTree_Task_Tick);
+
+						TaskResult = Task.Tick(*this, DeltaTime);
+					}
+
+					STATETREE_TRACE_TASK_EVENT(TaskIndex, TaskDataView,
+						TaskResult != EStateTreeRunStatus::Running ? EStateTreeTraceEventType::OnTaskCompleted : EStateTreeTraceEventType::OnTicked,
+						TaskResult);
+
+					// If a global task succeeds or fails, it will stop the whole tree.
+					if (TaskResult != EStateTreeRunStatus::Running)
+					{
+						Result = TaskResult;
+					}
+						
+					if (TaskResult == EStateTreeRunStatus::Failed)
+					{
+						bShouldTickTasks = false;
+					}
+				}
 			}
 		}
 	}
@@ -1212,73 +1463,84 @@ EStateTreeRunStatus FStateTreeExecutionContext::TickEvaluatorsAndGlobalTasks(con
 	return Result;
 }
 
+
 EStateTreeRunStatus FStateTreeExecutionContext::StartEvaluatorsAndGlobalTasks(FStateTreeIndex16& OutLastInitializedTaskIndex)
 {
 	CSV_SCOPED_TIMING_STAT_EXCLUSIVE(StateTree_StartEvaluators);
 	STATETREE_TRACE_SCOPED_PHASE(EStateTreeUpdatePhase::StartGlobalTasks);
 
-	STATETREE_CLOG(StateTree.EvaluatorsNum > 0 || StateTree.GlobalTasksNum > 0, Verbose, TEXT("Start Evaluators & Global tasks"));
+	STATETREE_LOG(Verbose, TEXT("Start Evaluators & Global tasks"));
+
+	FStateTreeExecutionState& Exec = GetExecState();
 
 	OutLastInitializedTaskIndex = FStateTreeIndex16();
 	EStateTreeRunStatus Result = EStateTreeRunStatus::Running;
-	
-	// Start evaluators
-	int32 InstanceStructIndex = 1; // Exec is at index 0
-	int32 InstanceObjectIndex = 0;
-	
-	for (int32 EvalIndex = StateTree.EvaluatorsBegin; EvalIndex < (StateTree.EvaluatorsBegin + StateTree.EvaluatorsNum); EvalIndex++)
-	{
-		const FStateTreeEvaluatorBase& Eval = StateTree.Nodes[EvalIndex].Get<const FStateTreeEvaluatorBase>();
-		SetNodeDataView(Eval, InstanceStructIndex, InstanceObjectIndex);
 
-		// Copy bound properties.
-		if (Eval.BindingsBatch.IsValid())
-		{
-			StateTree.PropertyBindings.CopyTo(DataViews, Eval.BindingsBatch, DataViews[Eval.DataViewIndex.Get()]);
-		}
-		STATETREE_LOG(Verbose, TEXT("  Start: '%s'"), *Eval.Name.ToString());
+	for (int32 FrameIndex = 0; FrameIndex < Exec.ActiveFrames.Num(); FrameIndex++)
+	{
+		FStateTreeExecutionFrame& CurrentFrame = Exec.ActiveFrames[FrameIndex];
+		if (CurrentFrame.bIsGlobalFrame)
 		{
-			QUICK_SCOPE_CYCLE_COUNTER(StateTree_Eval_TreeStart);
-			Eval.TreeStart(*this);
+			FCurrentlyProcessedFrameScope FrameScope(*this, CurrentFrame, FrameIndex);
+			
+			const UStateTree* CurrentStateTree = CurrentFrame.StateTree;
 
-			STATETREE_TRACE_EVALUATOR_EVENT(EvalIndex, DataViews[Eval.DataViewIndex.Get()], EStateTreeTraceEventType::OnTreeStarted);
-		}
-	}
+			// Start evaluators
+			for (int32 EvalIndex = CurrentStateTree->EvaluatorsBegin; EvalIndex < (CurrentStateTree->EvaluatorsBegin + CurrentStateTree->EvaluatorsNum); EvalIndex++)
+			{
+				const FStateTreeEvaluatorBase& Eval = CurrentStateTree->Nodes[EvalIndex].Get<const FStateTreeEvaluatorBase>();
+				const FStateTreeDataView EvalInstanceView = GetDataView(CurrentFrame, Eval.InstanceDataHandle);
 
-	// Start Global tasks
-	// Even if we call Enter/ExitState() on global tasks, they do not enter any specific state.
-	const FStateTreeTransitionResult Transition = {}; // Empty transition
-	
-	for (int32 TaskIndex = StateTree.GlobalTasksBegin; TaskIndex < (StateTree.GlobalTasksBegin + StateTree.GlobalTasksNum); TaskIndex++)
-	{
-		const FStateTreeTaskBase& Task =  StateTree.Nodes[TaskIndex].Get<const FStateTreeTaskBase>();
-		SetNodeDataView(Task, InstanceStructIndex, InstanceObjectIndex);
+				// Copy bound properties.
+				if (Eval.BindingsBatch.IsValid())
+				{
+					CopyBatch(CurrentFrame, EvalInstanceView, Eval.BindingsBatch);
+				}
+				STATETREE_LOG(Verbose, TEXT("  Start: '%s'"), *Eval.Name.ToString());
+				{
+					QUICK_SCOPE_CYCLE_COUNTER(StateTree_Eval_TreeStart);
+					Eval.TreeStart(*this);
 
-		// Copy bound properties.
-		if (Task.BindingsBatch.IsValid())
-		{
-			StateTree.PropertyBindings.CopyTo(DataViews, Task.BindingsBatch, DataViews[Task.DataViewIndex.Get()]);
-		}
+					STATETREE_TRACE_EVALUATOR_EVENT(EvalIndex, EvalInstanceView, EStateTreeTraceEventType::OnTreeStarted);
+				}
+			}
 
-		// Ignore disabled task
-		if (Task.bTaskEnabled == false)
-		{
-			STATETREE_LOG(VeryVerbose, TEXT("%*sSkipped 'EnterState' for disabled Task: '%s'"), UE::StateTree::DebugIndentSize, TEXT(""), *Task.Name.ToString());
-			continue;
-		}
+			// Start Global tasks
+			// Even if we call Enter/ExitState() on global tasks, they do not enter any specific state.
+			const FStateTreeTransitionResult Transition = {}; // Empty transition
+		
+			for (int32 TaskIndex = CurrentStateTree->GlobalTasksBegin; TaskIndex < (CurrentStateTree->GlobalTasksBegin + CurrentStateTree->GlobalTasksNum); TaskIndex++)
+			{
+				const FStateTreeTaskBase& Task =  CurrentStateTree->Nodes[TaskIndex].Get<const FStateTreeTaskBase>();
+				const FStateTreeDataView TaskDataView = GetDataView(CurrentFrame, Task.InstanceDataHandle);
 
-		STATETREE_LOG(Verbose, TEXT("  Start: '%s'"), *Task.Name.ToString());
-		{
-			QUICK_SCOPE_CYCLE_COUNTER(StateTree_Task_TreeStart);
-			const EStateTreeRunStatus TaskStatus = Task.EnterState(*this, Transition); 
+				// Copy bound properties.
+				if (Task.BindingsBatch.IsValid())
+				{
+					CopyBatch(CurrentFrame, TaskDataView, Task.BindingsBatch);
+				}
+
+				// Ignore disabled task
+				if (Task.bTaskEnabled == false)
+				{
+					STATETREE_LOG(VeryVerbose, TEXT("%*sSkipped 'EnterState' for disabled Task: '%s'"), UE::StateTree::DebugIndentSize, TEXT(""), *Task.Name.ToString());
+					continue;
+				}
 
-			STATETREE_TRACE_TASK_EVENT(TaskIndex, DataViews[Task.DataViewIndex.Get()], EStateTreeTraceEventType::OnEntered, TaskStatus);
+				STATETREE_LOG(Verbose, TEXT("  Start: '%s'"), *Task.Name.ToString());
+				{
+					QUICK_SCOPE_CYCLE_COUNTER(StateTree_Task_TreeStart);
+					const EStateTreeRunStatus TaskStatus = Task.EnterState(*this, Transition); 
 
-			if (TaskStatus != EStateTreeRunStatus::Running)
-			{
-				OutLastInitializedTaskIndex = FStateTreeIndex16(TaskIndex);
-				Result = TaskStatus;
-				break;
+					STATETREE_TRACE_TASK_EVENT(TaskIndex, TaskDataView, EStateTreeTraceEventType::OnEntered, TaskStatus);
+
+					if (TaskStatus != EStateTreeRunStatus::Running)
+					{
+						OutLastInitializedTaskIndex = FStateTreeIndex16(TaskIndex);
+						Result = TaskStatus;
+						break;
+					}
+				}
 			}
 		}
 	}
@@ -1291,77 +1553,97 @@ void FStateTreeExecutionContext::StopEvaluatorsAndGlobalTasks(const EStateTreeRu
 	CSV_SCOPED_TIMING_STAT_EXCLUSIVE(StateTree_StopEvaluators);
 	STATETREE_TRACE_SCOPED_PHASE(EStateTreeUpdatePhase::StopGlobalTasks);
 
-	STATETREE_CLOG(StateTree.EvaluatorsNum > 0, Verbose, TEXT("Stop Evaluators & Global Tasks"));
+	STATETREE_LOG(Verbose, TEXT("Stop Evaluators & Global Tasks"));
 
-	// Stop evaluators
-	int32 InstanceStructIndex = 1; // Exec is at index 0
-	int32 InstanceObjectIndex = 0;
-	
-	for (int32 EvalIndex = StateTree.EvaluatorsBegin; EvalIndex < (StateTree.EvaluatorsBegin + StateTree.EvaluatorsNum); EvalIndex++)
-	{
-		const FStateTreeEvaluatorBase& Eval = StateTree.Nodes[EvalIndex].Get<const FStateTreeEvaluatorBase>();
-		SetNodeDataView(Eval, InstanceStructIndex, InstanceObjectIndex);
+	FStateTreeExecutionState& Exec = GetExecState();
 
-		// Copy bound properties.
-		if (Eval.BindingsBatch.IsValid())
+	// Update bindings
+	for (int32 FrameIndex = 0; FrameIndex < Exec.ActiveFrames.Num(); FrameIndex++)
+	{
+		FStateTreeExecutionFrame& CurrentFrame = Exec.ActiveFrames[FrameIndex];
+		if (CurrentFrame.bIsGlobalFrame)
 		{
-			StateTree.PropertyBindings.CopyTo(DataViews, Eval.BindingsBatch, DataViews[Eval.DataViewIndex.Get()]);
-		}
-	}
+			FCurrentlyProcessedFrameScope FrameScope(*this, CurrentFrame, FrameIndex);
+			
+			const UStateTree* CurrentStateTree = CurrentFrame.StateTree;
 
-	// Stop Global tasks
-	for (int32 TaskIndex = StateTree.GlobalTasksBegin; TaskIndex < (StateTree.GlobalTasksBegin + StateTree.GlobalTasksNum); TaskIndex++)
-	{
-		const FStateTreeTaskBase& Task =  StateTree.Nodes[TaskIndex].Get<const FStateTreeTaskBase>();
-		SetNodeDataView(Task, InstanceStructIndex, InstanceObjectIndex);
+			for (int32 EvalIndex = CurrentStateTree->EvaluatorsBegin; EvalIndex < (CurrentStateTree->EvaluatorsBegin + CurrentStateTree->EvaluatorsNum); EvalIndex++)
+			{
+				const FStateTreeEvaluatorBase& Eval = CurrentStateTree->Nodes[EvalIndex].Get<const FStateTreeEvaluatorBase>();
+				const FStateTreeDataView EvalInstanceView = GetDataView(CurrentFrame, Eval.InstanceDataHandle);
 
-		// Copy bound properties.
-		if (Task.BindingsBatch.IsValid() && Task.bShouldCopyBoundPropertiesOnExitState)
-		{
-			StateTree.PropertyBindings.CopyTo(DataViews, Task.BindingsBatch, DataViews[Task.DataViewIndex.Get()]);
+				// Copy bound properties.
+				if (Eval.BindingsBatch.IsValid())
+				{
+					CopyBatch(CurrentFrame, EvalInstanceView, Eval.BindingsBatch);
+				}
+			}
+
+			for (int32 TaskIndex = CurrentStateTree->GlobalTasksBegin; TaskIndex < (CurrentStateTree->GlobalTasksBegin + CurrentStateTree->GlobalTasksNum); TaskIndex++)
+			{
+				const FStateTreeTaskBase& Task = CurrentStateTree->Nodes[TaskIndex].Get<const FStateTreeTaskBase>();
+				const FStateTreeDataView TaskDataView = GetDataView(CurrentFrame, Task.InstanceDataHandle);
+
+				// Copy bound properties.
+				if (Task.BindingsBatch.IsValid() && Task.bShouldCopyBoundPropertiesOnExitState)
+				{
+					CopyBatch(CurrentFrame, TaskDataView, Task.BindingsBatch);
+				}
+			}
 		}
 	}
 
-
 	// Call in reverse order.
 	FStateTreeTransitionResult Transition;
 	Transition.TargetState = FStateTreeStateHandle::FromCompletionStatus(CompletionStatus);
-	Transition.CurrentActiveStates = {};
 	Transition.CurrentRunStatus = CompletionStatus;
-	Transition.NextActiveStates = FStateTreeActiveStates(Transition.TargetState);
 
-	for (int32 TaskIndex = (StateTree.GlobalTasksBegin + StateTree.GlobalTasksNum) - 1;  TaskIndex >= StateTree.GlobalTasksBegin ; TaskIndex--)
+	for (int32 FrameIndex = Exec.ActiveFrames.Num() - 1; FrameIndex >= 0; FrameIndex--)
 	{
-		const FStateTreeTaskBase& Task =  StateTree.Nodes[TaskIndex].Get<const FStateTreeTaskBase>();
-
-		// Ignore disabled task
-		if (Task.bTaskEnabled == false)
+		const FStateTreeExecutionFrame& CurrentFrame = Exec.ActiveFrames[FrameIndex];
+		if (CurrentFrame.bIsGlobalFrame)
 		{
-			STATETREE_LOG(VeryVerbose, TEXT("%*sSkipped 'ExitState' for disabled Task: '%s'"), UE::StateTree::DebugIndentSize, TEXT(""), *Task.Name.ToString());
-			continue;
-		}
+			FCurrentlyProcessedFrameScope FrameScope(*this, CurrentFrame, FrameIndex);
 
-		// Relying here that invalid value of LastInitializedTaskIndex == MAX_uint16.
-		if (TaskIndex <= LastInitializedTaskIndex.Get())
-		{
-			STATETREE_LOG(Verbose, TEXT("  Stop: '%s'"), *Task.Name.ToString());
+			const UStateTree* CurrentStateTree = CurrentFrame.StateTree;
+
+			for (int32 TaskIndex = (CurrentStateTree->GlobalTasksBegin + CurrentStateTree->GlobalTasksNum) - 1;  TaskIndex >= CurrentStateTree->GlobalTasksBegin; TaskIndex--)
 			{
-				QUICK_SCOPE_CYCLE_COUNTER(StateTree_Task_TreeStop);
-				Task.ExitState(*this, Transition);
+				const FStateTreeTaskBase& Task =  CurrentStateTree->Nodes[TaskIndex].Get<const FStateTreeTaskBase>();
+				const FStateTreeDataView TaskDataView = GetDataView(CurrentFrame, Task.InstanceDataHandle);
+
+				// Ignore disabled task
+				if (Task.bTaskEnabled == false)
+				{
+					STATETREE_LOG(VeryVerbose, TEXT("%*sSkipped 'ExitState' for disabled Task: '%s'"), UE::StateTree::DebugIndentSize, TEXT(""), *Task.Name.ToString());
+					continue;
+				}
+
+				// Relying here that invalid value of LastInitializedTaskIndex == MAX_uint16.
+				if (TaskIndex <= LastInitializedTaskIndex.Get())
+				{
+					STATETREE_LOG(Verbose, TEXT("  Stop: '%s'"), *Task.Name.ToString());
+					{
+						QUICK_SCOPE_CYCLE_COUNTER(StateTree_Task_TreeStop);
+						Task.ExitState(*this, Transition);
+					}
+					STATETREE_TRACE_TASK_EVENT(TaskIndex, TaskDataView, EStateTreeTraceEventType::OnExited, Transition.CurrentRunStatus);
+				}
 			}
-			STATETREE_TRACE_TASK_EVENT(TaskIndex, DataViews[Task.DataViewIndex.Get()], EStateTreeTraceEventType::OnExited, Transition.CurrentRunStatus);
-		}
-	}
 
-	for (int32 EvalIndex = (StateTree.EvaluatorsBegin + StateTree.EvaluatorsNum) - 1; EvalIndex >= StateTree.EvaluatorsBegin; EvalIndex--)
-	{
-		const FStateTreeEvaluatorBase& Eval = StateTree.Nodes[EvalIndex].Get<const FStateTreeEvaluatorBase>();
-		STATETREE_LOG(Verbose, TEXT("  Stop: '%s'"), *Eval.Name.ToString());
-		{
-			QUICK_SCOPE_CYCLE_COUNTER(StateTree_Eval_TreeStop);
-			Eval.TreeStop(*this);
+			for (int32 EvalIndex = (CurrentStateTree->EvaluatorsBegin + CurrentStateTree->EvaluatorsNum) - 1; EvalIndex >= CurrentStateTree->EvaluatorsBegin; EvalIndex--)
+			{
+				const FStateTreeEvaluatorBase& Eval = CurrentStateTree->Nodes[EvalIndex].Get<const FStateTreeEvaluatorBase>();
+				const FStateTreeDataView EvalInstanceView = GetDataView(CurrentFrame, Eval.InstanceDataHandle);
+				
+				STATETREE_LOG(Verbose, TEXT("  Stop: '%s'"), *Eval.Name.ToString());
+				{
+					QUICK_SCOPE_CYCLE_COUNTER(StateTree_Eval_TreeStop);
+					Eval.TreeStop(*this);
 
-			STATETREE_TRACE_EVALUATOR_EVENT(EvalIndex, DataViews[Eval.DataViewIndex.Get()], EStateTreeTraceEventType::OnTreeStopped);
+					STATETREE_TRACE_EVALUATOR_EVENT(EvalIndex, EvalInstanceView, EStateTreeTraceEventType::OnTreeStopped);
+				}
+			}
 		}
 	}
 }
@@ -1373,7 +1655,7 @@ EStateTreeRunStatus FStateTreeExecutionContext::TickTasks(const float DeltaTime)
 
 	FStateTreeExecutionState& Exec = GetExecState();
 
-	if (Exec.ActiveStates.IsEmpty())
+	if (Exec.ActiveFrames.IsEmpty())
 	{
 		return EStateTreeRunStatus::Failed;
 	}
@@ -1382,113 +1664,119 @@ EStateTreeRunStatus FStateTreeExecutionContext::TickTasks(const float DeltaTime)
 	int32 NumTotalTasks = 0;
 
 	const bool bHasEvents = !EventsToProcess.IsEmpty();
-	
-	check(Exec.FirstTaskStructIndex.IsValid() && Exec.FirstTaskObjectIndex.IsValid()); 
-	int32 InstanceStructIndex = Exec.FirstTaskStructIndex.Get();
-	int32 InstanceObjectIndex = Exec.FirstTaskObjectIndex.Get();
 
+	Exec.CompletedFrameIndex = FStateTreeIndex16::Invalid;
 	Exec.CompletedStateHandle = FStateTreeStateHandle::Invalid;
 	
 	// Used to stop ticking tasks after one fails, but we still want to keep updating the data views so that property binding works properly.
 	bool bShouldTickTasks = true;
 
-	STATETREE_CLOG(Exec.ActiveStates.Num() > 0, VeryVerbose, TEXT("Ticking Tasks"));
+	STATETREE_CLOG(Exec.ActiveFrames.Num() > 0, VeryVerbose, TEXT("Ticking Tasks"));
 
-	for (int32 Index = 0; Index < Exec.ActiveStates.Num(); Index++)
+	for (int32 FrameIndex = 0; FrameIndex < Exec.ActiveFrames.Num(); FrameIndex++)
 	{
-		const FStateTreeStateHandle CurrentHandle = Exec.ActiveStates[Index];
-		const FCompactStateTreeState& State = StateTree.States[CurrentHandle.Index];
+		const FStateTreeExecutionFrame& CurrentFrame = Exec.ActiveFrames[FrameIndex];
+		const UStateTree* CurrentStateTree = CurrentFrame.StateTree;
 
-		FCurrentlyProcessedStateScope StateScope(*this, CurrentHandle);
-		STATETREE_TRACE_SCOPED_STATE(CurrentHandle);
+		FCurrentlyProcessedFrameScope FrameScope(*this, CurrentFrame, FrameIndex);
 
-		STATETREE_CLOG(State.TasksNum > 0, VeryVerbose, TEXT("%*sState '%s'"), Index*UE::StateTree::DebugIndentSize, TEXT(""), *DebugGetStatePath(Exec.ActiveStates, Index));
-
-		if (State.Type == EStateTreeStateType::Linked)
-		{
-			UpdateLinkedStateParameters(State, InstanceStructIndex);
-			InstanceStructIndex++;
-		}
-		else if (State.Type == EStateTreeStateType::Subtree)
+		for (int32 Index = 0; Index < CurrentFrame.ActiveStates.Num(); Index++)
 		{
-			UpdateSubtreeStateParameters(State);
-		}
+			const FStateTreeStateHandle CurrentHandle = CurrentFrame.ActiveStates[Index];
+			const FCompactStateTreeState& State = CurrentStateTree->States[CurrentHandle.Index];
 
-		// Update Tasks data and tick if possible (ie. if no task has yet failed and so bShouldTickTasks is true)
-		for (int32 TaskIndex = State.TasksBegin; TaskIndex < (State.TasksBegin + State.TasksNum); TaskIndex++)
-		{
-			const FStateTreeTaskBase& Task = StateTree.Nodes[TaskIndex].Get<const FStateTreeTaskBase>();
-			SetNodeDataView(Task, InstanceStructIndex, InstanceObjectIndex);
+			FCurrentlyProcessedStateScope StateScope(*this, CurrentHandle);
+			STATETREE_TRACE_SCOPED_STATE(CurrentHandle);
 
-			// Ignore disabled task
-			if (Task.bTaskEnabled == false)
-			{
-				STATETREE_LOG(VeryVerbose, TEXT("%*sSkipped 'Tick' for disabled Task: '%s'"), UE::StateTree::DebugIndentSize, TEXT(""), *Task.Name.ToString());
-				continue;
-			}
+			STATETREE_CLOG(State.TasksNum > 0, VeryVerbose, TEXT("%*sState '%s'"), Index*UE::StateTree::DebugIndentSize, TEXT(""), *DebugGetStatePath(Exec.ActiveFrames, &CurrentFrame, Index));
 
-			const bool bNeedsTick = bShouldTickTasks && (Task.bShouldCallTick || (bHasEvents && Task.bShouldCallTickOnlyOnEvents));
-			STATETREE_LOG(VeryVerbose, TEXT("%*s  Tick: '%s' %s"), Index*UE::StateTree::DebugIndentSize, TEXT(""), *Task.Name.ToString(), !bNeedsTick ? TEXT("[not ticked]") : TEXT(""));
-			if (!bNeedsTick)
+			if (State.Type == EStateTreeStateType::Linked)
 			{
-				continue;
+				if (State.ParameterDataHandle.IsValid()
+					&& State.ParameterBindingsBatch.IsValid())
+				{
+					const FStateTreeDataView StateParamsDataView = GetDataView(CurrentFrame, State.ParameterDataHandle);
+					CopyBatch(CurrentFrame, StateParamsDataView, State.ParameterBindingsBatch);
+				}
 			}
 
-			const FStateTreeDataView TaskDataView = DataViews[Task.DataViewIndex.Get()];
-			
-			// Copy bound properties.
-			// Only copy properties when the task is actually ticked, and copy properties at tick is requested.
-			if (Task.BindingsBatch.IsValid() && Task.bShouldCopyBoundPropertiesOnTick)
+			// Update Tasks data and tick if possible (ie. if no task has yet failed and so bShouldTickTasks is true)
+			for (int32 TaskIndex = State.TasksBegin; TaskIndex < (State.TasksBegin + State.TasksNum); TaskIndex++)
 			{
-				StateTree.PropertyBindings.CopyTo(DataViews, Task.BindingsBatch, TaskDataView);
-			}
+				const FStateTreeTaskBase& Task = CurrentStateTree->Nodes[TaskIndex].Get<const FStateTreeTaskBase>();
+				const FStateTreeDataView TaskDataView = GetDataView(CurrentFrame, Task.InstanceDataHandle);
 
-			//STATETREE_TRACE_TASK_EVENT(TaskIndex, TaskDataView, EStateTreeTraceEventType::OnTickingTask, EStateTreeRunStatus::Running);
-			EStateTreeRunStatus TaskResult = EStateTreeRunStatus::Unset;
-			{
-				QUICK_SCOPE_CYCLE_COUNTER(StateTree_Task_Tick);
-				CSV_SCOPED_TIMING_STAT_EXCLUSIVE(StateTree_Task_Tick);
+				// Ignore disabled task
+				if (Task.bTaskEnabled == false)
+				{
+					STATETREE_LOG(VeryVerbose, TEXT("%*sSkipped 'Tick' for disabled Task: '%s'"), UE::StateTree::DebugIndentSize, TEXT(""), *Task.Name.ToString());
+					continue;
+				}
 
-				TaskResult = Task.Tick(*this, DeltaTime);
-			}
+				const bool bNeedsTick = bShouldTickTasks && (Task.bShouldCallTick || (bHasEvents && Task.bShouldCallTickOnlyOnEvents));
+				STATETREE_LOG(VeryVerbose, TEXT("%*s  Tick: '%s' %s"), Index*UE::StateTree::DebugIndentSize, TEXT(""), *Task.Name.ToString(), !bNeedsTick ? TEXT("[not ticked]") : TEXT(""));
+				if (!bNeedsTick)
+				{
+					continue;
+				}
+				
+				// Copy bound properties.
+				// Only copy properties when the task is actually ticked, and copy properties at tick is requested.
+				if (Task.BindingsBatch.IsValid() && Task.bShouldCopyBoundPropertiesOnTick)
+				{
+					CopyBatch(CurrentFrame, TaskDataView, Task.BindingsBatch);
+				}
 
-			STATETREE_TRACE_TASK_EVENT(TaskIndex, TaskDataView,
-				TaskResult != EStateTreeRunStatus::Running ? EStateTreeTraceEventType::OnTaskCompleted : EStateTreeTraceEventType::OnTicked,
-				TaskResult);
-			
-			// TODO: Add more control over which states can control the failed/succeeded result.
-			if (TaskResult != EStateTreeRunStatus::Running)
-			{
-				// Store the first state that completed, will be used to decide where to trigger transitions.
-				if (!Exec.CompletedStateHandle.IsValid())
+				//STATETREE_TRACE_TASK_EVENT(TaskIndex, TaskDataView, EStateTreeTraceEventType::OnTickingTask, EStateTreeRunStatus::Running);
+				EStateTreeRunStatus TaskResult = EStateTreeRunStatus::Unset;
 				{
-					Exec.CompletedStateHandle = CurrentHandle;
+					QUICK_SCOPE_CYCLE_COUNTER(StateTree_Task_Tick);
+					CSV_SCOPED_TIMING_STAT_EXCLUSIVE(StateTree_Task_Tick);
+
+					TaskResult = Task.Tick(*this, DeltaTime);
+				}
+
+				STATETREE_TRACE_TASK_EVENT(TaskIndex, TaskDataView,
+					TaskResult != EStateTreeRunStatus::Running ? EStateTreeTraceEventType::OnTaskCompleted : EStateTreeTraceEventType::OnTicked,
+					TaskResult);
+				
+				// TODO: Add more control over which states can control the failed/succeeded result.
+				if (TaskResult != EStateTreeRunStatus::Running)
+				{
+					// Store the first state that completed, will be used to decide where to trigger transitions.
+					if (!Exec.CompletedStateHandle.IsValid())
+					{
+						Exec.CompletedFrameIndex = FStateTreeIndex16(FrameIndex);
+						Exec.CompletedStateHandle = CurrentHandle;
+					}
+					Result = TaskResult;
+				}
+				
+				if (TaskResult == EStateTreeRunStatus::Failed)
+				{
+					bShouldTickTasks = false;
 				}
-				Result = TaskResult;
-			}
-			
-			if (TaskResult == EStateTreeRunStatus::Failed)
-			{
-				bShouldTickTasks = false;
 			}
+			NumTotalTasks += State.TasksNum;
 		}
-		NumTotalTasks += State.TasksNum;
 	}
 
 	if (NumTotalTasks == 0)
 	{
 		// No tasks, done ticking.
 		Result = EStateTreeRunStatus::Succeeded;
+		Exec.CompletedFrameIndex = FStateTreeIndex16(0);
+		Exec.CompletedStateHandle = Exec.ActiveFrames[0].ActiveStates.GetStateSafe(0);
 	}
 
 	return Result;
 }
 
-bool FStateTreeExecutionContext::TestAllConditions(const int32 ConditionsOffset, const int32 ConditionsNum)
+bool FStateTreeExecutionContext::TestAllConditions(const FStateTreeExecutionFrame& CurrentFrame, const int32 ConditionsOffset, const int32 ConditionsNum)
 {
 	CSV_SCOPED_TIMING_STAT_EXCLUSIVE(StateTree_TestConditions);
 
-	if (ConditionsNum == 0 || !SharedInstanceData.IsValid())
+	if (ConditionsNum == 0)
 	{
 		return true;
 	}
@@ -1501,24 +1789,17 @@ bool FStateTreeExecutionContext::TestAllConditions(const int32 ConditionsOffset,
 	for (int32 Index = 0; Index < ConditionsNum; Index++)
 	{
 		const int32 ConditionIndex = ConditionsOffset + Index;
-		const FStateTreeConditionBase& Cond = StateTree.Nodes[ConditionIndex].Get<const FStateTreeConditionBase>();
-		FStateTreeDataView& DataView = DataViews[Cond.DataViewIndex.Get()]; 
-		if (Cond.bInstanceIsObject)
-		{
-			DataView = SharedInstanceData->GetMutableObject(Cond.InstanceIndex.Get());
-		}
-		else
-		{
-			DataView = SharedInstanceData->GetMutableStruct(Cond.InstanceIndex.Get());
-		}
-
+		const FStateTreeConditionBase& Cond = CurrentFrame.StateTree->Nodes[ConditionIndex].Get<const FStateTreeConditionBase>();
+		const FStateTreeDataView ConditionDataView = GetDataView(CurrentFrame, Cond.InstanceDataHandle);
+		
 		bool bValue = false;
 		if (Cond.EvaluationMode == EStateTreeConditionEvaluationMode::Evaluated)
 		{
 			// Copy bound properties.
 			if (Cond.BindingsBatch.IsValid())
 			{
-				if (!StateTree.PropertyBindings.CopyTo(DataViews, Cond.BindingsBatch, DataView))
+				// Use validated copy, since we test in situations where the sources are not always valid (e.g. enter conditions may try to access inactive parent state). 
+				if (!CopyBatchWithValidation(CurrentFrame, ConditionDataView, Cond.BindingsBatch))
 				{
 					// If the source data cannot be accessed, the whole expression evaluates to false.
 					Values[0] = false;
@@ -1527,12 +1808,12 @@ bool FStateTreeExecutionContext::TestAllConditions(const int32 ConditionsOffset,
 			}
 			
 			bValue = Cond.TestCondition(*this);
-			STATETREE_TRACE_CONDITION_EVENT(ConditionIndex, DataView, bValue ? EStateTreeTraceEventType::Passed : EStateTreeTraceEventType::Failed);
+			STATETREE_TRACE_CONDITION_EVENT(ConditionIndex, ConditionDataView, bValue ? EStateTreeTraceEventType::Passed : EStateTreeTraceEventType::Failed);
 			
 			// Reset copied properties that might contain object references.
 			if (Cond.BindingsBatch.IsValid())
 			{
-				StateTree.PropertyBindings.ResetObjects(Cond.BindingsBatch, DataView);
+				CurrentFrame.StateTree->PropertyBindings.ResetObjects(Cond.BindingsBatch, ConditionDataView);
 			}
 		}
 		else
@@ -1594,7 +1875,7 @@ FString FStateTreeExecutionContext::DebugGetEventsAsString() const
 	return Result;
 }
 
-bool FStateTreeExecutionContext::RequestTransition(const FStateTreeStateHandle NextState, const EStateTreeTransitionPriority Priority, const EStateTreeSelectionFallback Fallback)
+bool FStateTreeExecutionContext::RequestTransition(const FStateTreeExecutionFrame& CurrentFrame, const FStateTreeStateHandle NextState, const EStateTreeTransitionPriority Priority, const EStateTreeSelectionFallback Fallback)
 {
 	// Skip lower priority transitions.
 	if (NextTransition.Priority >= Priority)
@@ -1602,32 +1883,30 @@ bool FStateTreeExecutionContext::RequestTransition(const FStateTreeStateHandle N
 		return false;
 	}
 
-	const FStateTreeExecutionState& Exec = GetExecState();
-	
 	if (NextState.IsCompletionState())
 	{
-		SetupNextTransition(NextState, Priority);
-		STATETREE_LOG(Verbose, TEXT("Transition on state '%s' -[%s]-> state '%s'"),
-			*GetSafeStateName(NextTransition.CurrentActiveStates.Last()), *GetSafeStateName(NextState), *GetSafeStateName(NextTransition.NextActiveStates.Last()));
-
+		SetupNextTransition(CurrentFrame, NextState, Priority);
+		STATETREE_LOG(Verbose, TEXT("Transition on state '%s' -> state '%s'"),
+			*GetSafeStateName(CurrentFrame, CurrentFrame.ActiveStates.Last()), *NextState.Describe());
 		return true;
 	}
 	if (!NextState.IsValid())
 	{
 		// NotSet is no-operation, but can be used to mask a transition at parent state. Returning unset keeps updating current state.
-		SetupNextTransition(FStateTreeStateHandle::Invalid, Priority);
+		SetupNextTransition(CurrentFrame, FStateTreeStateHandle::Invalid, Priority);
 		return true;
 	}
 
-	FStateTreeActiveStates NewActiveState;
-	FStateTreeActiveStates VisitedStates;
-	if (SelectState(NextState, NewActiveState, VisitedStates, Fallback))
+	TArray<FStateTreeExecutionFrame, TFixedAllocator<MaxExecutionFrames>> NewNextActiveFrames;
+	if (SelectState(CurrentFrame, NextState, NewNextActiveFrames, Fallback))
 	{
-		SetupNextTransition(NextState, Priority);
-		NextTransition.NextActiveStates = NewActiveState;
+		SetupNextTransition(CurrentFrame, NextState, Priority);
+		NextTransition.NextActiveFrames = NewNextActiveFrames;
 
 		STATETREE_LOG(Verbose, TEXT("Transition on state '%s' -[%s]-> state '%s'"),
-			*GetSafeStateName(NextTransition.CurrentActiveStates.Last()), *GetSafeStateName(NextState), *GetSafeStateName(NextTransition.NextActiveStates.Last()));
+			*GetSafeStateName(CurrentFrame, CurrentFrame.ActiveStates.Last()),
+			*GetSafeStateName(CurrentFrame, NextState),
+			*GetSafeStateName(NextTransition.NextActiveFrames.Last(), NextTransition.NextActiveFrames.Last().ActiveStates.Last()));
 		
 		return true;
 	}
@@ -1635,23 +1914,29 @@ bool FStateTreeExecutionContext::RequestTransition(const FStateTreeStateHandle N
 	return false;
 }
 
-void FStateTreeExecutionContext::SetupNextTransition(const FStateTreeStateHandle NextState, const EStateTreeTransitionPriority Priority)
+void FStateTreeExecutionContext::SetupNextTransition(const FStateTreeExecutionFrame& CurrentFrame, const FStateTreeStateHandle NextState, const EStateTreeTransitionPriority Priority)
 {
 	const FStateTreeExecutionState& Exec = GetExecState();
 
-	NextTransition.CurrentActiveStates = Exec.ActiveStates;
 	NextTransition.CurrentRunStatus = Exec.LastTickStatus;
 	NextTransition.SourceState = CurrentlyProcessedState;
+	NextTransition.SourceStateTree = CurrentFrame.StateTree;
+	NextTransition.SourceRootState = CurrentFrame.ActiveStates.GetStateSafe(0);
 	NextTransition.TargetState = NextState;
+	NextTransition.Priority = Priority;
+
+	FStateTreeExecutionFrame& NewFrame = NextTransition.NextActiveFrames.AddDefaulted_GetRef();
+	NewFrame.StateTree = CurrentFrame.StateTree;
+	NewFrame.RootState = CurrentFrame.RootState;
+
 	if (NextState == FStateTreeStateHandle::Invalid)
 	{
-		NextTransition.NextActiveStates.Reset();
+		NewFrame.ActiveStates = {};
 	}
 	else
 	{
-		NextTransition.NextActiveStates = FStateTreeActiveStates(NextState);
+		NewFrame.ActiveStates = FStateTreeActiveStates(NextState);
 	}
-	NextTransition.Priority = Priority;
 }
 
 bool FStateTreeExecutionContext::TriggerTransitions()
@@ -1674,228 +1959,207 @@ bool FStateTreeExecutionContext::TriggerTransitions()
 	//
 	for (const FStateTreeTransitionRequest& Request : InstanceData.GetTransitionRequests())
 	{
-		if (RequestTransition(Request.TargetState, Request.Priority))
+		// Find frame associated with the request.
+		const int32 FrameIndex = Exec.ActiveFrames.IndexOfByPredicate([&Request](const FStateTreeExecutionFrame& Frame)
 		{
-			NextTransitionSource = FStateTreeTransitionSource(EStateTreeTransitionSourceType::ExternalRequest, Request.TargetState, Request.Priority);
+			return Frame.StateTree == Request.SourceStateTree && Frame.RootState == Request.SourceRootState;
+		});
+
+		if (FrameIndex != INDEX_NONE)
+		{
+			const FStateTreeExecutionFrame& CurrentFrame = Exec.ActiveFrames[FrameIndex];
+			if (RequestTransition(CurrentFrame, Request.TargetState, Request.Priority))
+			{
+				NextTransitionSource = FStateTreeTransitionSource(EStateTreeTransitionSourceType::ExternalRequest, Request.TargetState, Request.Priority);
+			}
 		}
 	}
 	InstanceData.ResetTransitionRequests();
-	
+
 	//
 	// Check tick, event, and task based transitions first.
 	//
-	if (Exec.ActiveStates.Num() > 0)
+	if (Exec.ActiveFrames.Num() > 0)
 	{
-		// Setup data views for the tasks that will get called.
-		// It is possible that not all tasks views are set up at this stage (e.g. failed tick, pending transition handling).
-		check(Exec.FirstTaskStructIndex.IsValid() && Exec.FirstTaskObjectIndex.IsValid());
-		int32 InstanceStructIndex = Exec.FirstTaskStructIndex.Get();
-		int32 InstanceObjectIndex = Exec.FirstTaskObjectIndex.Get();
-
-		for (int32 Index = 0; Index < Exec.ActiveStates.Num(); Index++)
+		for (int32 FrameIndex = 0; FrameIndex < Exec.ActiveFrames.Num(); FrameIndex++)
 		{
-			const FStateTreeStateHandle CurrentHandle = Exec.ActiveStates[Index];
-			const FCompactStateTreeState& State = StateTree.States[CurrentHandle.Index];
+			FStateTreeExecutionFrame& CurrentFrame = Exec.ActiveFrames[FrameIndex];
+			const UStateTree* CurrentStateTree = CurrentFrame.StateTree;
+
+			FCurrentlyProcessedFrameScope FrameScope(*this, CurrentFrame, FrameIndex);
 
-			if (State.bHasTransitionTasks)
+			for (int32 StateIndex = CurrentFrame.ActiveStates.Num() - 1; StateIndex >= 0; StateIndex--)
 			{
-				// Update index to skip over the linked state params
-				if (State.Type == EStateTreeStateType::Linked)
-				{
-					InstanceStructIndex++;
-				}
+				const FStateTreeStateHandle StateHandle = CurrentFrame.ActiveStates[StateIndex];
+				const FCompactStateTreeState& State = CurrentStateTree->States[StateHandle.Index];
 
-				for (int32 TaskIndex = State.TasksBegin; TaskIndex < (State.TasksBegin + State.TasksNum); TaskIndex++)
+				// Do not process any transitions from a disabled state
+				if (!State.bEnabled)
 				{
-					const FStateTreeTaskBase& Task = StateTree.Nodes[TaskIndex].Get<const FStateTreeTaskBase>();
-					if (Task.bShouldAffectTransitions)
-					{
-						SetNodeDataView(Task, InstanceStructIndex, InstanceObjectIndex);
-					}
-					else
-					{
-						if (Task.bInstanceIsObject)
-						{
-							InstanceObjectIndex++;
-						}
-						else
-						{
-							InstanceStructIndex++;
-						}
-					}
+					continue;
 				}
-			}
-			else
-			{
-				// Skip over all instances in the state.
-				InstanceStructIndex += (int32)State.TaskInstanceStructNum;
-				InstanceObjectIndex += (int32)State.TaskInstanceObjectNum;
-			}
-		}
-		
-		for (int32 StateIndex = Exec.ActiveStates.Num() - 1; StateIndex >= 0; StateIndex--)
-		{
-			const FStateTreeStateHandle StateHandle = Exec.ActiveStates[StateIndex];
-			const FCompactStateTreeState& State = StateTree.States[StateHandle.Index];
-
-			// Do not process any transitions from a disabled state
-			if (!State.bEnabled)
-			{
-				continue;
-			}
-
-			FCurrentlyProcessedStateScope StateScope(*this, StateHandle);
-			STATETREE_TRACE_SCOPED_STATE(StateHandle);
-
-			if (State.bHasTransitionTasks)
-			{
-				STATETREE_CLOG(State.TasksNum > 0, VeryVerbose, TEXT("%*sTrigger task transitions in state '%s'"), StateIndex*UE::StateTree::DebugIndentSize, TEXT(""), *DebugGetStatePath(Exec.ActiveStates, StateIndex));
+				
+				FCurrentlyProcessedStateScope StateScope(*this, StateHandle);
+				STATETREE_TRACE_SCOPED_STATE(StateHandle);
 
-				for (int32 TaskIndex = (State.TasksBegin + State.TasksNum) - 1; TaskIndex >= State.TasksBegin; TaskIndex--)
+				if (State.bHasTransitionTasks)
 				{
-					const FStateTreeTaskBase& Task = StateTree.Nodes[TaskIndex].Get<const FStateTreeTaskBase>();
-
-					// Ignore disabled task
-					if (Task.bTaskEnabled == false)
-					{
-						STATETREE_LOG(VeryVerbose, TEXT("%*sSkipped 'TriggerTransitions' for disabled Task: '%s'"), UE::StateTree::DebugIndentSize, TEXT(""), *Task.Name.ToString());
-						continue;
-					}
+					STATETREE_CLOG(State.TasksNum > 0, VeryVerbose, TEXT("%*sTrigger task transitions in state '%s'"), StateIndex*UE::StateTree::DebugIndentSize, TEXT(""), *DebugGetStatePath(Exec.ActiveFrames, &CurrentFrame, StateIndex));
 
-					if (Task.bShouldAffectTransitions)
+					for (int32 TaskIndex = (State.TasksBegin + State.TasksNum) - 1; TaskIndex >= State.TasksBegin; TaskIndex--)
 					{
-						STATETREE_LOG(VeryVerbose, TEXT("%*sTriggerTransitions: '%s'"), UE::StateTree::DebugIndentSize, TEXT(""), *Task.Name.ToString());
-						STATETREE_TRACE_TASK_EVENT(TaskIndex, DataViews[Task.DataViewIndex.Get()], EStateTreeTraceEventType::OnEvaluating, EStateTreeRunStatus::Running);
-						check(DataViews[Task.DataViewIndex.Get()].IsValid());
-						Task.TriggerTransitions(*this);
-					}
-				}
-			}
-			
-			
-			for (uint8 i = 0; i < State.TransitionsNum; i++)
-			{
-				// All transition conditions must pass
-				const int16 TransitionIndex = State.TransitionsBegin + i;
-				const FCompactStateTransition& Transition = StateTree.Transitions[TransitionIndex];
+						const FStateTreeTaskBase& Task = CurrentStateTree->Nodes[TaskIndex].Get<const FStateTreeTaskBase>();
+						const FStateTreeDataView TaskDataView = GetDataView(CurrentFrame, Task.InstanceDataHandle);
 
-				// Skip disabled transitions
-				if (Transition.bTransitionEnabled == false)
-				{
-					continue;
-				}
-				
-				// No need to test the transition if same or higher priority transition has already been processed.
-				if (Transition.Priority <= NextTransition.Priority)
-				{
-					continue;
-				}
+						// Ignore disabled task
+						if (Task.bTaskEnabled == false)
+						{
+							STATETREE_LOG(VeryVerbose, TEXT("%*sSkipped 'TriggerTransitions' for disabled Task: '%s'"), UE::StateTree::DebugIndentSize, TEXT(""), *Task.Name.ToString());
+							continue;
+						}
 
-				// Skip completion transitions
-				if (EnumHasAnyFlags(Transition.Trigger, EStateTreeTransitionTrigger::OnStateCompleted))
-				{
-					continue;
+						if (Task.bShouldAffectTransitions)
+						{
+							STATETREE_LOG(VeryVerbose, TEXT("%*sTriggerTransitions: '%s'"), UE::StateTree::DebugIndentSize, TEXT(""), *Task.Name.ToString());
+							STATETREE_TRACE_TASK_EVENT(TaskIndex, TaskDataView, EStateTreeTraceEventType::OnEvaluating, EStateTreeRunStatus::Running);
+							check(TaskDataView.IsValid());
+							Task.TriggerTransitions(*this);
+						}
+					}
 				}
-
-				// If a delayed transition has passed the delay, and remove it from the queue, and try trigger it.
-				FStateTreeTransitionDelayedState* DelayedState = nullptr;
-				if (Transition.HasDelay())
+				
+				for (uint8 i = 0; i < State.TransitionsNum; i++)
 				{
-					DelayedState = Exec.FindDelayedTransition(FStateTreeIndex16(TransitionIndex));
-					if (DelayedState != nullptr && DelayedState->TimeLeft <= 0.0f)
+					// All transition conditions must pass
+					const int16 TransitionIndex = State.TransitionsBegin + i;
+					const FCompactStateTransition& Transition = CurrentStateTree->Transitions[TransitionIndex];
+
+					// Skip disabled transitions
+					if (Transition.bTransitionEnabled == false)
+					{
+						continue;
+					}
+					
+					// No need to test the transition if same or higher priority transition has already been processed.
+					if (Transition.Priority <= NextTransition.Priority)
 					{
-						STATETREE_LOG(Verbose, TEXT("Passed delayed transition from '%s' (%s) -> '%s'"),
-							*GetSafeStateName(Exec.ActiveStates.Last()), *State.Name.ToString(), *GetSafeStateName(Transition.State));
+						continue;
+					}
 
-						Exec.DelayedTransitions.RemoveAllSwap([TransitionIndex](const FStateTreeTransitionDelayedState& DelayedState)
-							{
-								return DelayedState.TransitionIndex.Get() == TransitionIndex;
-							});
+					// Skip completion transitions
+					if (EnumHasAnyFlags(Transition.Trigger, EStateTreeTransitionTrigger::OnStateCompleted))
+					{
+						continue;
+					}
 
-						// Trigger Delayed Transition when the delay has passed.
-						if (RequestTransition(Transition.State, Transition.Priority, Transition.Fallback))
+					// If a delayed transition has passed the delay, and remove it from the queue, and try trigger it.
+					FStateTreeTransitionDelayedState* DelayedState = nullptr;
+					if (Transition.HasDelay())
+					{
+						DelayedState = Exec.FindDelayedTransition(CurrentFrame.StateTree, FStateTreeIndex16(TransitionIndex));
+						if (DelayedState != nullptr && DelayedState->TimeLeft <= 0.0f)
 						{
-							NextTransitionSource = FStateTreeTransitionSource(FStateTreeIndex16(TransitionIndex), Transition.State, Transition.Priority);
+							STATETREE_LOG(Verbose, TEXT("Passed delayed transition from '%s' (%s) -> '%s'"),
+								*GetSafeStateName(CurrentFrame, CurrentFrame.ActiveStates.Last()), *State.Name.ToString(), *GetSafeStateName(CurrentFrame, Transition.State));
+
+							Exec.DelayedTransitions.RemoveAllSwap([StateTree = CurrentFrame.StateTree, TransitionIndex](const FStateTreeTransitionDelayedState& DelayedState)
+								{
+									return DelayedState.StateTree == StateTree && DelayedState.TransitionIndex.Get() == TransitionIndex;
+								});
+
+							// Trigger Delayed Transition when the delay has passed.
+							if (RequestTransition(CurrentFrame, Transition.State, Transition.Priority, Transition.Fallback))
+							{
+								NextTransitionSource = FStateTreeTransitionSource(FStateTreeIndex16(TransitionIndex), Transition.State, Transition.Priority);
+							}
+							continue;
 						}
-						continue;
 					}
-				}
 
-				const bool bShouldTrigger = Transition.Trigger == EStateTreeTransitionTrigger::OnTick
-											|| (Transition.Trigger == EStateTreeTransitionTrigger::OnEvent
-												&& HasEventToProcess(Transition.EventTag));
+					const bool bShouldTrigger = Transition.Trigger == EStateTreeTransitionTrigger::OnTick
+												|| (Transition.Trigger == EStateTreeTransitionTrigger::OnEvent
+													&& HasEventToProcess(Transition.EventTag));
 
-				bool bPassed = false; 
-				if (bShouldTrigger)
-				{
-					STATETREE_TRACE_TRANSITION_EVENT(FStateTreeTransitionSource(FStateTreeIndex16(TransitionIndex), Transition.State, Transition.Priority), EStateTreeTraceEventType::OnEvaluating);
-					STATETREE_TRACE_SCOPED_PHASE(EStateTreeUpdatePhase::TransitionConditions);
-					bPassed = TestAllConditions(Transition.ConditionsBegin, Transition.ConditionsNum);
-				}
+					bool bPassed = false; 
+					if (bShouldTrigger)
+					{
+						STATETREE_TRACE_TRANSITION_EVENT(FStateTreeTransitionSource(FStateTreeIndex16(TransitionIndex), Transition.State, Transition.Priority), EStateTreeTraceEventType::OnEvaluating);
+						STATETREE_TRACE_SCOPED_PHASE(EStateTreeUpdatePhase::TransitionConditions);
+						bPassed = TestAllConditions(CurrentFrame, Transition.ConditionsBegin, Transition.ConditionsNum);
+					}
 
-				if (bPassed)
-				{
-					// If the transitions is delayed, set up the delay. 
-					if (Transition.HasDelay())
+					if (bPassed)
 					{
-						if (DelayedState == nullptr)
+						// If the transitions is delayed, set up the delay. 
+						if (Transition.HasDelay())
 						{
-							// Initialize new delayed transition.
-							const float DelayDuration = Transition.Delay.GetRandomDuration();
-							if (DelayDuration > 0.0f)
+							if (DelayedState == nullptr)
+							{
+								// Initialize new delayed transition.
+								const float DelayDuration = Transition.Delay.GetRandomDuration();
+								if (DelayDuration > 0.0f)
+								{
+									DelayedState = &Exec.DelayedTransitions.AddDefaulted_GetRef();
+									DelayedState->StateTree = CurrentFrame.StateTree;
+									DelayedState->TransitionIndex = FStateTreeIndex16(TransitionIndex);
+									DelayedState->TimeLeft = DelayDuration;
+									BeginDelayedTransition(*DelayedState);
+									STATETREE_LOG(Verbose, TEXT("Delayed transition triggered from '%s' (%s) -> '%s' %.1fs"),
+										*GetSafeStateName(CurrentFrame, CurrentFrame.ActiveStates.Last()), *State.Name.ToString(), *GetSafeStateName(CurrentFrame, Transition.State), DelayedState->TimeLeft);
+									
+									// Delay state added, skip requesting the transition.
+									continue;
+								}
+								// Fallthrough to request transition if duration was zero. 
+							}
+							else
 							{
-								DelayedState = &Exec.DelayedTransitions.AddDefaulted_GetRef();
-								DelayedState->TransitionIndex = FStateTreeIndex16(TransitionIndex);
-								DelayedState->TimeLeft = DelayDuration;
-								BeginDelayedTransition(*DelayedState);
-								STATETREE_LOG(Verbose, TEXT("Delayed transition triggered from '%s' (%s) -> '%s' %.1fs"),
-									*GetSafeStateName(Exec.ActiveStates.Last()), *State.Name.ToString(), *GetSafeStateName(Transition.State), DelayedState->TimeLeft);
-								
-								// Delay state added, skip requesting the transition.
+								// We get here if the transitions re-triggers during the delay, on which case we'll just ignore it.
 								continue;
 							}
-							// Fallthrough to request transition if duration was zero. 
 						}
-						else
+
+						if (RequestTransition(CurrentFrame, Transition.State, Transition.Priority, Transition.Fallback))
 						{
-							// We get here if the transitions re-triggers during the delay, on which case we'll just ignore it.
-							continue;
+							NextTransitionSource = FStateTreeTransitionSource(FStateTreeIndex16(TransitionIndex), Transition.State, Transition.Priority);
 						}
 					}
-
-					if (RequestTransition(Transition.State, Transition.Priority, Transition.Fallback))
-					{
-						NextTransitionSource = FStateTreeTransitionSource(FStateTreeIndex16(TransitionIndex), Transition.State, Transition.Priority);
-					}
 				}
 			}
-		}
-	}
-
-	if (StateTree.bHasGlobalTransitionTasks)
-	{
-		STATETREE_LOG(VeryVerbose, TEXT("Trigger global task transitions"));
-		for (int32 TaskIndex = (StateTree.GlobalTasksBegin + StateTree.GlobalTasksNum) - 1; TaskIndex >= StateTree.GlobalTasksBegin; TaskIndex--)
-		{
-			const FStateTreeTaskBase& Task =  StateTree.Nodes[TaskIndex].Get<const FStateTreeTaskBase>();
 
-			// Ignore disabled task
-			if (Task.bTaskEnabled == false)
+			if (CurrentFrame.bIsGlobalFrame)
 			{
-				STATETREE_LOG(VeryVerbose, TEXT("%*sSkipped 'TriggerTransitions' for disabled Task: '%s'"), UE::StateTree::DebugIndentSize, TEXT(""), *Task.Name.ToString());
-				continue;
-			}
+				// Global frame
+				if (CurrentFrame.StateTree->bHasGlobalTransitionTasks)
+				{
+					STATETREE_LOG(VeryVerbose, TEXT("Trigger global task transitions"));
+					for (int32 TaskIndex = (CurrentStateTree->GlobalTasksBegin + CurrentStateTree->GlobalTasksNum) - 1; TaskIndex >= CurrentFrame.StateTree->GlobalTasksBegin; TaskIndex--)
+					{
+						const FStateTreeTaskBase& Task =  CurrentStateTree->Nodes[TaskIndex].Get<const FStateTreeTaskBase>();
+						const FStateTreeDataView TaskDataView = GetDataView(CurrentFrame, Task.InstanceDataHandle);
 
-			if (Task.bShouldAffectTransitions)
-			{
-				STATETREE_LOG(VeryVerbose, TEXT("%*sTriggerTransitions: '%s'"), UE::StateTree::DebugIndentSize, TEXT(""), *Task.Name.ToString());
-				STATETREE_TRACE_TASK_EVENT(TaskIndex, DataViews[Task.DataViewIndex.Get()], EStateTreeTraceEventType::OnEvaluating, EStateTreeRunStatus::Running);
-				check(DataViews[Task.DataViewIndex.Get()].IsValid());
-				Task.TriggerTransitions(*this);
+						// Ignore disabled task
+						if (Task.bTaskEnabled == false)
+						{
+							STATETREE_LOG(VeryVerbose, TEXT("%*sSkipped 'TriggerTransitions' for disabled Task: '%s'"), UE::StateTree::DebugIndentSize, TEXT(""), *Task.Name.ToString());
+							continue;
+						}
+
+						if (Task.bShouldAffectTransitions)
+						{
+							STATETREE_LOG(VeryVerbose, TEXT("%*sTriggerTransitions: '%s'"), UE::StateTree::DebugIndentSize, TEXT(""), *Task.Name.ToString());
+							STATETREE_TRACE_TASK_EVENT(TaskIndex, TaskDataView, EStateTreeTraceEventType::OnEvaluating, EStateTreeRunStatus::Running);
+							check(TaskDataView.IsValid());
+							Task.TriggerTransitions(*this);
+						}
+					}
+				}
+				
 			}
 		}
 	}
 
+
 	//
 	// Check state completion transitions.
 	//
@@ -1904,77 +2168,94 @@ bool FStateTreeExecutionContext::TriggerTransitions()
 	if (NextTransition.Priority == EStateTreeTransitionPriority::None
 		&& Exec.LastTickStatus != EStateTreeRunStatus::Running)
 	{
-		// Start from the last completed state.
-		const int32 StateStartIndex = Exec.CompletedStateHandle.IsValid() ? Exec.ActiveStates.IndexOfReverse(Exec.CompletedStateHandle) : (Exec.ActiveStates.Num() - 1);
-		const EStateTreeTransitionTrigger CompletionTrigger = Exec.LastTickStatus == EStateTreeRunStatus::Succeeded ? EStateTreeTransitionTrigger::OnStateSucceeded : EStateTreeTransitionTrigger::OnStateFailed;
-
-		check(StateStartIndex >= 0 && StateStartIndex < Exec.ActiveStates.Num());
+		// Start from the last completed state if specified.
+		const int32 FrameStartIndex = Exec.CompletedFrameIndex.IsValid() ? Exec.CompletedFrameIndex.AsInt32() : (Exec.ActiveFrames.Num() - 1);
+		check(FrameStartIndex >= 0 && FrameStartIndex < Exec.ActiveFrames.Num());
 		
-		// Check completion transitions
-		for (int32 StateIndex = StateStartIndex; StateIndex >= 0; StateIndex--)
+		for (int32 FrameIndex = FrameStartIndex; FrameIndex >= 0; FrameIndex--)
 		{
-			const FStateTreeStateHandle StateHandle = Exec.ActiveStates[StateIndex];
-			const FCompactStateTreeState& State = StateTree.States[StateHandle.Index];
+			FStateTreeExecutionFrame& CurrentFrame = Exec.ActiveFrames[FrameIndex];
+			const UStateTree* CurrentStateTree = CurrentFrame.StateTree;
 
-			FCurrentlyProcessedStateScope StateScope(*this, StateHandle);
-			STATETREE_TRACE_SCOPED_STATE_PHASE(StateHandle, EStateTreeUpdatePhase::TriggerTransitions);
+			FCurrentlyProcessedFrameScope FrameScope(*this, CurrentFrame, FrameIndex);
 
-			for (uint8 i = 0; i < State.TransitionsNum; i++)
+			int32 StateStartIndex = CurrentFrame.ActiveStates.Num() - 1; // This is ok, even if the ActiveStates is 0, -1 will skip the whole state loop below.
+			if (FrameIndex == FrameStartIndex
+				&& Exec.CompletedStateHandle.IsValid())
+			{
+				StateStartIndex = CurrentFrame.ActiveStates.IndexOfReverse(Exec.CompletedStateHandle);
+				// INDEX_NONE (-1) will skip the whole state loop below. We still want to warn.
+				ensureMsgf(StateStartIndex != INDEX_NONE, TEXT("If CompletedFrameIndex and CompletedStateHandle are specified, we expect that the state is found"));
+			}
+			
+			const EStateTreeTransitionTrigger CompletionTrigger = Exec.LastTickStatus == EStateTreeRunStatus::Succeeded ? EStateTreeTransitionTrigger::OnStateSucceeded : EStateTreeTransitionTrigger::OnStateFailed;
+		
+			// Check completion transitions
+			for (int32 StateIndex = StateStartIndex; StateIndex >= 0; StateIndex--)
 			{
-				// All transition conditions must pass
-				const int16 TransitionIndex = State.TransitionsBegin + i;
-				const FCompactStateTransition& Transition = StateTree.Transitions[TransitionIndex];
+				const FStateTreeStateHandle StateHandle = CurrentFrame.ActiveStates[StateIndex];
+				const FCompactStateTreeState& State = CurrentStateTree->States[StateHandle.Index];
 
-				// Skip disabled transitions
-				if (Transition.bTransitionEnabled == false)
-				{
-					continue;
-				}
+				FCurrentlyProcessedStateScope StateScope(*this, StateHandle);
+				STATETREE_TRACE_SCOPED_STATE_PHASE(StateHandle, EStateTreeUpdatePhase::TriggerTransitions);
 
-				if (EnumHasAnyFlags(Transition.Trigger, CompletionTrigger))
+				for (uint8 i = 0; i < State.TransitionsNum; i++)
 				{
-					bool bPassed = false;
+					// All transition conditions must pass
+					const int16 TransitionIndex = State.TransitionsBegin + i;
+					const FCompactStateTransition& Transition = CurrentStateTree->Transitions[TransitionIndex];
+
+					// Skip disabled transitions
+					if (Transition.bTransitionEnabled == false)
 					{
-						STATETREE_TRACE_TRANSITION_EVENT(FStateTreeTransitionSource(FStateTreeIndex16(TransitionIndex), Transition.State, Transition.Priority), EStateTreeTraceEventType::OnEvaluating);
-						STATETREE_TRACE_SCOPED_PHASE(EStateTreeUpdatePhase::TransitionConditions);
-						bPassed = TestAllConditions(Transition.ConditionsBegin, Transition.ConditionsNum);
+						continue;
 					}
 
-					if (bPassed)
+					if (EnumHasAnyFlags(Transition.Trigger, CompletionTrigger))
 					{
-						// No delay allowed on completion conditions.
-						// No priority on completion transitions, use the priority to signal that state is selected.
-						if (RequestTransition(Transition.State, EStateTreeTransitionPriority::Normal, Transition.Fallback))
+						bool bPassed = false;
 						{
-							NextTransitionSource = FStateTreeTransitionSource(FStateTreeIndex16(TransitionIndex), Transition.State, Transition.Priority);
-							break;
+							STATETREE_TRACE_TRANSITION_EVENT(FStateTreeTransitionSource(FStateTreeIndex16(TransitionIndex), Transition.State, Transition.Priority), EStateTreeTraceEventType::OnEvaluating);
+							STATETREE_TRACE_SCOPED_PHASE(EStateTreeUpdatePhase::TransitionConditions);
+							bPassed = TestAllConditions(CurrentFrame, Transition.ConditionsBegin, Transition.ConditionsNum);
+						}
+
+						if (bPassed)
+						{
+							// No delay allowed on completion conditions.
+							// No priority on completion transitions, use the priority to signal that state is selected.
+							if (RequestTransition(CurrentFrame, Transition.State, EStateTreeTransitionPriority::Normal, Transition.Fallback))
+							{
+								NextTransitionSource = FStateTreeTransitionSource(FStateTreeIndex16(TransitionIndex), Transition.State, Transition.Priority);
+								break;
+							}
 						}
 					}
 				}
-			}
 
-			if (NextTransition.Priority != EStateTreeTransitionPriority::None)
-			{
-				break;
+				if (NextTransition.Priority != EStateTreeTransitionPriority::None)
+				{
+					break;
+				}
 			}
-		}
 
-		if (NextTransition.Priority == EStateTreeTransitionPriority::None)
-		{
-			STATETREE_LOG_AND_TRACE(Verbose, TEXT("Could not trigger completion transition, jump back to root state."));
-			FCurrentlyProcessedStateScope StateScope(*this, FStateTreeStateHandle::Root);
-			if (RequestTransition(FStateTreeStateHandle::Root, EStateTreeTransitionPriority::Normal))
-			{
-				NextTransitionSource = FStateTreeTransitionSource(EStateTreeTransitionSourceType::Internal, FStateTreeStateHandle::Root, EStateTreeTransitionPriority::Normal);
-			}
-			else
+			if (NextTransition.Priority == EStateTreeTransitionPriority::None)
 			{
-				STATETREE_LOG_AND_TRACE(Warning, TEXT("Failed to select root state. Stopping the tree with failure."));
+				STATETREE_LOG_AND_TRACE(Verbose, TEXT("Could not trigger completion transition, jump back to root state."));
+				FCurrentlyProcessedStateScope StateScope(*this, FStateTreeStateHandle::Root);
+				if (RequestTransition(CurrentFrame, FStateTreeStateHandle::Root, EStateTreeTransitionPriority::Normal))
+				{
+					NextTransitionSource = FStateTreeTransitionSource(EStateTreeTransitionSourceType::Internal, FStateTreeStateHandle::Root, EStateTreeTransitionPriority::Normal);
+				}
+				else
+				{
+					STATETREE_LOG_AND_TRACE(Warning, TEXT("Failed to select root state. Stopping the tree with failure."));
 
-				SetupNextTransition(FStateTreeStateHandle::Failed, EStateTreeTransitionPriority::Critical);
+					SetupNextTransition(CurrentFrame, FStateTreeStateHandle::Failed, EStateTreeTransitionPriority::Critical);
 
-				// In this case we don't want to complete subtrees, we want to force the whole tree to stop.
-				bProcessSubTreeCompletion = false;
+					// In this case we don't want to complete subtrees, we want to force the whole tree to stop.
+					bProcessSubTreeCompletion = false;
+				}
 			}
 		}
 	}
@@ -1982,163 +2263,174 @@ bool FStateTreeExecutionContext::TriggerTransitions()
 	// Check if the transition was succeed/failed, if we're on a sub-tree, complete the subtree instead of transition.
 	if (NextTransition.TargetState.IsCompletionState() && bProcessSubTreeCompletion)
 	{
-		const FStateTreeStateHandle ParentLinkedState = GetParentLinkedStateHandle(Exec.ActiveStates, NextTransition.SourceState);
-		if (ParentLinkedState.IsValid())
+		const int32 SourceFrameIndex = Exec.ActiveFrames.IndexOfByPredicate([&NextTransition = NextTransition](const FStateTreeExecutionFrame& Frame)
 		{
-			const EStateTreeRunStatus RunStatus = NextTransition.TargetState.ToCompletionStatus(); 
-			STATETREE_LOG(Verbose, TEXT("Completed subtree '%s' from state '%s' (%s): %s"),
-				*GetSafeStateName(ParentLinkedState), *GetSafeStateName(Exec.ActiveStates.Last()), *GetSafeStateName(NextTransition.SourceState), *UEnum::GetDisplayValueAsText(RunStatus).ToString());
-
-			// Set the parent linked state as last completed state, and update tick status to the status from the transition.
-			Exec.CompletedStateHandle = ParentLinkedState;
-			Exec.LastTickStatus = RunStatus;
+			return Frame.StateTree == NextTransition.SourceStateTree && Frame.RootState == NextTransition.SourceRootState;
+		});
+		// Check that the transition source frame is a sub-tree, the first frame (0 index) is not a subtree. 
+		if (SourceFrameIndex > 0)
+		{
+			const FStateTreeExecutionFrame& SourceFrame = Exec.ActiveFrames[SourceFrameIndex];
+			const int32 ParentFrameIndex = SourceFrameIndex - 1;
+			const FStateTreeExecutionFrame& ParentFrame = Exec.ActiveFrames[ParentFrameIndex];
+			const FStateTreeStateHandle ParentLinkedState = ParentFrame.ActiveStates.Last();
 
-			// Clear the transition and return that no transition took place.
-			// Since the LastTickStatus != running, the transition loop will try another transition
-			// now starting from the linked parent state. If we run out of retires in the selection loop (e.g. very deep hierarchy)
-			// we will continue on next tick.
-			NextTransition.Reset();
-			return false;
+			if (ParentLinkedState.IsValid())
+			{
+				const EStateTreeRunStatus RunStatus = NextTransition.TargetState.ToCompletionStatus(); 
+				STATETREE_LOG(Verbose, TEXT("Completed subtree '%s' from state '%s': %s"),
+					*GetSafeStateName(ParentFrame, ParentLinkedState), *GetSafeStateName(SourceFrame, NextTransition.SourceState), *UEnum::GetDisplayValueAsText(RunStatus).ToString());
+
+				// Set the parent linked state as last completed state, and update tick status to the status from the transition.
+				Exec.CompletedFrameIndex = FStateTreeIndex16(ParentFrameIndex);
+				Exec.CompletedStateHandle = ParentLinkedState;
+				Exec.LastTickStatus = RunStatus;
+
+				// Clear the transition and return that no transition took place.
+				// Since the LastTickStatus != running, the transition loop will try another transition
+				// now starting from the linked parent state. If we run out of retires in the selection loop (e.g. very deep hierarchy)
+				// we will continue on next tick.
+				NextTransition.Reset();
+				return false;
+			}
 		}
 	}
 
 	return NextTransition.TargetState.IsValid();
 }
 
-FStateTreeStateHandle FStateTreeExecutionContext::GetParentLinkedStateHandle(const FStateTreeActiveStates& ActiveStates, const int32 StartStartIndex) const
-{
-	check(ActiveStates.IsValidIndex(StartStartIndex));
-	for (int32 StateIndex = StartStartIndex; StateIndex >= 0; StateIndex--)
-	{
-		const FCompactStateTreeState& State = StateTree.States[ActiveStates[StateIndex].Index];
-		if (State.LinkedState.IsValid())
-		{
-			return ActiveStates[StateIndex];
-		}
-	}
-
-	return FStateTreeStateHandle();
-}	
-
-FStateTreeStateHandle FStateTreeExecutionContext::GetParentLinkedStateHandle(const FStateTreeActiveStates& ActiveStates, const FStateTreeStateHandle StartStateHandle) const
+bool FStateTreeExecutionContext::SelectState(const FStateTreeExecutionFrame& CurrentFrame,
+											const FStateTreeStateHandle NextState,
+											TArray<FStateTreeExecutionFrame, TFixedAllocator<MaxExecutionFrames>>& OutNextActiveFrames,
+											const EStateTreeSelectionFallback Fallback)
 {
-	// Find start state
-	int32 StateIndex = ActiveStates.Num() - 1;
-	while (StateIndex >= 0)
-	{
-		if (ActiveStates[StateIndex] == StartStateHandle)
-		{
-			break;
-		}
-		StateIndex--;
-	}
-
-	// The function result is used to iteratively traverse to the root-most parent linked state.
-	// Skip the start state, as we want to always find a parent state to the start state, or else the iteration will hit infinite loop. 
-	StateIndex--;
+	const FStateTreeExecutionState& Exec = GetExecState();
 
-	// Find parent linked state.
-	while (StateIndex >= 0)
+	if (Exec.ActiveFrames.IsEmpty())
 	{
-		const FCompactStateTreeState& State = StateTree.States[ActiveStates[StateIndex].Index];
-		if (State.LinkedState.IsValid())
-		{
-			return ActiveStates[StateIndex];
-		}
-		
-		StateIndex--;
+		STATETREE_LOG(Error, TEXT("%hs: SelectState can only be called on initialized tree.  '%s' using StateTree '%s'."),
+			__FUNCTION__, *GetNameSafe(&Owner), *GetFullNameSafe(&RootStateTree));
+		return false;
 	}
-
-	return FStateTreeStateHandle();
-}
-
-bool FStateTreeExecutionContext::SelectState(const FStateTreeStateHandle NextState, FStateTreeActiveStates& OutNewActiveState, FStateTreeActiveStates& VisitedStates, const EStateTreeSelectionFallback Fallback)
-{
-	const FStateTreeExecutionState& Exec = GetExecState();
-
+	
 	if (!NextState.IsValid())
 	{
 		return false;
 	}
 
-	// Find common ancestor of `NextState` in the current active states and connect.
-	// This allows transitions within a subtree.
-	OutNewActiveState = Exec.ActiveStates;
-	
-	TStaticArray<FStateTreeStateHandle, FStateTreeActiveStates::MaxStates> InBetweenStates;
-	int32 NumInBetweenStates = 0;
-	int32 CommonActiveAncestorIndex = INDEX_NONE;
-
 	// Walk towards the root from current state.
+	TStaticArray<FStateTreeStateHandle, FStateTreeActiveStates::MaxStates> ParentStates;
+	int32 NumParentStates = 0;
 	FStateTreeStateHandle CurrState = NextState;
 	while (CurrState.IsValid())
 	{
-		// Store the states that are in between the 'NextState' and common ancestor. 
-		InBetweenStates[NumInBetweenStates++] = CurrState;
-		// Check if the state can be found in the active states.
-		CommonActiveAncestorIndex = OutNewActiveState.IndexOfReverse(CurrState); 
-		if (CommonActiveAncestorIndex != INDEX_NONE)
+		if (NumParentStates == FStateTreeActiveStates::MaxStates)
 		{
-			break;
-		}
-		if (NumInBetweenStates == InBetweenStates.Num())
-		{
-			STATETREE_LOG(Error, TEXT("%hs: Too many parent states when selecting state '%s' from '%s'.  '%s' using StateTree '%s'."),
-				__FUNCTION__, *GetSafeStateName(NextState), *GetStateStatusString(Exec), *GetNameSafe(&Owner), *GetFullNameSafe(&StateTree));
+			STATETREE_LOG(Error, TEXT("%hs: Reached max execution depth when trying to select state %s from '%s'.  '%s' using StateTree '%s'."),
+				__FUNCTION__, *GetSafeStateName(CurrentFrame, NextState), *GetStateStatusString(Exec), *GetNameSafe(&Owner), *GetFullNameSafe(&RootStateTree));
 			return false;
 		}
+		// Store the states that are in between the 'NextState' and common ancestor. 
+		ParentStates[NumParentStates++] = CurrState;
+		CurrState = CurrentFrame.StateTree->States[CurrState.Index].Parent;
+	}
+
+	const UStateTree* NextStateTree = CurrentFrame.StateTree;
+	const FStateTreeStateHandle NextRootState = ParentStates[NumParentStates - 1]; 
+
+	// Find the frame that the next state belongs to.
+	int32 CurrentFrameIndex = INDEX_NONE;
+	int32 CurrentStateTreeIndex = INDEX_NONE;
 
-		CurrState = StateTree.States[CurrState.Index].Parent;
+	for (int32 FrameIndex = Exec.ActiveFrames.Num() - 1; FrameIndex >= 0; FrameIndex--)
+	{
+		const FStateTreeExecutionFrame& Frame = Exec.ActiveFrames[FrameIndex]; 
+		if (Frame.StateTree == NextStateTree)
+		{
+			CurrentStateTreeIndex = FrameIndex;
+			if (Frame.RootState == NextRootState)
+			{
+				CurrentFrameIndex = FrameIndex;
+				break;
+			}
+		}
 	}
 
-	// Max takes care of INDEX_NONE, by setting the num to 0.
-	OutNewActiveState.SetNum(FMath::Max(0, CommonActiveAncestorIndex));
+	// Copy common frames over.
+	if (CurrentFrameIndex != INDEX_NONE)
+	{
+		const int32 NumCommonFrames = CurrentFrameIndex + 1;
+		OutNextActiveFrames = MakeArrayView(Exec.ActiveFrames.GetData(), NumCommonFrames);
+	}
+	else if (CurrentStateTreeIndex != INDEX_NONE)
+	{
+		// If we could not find a common frame, we assume that we jumped to different subtree in same asset.
+		const int32 NumCommonFrames = CurrentStateTreeIndex + 1;
+		OutNextActiveFrames = MakeArrayView(Exec.ActiveFrames.GetData(), NumCommonFrames);
+	}
+	else
+	{
+		STATETREE_LOG(Error, TEXT("%hs: Encountered unrecognized state %s during state selection from '%s'.  '%s' using StateTree '%s'."),
+			__FUNCTION__, *GetNameSafe(NextStateTree), *GetStateStatusString(Exec), *GetNameSafe(&Owner), *GetFullNameSafe(&RootStateTree));
+		return false;
+	}
 	
 	// Append in between state in reverse order, they were collected from leaf towards the root.
-	bool bActiveStatesOverflow = false;
-	for (int32 Index = NumInBetweenStates - 1; Index > 0; Index--)
+	// Note: NextState will be added by SelectStateInternal() if conditions pass.
+	const int32 LastFrameIndex = OutNextActiveFrames.Num() - 1;
+	FStateTreeExecutionFrame& LastFrame = OutNextActiveFrames[LastFrameIndex];
+
+	LastFrame.ActiveStates.Reset();
+	for (int32 Index = NumParentStates - 1; Index > 0; Index--)
 	{
-		bActiveStatesOverflow |= !OutNewActiveState.Push(InBetweenStates[Index]);
+		LastFrame.ActiveStates.Push(ParentStates[Index]);
 	}
 
-	if (bActiveStatesOverflow)
+	TArray<FStateTreeExecutionFrame, TFixedAllocator<MaxExecutionFrames>> InitialNextActiveFrames;
+	
+	if (Fallback == EStateTreeSelectionFallback::NextSelectableSibling)
 	{
-		STATETREE_LOG(Error, TEXT("%hs: Reached max execution depth when trying to select state %s from '%s'.  '%s' using StateTree '%s'."),
-			__FUNCTION__, *GetSafeStateName(NextState), *GetStateStatusString(Exec), *GetNameSafe(&Owner), *GetFullNameSafe(&StateTree));
-		return false;
+		InitialNextActiveFrames = OutNextActiveFrames;
 	}
-
-	if (SelectStateInternal(NextState, OutNewActiveState, VisitedStates))
+	
+	// We take copy of the last frame and assign it later, as SelectStateInternal() might change the array and invalidate the pointer.
+	if (SelectStateInternal(OutNextActiveFrames[LastFrameIndex], NextState, OutNextActiveFrames))
 	{
 		return true;
 	}
 
 	// Failed to Select Next State, handle fallback here
 	// Return true on the first next sibling that gets selected successfully
-	if (Fallback == EStateTreeSelectionFallback::NextSelectableSibling && NumInBetweenStates >= 2)
+	if (Fallback == EStateTreeSelectionFallback::NextSelectableSibling && NumParentStates >= 2)
 	{
 		// InBetweenStates is in reversed order (i.e. from leaf to root)
-		FStateTreeStateHandle Parent = InBetweenStates[1];
+		const FStateTreeStateHandle Parent = ParentStates[1];
 		if (Parent.IsValid())
 		{
-			const FCompactStateTreeState& ParentState = StateTree.States[Parent.Index];
+			const FCompactStateTreeState& ParentState = CurrentFrame.StateTree->States[Parent.Index];
 
-			uint16 ChildState = StateTree.States[NextState.Index].GetNextSibling();
-			for (; ChildState < ParentState.ChildrenEnd; ChildState = StateTree.States[ChildState].GetNextSibling())
+			uint16 ChildState = CurrentFrame.StateTree->States[NextState.Index].GetNextSibling();
+			for (; ChildState < ParentState.ChildrenEnd; ChildState = CurrentFrame.StateTree->States[ChildState].GetNextSibling())
 			{
 				FStateTreeStateHandle ChildStateHandle = FStateTreeStateHandle(ChildState);
-				if (SelectStateInternal(ChildStateHandle, OutNewActiveState, VisitedStates))
+
+				// Start selection from blank slate.
+				OutNextActiveFrames = InitialNextActiveFrames;
+	
+				// We take copy of the last frame and assign it later, as SelectStateInternal() might change the array and invalidate the pointer.
+				if (SelectStateInternal(OutNextActiveFrames[LastFrameIndex], ChildStateHandle, OutNextActiveFrames))
 				{
 					return true;
 				}
 			}
 		}
 	}
-
+	
 	return false;
 }
 
-bool FStateTreeExecutionContext::SelectStateInternal(const FStateTreeStateHandle NextState, FStateTreeActiveStates& OutNewActiveState, FStateTreeActiveStates& VisitedStates)
+bool FStateTreeExecutionContext::SelectStateInternal(FStateTreeExecutionFrame& CurrentFrame, const FStateTreeStateHandle NextState,
+														TArray<FStateTreeExecutionFrame, TFixedAllocator<MaxExecutionFrames>>& OutNextActiveFrames)
 {
 	CSV_SCOPED_TIMING_STAT_EXCLUSIVE(StateTree_SelectState);
 
@@ -2148,68 +2440,107 @@ bool FStateTreeExecutionContext::SelectStateInternal(const FStateTreeStateHandle
 	{
 		// Trying to select non-existing state.
 		STATETREE_LOG(Error, TEXT("%hs: Trying to select invalid state from '%s'.  '%s' using StateTree '%s'."),
-            __FUNCTION__, *GetStateStatusString(Exec), *GetNameSafe(&Owner), *GetFullNameSafe(&StateTree));
+            __FUNCTION__, *GetStateStatusString(Exec), *GetNameSafe(&Owner), *GetFullNameSafe(CurrentFrame.StateTree));
 		return false;
 	}
 
-	const FCompactStateTreeState& State = StateTree.States[NextState.Index];
+	FCurrentlyProcessedFrameScope FrameScope(*this, CurrentFrame, INDEX_NONE);
+
+	const UStateTree* CurrentStateTree = CurrentFrame.StateTree;
+	const FCompactStateTreeState& CurrentState = CurrentStateTree->States[NextState.Index];
 
-	if (State.bEnabled == false)
+	if (CurrentState.bEnabled == false)
 	{
 		// Do not select disabled state
 		STATETREE_LOG(VeryVerbose, TEXT("%hs: Ignoring disabled state '%s'.  '%s' using StateTree '%s'."),
-			__FUNCTION__, *GetSafeStateName(NextState), *GetNameSafe(&Owner), *GetFullNameSafe(&StateTree));
+			__FUNCTION__, *GetSafeStateName(CurrentFrame, NextState), *GetNameSafe(&Owner), *GetFullNameSafe(CurrentFrame.StateTree));
 		return false;
 	}
 
 	STATETREE_TRACE_SCOPED_STATE_PHASE(NextState, EStateTreeUpdatePhase::StateSelection);
-	
-	// Check that the state can be entered
-	bool bEnterConditionsPassed = false;
-	if (State.SelectionBehavior != EStateTreeStateSelectionBehavior::None)
+
+	// The state cannot be directly selected.
+	if (CurrentState.SelectionBehavior == EStateTreeStateSelectionBehavior::None)
 	{
-		STATETREE_TRACE_SCOPED_PHASE(EStateTreeUpdatePhase::EnterConditions);
-		bEnterConditionsPassed = TestAllConditions(State.EnterConditionsBegin, State.EnterConditionsNum);
+		return false;
 	}
 
+	// Check that the state can be entered
+	STATETREE_TRACE_SCOPED_PHASE(EStateTreeUpdatePhase::EnterConditions);
+	const bool bEnterConditionsPassed = TestAllConditions(CurrentFrame, CurrentState.EnterConditionsBegin, CurrentState.EnterConditionsNum);
+
 	if (bEnterConditionsPassed)
 	{
-		if (!OutNewActiveState.Push(NextState))
+		if (!CurrentFrame.ActiveStates.Push(NextState))
 		{
 			STATETREE_LOG(Error, TEXT("%hs: Reached max execution depth when trying to select state %s from '%s'.  '%s' using StateTree '%s'."),
-				__FUNCTION__, *GetSafeStateName(NextState), *GetStateStatusString(Exec), *GetNameSafe(&Owner), *GetFullNameSafe(&StateTree));
-			return false;
-		}
-		if (!VisitedStates.Push(NextState))
-		{
-			STATETREE_LOG(Error, TEXT("%hs: Reached max visited state depth when trying to select state %s from '%s'.  '%s' using StateTree '%s'."),
-				__FUNCTION__, *GetSafeStateName(NextState), *GetStateStatusString(Exec), *GetNameSafe(&Owner), *GetFullNameSafe(&StateTree));
+				__FUNCTION__, *GetSafeStateName(CurrentFrame, NextState), *GetStateStatusString(Exec), *GetNameSafe(&Owner), *GetFullNameSafe(CurrentFrame.StateTree));
 			return false;
 		}
 		
-		if (State.LinkedState.IsValid())
+		if (CurrentState.Type == EStateTreeStateType::Linked && CurrentState.LinkedState.IsValid())
 		{
+			if (OutNextActiveFrames.Num() == MaxExecutionFrames)
+			{
+				STATETREE_LOG(Error, TEXT("%hs: Reached max execution depth when trying to select state %s from '%s'.  '%s' using StateTree '%s'."),
+					__FUNCTION__, *GetSafeStateName(CurrentFrame, NextState), *GetStateStatusString(Exec), *GetNameSafe(&Owner), *GetFullNameSafe(CurrentFrame.StateTree));
+				return false;
+			}
+
+			FStateTreeExecutionFrame NewFrame;
+			NewFrame.StateTree = CurrentFrame.StateTree;
+			NewFrame.RootState = CurrentState.LinkedState;
+			NewFrame.GlobalInstanceIndexBase = CurrentFrame.GlobalInstanceIndexBase;
+
+			// If the Frame already exists, copy instance indices so that conditions that rely on active states work correctly.
+			const FStateTreeExecutionFrame* ExistingFrame = Exec.ActiveFrames.FindByPredicate(
+				[StateTree = CurrentFrame.StateTree, RootState = CurrentState.LinkedState](const FStateTreeExecutionFrame& Frame)
+				{
+					return Frame.StateTree == StateTree && Frame.RootState == RootState;
+				});
+			if (ExistingFrame)
+			{
+				NewFrame.ActiveInstanceIndexBase = ExistingFrame->ActiveInstanceIndexBase;
+				NewFrame.StateParameterDataIndex = ExistingFrame->StateParameterDataIndex;
+				// No need to update the global instance index base since we're on the same asset.
+			}
+
+			// Check and prevent recursion.
+			const bool bNewFrameAlreadySelected = OutNextActiveFrames.ContainsByPredicate([&NewFrame](const FStateTreeExecutionFrame& Frame) {
+				return Frame.IsSameFrame(NewFrame);
+			});
+			
+			if (bNewFrameAlreadySelected)
+			{
+				STATETREE_LOG(Error, TEXT("%hs: Trying to recursively enter subtree '%s' from '%s'.  '%s' using StateTree '%s'."),
+					__FUNCTION__, *GetSafeStateName(NewFrame, NewFrame.RootState), *GetStateStatusString(Exec), *GetNameSafe(&Owner), *GetFullNameSafe(CurrentFrame.StateTree));
+				return false;
+			}
+			
+			OutNextActiveFrames.Push(NewFrame);
+
 			// If State is linked, proceed to the linked state.
-			if (SelectStateInternal(State.LinkedState, OutNewActiveState, VisitedStates))
+			if (SelectStateInternal(OutNextActiveFrames.Last(), CurrentState.LinkedState, OutNextActiveFrames))
 			{
-				// Selection succeeded
 				return true;
 			}
+			
+			OutNextActiveFrames.Pop();
 		}
-		else if (State.SelectionBehavior == EStateTreeStateSelectionBehavior::TryEnterState)
+		else if (CurrentState.SelectionBehavior == EStateTreeStateSelectionBehavior::TryEnterState)
 		{
 			// Select this state.
 			STATETREE_TRACE_STATE_EVENT(NextState, EStateTreeTraceEventType::OnStateSelected);
 			return true;
 		}
-		else if (State.SelectionBehavior == EStateTreeStateSelectionBehavior::TryFollowTransitions)
+		else if (CurrentState.SelectionBehavior == EStateTreeStateSelectionBehavior::TryFollowTransitions)
 		{
 			EStateTreeTransitionPriority CurrentPriority = EStateTreeTransitionPriority::None;
-			
-			for (uint8 i = 0; i < State.TransitionsNum; i++)
+
+			for (uint8 i = 0; i < CurrentState.TransitionsNum; i++)
 			{
-				const int16 TransitionIndex = State.TransitionsBegin + i;
-				const FCompactStateTransition& Transition = StateTree.Transitions[TransitionIndex];
+				const int16 TransitionIndex = CurrentState.TransitionsBegin + i;
+				const FCompactStateTransition& Transition = RootStateTree.Transitions[TransitionIndex];
 
 				// Skip disabled transitions
 				if (Transition.bTransitionEnabled == false)
@@ -2236,10 +2567,10 @@ bool FStateTreeExecutionContext::SelectStateInternal(const FStateTreeStateHandle
 				}
 
 				// Try to prevent (infinite) loops in the selection.
-				if (VisitedStates.Contains(Transition.State))
+				if (CurrentFrame.ActiveStates.Contains(Transition.State))
 				{
 					STATETREE_LOG(Error, TEXT("%hs: Loop detected when trying to select state %s from '%s'. Prior states: %s.  '%s' using StateTree '%s'."),
-						__FUNCTION__, *GetSafeStateName(NextState), *GetStateStatusString(Exec), *DebugGetStatePath(VisitedStates), *GetNameSafe(&Owner), *GetFullNameSafe(&StateTree));
+						__FUNCTION__, *GetSafeStateName(CurrentFrame, NextState), *GetStateStatusString(Exec), *DebugGetStatePath(OutNextActiveFrames, &CurrentFrame), *GetNameSafe(&Owner), *GetFullNameSafe(CurrentFrame.StateTree));
 					continue;
 				}
 
@@ -2252,18 +2583,19 @@ bool FStateTreeExecutionContext::SelectStateInternal(const FStateTreeStateHandle
 				{
 					STATETREE_TRACE_TRANSITION_EVENT(FStateTreeTransitionSource(FStateTreeIndex16(TransitionIndex), Transition.State, Transition.Priority), EStateTreeTraceEventType::OnEvaluating);
 					STATETREE_TRACE_SCOPED_PHASE(EStateTreeUpdatePhase::TransitionConditions);
-					bTransitionConditionsPassed = bShouldTrigger && TestAllConditions(Transition.ConditionsBegin, Transition.ConditionsNum);
+					bTransitionConditionsPassed = bShouldTrigger && TestAllConditions(CurrentFrame, Transition.ConditionsBegin, Transition.ConditionsNum);
 				}
 
 				if (bTransitionConditionsPassed)
 				{
 					// Using SelectState() instead of SelectStateInternal to treat the transitions the same way as regular transitions,
 					// e.g. it may jump to a completely different branch.
-					FStateTreeActiveStates NewActiveState;
-					if (SelectState(Transition.State, NewActiveState, VisitedStates, Transition.Fallback))
+					TArray<FStateTreeExecutionFrame, TFixedAllocator<MaxExecutionFrames>> NewActiveFrames;
+					if (SelectState(CurrentFrame, Transition.State, NewActiveFrames, Transition.Fallback))
 					{
-						// Selection succeeded
-						OutNewActiveState = NewActiveState;
+						// Selection succeeded.
+						// Cannot break yet because higher priority transitions may override the selection. 
+						OutNextActiveFrames = NewActiveFrames;
 						CurrentPriority = Transition.Priority;
 					}
 				}
@@ -2274,14 +2606,14 @@ bool FStateTreeExecutionContext::SelectStateInternal(const FStateTreeStateHandle
 				return true;
 			}
 		}
-		else if (State.SelectionBehavior == EStateTreeStateSelectionBehavior::TrySelectChildrenInOrder)
+		else if (CurrentState.SelectionBehavior == EStateTreeStateSelectionBehavior::TrySelectChildrenInOrder)
 		{
-			if (State.HasChildren())
+			if (CurrentState.HasChildren())
 			{
 				// If the state has children, proceed to select children.
-				for (uint16 ChildState = State.ChildrenBegin; ChildState < State.ChildrenEnd; ChildState = StateTree.States[ChildState].GetNextSibling())
+				for (uint16 ChildState = CurrentState.ChildrenBegin; ChildState < CurrentState.ChildrenEnd; ChildState = CurrentStateTree->States[ChildState].GetNextSibling())
 				{
-					if (SelectStateInternal(FStateTreeStateHandle(ChildState), OutNewActiveState, VisitedStates))
+					if (SelectStateInternal(CurrentFrame, FStateTreeStateHandle(ChildState), OutNextActiveFrames))
 					{
 						// Selection succeeded
 						return true;
@@ -2296,15 +2628,14 @@ bool FStateTreeExecutionContext::SelectStateInternal(const FStateTreeStateHandle
 			}
 		}
 		
-		OutNewActiveState.Pop();
-		VisitedStates.Pop();
+		CurrentFrame.ActiveStates.Pop();
 	}
 
 	// Nothing got selected.
 	return false;
 }
 
-FString FStateTreeExecutionContext::GetSafeStateName(const FStateTreeStateHandle State) const
+FString FStateTreeExecutionContext::GetSafeStateName(const FStateTreeExecutionFrame& CurrentFrame, const FStateTreeStateHandle State) const
 {
 	if (State == FStateTreeStateHandle::Invalid)
 	{
@@ -2318,33 +2649,55 @@ FString FStateTreeExecutionContext::GetSafeStateName(const FStateTreeStateHandle
 	{
 		return TEXT("(State Failed)");
 	}
-	else if (StateTree.States.IsValidIndex(State.Index))
+	else if (CurrentFrame.StateTree && CurrentFrame.StateTree->States.IsValidIndex(State.Index))
 	{
-		return *StateTree.States[State.Index].Name.ToString();
+		return *CurrentFrame.StateTree->States[State.Index].Name.ToString();
 	}
 	return TEXT("(Unknown)");
 }
 
-FString FStateTreeExecutionContext::DebugGetStatePath(const FStateTreeActiveStates& ActiveStates, const int32 ActiveStateIndex) const
+FString FStateTreeExecutionContext::DebugGetStatePath(TConstArrayView<FStateTreeExecutionFrame> ActiveFrames, const FStateTreeExecutionFrame* CurrentFrame, const int32 ActiveStateIndex) const
 {
 	FString StatePath;
-	const int32 Num = ActiveStateIndex == INDEX_NONE ? ActiveStates.Num() : (ActiveStateIndex + 1);  
-	if (!ensureMsgf(ActiveStates.IsValidIndex(Num - 1), TEXT("Provided index must be valid")))
+	const UStateTree* LastStateTree = &RootStateTree;
+		
+	for (const FStateTreeExecutionFrame& Frame : ActiveFrames)
 	{
-		return StatePath;
-	}
+		if (!ensure(Frame.StateTree))
+		{
+			return StatePath;
+		}
 
-	for (int32 i = 0; i < Num; i++)
-	{
-		const FCompactStateTreeState& State = StateTree.States[ActiveStates[i].Index];
-		StatePath.Appendf(TEXT("%s%s"), i == 0 ? TEXT("") : TEXT("."), *State.Name.ToString());
+		// If requested up the active state, clamp count.
+		int32 Num = Frame.ActiveStates.Num();
+		if (CurrentFrame == &Frame && Frame.ActiveStates.IsValidIndex(ActiveStateIndex))
+		{
+			Num = ActiveStateIndex + 1;
+		}
+
+		if (Frame.StateTree != LastStateTree)
+		{
+			StatePath.Appendf(TEXT("[%s]"), *GetNameSafe(Frame.StateTree));
+			LastStateTree = Frame.StateTree;
+		}
+		
+		for (int32 i = 0; i < Num; i++)
+		{
+			const FCompactStateTreeState& State = Frame.StateTree->States[Frame.ActiveStates[i].Index];
+			StatePath.Appendf(TEXT("%s%s"), i == 0 ? TEXT("") : TEXT("."), *State.Name.ToString());
+		}
 	}
+		
 	return StatePath;
 }
 
 FString FStateTreeExecutionContext::GetStateStatusString(const FStateTreeExecutionState& ExecState) const
 {
-	return GetSafeStateName(ExecState.ActiveStates.Last()) + TEXT(":") + UEnum::GetDisplayValueAsText(ExecState.LastTickStatus).ToString();
+	if (ExecState.TreeRunStatus != EStateTreeRunStatus::Running)
+	{
+		return TEXT("--:") + UEnum::GetDisplayValueAsText(ExecState.LastTickStatus).ToString();
+	}
+	return GetSafeStateName(ExecState.ActiveFrames.Last(), ExecState.ActiveFrames.Last().ActiveStates.Last()) + TEXT(":") + UEnum::GetDisplayValueAsText(ExecState.LastTickStatus).ToString();
 }
 
 EStateTreeRunStatus FStateTreeExecutionContext::GetLastTickStatus() const
@@ -2358,10 +2711,10 @@ FString FStateTreeExecutionContext::GetInstanceDescription() const
 	return FString::Printf(TEXT("%s"), *GetNameSafe(&Owner));
 }
 
-const FStateTreeActiveStates& FStateTreeExecutionContext::GetActiveStates() const
+TConstArrayView<FStateTreeExecutionFrame> FStateTreeExecutionContext::GetActiveFrames() const
 {
 	const FStateTreeExecutionState& Exec = GetExecState();
-	return Exec.ActiveStates;
+	return Exec.ActiveFrames;
 }
 
 
@@ -2371,7 +2724,7 @@ FString FStateTreeExecutionContext::GetDebugInfoString() const
 {
 	const FStateTreeExecutionState& Exec = GetExecState();
 
-	FString DebugString = FString::Printf(TEXT("StateTree (asset: '%s')\n"), *GetFullNameSafe(&StateTree));
+	FString DebugString = FString::Printf(TEXT("StateTree (asset: '%s')\n"), *GetFullNameSafe(&RootStateTree));
 
 	DebugString += TEXT("Status: ");
 	switch (Exec.TreeRunStatus)
@@ -2389,35 +2742,52 @@ FString FStateTreeExecutionContext::GetDebugInfoString() const
 		DebugString += TEXT("--\n");
 	}
 
-	if (StateTree.EvaluatorsNum > 0)
-	{
-		DebugString += TEXT("\nEvaluators:\n");
-		for (int32 EvalIndex = StateTree.EvaluatorsBegin; EvalIndex < (StateTree.EvaluatorsBegin + StateTree.EvaluatorsNum); EvalIndex++)
-		{
-			const FStateTreeEvaluatorBase& Eval = StateTree.Nodes[EvalIndex].Get<const FStateTreeEvaluatorBase>();
-			Eval.AppendDebugInfoString(DebugString, *this);
-		}
-	}
-
 	// Active States
 	DebugString += TEXT("Current State:\n");
-	for (int32 Index = 0; Index < Exec.ActiveStates.Num(); Index++)
+	for (const FStateTreeExecutionFrame& CurrentFrame : Exec.ActiveFrames)
 	{
-		FStateTreeStateHandle Handle = Exec.ActiveStates[Index];
-		if (Handle.IsValid())
-		{
-			const FCompactStateTreeState& State = StateTree.States[Handle.Index];
-			DebugString += FString::Printf(TEXT("[%s]\n"), *State.Name.ToString());
+		const UStateTree* CurrentStateTree = CurrentFrame.StateTree;
 
-			if (State.TasksNum > 0)
+		if (CurrentFrame.bIsGlobalFrame)
+		{
+			DebugString += FString::Printf(TEXT("\nEvaluators\n  [ %-30s | %8s | %15s ]\n"),
+				TEXT("Name"), TEXT("Bindings"), TEXT("Data Handle"));
+			for (int32 EvalIndex = CurrentStateTree->EvaluatorsBegin; EvalIndex < (CurrentStateTree->EvaluatorsBegin + CurrentStateTree->EvaluatorsNum); EvalIndex++)
+			{
+				const FStateTreeEvaluatorBase& Eval = CurrentStateTree->Nodes[EvalIndex].Get<const FStateTreeEvaluatorBase>();
+				DebugString += FString::Printf(TEXT("| %-30s | %8d | %15s |\n"),
+					*Eval.Name.ToString(), Eval.BindingsBatch.Get(), *Eval.InstanceDataHandle.Describe());
+			}
+			
+			DebugString += FString::Printf(TEXT("\nGlobal Tasks\n"));
+			for (int32 TaskIndex = CurrentStateTree->GlobalTasksBegin; TaskIndex < (CurrentStateTree->GlobalTasksBegin + CurrentStateTree->GlobalTasksNum); TaskIndex++)
+			{
+				const FStateTreeTaskBase& Task = CurrentStateTree->Nodes[TaskIndex].Get<const FStateTreeTaskBase>();
+				if (Task.bTaskEnabled)
+				{
+					Task.AppendDebugInfoString(DebugString, *this);
+				}
+			}
+		}
+		
+		for (int32 Index = 0; Index < CurrentFrame.ActiveStates.Num(); Index++)
+		{
+			FStateTreeStateHandle Handle = CurrentFrame.ActiveStates[Index];
+			if (Handle.IsValid())
 			{
-				DebugString += TEXT("\nTasks:\n");
-				for (int32 TaskIndex = State.TasksBegin; TaskIndex < (State.TasksBegin + State.TasksNum); TaskIndex++)
+				const FCompactStateTreeState& State = RootStateTree.States[Handle.Index];
+				DebugString += FString::Printf(TEXT("[%s]\n"), *State.Name.ToString());
+
+				if (State.TasksNum > 0)
 				{
-					const FStateTreeTaskBase& Task = StateTree.Nodes[TaskIndex].Get<const FStateTreeTaskBase>();
-					if (Task.bTaskEnabled)
+					DebugString += TEXT("\nTasks:\n");
+					for (int32 TaskIndex = State.TasksBegin; TaskIndex < (State.TasksBegin + State.TasksNum); TaskIndex++)
 					{
-						Task.AppendDebugInfoString(DebugString, *this);
+						const FStateTreeTaskBase& Task = RootStateTree.Nodes[TaskIndex].Get<const FStateTreeTaskBase>();
+						if (Task.bTaskEnabled)
+						{
+							Task.AppendDebugInfoString(DebugString, *this);
+						}
 					}
 				}
 			}
@@ -2433,44 +2803,48 @@ void FStateTreeExecutionContext::DebugPrintInternalLayout()
 {
 	LOG_SCOPE_VERBOSITY_OVERRIDE(LogStateTree, ELogVerbosity::Log);
 
-	FString DebugString = FString::Printf(TEXT("StateTree (asset: '%s')\n"), *GetFullNameSafe(&StateTree));
+	// @todo: this looks like more a UStateTree thing...
+	
+	FString DebugString = FString::Printf(TEXT("StateTree (asset: '%s')\n"), *GetFullNameSafe(&RootStateTree));
 
 	// Tree items (e.g. tasks, evaluators, conditions)
-	DebugString += FString::Printf(TEXT("\nItems(%d)\n"), StateTree.Nodes.Num());
-	for (int32 Index = 0; Index < StateTree.Nodes.Num(); Index++)
+	DebugString += FString::Printf(TEXT("\nItems(%d)\n"), RootStateTree.Nodes.Num());
+	for (int32 Index = 0; Index < RootStateTree.Nodes.Num(); Index++)
 	{
-		const FConstStructView Node = StateTree.Nodes[Index];
+		const FConstStructView Node = RootStateTree.Nodes[Index];
 		DebugString += FString::Printf(TEXT("  %s\n"), Node.IsValid() ? *Node.GetScriptStruct()->GetName() : TEXT("null"));
 	}
 
 	// Instance InstanceData data (e.g. tasks)
-	DebugString += FString::Printf(TEXT("\nInstance Structs(%d)\n"), StateTree.DefaultInstanceData.NumStructs());
-	for (int32 Index = 0; Index < StateTree.DefaultInstanceData.NumStructs(); Index++)
+	DebugString += FString::Printf(TEXT("\nInstance Data(%d)\n"), RootStateTree.DefaultInstanceData.Num());
+	for (int32 Index = 0; Index < RootStateTree.DefaultInstanceData.Num(); Index++)
 	{
-		const FConstStructView Data = StateTree.DefaultInstanceData.GetStruct(Index);
-		DebugString += FString::Printf(TEXT("  %s\n"), Data.IsValid() ? *Data.GetScriptStruct()->GetName() : TEXT("null"));
-	}
-	DebugString += FString::Printf(TEXT("\nInstance Objects(%d)\n"), StateTree.DefaultInstanceData.NumObjects());
-	for (int32 Index = 0; Index < StateTree.DefaultInstanceData.NumObjects(); Index++)
-	{
-		const UObject* Data = StateTree.DefaultInstanceData.GetObject(Index);
-		DebugString += FString::Printf(TEXT("  %s\n"), *GetNameSafe(Data));
+		if (RootStateTree.DefaultInstanceData.IsObject(Index))
+		{
+			const UObject* Data = RootStateTree.DefaultInstanceData.GetObject(Index);
+			DebugString += FString::Printf(TEXT("  %s\n"), *GetNameSafe(Data));
+		}
+		else
+		{
+			const FConstStructView Data = RootStateTree.DefaultInstanceData.GetStruct(Index);
+			DebugString += FString::Printf(TEXT("  %s\n"), Data.IsValid() ? *Data.GetScriptStruct()->GetName() : TEXT("null"));
+		}
 	}
 
 	// External data (e.g. fragments, subsystems)
-	DebugString += FString::Printf(TEXT("\nExternal Data(%d)\n  [ %-40s | %-8s | %5s ]\n"), StateTree.ExternalDataDescs.Num(), TEXT("Name"), TEXT("Optional"), TEXT("Index"));
-	for (const FStateTreeExternalDataDesc& Desc : StateTree.ExternalDataDescs)
+	DebugString += FString::Printf(TEXT("\nExternal Data(%d)\n  [ %-40s | %-8s | %15s ]\n"), RootStateTree.ExternalDataDescs.Num(), TEXT("Name"), TEXT("Optional"), TEXT("Handle"));
+	for (const FStateTreeExternalDataDesc& Desc : RootStateTree.ExternalDataDescs)
 	{
-		DebugString += FString::Printf(TEXT("  | %-40s | %8s | %5d |\n"), Desc.Struct ? *Desc.Struct->GetName() : TEXT("null"), *UEnum::GetDisplayValueAsText(Desc.Requirement).ToString(), Desc.Handle.DataViewIndex.Get());
+		DebugString += FString::Printf(TEXT("  | %-40s | %8s | %15s |\n"), Desc.Struct ? *Desc.Struct->GetName() : TEXT("null"), *UEnum::GetDisplayValueAsText(Desc.Requirement).ToString(), *Desc.Handle.DataHandle.Describe());
 	}
 
 	// Bindings
-	StateTree.PropertyBindings.DebugPrintInternalLayout(DebugString);
+	RootStateTree.PropertyBindings.DebugPrintInternalLayout(DebugString);
 
 	// Transitions
-	DebugString += FString::Printf(TEXT("\nTransitions(%d)\n  [ %-3s | %15s | %-20s | %-40s | %-8s ]\n"), StateTree.Transitions.Num()
+	DebugString += FString::Printf(TEXT("\nTransitions(%d)\n  [ %-3s | %15s | %-20s | %-40s | %-8s ]\n"), RootStateTree.Transitions.Num()
 		, TEXT("Idx"), TEXT("State"), TEXT("Transition Trigger"), TEXT("Transition Event Tag"), TEXT("Num Cond"));
-	for (const FCompactStateTransition& Transition : StateTree.Transitions)
+	for (const FCompactStateTransition& Transition : RootStateTree.Transitions)
 	{
 		DebugString += FString::Printf(TEXT("  | %3d | %15s | %-20s | %-40s | %8d |\n"),
 									Transition.ConditionsBegin, *Transition.State.Describe(),
@@ -2479,22 +2853,15 @@ void FStateTreeExecutionContext::DebugPrintInternalLayout()
 									Transition.ConditionsNum);
 	}
 
-	// DataViews
-	DebugString += FString::Printf(TEXT("\nDataViews(%d)\n"), DataViews.Num());
-	for (const FStateTreeDataView& DataView : DataViews)
-	{
-		DebugString += FString::Printf(TEXT("  [%s]\n"), DataView.IsValid() ? *DataView.GetStruct()->GetName() : TEXT("null"));
-	}
-
 	// States
 	DebugString += FString::Printf(TEXT("\nStates(%d)\n"
 		"  [ %-30s | %15s | %5s [%3s:%-3s[ | Begin Idx : %4s %4s %4s %4s | Num : %4s %4s %4s %4s | Transitions : %-16s %-40s %-16s %-40s ]\n"),
-		StateTree.States.Num(),
+		RootStateTree.States.Num(),
 		TEXT("Name"), TEXT("Parent"), TEXT("Child"), TEXT("Beg"), TEXT("End"),
 		TEXT("Cond"), TEXT("Tr"), TEXT("Tsk"), TEXT("Evt"), TEXT("Cond"), TEXT("Tr"), TEXT("Tsk"), TEXT("Evt"),
 		TEXT("Done State"), TEXT("Done Type"), TEXT("Failed State"), TEXT("Failed Type")
 		);
-	for (const FCompactStateTreeState& State : StateTree.States)
+	for (const FCompactStateTreeState& State : RootStateTree.States)
 	{
 		DebugString += FString::Printf(TEXT("  | %-30s | %15s | %5s [%3d:%-3d[ | %9s   %4d %4d %4d | %3s   %4d %4d %4d\n"),
 									*State.Name.ToString(), *State.Parent.Describe(),
@@ -2504,31 +2871,31 @@ void FStateTreeExecutionContext::DebugPrintInternalLayout()
 	}
 
 	// Evaluators
-	if (StateTree.EvaluatorsNum)
+	if (RootStateTree.EvaluatorsNum)
 	{
 		DebugString += FString::Printf(TEXT("\nEvaluators\n  [ %-30s | %8s | %10s ]\n"),
 			TEXT("Name"), TEXT("Bindings"), TEXT("Struct Idx"));
-		for (int32 EvalIndex = StateTree.EvaluatorsBegin; EvalIndex < (StateTree.EvaluatorsBegin + StateTree.EvaluatorsNum); EvalIndex++)
+		for (int32 EvalIndex = RootStateTree.EvaluatorsBegin; EvalIndex < (RootStateTree.EvaluatorsBegin + RootStateTree.EvaluatorsNum); EvalIndex++)
 		{
-			const FStateTreeEvaluatorBase& Eval = StateTree.Nodes[EvalIndex].Get<const FStateTreeEvaluatorBase>();
-			DebugString += FString::Printf(TEXT("| %-30s | %8d | %10d |\n"),
-				*Eval.Name.ToString(), Eval.BindingsBatch.Get(), Eval.DataViewIndex.Get());
+			const FStateTreeEvaluatorBase& Eval = RootStateTree.Nodes[EvalIndex].Get<const FStateTreeEvaluatorBase>();
+			DebugString += FString::Printf(TEXT("| %-30s | %8d | %10s |\n"),
+				*Eval.Name.ToString(), Eval.BindingsBatch.Get(), *Eval.InstanceDataHandle.Describe());
 		}
 	}
 
 
 	DebugString += FString::Printf(TEXT("\nTasks\n  [ %-30s | %-30s | %8s | %10s ]\n"),
 		TEXT("State"), TEXT("Name"), TEXT("Bindings"), TEXT("Struct Idx"));
-	for (const FCompactStateTreeState& State : StateTree.States)
+	for (const FCompactStateTreeState& State : RootStateTree.States)
 	{
 		// Tasks
 		if (State.TasksNum)
 		{
 			for (int32 TaskIndex = State.TasksBegin; TaskIndex < (State.TasksBegin + State.TasksNum); TaskIndex++)
 			{
-				const FStateTreeTaskBase& Task = StateTree.Nodes[TaskIndex].Get<const FStateTreeTaskBase>();
-				DebugString += FString::Printf(TEXT("  | %-30s | %-30s | %8d | %10d |\n"), *State.Name.ToString(),
-					*Task.Name.ToString(), Task.BindingsBatch.Get(), Task.DataViewIndex.Get());
+				const FStateTreeTaskBase& Task = RootStateTree.Nodes[TaskIndex].Get<const FStateTreeTaskBase>();
+				DebugString += FString::Printf(TEXT("  | %-30s | %-30s | %8d | %10s |\n"), *State.Name.ToString(),
+					*Task.Name.ToString(), Task.BindingsBatch.Get(), *Task.InstanceDataHandle.Describe());
 			}
 		}
 	}
@@ -2559,24 +2926,39 @@ FString FStateTreeExecutionContext::GetActiveStateName() const
 
 	FString FullStateName;
 	
-	// Active States
-	for (int32 Index = 0; Index < Exec.ActiveStates.Num(); Index++)
+	const UStateTree* LastStateTree = &RootStateTree;
+	int32 Indent = 0;
+
+	for (int32 FrameIndex = 0; FrameIndex < Exec.ActiveFrames.Num(); FrameIndex++)
 	{
-		const FStateTreeStateHandle Handle = Exec.ActiveStates[Index];
-		if (Handle.IsValid())
+		const FStateTreeExecutionFrame& CurrentFrame = Exec.ActiveFrames[FrameIndex];
+		const UStateTree* CurrentStateTree = CurrentFrame.StateTree;
+
+		// Append linked state marker at the end of the previous line.
+		if (Indent > 0)
 		{
-			const FCompactStateTreeState& State = StateTree.States[Handle.Index];
-			bool bIsLinked = false;
-			if (Index > 0)
-			{
-				FullStateName += TEXT("\n");
-				bIsLinked = Exec.ActiveStates[Index - 1] != State.Parent;
-			}
-			FullStateName += FString::Printf(TEXT("%*s-"), Index * 3, TEXT("")); // Indent
-			FullStateName += *State.Name.ToString();
-			if (bIsLinked)
+			FullStateName += TEXT(" >");
+		}
+		// If tree has changed, append that too.
+		if (CurrentFrame.StateTree != LastStateTree)
+		{
+			FullStateName.Appendf(TEXT(" [%s]"), *GetNameSafe(CurrentFrame.StateTree));
+			LastStateTree = CurrentFrame.StateTree;
+		}
+
+		for (int32 Index = 0; Index < CurrentFrame.ActiveStates.Num(); Index++)
+		{
+			const FStateTreeStateHandle Handle = CurrentFrame.ActiveStates[Index];
+			if (Handle.IsValid())
 			{
-				FullStateName += TEXT(" >");
+				const FCompactStateTreeState& State = CurrentStateTree->States[Handle.Index];
+				if (Indent > 0)
+				{
+					FullStateName += TEXT("\n");
+				}
+				FullStateName += FString::Printf(TEXT("%*s-"), Indent * 3, TEXT("")); // Indent
+				FullStateName += *State.Name.ToString();
+				Indent++;
 			}
 		}
 	}
@@ -2611,13 +2993,17 @@ TArray<FName> FStateTreeExecutionContext::GetActiveStateNames() const
 	const FStateTreeExecutionState& Exec = GetExecState();
 
 	// Active States
-	for (int32 Index = 0; Index < Exec.ActiveStates.Num(); Index++)
+	for (const FStateTreeExecutionFrame& CurrentFrame : Exec.ActiveFrames)
 	{
-		const FStateTreeStateHandle Handle = Exec.ActiveStates[Index];
-		if (Handle.IsValid())
+		const UStateTree* CurrentStateTree = CurrentFrame.StateTree;
+		for (int32 Index = 0; Index < CurrentFrame.ActiveStates.Num(); Index++)
 		{
-			const FCompactStateTreeState& State = StateTree.States[Handle.Index];
-			Result.Add(State.Name);
+			const FStateTreeStateHandle Handle = CurrentFrame.ActiveStates[Index];
+			if (Handle.IsValid())
+			{
+				const FCompactStateTreeState& State = CurrentStateTree->States[Handle.Index];
+				Result.Add(State.Name);
+			}
 		}
 	}
 
