StateTree: Changed how instance data is accessed

diff --git a/Engine/Plugins/Runtime/StateTree/Source/StateTreeModule/Private/StateTreeInstanceData.cpp b/Engine/Plugins/Runtime/StateTree/Source/StateTreeModule/Private/StateTreeInstanceData.cpp
--- a/Engine/Plugins/Runtime/StateTree/Source/StateTreeModule/Private/StateTreeInstanceData.cpp
+++ b/Engine/Plugins/Runtime/StateTree/Source/StateTreeModule/Private/StateTreeInstanceData.cpp
@@ -80,12 +80,12 @@ bool FStateTreeInstanceStorage::AreAllInstancesValid() const
 		{
 			return false;
 		}
-	}
-	for (const UObject* Instance : InstanceObjects)
-	{
-		if (!Instance)
+		if (const FStateTreeInstanceObjectWrapper* Wrapper = Instance.GetPtr<const FStateTreeInstanceObjectWrapper>())
 		{
-			return false;
+			if (!Wrapper->InstanceObject)
+			{
+				return false;
+			}
 		}
 	}
 	return true;
@@ -106,38 +106,18 @@ FStateTreeInstanceData::~FStateTreeInstanceData()
 	Reset();
 }
 
-bool FStateTreeInstanceStorage::IsValid() const
-{
-	return InstanceStructs.Num() > 0 || InstanceObjects.Num() > 0;
-}
-
 const FStateTreeInstanceStorage& FStateTreeInstanceData::GetStorage() const
 {
-	check(InstanceStorage.GetMemory() != nullptr && InstanceStorage.GetScriptStruct() == TBaseStructure<FStateTreeInstanceStorage>::Get());
+	check(InstanceStorage.GetMemory() != nullptr);
 	return *reinterpret_cast<const FStateTreeInstanceStorage*>(InstanceStorage.GetMemory());
 }
 
 FStateTreeInstanceStorage& FStateTreeInstanceData::GetMutableStorage()
 {
-	check(InstanceStorage.GetMemory() != nullptr && InstanceStorage.GetScriptStruct() == TBaseStructure<FStateTreeInstanceStorage>::Get());
+	check(InstanceStorage.GetMemory() != nullptr);
 	return *reinterpret_cast<FStateTreeInstanceStorage*>(InstanceStorage.GetMutableMemory());
 }
 
-const FStateTreeExecutionState* FStateTreeInstanceData::GetExecutionState() const
-{
-	if (!IsValid())
-	{
-		return nullptr;
-	}
-	const FConstStructView ExecView = GetStruct(0); // Execution state is fixed at index 0. 
-	return ExecView.GetPtr<const FStateTreeExecutionState>();
-}
-
-TArray<FStateTreeEvent>& FStateTreeInstanceData::GetEvents() const
-{
-	return const_cast<FStateTreeInstanceData*>(this)->GetMutableStorage().EventQueue.GetEventsArray();
-}
-
 FStateTreeEventQueue& FStateTreeInstanceData::GetMutableEventQueue()
 {
 	return GetMutableStorage().EventQueue;	
@@ -175,23 +155,20 @@ int32 FStateTreeInstanceData::GetEstimatedMemoryUsage() const
 
 	Size += Storage.InstanceStructs.GetAllocatedMemory();
 
-	for (const UObject* InstanceObject : Storage.InstanceObjects)
+	for (FConstStructView Instance : Storage.InstanceStructs)
 	{
-		if (InstanceObject)
+		if (const FStateTreeInstanceObjectWrapper* Wrapper = Instance.GetPtr<const FStateTreeInstanceObjectWrapper>())
 		{
-			Size += InstanceObject->GetClass()->GetStructureSize();
+			if (Wrapper->InstanceObject)
+			{
+				Size += Wrapper->InstanceObject->GetClass()->GetStructureSize();
+			}
 		}
 	}
 
 	return Size;
 }
 
-int32 FStateTreeInstanceData::GetNumItems() const
-{
-	const FStateTreeInstanceStorage& Storage = GetStorage();
-	return Storage.InstanceStructs.Num() + Storage.InstanceObjects.Num();
-}
-
 bool FStateTreeInstanceData::Identical(const FStateTreeInstanceData* Other, uint32 PortFlags) const
 {
 	if (Other == nullptr)
@@ -214,12 +191,6 @@ bool FStateTreeInstanceData::Identical(const FStateTreeInstanceData* Other, uint
 	const FStateTreeInstanceStorage& Storage = GetStorage();
 	const FStateTreeInstanceStorage& OtherStorage = Other->GetStorage();
 
-	// Not identical if different amount of instanced objects.
-	if (Storage.InstanceObjects.Num() != OtherStorage.InstanceObjects.Num())
-	{
-		return false;
-	}
-
 	// Not identical if structs are different.
 	if (Storage.InstanceStructs.Identical(&OtherStorage.InstanceStructs, PortFlags) == false)
 	{
@@ -250,47 +221,31 @@ bool FStateTreeInstanceData::Identical(const FStateTreeInstanceData* Other, uint
 	};
 
 	bool bResult = true;
-	for (int32 Index = 0; Index < Storage.InstanceObjects.Num(); Index++)
+
+	for (int32 Index = 0; Index < Storage.InstanceStructs.Num(); Index++)
 	{
-		if (Storage.InstanceObjects[Index] != nullptr && OtherStorage.InstanceObjects[Index] != nullptr)
+		const FStateTreeInstanceObjectWrapper* Wrapper = Storage.InstanceStructs[Index].GetPtr<const FStateTreeInstanceObjectWrapper>();
+		const FStateTreeInstanceObjectWrapper* OtherWrapper = OtherStorage.InstanceStructs[Index].GetPtr<const FStateTreeInstanceObjectWrapper>();
+
+		if (Wrapper)
 		{
-			if (!AreObjectsIdentical(Storage.InstanceObjects[Index], OtherStorage.InstanceObjects[Index], PortFlags))
+			if (!OtherWrapper)
 			{
 				bResult = false;
 				break;
 			}
-		}
-		else
-		{
-			bResult = false;
-			break;
-		}
-	}
-	
-	return bResult;
-}
-
-void FStateTreeInstanceData::PostSerialize(const FArchive& Ar)
-{
-#if WITH_EDITORONLY_DATA
-	if (Ar.IsLoading())
-	{
-		FStateTreeInstanceStorage& Storage = GetMutableStorage();
-		PRAGMA_DISABLE_DEPRECATION_WARNINGS
-		if (InstanceStructs_DEPRECATED.Num() > 0 || InstanceObjects_DEPRECATED.Num() > 0)
-		{
-			if (!Storage.IsValid())
+			if (Wrapper->InstanceObject && OtherWrapper->InstanceObject)
 			{
-				Storage.InstanceStructs.Reset();
-				Storage.InstanceStructs.Append(InstanceStructs_DEPRECATED);
-				Storage.InstanceObjects = InstanceObjects_DEPRECATED;
+				if (!AreObjectsIdentical(Wrapper->InstanceObject, OtherWrapper->InstanceObject, PortFlags))
+				{
+					bResult = false;
+					break;
+				}
 			}
-			InstanceStructs_DEPRECATED.Reset();
-			InstanceObjects_DEPRECATED.Reset();
 		}
-		PRAGMA_ENABLE_DEPRECATION_WARNINGS
 	}
-#endif
+	
+	return bResult;
 }
 
 void FStateTreeInstanceData::CopyFrom(UObject& InOwner, const FStateTreeInstanceData& InOther)
@@ -307,81 +262,84 @@ void FStateTreeInstanceData::CopyFrom(UObject& InOwner, const FStateTreeInstance
 	Storage.InstanceStructs = OtherStorage.InstanceStructs;
 
 	// Copy instance objects.
-	Storage.InstanceObjects.Reset();
-	for (const UObject* Instance : OtherStorage.InstanceObjects)
+	for (FStructView Instance : Storage.InstanceStructs)
 	{
-		if (ensure(Instance != nullptr))
+		if (FStateTreeInstanceObjectWrapper* Wrapper = Instance.GetPtr<FStateTreeInstanceObjectWrapper>())
 		{
-			Storage.InstanceObjects.Add(UE::StateTree::DuplicateNodeInstance(*Instance, InOwner));
+			if (Wrapper->InstanceObject)
+			{
+				Wrapper->InstanceObject = UE::StateTree::DuplicateNodeInstance(*Wrapper->InstanceObject, InOwner);
+			}
 		}
 	}
 }
 
-void FStateTreeInstanceData::Init(UObject& InOwner, TConstArrayView<FInstancedStruct> InStructs, TConstArrayView<const UObject*> InObjects)
+void FStateTreeInstanceData::Init(UObject& InOwner, TConstArrayView<FInstancedStruct> InStructs)
 {
 	Reset();
-	Append(InOwner, InStructs, InObjects);
+	Append(InOwner, InStructs);
 }
 
-void FStateTreeInstanceData::Init(UObject& InOwner, TConstArrayView<FConstStructView> InStructs, TConstArrayView<const UObject*> InObjects)
+void FStateTreeInstanceData::Init(UObject& InOwner, TConstArrayView<FConstStructView> InStructs)
 {
 	Reset();
-	Append(InOwner, InStructs, InObjects);
+	Append(InOwner, InStructs);
 }
 
-void FStateTreeInstanceData::Append(UObject& InOwner, TConstArrayView<FInstancedStruct> InStructs, TConstArrayView<const UObject*> InObjects)
+void FStateTreeInstanceData::Append(UObject& InOwner, TConstArrayView<FInstancedStruct> InStructs)
 {
 	FStateTreeInstanceStorage& Storage = GetMutableStorage();
 
+	const int32 StartIndex = Storage.InstanceStructs.Num();
 	Storage.InstanceStructs.Append(InStructs);
-	
-	Storage.InstanceObjects.Reserve(Storage.InstanceObjects.Num() + InObjects.Num());
-	for (const UObject* Instance : InObjects)
+
+	for (int32 Index = StartIndex; Index < Storage.InstanceStructs.Num(); Index++)
 	{
-		if (ensure(Instance != nullptr))
+		if (FStateTreeInstanceObjectWrapper* Wrapper = Storage.InstanceStructs[Index].GetPtr<FStateTreeInstanceObjectWrapper>())
 		{
-			Storage.InstanceObjects.Add(UE::StateTree::DuplicateNodeInstance(*Instance, InOwner));
+			if (Wrapper->InstanceObject)
+			{
+				Wrapper->InstanceObject = UE::StateTree::DuplicateNodeInstance(*Wrapper->InstanceObject, InOwner);
+			}
 		}
 	}
 }
 
-void FStateTreeInstanceData::Append(UObject& InOwner, TConstArrayView<FConstStructView> InStructs, TConstArrayView<const UObject*> InObjects)
+void FStateTreeInstanceData::Append(UObject& InOwner, TConstArrayView<FConstStructView> InStructs)
 {
 	FStateTreeInstanceStorage& Storage = GetMutableStorage();
 
+	const int32 StartIndex = Storage.InstanceStructs.Num();
 	Storage.InstanceStructs.Append(InStructs);
 	
-	Storage.InstanceObjects.Reserve(Storage.InstanceObjects.Num() + InObjects.Num());
-	for (const UObject* Instance : InObjects)
+	for (int32 Index = StartIndex; Index < Storage.InstanceStructs.Num(); Index++)
 	{
-		if (ensure(Instance != nullptr))
+		if (FStateTreeInstanceObjectWrapper* Wrapper = Storage.InstanceStructs[Index].GetPtr<FStateTreeInstanceObjectWrapper>())
 		{
-			Storage.InstanceObjects.Add(UE::StateTree::DuplicateNodeInstance(*Instance, InOwner));
+			if (Wrapper->InstanceObject)
+			{
+				Wrapper->InstanceObject = UE::StateTree::DuplicateNodeInstance(*Wrapper->InstanceObject, InOwner);
+			}
 		}
 	}
 }
 
-void FStateTreeInstanceData::ShrinkTo(const int32 NumStructs, const int32 NumObjects)
+void FStateTreeInstanceData::ShrinkTo(const int32 NumStructs)
 {
 	FStateTreeInstanceStorage& Storage = GetMutableStorage();
-	check(NumStructs <= Storage.InstanceStructs.Num() && NumObjects <= Storage.InstanceObjects.Num());  
+	check(NumStructs <= Storage.InstanceStructs.Num());  
 	Storage.InstanceStructs.SetNum(NumStructs);
-	Storage.InstanceObjects.SetNum(NumObjects);
 }
 
 bool FStateTreeInstanceData::IsValid() const
 {
-	if (!InstanceStorage.IsValid())
-	{
-		return false;
-	}
-	return GetStorage().IsValid();
+	return InstanceStorage.IsValid();
 }
 
 void FStateTreeInstanceData::Reset()
 {
 	FStateTreeInstanceStorage& Storage = GetMutableStorage();
 	Storage.InstanceStructs.Reset();
-	Storage.InstanceObjects.Reset();
 	Storage.EventQueue.Reset();
+	ExecutionState.Reset();
 }
diff --git a/Engine/Plugins/Runtime/StateTree/Source/StateTreeModule/Private/StateTreePropertyBindings.cpp b/Engine/Plugins/Runtime/StateTree/Source/StateTreeModule/Private/StateTreePropertyBindings.cpp
--- a/Engine/Plugins/Runtime/StateTree/Source/StateTreeModule/Private/StateTreePropertyBindings.cpp
+++ b/Engine/Plugins/Runtime/StateTree/Source/StateTreeModule/Private/StateTreePropertyBindings.cpp
@@ -135,6 +135,28 @@ void FStateTreePropertyBindings::Reset()
 	bBindingsResolved = false;
 }
 
+const FStateTreeBindableStructDesc* FStateTreePropertyBindings::GetSourceDescByHandle(const FStateTreeDataHandle SourceDataHandle)
+{
+	TArray<FStateTreeBindableStructDesc> FoundDescs;
+	for (const FStateTreeBindableStructDesc& Desc : SourceStructs)
+	{
+		if (Desc.DataHandle == SourceDataHandle)
+		{
+			FoundDescs.Add(Desc);
+		}
+	}
+
+	if (FoundDescs.Num() > 1)
+	{
+		UE_LOG(LogStateTree, Error, TEXT("%hs: Found %d entries for handle %s."), __FUNCTION__, FoundDescs.Num(), *SourceDataHandle.Describe());
+	}
+	
+	return SourceStructs.FindByPredicate([SourceDataHandle](const FStateTreeBindableStructDesc& Desc)
+	{
+		return Desc.DataHandle == SourceDataHandle;
+	});
+}
+
 bool FStateTreePropertyBindings::ResolvePaths()
 {
 	PropertyIndirections.Reset();
@@ -149,18 +171,30 @@ bool FStateTreePropertyBindings::ResolvePaths()
 		for (int32 i = Batch.BindingsBegin; i != Batch.BindingsEnd; i++)
 		{
 			const FStateTreePropertyPathBinding& Binding = PropertyPathBindings[i];
+			
 			FStateTreePropertyCopy& Copy = PropertyCopies[i];
+			Copy.SourceDataHandle = Binding.GetSourceDataHandle();
 
-			if (!Binding.GetCompiledSourceStructIndex().IsValid())
+			if (!Binding.GetSourceDataHandle().IsValid())
 			{
 				UE_LOG(LogStateTree, Error, TEXT("%hs: Invalid source struct for property binding %s."), __FUNCTION__, *Binding.GetSourcePath().ToString());
+				Copy.Type = EStateTreePropertyCopyType::None;
+				bBindingsResolved = false;
 				bResult = false;
 				continue;
 			}
 
-			Copy.SourceStructIndex = Binding.GetCompiledSourceStructIndex();
+			const FStateTreeBindableStructDesc* SourceDesc = GetSourceDescByHandle(Copy.SourceDataHandle);
+			if (!SourceDesc)
+			{
+				UE_LOG(LogStateTree, Error, TEXT("%hs: Could not find data source for binding %s."), __FUNCTION__, *Binding.GetSourcePath().ToString());
+				Copy.Type = EStateTreePropertyCopyType::None;
+				bBindingsResolved = false;
+				bResult = false;
+				continue;
+			}
 
-			const UStruct* SourceStruct = SourceStructs[Copy.SourceStructIndex.Get()].Struct;
+			const UStruct* SourceStruct = SourceDesc->Struct;
 			const UStruct* TargetStruct = Batch.TargetStruct.Struct;
 			if (!SourceStruct || !TargetStruct)
 			{
@@ -170,6 +204,8 @@ bool FStateTreePropertyBindings::ResolvePaths()
 				continue;
 			}
 
+			Copy.SourceStructType = SourceStruct;
+
 			// Resolve paths and validate the copy. Stops on first failure.
 			bool bSuccess = true;
 			FStateTreePropertyPathIndirection SourceLeafIndirection;
@@ -960,7 +996,7 @@ void FStateTreePropertyBindings::PerformCopy(const FStateTreePropertyCopy& Copy,
 	}
 }
 
-bool FStateTreePropertyBindings::CopyTo(TConstArrayView<FStateTreeDataView> SourceStructViews, const FStateTreeIndex16 TargetBatchIndex, FStateTreeDataView TargetStructView) const
+bool FStateTreePropertyBindings::CopyProperty(const FStateTreePropertyCopy& Copy, FStateTreeDataView SourceStructView, FStateTreeDataView TargetStructView) const
 {
 	// This is made ensure so that the programmers have the change to catch it (it's usually programming error not to call ResolvePaths(), and it wont spam log for others.
 	if (!ensureMsgf(bBindingsResolved, TEXT("Bindings must be resolved successfully before copying. See ResolvePaths()")))
@@ -968,51 +1004,35 @@ bool FStateTreePropertyBindings::CopyTo(TConstArrayView<FStateTreeDataView> Sour
 		return false;
 	}
 
-	if (TargetBatchIndex.IsValid() == false)
+	// Copies that fail to be resolved (i.e. property path does not resolve, types changed) will be marked as None, skip them.
+	if (Copy.Type == EStateTreePropertyCopyType::None)
 	{
-		return false;
+		return true;
 	}
 
-	check(CopyBatches.IsValidIndex(TargetBatchIndex.Get()));
-	const FStateTreePropertyCopyBatch& Batch = CopyBatches[TargetBatchIndex.Get()];
-
-	check(TargetStructView.IsValid());
-	check(TargetStructView.GetStruct() == Batch.TargetStruct.Struct);
-
 	bool bResult = true;
 	
-	for (int32 i = Batch.BindingsBegin; i != Batch.BindingsEnd; i++)
+	if (SourceStructView.IsValid() && TargetStructView.IsValid())
 	{
-		const FStateTreePropertyCopy& Copy = PropertyCopies[i];
-		// Copies that fail to be resolved (i.e. property path does not resolve, types changed) will be marked as None, skip them.
-		if (Copy.Type == EStateTreePropertyCopyType::None)
-		{
-			continue;
-		}
-		
-		const FStateTreeDataView SourceStructView = SourceStructViews[Copy.SourceStructIndex.Get()];
-		if (SourceStructView.IsValid())
-		{
-			check(SourceStructView.GetStruct() == SourceStructs[Copy.SourceStructIndex.Get()].Struct
-				|| (SourceStructView.GetStruct() && SourceStructView.GetStruct()->IsChildOf(SourceStructs[Copy.SourceStructIndex.Get()].Struct)));
-				
-			uint8* SourceAddress = GetAddress(SourceStructView, Copy.SourceIndirection, Copy.SourceLeafProperty);
-			uint8* TargetAddress = GetAddress(TargetStructView, Copy.TargetIndirection, Copy.TargetLeafProperty);
+		check(SourceStructView.GetStruct() == Copy.SourceStructType
+			|| (SourceStructView.GetStruct() && SourceStructView.GetStruct()->IsChildOf(Copy.SourceStructType)));
 			
-			if (SourceAddress != nullptr && TargetAddress != nullptr)
-			{
-				PerformCopy(Copy, SourceAddress, TargetAddress);
-			}
-			else
-			{
-				bResult = false;
-			}
+		uint8* SourceAddress = GetAddress(SourceStructView, Copy.SourceIndirection, Copy.SourceLeafProperty);
+		uint8* TargetAddress = GetAddress(TargetStructView, Copy.TargetIndirection, Copy.TargetLeafProperty);
+		
+		if (SourceAddress != nullptr && TargetAddress != nullptr)
+		{
+			PerformCopy(Copy, SourceAddress, TargetAddress);
 		}
 		else
 		{
 			bResult = false;
 		}
 	}
+	else
+	{
+		bResult = false;
+	}
 
 	return bResult;
 }
@@ -1165,13 +1185,13 @@ void FStateTreePropertyBindings::DebugPrintInternalLayout(FString& OutString) co
 	}
 
 	/** Array of property copies */
-	OutString += FString::Printf(TEXT("\nPropertyCopies (%d)\n  [ %-7s | %-4s | %-4s | %-10s | %-7s | %-4s | %-4s | %-10s | %-7s | %-20s | %-4s ]\n"), PropertyCopies.Num(),
+	OutString += FString::Printf(TEXT("\nPropertyCopies (%d)\n  [ %-7s | %-4s | %-4s | %-10s | %-7s | %-4s | %-4s | %-10s | %-10s | %-20s | %-4s ]\n"), PropertyCopies.Num(),
 		TEXT("Src Idx"), TEXT("Off."), TEXT("Next"), TEXT("Type"),
 		TEXT("Tgt Idx"), TEXT("Off."), TEXT("Next"), TEXT("Type"),
-		TEXT("Struct"), TEXT("Copy Type"), TEXT("Size"));
+		TEXT("Source"), TEXT("Copy Type"), TEXT("Size"));
 	for (const FStateTreePropertyCopy& PropertyCopy : PropertyCopies)
 	{
-		OutString += FString::Printf(TEXT("  | %7d | %4d | %4d | %-10s | %7d | %4d | %4d | %-10s | %7d | %-20s | %4d |\n"),
+		OutString += FString::Printf(TEXT("  | %7d | %4d | %4d | %-10s | %7d | %4d | %4d | %-10s | %10s | %-20s | %4d |\n"),
 					PropertyCopy.SourceIndirection.ArrayIndex.Get(),
 					PropertyCopy.SourceIndirection.Offset,
 					PropertyCopy.SourceIndirection.NextIndex.Get(),
@@ -1180,7 +1200,7 @@ void FStateTreePropertyBindings::DebugPrintInternalLayout(FString& OutString) co
 					PropertyCopy.TargetIndirection.Offset,
 					PropertyCopy.TargetIndirection.NextIndex.Get(),
 					*UEnum::GetDisplayValueAsText(PropertyCopy.TargetIndirection.Type).ToString(),
-					PropertyCopy.SourceStructIndex.Get(),
+					*PropertyCopy.SourceDataHandle.Describe(),
 					*UEnum::GetDisplayValueAsText(PropertyCopy.Type).ToString(),
 					PropertyCopy.CopySize);
 	}
diff --git a/Engine/Plugins/Runtime/StateTree/Source/StateTreeModule/Private/StateTreeTypes.cpp b/Engine/Plugins/Runtime/StateTree/Source/StateTreeModule/Private/StateTreeTypes.cpp
--- a/Engine/Plugins/Runtime/StateTree/Source/StateTreeModule/Private/StateTreeTypes.cpp
+++ b/Engine/Plugins/Runtime/StateTree/Source/StateTreeModule/Private/StateTreeTypes.cpp
@@ -13,6 +13,8 @@ const FStateTreeStateHandle FStateTreeStateHandle::Failed = FStateTreeStateHandl
 const FStateTreeStateHandle FStateTreeStateHandle::Stopped = FStateTreeStateHandle(StoppedIndex);
 const FStateTreeStateHandle FStateTreeStateHandle::Root = FStateTreeStateHandle(0);
 
+const FStateTreeDataHandle FStateTreeDataHandle::Invalid = FStateTreeDataHandle();
+
 const FStateTreeIndex16 FStateTreeIndex16::Invalid = FStateTreeIndex16();
 const FStateTreeIndex8 FStateTreeIndex8::Invalid = FStateTreeIndex8();
 
diff --git a/Engine/Plugins/Runtime/StateTree/Source/StateTreeModule/Public/StateTree.h b/Engine/Plugins/Runtime/StateTree/Source/StateTreeModule/Public/StateTree.h
--- a/Engine/Plugins/Runtime/StateTree/Source/StateTreeModule/Public/StateTree.h
+++ b/Engine/Plugins/Runtime/StateTree/Source/StateTreeModule/Public/StateTree.h
@@ -37,6 +37,8 @@ struct STATETREEMODULE_API FStateTreeCustomVersion
 		ChangedBindingsRepresentation,
 		// Added guid to transitions
 		AddedTransitionIds,
+		// Added data handles
+		AddedDataHandlesIds,
 
 		// -----<new versions can be added above this line>-------------------------------------------------
 		VersionPlusOne,
@@ -90,8 +92,8 @@ public:
 	/** @return Shared instance data. */
 	TSharedPtr<FStateTreeInstanceData> GetSharedInstanceData() const;
 
-	/** @return Number of data views required for StateTree execution (Evaluators, Tasks, Conditions, External data). */
-	int32 GetNumDataViews() const { return NumDataViews; }
+	/** @return Number of context data views required for StateTree execution (Tree params, context data, External data). */
+	int32 GetNumContextDataViews() const { return NumContextDataViews; }
 
 	/** @return List of external data required by the state tree */
 	TConstArrayView<FStateTreeExternalDataDesc> GetExternalDataDescs() const { return ExternalDataDescs; }
@@ -141,6 +143,9 @@ public:
 	/** @return Id of the transition matching a given runtime transition index; invalid Id if transition not found. */
 	FGuid GetTransitionIdFromIndex(const FStateTreeIndex16 Index) const;	
 
+	UE_DEPRECATED(5.4, "Replaced with GetNumContextDataViews() which contains context data and external data only.")
+	int32 GetNumDataViews() const { return 0; }
+
 #if WITH_EDITOR
 	/** Resets the compiled data to empty. */
 	void ResetCompiled();
@@ -257,8 +262,16 @@ private:
 
 	/** Data view index of the tree Parameters */
 	UPROPERTY()
-	FStateTreeIndex8 ParametersDataViewIndex = FStateTreeIndex8::Invalid;
+	FStateTreeDataHandle ParametersDataHandle = FStateTreeDataHandle::Invalid; 
+
+	/** Number of context data. */
+	UPROPERTY()
+	uint16 NumContextData = 0;
 
+	/** Number of global instance data. */
+	UPROPERTY()
+	uint16 NumGlobalInstanceData = 0;
+	
 	/** Index of first evaluator in Nodes. */
 	UPROPERTY()
 	uint16 EvaluatorsBegin = 0;
@@ -289,10 +302,10 @@ private:
 	UPROPERTY(Transient)
 	int32 ExternalDataBaseIndex = 0;
 
-	/** Total number of data views, created during linking. */
+	/** Total number of context data views, created during linking. */
 	UPROPERTY(Transient)
-	int32 NumDataViews = 0;
-
+	int32 NumContextDataViews = 0;
+	
 	/** True if the StateTree was linked successfully. */
 	bool bIsLinked = false;
 
diff --git a/Engine/Plugins/Runtime/StateTree/Source/StateTreeModule/Public/StateTreeExecutionContext.h b/Engine/Plugins/Runtime/StateTree/Source/StateTreeModule/Public/StateTreeExecutionContext.h
--- a/Engine/Plugins/Runtime/StateTree/Source/StateTreeModule/Public/StateTreeExecutionContext.h
+++ b/Engine/Plugins/Runtime/StateTree/Source/StateTreeModule/Public/StateTreeExecutionContext.h
@@ -58,7 +58,7 @@ public:
 	void SetParameters(const FInstancedPropertyBag& Parameters);
 	
 	/** @return the StateTree asset in use. */
-	const UStateTree* GetStateTree() const { return &StateTree; }
+	const UStateTree* GetStateTree() const { return &RootStateTree; }
 
 	/** @return const references to the instance data in use, or nullptr if the context is not valid. */
 	const FStateTreeInstanceData* GetInstanceData() const { return &InstanceData; }
@@ -78,7 +78,7 @@ public:
 	UWorld* GetWorld() const { return Owner.GetWorld(); };
 
 	/** @return True of the the execution context is valid and initialized. */ 
-	bool IsValid() const { return StateTree.IsReadyToRun(); }
+	bool IsValid() const { return RootStateTree.IsReadyToRun(); }
 	
 	/** Start executing. */
 	EStateTreeRunStatus Start();
@@ -104,7 +104,15 @@ public:
 	EStateTreeRunStatus GetLastTickStatus() const;
 
 	/** @return reference to the list of currently active states. */
-	const FStateTreeActiveStates& GetActiveStates() const;
+	UE_DEPRECATED(5.4, "Use GetActiveFrames() instead.")
+	const FStateTreeActiveStates& GetActiveStates() const
+	{
+		static FStateTreeActiveStates Dummy;
+		return Dummy;
+	}
+
+	/** @return reference to the list of currently active frames and states. */
+	TConstArrayView<FStateTreeExecutionFrame> GetActiveFrames() const;
 
 #if WITH_GAMEPLAY_DEBUGGER
 	/** @return Debug string describing the current state of the execution */
@@ -166,19 +174,19 @@ public:
 	/** @return Pointer to a State or null if state not found */ 
 	const FCompactStateTreeState* GetStateFromHandle(const FStateTreeStateHandle StateHandle) const
 	{
-		return StateTree.GetStateFromHandle(StateHandle);
+		return RootStateTree.GetStateFromHandle(StateHandle);
 	}
 
 	/** @return Array view to external data descriptors associated with this context. Note: Init() must be called before calling this method. */
 	TConstArrayView<FStateTreeExternalDataDesc> GetExternalDataDescs() const
 	{
-		return StateTree.ExternalDataDescs;
+		return RootStateTree.ExternalDataDescs;
 	}
 
 	/** @return Array view to named external data descriptors associated with this context. Note: Init() must be called before calling this method. */
 	TConstArrayView<FStateTreeExternalDataDesc> GetContextDataDescs() const
 	{
-		return StateTree.GetContextDataDescs();
+		return RootStateTree.GetContextDataDescs();
 	}
 
 	/** @return True if all required external data pointers are set. */ 
@@ -187,7 +195,7 @@ public:
 	/** @return Handle to external data of type InStruct, or invalid handle if struct not found. */ 
 	FStateTreeExternalDataHandle GetExternalDataHandleByStruct(const UStruct* InStruct) const
 	{
-		const FStateTreeExternalDataDesc* DataDesc = StateTree.ExternalDataDescs.FindByPredicate([InStruct](const FStateTreeExternalDataDesc& Item) { return Item.Struct == InStruct; });
+		const FStateTreeExternalDataDesc* DataDesc = RootStateTree.ExternalDataDescs.FindByPredicate([InStruct](const FStateTreeExternalDataDesc& Item) { return Item.Struct == InStruct; });
 		return DataDesc != nullptr ? DataDesc->Handle : FStateTreeExternalDataHandle::Invalid;
 	}
 
@@ -195,7 +203,8 @@ public:
 	void SetExternalData(const FStateTreeExternalDataHandle Handle, FStateTreeDataView DataView)
 	{
 		check(Handle.IsValid());
-		DataViews[Handle.DataViewIndex.Get()] = DataView;
+		check(Handle.DataHandle.GetSource() == EStateTreeDataSourceType::ContextData);
+		ContextDataViews[Handle.DataHandle.GetIndex()] = DataView;
 	}
 
 	/**
@@ -207,8 +216,9 @@ public:
 	typename T::DataType& GetExternalData(const T Handle) const
 	{
 		check(Handle.IsValid());
-		checkSlow(StateTree.ExternalDataDescs[Handle.DataViewIndex.Get() - StateTree.ExternalDataBaseIndex].Requirement != EStateTreeExternalDataRequirement::Optional); // Optionals should query pointer instead.
-		return DataViews[Handle.DataViewIndex.Get()].template GetMutable<typename T::DataType>();
+		check(RootStateTree.ExternalDataDescs[Handle.DataHandle.GetIndex() - RootStateTree.ExternalDataBaseIndex].Requirement != EStateTreeExternalDataRequirement::Optional); // Optionals should query pointer instead.
+		check(Handle.DataHandle.GetSource() == EStateTreeDataSourceType::ContextData);
+		return ContextDataViews[Handle.DataHandle.GetIndex()].template GetMutable<typename T::DataType>();
 	}
 
 	/**
@@ -219,14 +229,20 @@ public:
 	template <typename T>
 	typename T::DataType* GetExternalDataPtr(const T Handle) const
 	{
-		return Handle.IsValid() ? DataViews[Handle.DataViewIndex.Get()].template GetMutablePtr<typename T::DataType>() : nullptr;
+		if (Handle.IsValid())
+		{
+			check(Handle.DataHandle.GetSource() == EStateTreeDataSourceType::ContextData);
+			return ContextDataViews[Handle.DataHandle.GetIndex()].template GetMutablePtr<typename T::DataType>();
+		}
+		return nullptr;
 	}
 
 	FStateTreeDataView GetExternalDataView(const FStateTreeExternalDataHandle Handle)
 	{
 		if (Handle.IsValid())
 		{
-			return DataViews[Handle.DataViewIndex.Get()];
+			check(Handle.DataHandle.GetSource() == EStateTreeDataSourceType::ContextData);
+			return ContextDataViews[Handle.DataHandle.GetIndex()];
 		}
 		return FStateTreeDataView();
 	}
@@ -235,14 +251,16 @@ public:
 	template <typename T>
 	T* GetInstanceDataPtr(const FStateTreeNodeBase& Node) const
 	{
-		return DataViews[Node.DataViewIndex.Get()].template GetMutablePtr<T>();
+		check(CurrentlyProcessedFrame)
+		return GetDataView(*CurrentlyProcessedFrame, Node.InstanceDataHandle).template GetMutablePtr<T>();
 	}
 
 	/** @returns reference to the instance data of specified node. */
 	template <typename T>
 	T& GetInstanceData(const FStateTreeNodeBase& Node) const
 	{
-		return DataViews[Node.DataViewIndex.Get()].template GetMutable<T>();
+		check(CurrentlyProcessedFrame);
+		return GetDataView(*CurrentlyProcessedFrame, Node.InstanceDataHandle).template GetMutable<T>();
 	}
 
 	/** @returns reference to the instance data of specified node. Infers the instance data type from the node's FInstanceDataType. */
@@ -250,7 +268,8 @@ public:
 	typename T::FInstanceDataType& GetInstanceData(const T& Node) const
 	{
 		static_assert(TIsDerivedFrom<T, FStateTreeNodeBase>::IsDerived, "Expecting Node to derive from FStateTreeNodeBase.");
-		return DataViews[Node.DataViewIndex.Get()].template GetMutable<typename T::FInstanceDataType>();
+		check(CurrentlyProcessedFrame);
+		return GetDataView(*CurrentlyProcessedFrame, Node.InstanceDataHandle).template GetMutable<typename T::FInstanceDataType>();
 	}
 
 	/** @returns reference to instance data struct that can be passed to lambdas. See TStateTreeInstanceDataStructRef for usage. */
@@ -258,7 +277,9 @@ public:
 	TStateTreeInstanceDataStructRef<typename T::FInstanceDataType> GetInstanceDataStructRef(const T& Node) const
 	{
 		static_assert(TIsDerivedFrom<T, FStateTreeNodeBase>::IsDerived, "Expecting Node to derive from FStateTreeNodeBase.");
-		return TStateTreeInstanceDataStructRef<typename T::FInstanceDataType>(InstanceData, DataViews[Node.DataViewIndex.Get()].template GetMutable<typename T::FInstanceDataType>());
+		check(CurrentlyProcessedFrame);
+		typename T::FInstanceDataType& Data = GetDataView(*CurrentlyProcessedFrame, Node.InstanceDataHandle).template GetMutable<typename T::FInstanceDataType>();
+		return TStateTreeInstanceDataStructRef<typename T::FInstanceDataType>(InstanceData, Data);
 	}
 
 	/**
@@ -272,6 +293,9 @@ public:
 
 protected:
 
+	/** Max number of execution frames handled during state selection. */
+	static constexpr int32 MaxExecutionFrames = 8;
+	
 #if WITH_STATETREE_DEBUGGER
 	FStateTreeInstanceDebugId GetInstanceDebugId() const;
 #endif // WITH_STATETREE_DEBUGGER
@@ -279,13 +303,10 @@ protected:
 	/** @return Prefix that will be used by STATETREE_LOG and STATETREE_CLOG, Owner name by default. */
 	virtual FString GetInstanceDescription() const;
 
-	UE_DEPRECATED(5.2, "Use BeginDelayedTransition() instead.")
-	virtual void BeginGatedTransition(const FStateTreeExecutionState& Exec) final {};
-	
 	/** Callback when delayed transition is triggered. Contexts that are event based can use this to trigger a future event. */
 	virtual void BeginDelayedTransition(const FStateTreeTransitionDelayedState& DelayedState) {};
 
-	void UpdateInstanceData(const FStateTreeActiveStates& CurrentActiveStates, const FStateTreeActiveStates& NextActiveStates);
+	void UpdateInstanceData(TConstArrayView<FStateTreeExecutionFrame> CurrentActiveFrames, TArrayView<FStateTreeExecutionFrame> NextActiveFrames);
 
 	/**
 	 * Handles logic for entering State. EnterState is called on new active Evaluators and Tasks that are part of the re-planned tree.
@@ -293,7 +314,7 @@ protected:
 	 * and still active after the transition will remain intact.
 	 * @return Run status returned by the tasks.
 	 */
-	EStateTreeRunStatus EnterState(const FStateTreeTransitionResult& Transition);
+	EStateTreeRunStatus EnterState(FStateTreeTransitionResult& Transition);
 
 	/**
 	 * Handles logic for exiting State. ExitState is called on current active Evaluators and Tasks that are part of the re-planned tree.
@@ -333,17 +354,21 @@ protected:
 	 * Checks all conditions at given range
 	 * @return True if all conditions pass.
 	 */
-	bool TestAllConditions(const int32 ConditionsOffset, const int32 ConditionsNum);
+	bool TestAllConditions(const FStateTreeExecutionFrame& CurrentFrame, const int32 ConditionsOffset, const int32 ConditionsNum);
 
 	/**
 	 * Requests transition to a specified state with specified priority.
 	 */
-	bool RequestTransition(const FStateTreeStateHandle NextState, const EStateTreeTransitionPriority Priority, const EStateTreeSelectionFallback Fallback = EStateTreeSelectionFallback::None);
+	bool RequestTransition(
+		const FStateTreeExecutionFrame& CurrentFrame,
+		const FStateTreeStateHandle NextState,
+		const EStateTreeTransitionPriority Priority,
+		const EStateTreeSelectionFallback Fallback = EStateTreeSelectionFallback::None);
 
 	/**
 	 * Sets up NextTransition based on the provided parameters and the current execution status. 
 	 */
-	void SetupNextTransition(const FStateTreeStateHandle NextState, const EStateTreeTransitionPriority Priority);
+	void SetupNextTransition(const FStateTreeExecutionFrame& CurrentFrame, const FStateTreeStateHandle NextState, const EStateTreeTransitionPriority Priority);
 
 	/**
 	 * Triggers transitions based on current run status. CurrentStatus is used to select which transitions events are triggered.
@@ -354,125 +379,133 @@ protected:
 	 */
 	bool TriggerTransitions();
 
-	/**
-	 * Traverses the ActiveStates from StartStateIndex to 0 and returns first linked state.
-	 * @return Parent linked state, or invalid state if no linked state found. 
-	 */
-	FStateTreeStateHandle GetParentLinkedStateHandle(const FStateTreeActiveStates& ActiveStates, const int32 StartStateIndex) const;
-
-	FStateTreeStateHandle GetParentLinkedStateHandle(const FStateTreeActiveStates& ActiveStates, const FStateTreeStateHandle StartStateHandle) const;
-
 	/**
 	 * Runs state selection logic starting at the specified state, walking towards the leaf states.
 	 * If a state cannot be selected, false is returned. 
 	 * If NextState is a selector state, SelectStateInternal is called recursively (depth-first) to all child states (where NextState will be one of child states).
 	 * If NextState is a leaf state, the active states leading from root to the leaf are returned.
+	 * @param CurrentFrame The frame where the NextState is valid. 
 	 * @param NextState The state which we try to select next.
-	 * @param OutNewActiveStates Active states that got selected.
-	 * @param VisitedStates States visited so far during selection (used for detecting selection loops)
+	 * @param OutNextActiveFrames Active frames and states that got selected.
 	 * @param Fallback selection behavior to execute if it fails to select the desired state
 	 * @return True if succeeded to select new active states.
 	 */
-	bool SelectState(const FStateTreeStateHandle NextState, FStateTreeActiveStates& OutNewActiveStates, FStateTreeActiveStates& VisitedStates, const EStateTreeSelectionFallback Fallback = EStateTreeSelectionFallback::None);
+	bool SelectState(
+		const FStateTreeExecutionFrame& CurrentFrame,
+		const FStateTreeStateHandle NextState,
+		TArray<FStateTreeExecutionFrame, TFixedAllocator<MaxExecutionFrames>>& OutNextActiveFrames,
+		const EStateTreeSelectionFallback Fallback = EStateTreeSelectionFallback::None);
 
 	/**
 	 * Used internally to do the recursive part of the SelectState().
 	 */
-	bool SelectStateInternal(const FStateTreeStateHandle NextState, FStateTreeActiveStates& OutNewActiveStates, FStateTreeActiveStates& VisitedStates);
+	bool SelectStateInternal(
+		FStateTreeExecutionFrame& CurrentFrame,
+		const FStateTreeStateHandle NextState,
+		TArray<FStateTreeExecutionFrame, TFixedAllocator<MaxExecutionFrames>>& OutNextActiveFrames);
 
 	/** @return StateTree execution state from the instance storage. */
 	FStateTreeExecutionState& GetExecState()
 	{
-		return InstanceData.GetMutableStruct(0).Get<FStateTreeExecutionState>();
+		return *InstanceData.GetMutableExecutionState();
 	}
 
 	/** @return const StateTree execution state from the instance storage. */
 	const FStateTreeExecutionState& GetExecState() const
 	{
-		return InstanceData.GetStruct(0).Get<const FStateTreeExecutionState>();
+		return *InstanceData.GetExecutionState();
 	}
 
-	/** Sets up parameter data view for a linked state and copies bound properties. */
-	void UpdateLinkedStateParameters(const FCompactStateTreeState& State, const int32 ParameterInstanceIndex);
-
-	/** Sets up parameter data view for subtree state. */
-	void UpdateSubtreeStateParameters(const FCompactStateTreeState& State);
-	
 	/** @return String describing state status for logging and debug. */
 	FString GetStateStatusString(const FStateTreeExecutionState& ExecState) const;
 
 	/** @return String describing state name for logging and debug. */
-	FString GetSafeStateName(const FStateTreeStateHandle State) const;
+	FString GetSafeStateName(const FStateTreeExecutionFrame& CurrentFrame, const FStateTreeStateHandle State) const;
 
 	/** @return String describing full path of an activate state for logging and debug. */
-	FString DebugGetStatePath(const FStateTreeActiveStates& ActiveStates, const int32 ActiveStateIndex = INDEX_NONE) const;
+	FString DebugGetStatePath(TConstArrayView<FStateTreeExecutionFrame> ActiveFrames, const FStateTreeExecutionFrame* CurrentFrame = nullptr, const int32 ActiveStateIndex = INDEX_NONE) const;
 
 	/** @return String describing all events that are currently being processed  for logging and debug. */
 	FString DebugGetEventsAsString() const;
+
+	/** @return data view of the specified handle relative to given frame. */
+	FStateTreeDataView GetDataView(const FStateTreeExecutionFrame& CurrentFrame, const FStateTreeDataHandle Handle) const;
+
+	/** @return data view of the specified handle relative to given frame. */
+	bool IsHandleSourceValid(const FStateTreeExecutionFrame& CurrentFrame, const FStateTreeDataHandle Handle) const;
 	
-	/** Helper function to update struct or object dataview of a node. */
-	template<typename T>
-	void SetNodeDataView(T& Node, int32& InstanceStructIndex, int32& InstanceObjectIndex)
-	{
-		if (Node.bInstanceIsObject)
-		{
-			DataViews[Node.DataViewIndex.Get()] = InstanceData.GetMutableObject(InstanceObjectIndex);
-			InstanceObjectIndex++;
-		}
-		else
-		{
-			DataViews[Node.DataViewIndex.Get()] = InstanceData.GetMutableStruct(InstanceStructIndex);
-			InstanceStructIndex++;
-		}
-	}
-	
+	bool CopyBatch(const FStateTreeExecutionFrame& CurrentFrame, const FStateTreeDataView TargetView, const FStateTreeIndex16 BindingsBatch) const;
+	bool CopyBatchWithValidation(const FStateTreeExecutionFrame& CurrentFrame, const FStateTreeDataView TargetView, const FStateTreeIndex16 BindingsBatch) const;
 
 	/** Owner of the instance data. */
 	UObject& Owner;
 
 	/** The StateTree asset the context is initialized for */
-	const UStateTree& StateTree;
+	const UStateTree& RootStateTree;
 
 	/** Instance data used during current tick. */
 	FStateTreeInstanceData& InstanceData;
 
-	/** Array of data pointers (external data, tasks, evaluators, conditions), used during evaluation. Initialized to match the number of items in the asset. */
-	TArray<FStateTreeDataView, TConcurrentLinearArrayAllocator<FDefaultBlockAllocationTag>> DataViews;
+	/** Data storage of the instance data, cached for less indirections. */
+	FStateTreeInstanceStorage* InstanceDataStorage = nullptr;
 
+	/** Data view of the context data. */
+	TArray<FStateTreeDataView, TConcurrentLinearArrayAllocator<FDefaultBlockAllocationTag>> ContextDataViews;
+	
 	/** Events to process in current tick. */
 	TArray<FStateTreeEvent, TConcurrentLinearArrayAllocator<FDefaultBlockAllocationTag>> EventsToProcess;
 
-	/** Shared instance data for the duration of the context. */
-	TSharedPtr<FStateTreeInstanceData> SharedInstanceData;
-
+	
 	/** Next transition, used by RequestTransition(). */
 	FStateTreeTransitionResult NextTransition;
 
 	/** Structure describing the origin of the state transition that caused the state change. */
 	FStateTreeTransitionSource NextTransitionSource;
 
-	/** Current state we're processing, or invalid if not applicable. */
-	FStateTreeStateHandle CurrentlyProcessedState;
+	
+	/** Index of current frame we're processing. */
+	int32 CurrentlyProcessedFrameIndex = INDEX_NONE;
 
-	/** True if transitions are allowed to be requested directly instead of buffering. */
-	bool bAllowDirectTransitions = false;
+	/** Current frame we're processing. */
+	const FStateTreeExecutionFrame* CurrentlyProcessedFrame = nullptr; 
 
-	/** Helper struct to track when it is allowed to request transitions. */
-	struct FAllowDirectTransitionsScope
+	/** Pointer to the shared instance data of the current frame we're processing.  */
+	FStateTreeInstanceStorage* CurrentlyProcessedSharedInstanceStorage = nullptr;
+
+	/** Helper struct to track currently processed frame. */
+	struct FCurrentlyProcessedFrameScope
 	{
-		FAllowDirectTransitionsScope(FStateTreeExecutionContext& InContext)
+		FCurrentlyProcessedFrameScope(FStateTreeExecutionContext& InContext, const FStateTreeExecutionFrame& CurrentFrame, int32 FrameIndex)
 			: Context(InContext)
 		{
-			Context.bAllowDirectTransitions = true;
+			check(CurrentFrame.StateTree);
+			FStateTreeInstanceStorage* SharedInstanceDataStorage = &CurrentFrame.StateTree->GetSharedInstanceData()->GetMutableStorage();
+
+			SavedFrameIndex = Context.CurrentlyProcessedFrameIndex;
+			SavedFrame = Context.CurrentlyProcessedFrame;
+			SavedSharedInstanceDataStorage = Context.CurrentlyProcessedSharedInstanceStorage;
+			Context.CurrentlyProcessedFrameIndex = FrameIndex;
+			Context.CurrentlyProcessedFrame = &CurrentFrame;
+			Context.CurrentlyProcessedSharedInstanceStorage = SharedInstanceDataStorage;
 		}
 
-		~FAllowDirectTransitionsScope()
+		~FCurrentlyProcessedFrameScope()
 		{
-			Context.bAllowDirectTransitions = false;
+			Context.CurrentlyProcessedFrameIndex = SavedFrameIndex;
+			Context.CurrentlyProcessedFrame = SavedFrame;
+			Context.CurrentlyProcessedSharedInstanceStorage = SavedSharedInstanceDataStorage;
 		}
-		
+
+	private:
 		FStateTreeExecutionContext& Context;
+		int32 SavedFrameIndex = 0;
+		FStateTreeInstanceStorage* SavedSharedInstanceDataStorage = nullptr;
+		const FStateTreeExecutionFrame* SavedFrame = nullptr;
 	};
+
+	
+	/** Current state we're processing, or invalid if not applicable. */
+	FStateTreeStateHandle CurrentlyProcessedState;
 	
 	/** Helper struct to track currently processed state. */
 	struct FCurrentlyProcessedStateScope
@@ -480,14 +513,41 @@ protected:
 		FCurrentlyProcessedStateScope(FStateTreeExecutionContext& InContext, const FStateTreeStateHandle State)
 			: Context(InContext)
 		{
+			SavedState = Context.CurrentlyProcessedState;
 			Context.CurrentlyProcessedState = State;
 		}
 
 		~FCurrentlyProcessedStateScope()
 		{
-			Context.CurrentlyProcessedState = FStateTreeStateHandle::Invalid;
+			Context.CurrentlyProcessedState = SavedState;
 		}
-		
+
+	private:
+		FStateTreeExecutionContext& Context;
+		FStateTreeStateHandle SavedState = FStateTreeStateHandle::Invalid; 
+	};
+	
+
+	/** True if transitions are allowed to be requested directly instead of buffering. */
+	bool bAllowDirectTransitions = false;
+
+	/** Helper struct to track when it is allowed to request transitions. */
+	struct FAllowDirectTransitionsScope
+	{
+		FAllowDirectTransitionsScope(FStateTreeExecutionContext& InContext)
+			: Context(InContext)
+		{
+			bCachedAllowDirectTransitions = Context.bAllowDirectTransitions; 
+			Context.bAllowDirectTransitions = true;
+		}
+
+		~FAllowDirectTransitionsScope()
+		{
+			Context.bAllowDirectTransitions = bCachedAllowDirectTransitions;
+		}
+
+	private:
 		FStateTreeExecutionContext& Context;
+		bool bCachedAllowDirectTransitions = false;
 	};
 };
diff --git a/Engine/Plugins/Runtime/StateTree/Source/StateTreeModule/Public/StateTreeExecutionTypes.h b/Engine/Plugins/Runtime/StateTree/Source/StateTreeModule/Public/StateTreeExecutionTypes.h
--- a/Engine/Plugins/Runtime/StateTree/Source/StateTreeModule/Public/StateTreeExecutionTypes.h
+++ b/Engine/Plugins/Runtime/StateTree/Source/StateTreeModule/Public/StateTreeExecutionTypes.h
@@ -5,6 +5,8 @@
 
 #include "StateTreeExecutionTypes.generated.h"
 
+class UStateTree;
+
 /**
  * Enumeration for the different update phases.
  * This is used as context information when tracing debug events.
@@ -90,15 +92,30 @@
 {
 	GENERATED_BODY()
 
+	PRAGMA_DISABLE_DEPRECATION_WARNINGS
+	FStateTreeExternalDataHandle() = default;
+	FStateTreeExternalDataHandle(const FStateTreeExternalDataHandle& Other) = default;
+	FStateTreeExternalDataHandle(FStateTreeExternalDataHandle&& Other) = default;
+	FStateTreeExternalDataHandle& operator=(FStateTreeExternalDataHandle const& Other) = default;
+	FStateTreeExternalDataHandle& operator=(FStateTreeExternalDataHandle&& Other) = default;
+	PRAGMA_ENABLE_DEPRECATION_WARNINGS
+
 	static const FStateTreeExternalDataHandle Invalid;
 	
-	static bool IsValidIndex(const int32 Index) { return FStateTreeIndex16::IsValidIndex(Index); }
-	bool IsValid() const { return DataViewIndex.IsValid(); }
+	bool IsValid() const { return DataHandle.IsValid(); }
+
+	UE_DEPRECATED(5.4, "Index is deprecated, use DataHandle instead.")
+	static bool IsValidIndex(const int32 Index) { return FStateTreeDataHandle::IsValidIndex(Index); }
 
 	UPROPERTY()
-	FStateTreeIndex16 DataViewIndex = FStateTreeIndex16::Invalid;
+	FStateTreeDataHandle DataHandle = FStateTreeDataHandle::Invalid;
+
+#if WITH_EDITORONLY_DATA
+	UE_DEPRECATED(5.4, "Use DataHandle instead.")
+	UPROPERTY()
+	FStateTreeIndex16 DataViewIndex_DEPRECATED = FStateTreeIndex16::Invalid;
+#endif // WITH_EDITORONLY_DATA
 };
 
-
 /**
  * Handle to access an external struct or object.
@@ -201,9 +218,17 @@
 	}
 
 	/** Source state of the transition. Filled in by the StateTree execution context. */
-	UPROPERTY(EditDefaultsOnly, Category = "Default")
+	UPROPERTY()
 	FStateTreeStateHandle SourceState;
 
+	/** StateTree asset that was active when the transition was requested. Filled in by the StateTree execution context. */
+	UPROPERTY()
+	TObjectPtr<const UStateTree> SourceStateTree = nullptr;
+
+	/** Root state the execution frame where the transition was requested. Filled in by the StateTree execution context. */
+	UPROPERTY()
+	FStateTreeStateHandle SourceRootState = FStateTreeStateHandle::Invalid;
+
 	/** Target state of the transition. */
 	UPROPERTY(EditDefaultsOnly, Category = "Default")
 	FStateTreeStateHandle TargetState;
@@ -411,6 +436,7 @@
 		*this = {};
 	}
 
+	/** Describes where the transition originated. */
 	EStateTreeTransitionSourceType SourceType = EStateTreeTransitionSourceType::Unset;
 
 	/* Index of the transition if from predefined asset transitions, invalid otherwise */
@@ -424,63 +450,6 @@
 };
 
 
-/**
- * Describes a state tree transition. Source is the state where the transition started, Target describes the state where the transition pointed at,
- * and Next describes the selected state. The reason Transition and Next are different is that Transition state can be a selector state,
- * in which case the children will be visited until a leaf state is found, which will be the next state.
- */
-USTRUCT(BlueprintType)
-struct STATETREEMODULE_API FStateTreeTransitionResult
-{
-	GENERATED_BODY()
-
-	FStateTreeTransitionResult() = default;
-
-	void Reset()
-	{
-		CurrentActiveStates.Reset();
-		CurrentRunStatus = EStateTreeRunStatus::Unset;
-		TargetState = FStateTreeStateHandle::Invalid;
-		NextActiveStates.Reset();
-		CurrentState = FStateTreeStateHandle::Invalid;
-		ChangeType = EStateTreeStateChangeType::Changed;
-		Priority = EStateTreeTransitionPriority::None;
-	}
-	
-	/** Current active states, where the transition started. */
-	UPROPERTY(EditDefaultsOnly, Category = "Default", BlueprintReadOnly)
-	FStateTreeActiveStates CurrentActiveStates;
-
-	/** Current Run status. */
-	UPROPERTY(EditDefaultsOnly, Category = "Default", BlueprintReadOnly)
-	EStateTreeRunStatus CurrentRunStatus = EStateTreeRunStatus::Unset;
-
-	/** Transition source state */
-	UPROPERTY(EditDefaultsOnly, Category = "Default", BlueprintReadOnly)
-	FStateTreeStateHandle SourceState = FStateTreeStateHandle::Invalid;
-
-	/** Transition target state */
-	UPROPERTY(EditDefaultsOnly, Category = "Default", BlueprintReadOnly)
-	FStateTreeStateHandle TargetState = FStateTreeStateHandle::Invalid;
-
-	/** States selected as result of the transition. */
-	UPROPERTY(EditDefaultsOnly, Category = "Default", BlueprintReadOnly)
-	FStateTreeActiveStates NextActiveStates;
-
-	/** The current state being executed. On enter/exit callbacks this is the state of the task. */
-	UPROPERTY(EditDefaultsOnly, Category = "Default", BlueprintReadOnly)
-	FStateTreeStateHandle CurrentState = FStateTreeStateHandle::Invalid;
-
-	/** If the change type is Sustained, then the CurrentState was reselected, or if Changed then the state was just activated. */
-	UPROPERTY(EditDefaultsOnly, Category = "Default", BlueprintReadOnly)
-	EStateTreeStateChangeType ChangeType = EStateTreeStateChangeType::Changed; 
-
-	/** Priority of the transition that caused the state change. */
-	UPROPERTY(EditDefaultsOnly, Category = "Default", BlueprintReadOnly)
-	EStateTreeTransitionPriority Priority = EStateTreeTransitionPriority::None;
-};
-
-
 #if WITH_STATETREE_DEBUGGER
 struct STATETREEMODULE_API FStateTreeInstanceDebugId
 {
@@ -521,52 +490,235 @@
 };
 #endif // WITH_STATETREE_DEBUGGER
 
+/** Describes current state of a delayed transition. */
 USTRUCT()
+struct STATETREEMODULE_API FStateTreeTransitionDelayedState
+{
+	GENERATED_BODY()
+
+	UPROPERTY()
+	TObjectPtr<const UStateTree> StateTree = nullptr;
+
+	UPROPERTY()
+	FStateTreeIndex16 TransitionIndex = FStateTreeIndex16::Invalid;
+
+	UPROPERTY()
+	float TimeLeft = 0.0f;
+};
+
+/** Describes an active branch of a State Tree. */
+USTRUCT(BlueprintType)
+struct STATETREEMODULE_API FStateTreeExecutionFrame
+{
+	GENERATED_BODY()
+
+	bool IsSameFrame(const FStateTreeExecutionFrame& OtherFrame) const
+	{
+		return StateTree == OtherFrame.StateTree && RootState == OtherFrame.RootState;
+	}
+	
+	/** The State Tree used for ticking this frame. */
+	UPROPERTY()
+	TObjectPtr<const UStateTree> StateTree = nullptr;
+
+	/** The root state of the frame (e.g. Root state or a subtree). */
+	UPROPERTY()
+	FStateTreeStateHandle RootState = FStateTreeStateHandle::Root; 
+	
+	/** Active states in this frame */
+	UPROPERTY()
+	FStateTreeActiveStates ActiveStates;
+
+	/** Index within the instance data to the first global instance data (e.g. global tasks) */
+	UPROPERTY()
+	FStateTreeIndex16 GlobalInstanceIndexBase = FStateTreeIndex16::Invalid;
+
+	/** Index within the instance data to the first active state's instance data (e.g. tasks) */
+	UPROPERTY()
+	FStateTreeIndex16 ActiveInstanceIndexBase = FStateTreeIndex16::Invalid;
+
+	/** Index in the active instance data for the parameters of the root state. */
+	UPROPERTY()
+	FStateTreeIndex16 StateParameterDataIndex = FStateTreeIndex16::Invalid;
+
+	/** If true, the global tasks of the State Tree should be handle in this frame. */
+	UPROPERTY()
+	uint8 bIsGlobalFrame : 1 = false;
+};
+
+/** Describes the execution state of the current State Tree instance. */
+USTRUCT()
 struct STATETREEMODULE_API FStateTreeExecutionState
 {
 	GENERATED_BODY()
 
+PRAGMA_DISABLE_DEPRECATION_WARNINGS
+	FStateTreeExecutionState() = default;
+	FStateTreeExecutionState(const FStateTreeExecutionState&) = default;
+	FStateTreeExecutionState(FStateTreeExecutionState&&) = default;
+	FStateTreeExecutionState& operator=(const FStateTreeExecutionState&) = default;
+	FStateTreeExecutionState& operator=(FStateTreeExecutionState&&) = default;
+PRAGMA_ENABLE_DEPRECATION_WARNINGS
+	
+	void Reset()
+	{
+		ActiveFrames.Reset();
+#if WITH_STATETREE_DEBUGGER
+		InstanceDebugId.Reset();
+#endif
+		EnterStateFailedTaskIndex = FStateTreeIndex16::Invalid;
+		LastTickStatus = EStateTreeRunStatus::Failed;
+		TreeRunStatus = EStateTreeRunStatus::Unset;
+		RequestedStop = EStateTreeRunStatus::Unset;
+		CurrentPhase = EStateTreeUpdatePhase::Unset;
+		CompletedStateHandle = FStateTreeStateHandle::Invalid;
+		StateChangeCount = 0;
+	}
+
 	/** @returns Delayed transition state for a specific transition, or nullptr if it does not exists. */
-	FStateTreeTransitionDelayedState* FindDelayedTransition(const FStateTreeIndex16 TransitionIndex)
+	FStateTreeTransitionDelayedState* FindDelayedTransition(const UStateTree* OwnerStateTree, const FStateTreeIndex16 TransitionIndex)
 	{
-		return DelayedTransitions.FindByPredicate([TransitionIndex](const FStateTreeTransitionDelayedState& State){ return State.TransitionIndex == TransitionIndex; });
+		return DelayedTransitions.FindByPredicate([OwnerStateTree, TransitionIndex](const FStateTreeTransitionDelayedState& TransitionState)
+		{
+			return TransitionState.StateTree == OwnerStateTree && TransitionState.TransitionIndex == TransitionIndex;
+		});
 	}
 
-	/** Currently active states */
-	FStateTreeActiveStates ActiveStates;
+	/** Currently active frames (and states) */
+	UPROPERTY()
+	TArray<FStateTreeExecutionFrame> ActiveFrames;
+
+	/** Pending delayed transitions. */
+	UPROPERTY()
+	TArray<FStateTreeTransitionDelayedState> DelayedTransitions;
 
 #if WITH_STATETREE_DEBUGGER
 	/** Id for the active instance used for debugging. */
 	mutable FStateTreeInstanceDebugId InstanceDebugId;
 #endif
 
-	/** Index of the first task struct in the currently initialized instance data. */
-	FStateTreeIndex16 FirstTaskStructIndex = FStateTreeIndex16::Invalid;
-	
-	/** Index of the first task object in the currently initialized instance data. */
-	FStateTreeIndex16 FirstTaskObjectIndex = FStateTreeIndex16::Invalid;
-
 	/** The index of the task that failed during enter state. Exit state uses it to call ExitState() symmetrically. */
+	UPROPERTY()
+	FStateTreeIndex16 EnterStateFailedFrameIndex = FStateTreeIndex16::Invalid;
+
+	/** The index of the frame that failed during enter state. Exit state uses it to call ExitState() symmetrically. */
+	UPROPERTY()
 	FStateTreeIndex16 EnterStateFailedTaskIndex = FStateTreeIndex16::Invalid;
 
 	/** Result of last tick */
+	UPROPERTY()
 	EStateTreeRunStatus LastTickStatus = EStateTreeRunStatus::Failed;
 
 	/** Running status of the instance */
+	UPROPERTY()
 	EStateTreeRunStatus TreeRunStatus = EStateTreeRunStatus::Unset;
 
 	/** Completion status stored if Stop was called during the Tick and needed to be deferred. */
+	UPROPERTY()
 	EStateTreeRunStatus RequestedStop = EStateTreeRunStatus::Unset;
 
 	/** Current update phase used to validate reentrant calls to the main entry points of the execution context (i.e. Start, Stop, Tick). */
+	UPROPERTY()
 	EStateTreeUpdatePhase CurrentPhase = EStateTreeUpdatePhase::Unset;
 
 	/** Handle of the state that was first to report state completed (success or failure), used to trigger completion transitions. */
+	UPROPERTY()
+	FStateTreeIndex16 CompletedFrameIndex = FStateTreeIndex16::Invalid; 
+	
+	UPROPERTY()
 	FStateTreeStateHandle CompletedStateHandle = FStateTreeStateHandle::Invalid;
 
 	/** Number of times a new state has been changed. */
+	UPROPERTY()
 	uint16 StateChangeCount = 0;
 
-	/** Running time of the delayed transition */
-	TArray<FStateTreeTransitionDelayedState> DelayedTransitions;
+#if WITH_EDITORONLY_DATA
+	UE_DEPRECATED(5.3, "Use DataHandle instead.")
+	UPROPERTY()
+	FStateTreeActiveStates CurrentActiveStates_DEPRECATED;
+	
+	/** Index of the first task struct in the currently initialized instance data. */
+	UE_DEPRECATED(5.4, "Superceded by State Tree Data Handles.")
+	FStateTreeIndex16 FirstTaskStructIndex_DEPRECATED = FStateTreeIndex16::Invalid;
+	
+	/** Index of the first task object in the currently initialized instance data. */
+	UE_DEPRECATED(5.4, "Superceded by State Tree Data Handles.")
+	FStateTreeIndex16 FirstTaskObjectIndex_DEPRECATED = FStateTreeIndex16::Invalid;
+#endif	
 };
+
+/**
+ * Describes a state tree transition. Source is the state where the transition started, Target describes the state where the transition pointed at,
+ * and Next describes the selected state. The reason Transition and Next are different is that Transition state can be a selector state,
+ * in which case the children will be visited until a leaf state is found, which will be the next state.
+ */
+USTRUCT(BlueprintType)
+struct STATETREEMODULE_API FStateTreeTransitionResult
+{
+	GENERATED_BODY()
+
+PRAGMA_DISABLE_DEPRECATION_WARNINGS
+	FStateTreeTransitionResult() = default;
+	FStateTreeTransitionResult(const FStateTreeTransitionResult&) = default;
+	FStateTreeTransitionResult(FStateTreeTransitionResult&&) = default;
+	FStateTreeTransitionResult& operator=(const FStateTreeTransitionResult&) = default;
+	FStateTreeTransitionResult& operator=(FStateTreeTransitionResult&&) = default;
+PRAGMA_ENABLE_DEPRECATION_WARNINGS
+
+	
+	void Reset()
+	{
+		NextActiveFrames.Reset();
+		CurrentRunStatus = EStateTreeRunStatus::Unset;
+		TargetState = FStateTreeStateHandle::Invalid;
+		CurrentState = FStateTreeStateHandle::Invalid;
+		ChangeType = EStateTreeStateChangeType::Changed;
+		Priority = EStateTreeTransitionPriority::None;
+	}
+
+	/** States selected as result of the transition. */
+	UPROPERTY(EditDefaultsOnly, Category = "Default", BlueprintReadOnly)
+	TArray<FStateTreeExecutionFrame> NextActiveFrames;
+
+	/** Current Run status. */
+	UPROPERTY(EditDefaultsOnly, Category = "Default", BlueprintReadOnly)
+	EStateTreeRunStatus CurrentRunStatus = EStateTreeRunStatus::Unset;
+
+	/** Transition source state */
+	UPROPERTY(EditDefaultsOnly, Category = "Default", BlueprintReadOnly)
+	FStateTreeStateHandle SourceState = FStateTreeStateHandle::Invalid;
+
+	/** Transition target state */
+	UPROPERTY(EditDefaultsOnly, Category = "Default", BlueprintReadOnly)
+	FStateTreeStateHandle TargetState = FStateTreeStateHandle::Invalid;
+
+	/** The current state being executed. On enter/exit callbacks this is the state of the task. */
+	UPROPERTY(EditDefaultsOnly, Category = "Default", BlueprintReadOnly)
+	FStateTreeStateHandle CurrentState = FStateTreeStateHandle::Invalid;
+	
+	/** If the change type is Sustained, then the CurrentState was reselected, or if Changed then the state was just activated. */
+	UPROPERTY(EditDefaultsOnly, Category = "Default", BlueprintReadOnly)
+	EStateTreeStateChangeType ChangeType = EStateTreeStateChangeType::Changed; 
+
+	/** Priority of the transition that caused the state change. */
+	UPROPERTY(EditDefaultsOnly, Category = "Default", BlueprintReadOnly)
+	EStateTreeTransitionPriority Priority = EStateTreeTransitionPriority::None;
+
+	/** StateTree asset that was active when the transition was requested. */
+	UPROPERTY(EditDefaultsOnly, Category = "Default", BlueprintReadOnly)
+	TObjectPtr<const UStateTree> SourceStateTree = nullptr;
+
+	/** Root state the execution frame where the transition was requested. */
+	UPROPERTY(EditDefaultsOnly, Category = "Default", BlueprintReadOnly)
+	FStateTreeStateHandle SourceRootState = FStateTreeStateHandle::Invalid;
+	
+#if WITH_EDITORONLY_DATA
+	UE_DEPRECATED(5.4, "Use the ActiveFrames on FStateTreeExecutionState instead.")
+	UPROPERTY()
+	FStateTreeActiveStates CurrentActiveStates_DEPRECATED;
+
+	UE_DEPRECATED(5.4, "Use the NextActiveFrames instead.")
+	UPROPERTY()
+	FStateTreeActiveStates NextActiveStates_DEPRECATED;
+#endif	
+};
diff --git a/Engine/Plugins/Runtime/StateTree/Source/StateTreeModule/Public/StateTreeIndexTypes.h b/Engine/Plugins/Runtime/StateTree/Source/StateTreeModule/Public/StateTreeIndexTypes.h
--- a/Engine/Plugins/Runtime/StateTree/Source/StateTreeModule/Public/StateTreeIndexTypes.h
+++ b/Engine/Plugins/Runtime/StateTree/Source/StateTreeModule/Public/StateTreeIndexTypes.h
@@ -13,6 +13,11 @@ struct STATETREEMODULE_API FStateTreeIndex16
 	static constexpr uint16 InvalidValue = MAX_uint16;
 	static const FStateTreeIndex16 Invalid;
 
+	friend FORCEINLINE uint32 GetTypeHash(const FStateTreeIndex16 Index)
+	{
+		return GetTypeHash(Index.Value);
+	}
+
 	/** @return true if the given index can be represented by the type. */
 	static bool IsValidIndex(const int32 Index)
 	{
diff --git a/Engine/Plugins/Runtime/StateTree/Source/StateTreeModule/Public/StateTreeInstanceData.h b/Engine/Plugins/Runtime/StateTree/Source/StateTreeModule/Public/StateTreeInstanceData.h
--- a/Engine/Plugins/Runtime/StateTree/Source/StateTreeModule/Public/StateTreeInstanceData.h
+++ b/Engine/Plugins/Runtime/StateTree/Source/StateTreeModule/Public/StateTreeInstanceData.h
@@ -11,6 +11,19 @@
 struct FStateTreeTransitionRequest;
 struct FStateTreeExecutionState;
 
+/** Wrapper class to store an object amongst the structs. */
+USTRUCT()
+struct STATETREEMODULE_API FStateTreeInstanceObjectWrapper
+{
+	GENERATED_BODY()
+
+	FStateTreeInstanceObjectWrapper() = default;
+	FStateTreeInstanceObjectWrapper(UObject* Object) : InstanceObject(Object) {}
+	
+	UPROPERTY()
+	TObjectPtr<UObject> InstanceObject = nullptr;
+};
+
 /**
  * State Tree instance data is used to store the runtime state of a State Tree. It is used together with FStateTreeExecution context to tick the state tree.
  * You are supposed to use FStateTreeInstanceData as a property to store the instance data. That ensures that any UObject references will get GC'd correctly.
@@ -28,35 +41,17 @@ struct STATETREEMODULE_API FStateTreeInstanceStorage
 {
 	GENERATED_BODY()
 
-	/** @return true if the instance is correctly initialized. */
-	bool IsValid() const;
-
-	/** @return Number of items in the instance data. */
-	int32 NumStructs() const { return InstanceStructs.Num(); }
-
-	/** @return true if the specified index is valid index into the struct data container. */
-	bool IsValidStructIndex(const int32 Index) const { return InstanceStructs.IsValidIndex(Index); }
-	
-	/** @return mutable view to the struct at specified index. */
-	FStructView GetMutableStruct(const int32 Index) { return InstanceStructs[Index]; }
-
-	/** @return const view to the struct at specified index. */
-	FConstStructView GetStruct(const int32 Index) const { return InstanceStructs[Index]; }
-
-	/** @return number of instance objects */
-	int32 NumObjects() const { return InstanceObjects.Num(); }
-
-	/** @return pointer to an instance object   */
-	UObject* GetMutableObject(const int32 Index) const { return InstanceObjects[Index]; }
-
-	/** @return const pointer to an instance object   */
-	const UObject* GetObject(const int32 Index) const { return InstanceObjects[Index]; }
-
 	/** @return reference to the event queue. */
-	FStateTreeEventQueue& GetMutableEventQueue() { return EventQueue; }
+	FStateTreeEventQueue& GetMutableEventQueue()
+	{
+		return EventQueue;
+	}
 
 	/** @return reference to the event queue. */
-	const FStateTreeEventQueue& GetEventQueue() const { return EventQueue; }
+	const FStateTreeEventQueue& GetEventQueue() const
+	{
+		return EventQueue;
+	}
 
 	/** 
 	 * Buffers a transition request to be sent to the State Tree.
@@ -68,7 +63,7 @@ struct STATETREEMODULE_API FStateTreeInstanceStorage
 	/** @return currently pending transition requests. */
 	TConstArrayView<FStateTreeTransitionRequest> GetTransitionRequests() const
 	{
-		return TransitionRequests;		
+		return TransitionRequests;
 	}
 
 	/** Reset all pending transition requests. */
@@ -77,15 +72,67 @@ struct STATETREEMODULE_API FStateTreeInstanceStorage
 	/** @return true if all instances are valid. */
 	bool AreAllInstancesValid() const;
 
+	/** @return number of items in the storage. */
+	int32 Num() const
+	{
+		return InstanceStructs.Num();
+	}
+
+	/** @return true if the index can be used to get data. */
+	bool IsValidIndex(const int32 Index) const
+	{
+		return InstanceStructs.IsValidIndex(Index);
+	}
+
+	/** @return true if item at the specified index is object type. */
+	bool IsObject(const int32 Index) const
+	{
+		return InstanceStructs[Index].GetScriptStruct() == TBaseStructure<FStateTreeInstanceObjectWrapper>::Get();
+	}
+
+	/** @return specified item as struct. */
+	FConstStructView GetStruct(const int32 Index) const
+	{
+		return InstanceStructs[Index];
+	}
+	
+	/** @return specified item as mutable struct. */
+	FStructView GetMutableStruct(const int32 Index)
+	{
+		return InstanceStructs[Index];
+	}
+
+	/** @return specified item as object, will check() if the item is not an object. */
+	const UObject* GetObject(const int32 Index) const
+	{
+		const FStateTreeInstanceObjectWrapper& Wrapper = InstanceStructs[Index].Get<const FStateTreeInstanceObjectWrapper>();
+		return Wrapper.InstanceObject;
+	}
+
+	/** @return specified item as mutable Object, will check() if the item is not an object. */
+	UObject* GetMutableObject(const int32 Index) const
+	{
+		const FStateTreeInstanceObjectWrapper& Wrapper = InstanceStructs[Index].Get<const FStateTreeInstanceObjectWrapper>();
+		return Wrapper.InstanceObject;
+	}
+
+	UE_DEPRECATED(5.4, "Use Num() instead.")
+	int32 NumStructs() const { return 0; }
+
+	UE_DEPRECATED(5.4, "Use IsValidIndex() instead.")
+	bool IsValidStructIndex(const int32 Index) const { return false; }
+	
+	UE_DEPRECATED(5.4, "Use Num() instead.")
+	int32 NumObjects() const { return 0; }
+
+	UE_DEPRECATED(5.4, "Not used anymore, since ExecutionState is stored separately.")
+	bool IsValid() const { return false; }
+
 protected:
 	/** Struct instances */
 	UPROPERTY()
 	FInstancedStructContainer InstanceStructs;
 
-	/** Object instances. */
-	UPROPERTY()
-	TArray<TObjectPtr<UObject>> InstanceObjects;
-
 	/** Events */
 	UPROPERTY()
 	FStateTreeEventQueue EventQueue;
@@ -112,18 +159,18 @@ struct STATETREEMODULE_API FStateTreeInstanceData
 
 	FStateTreeInstanceData();
 	~FStateTreeInstanceData();
-
+	
 	/** Initializes the array with specified items. */
-	void Init(UObject& InOwner, TConstArrayView<FInstancedStruct> InStructs, TConstArrayView<const UObject*> InObjects);
-	void Init(UObject& InOwner, TConstArrayView<FConstStructView> InStructs, TConstArrayView<const UObject*> InObjects);
+	void Init(UObject& InOwner, TConstArrayView<FInstancedStruct> InStructs);
+	void Init(UObject& InOwner, TConstArrayView<FConstStructView> InStructs);
 
 	/** Appends new items to the instance. */
-	void Append(UObject& InOwner, TConstArrayView<FInstancedStruct> InStructs, TConstArrayView<const UObject*> InObjects);
-	void Append(UObject& InOwner, TConstArrayView<FConstStructView> InStructs, TConstArrayView<const UObject*> InObjects);
+	void Append(UObject& InOwner, TConstArrayView<FInstancedStruct> InStructs);
+	void Append(UObject& InOwner, TConstArrayView<FConstStructView> InStructs);
 
 	/** Shrinks the array sizes to specified lengths. Sizes must be small or equal than current size. */
-	void ShrinkTo(const int32 NumStructs, const int32 NumObjects);
-
+	void ShrinkTo(const int32 Num);
+	
 	/** Shares the layout from another instance data, and copies the data over. */
 	void CopyFrom(UObject& InOwner, const FStateTreeInstanceData& InOther);
 
@@ -134,32 +181,57 @@ struct STATETREEMODULE_API FStateTreeInstanceData
 	bool IsValid() const;
 
 	/** @return Number of items in the instance data. */
-	int32 NumStructs() const { return GetStorage().InstanceStructs.Num(); }
+	int32 Num() const
+	{
+		return GetStorage().Num();
+	}
+
+	/** @return true if the specified index is valid index into the instance data container. */
+	bool IsValidIndex(const int32 Index) const
+	{
+		return GetStorage().IsValidIndex(Index);
+	}
+
+	/** @return true if the data at specified index is object. */
+	bool IsObject(const int32 Index) const
+	{
+		return GetStorage().IsObject(Index);
+	}
 
-	/** @return true if the specified index is valid index into the struct data container. */
-	bool IsValidStructIndex(const int32 Index) const { return GetStorage().InstanceStructs.IsValidIndex(Index); }
-	
 	/** @return mutable view to the struct at specified index. */
-	FStructView GetMutableStruct(const int32 Index) { return GetMutableStorage().InstanceStructs[Index]; }
+	FStructView GetMutableStruct(const int32 Index)
+	{
+		return GetMutableStorage().GetMutableStruct(Index);
+	}
 
 	/** @return const view to the struct at specified index. */
-	FConstStructView GetStruct(const int32 Index) const { return GetStorage().InstanceStructs[Index]; }
-
-	/** @return number of instance objects */
-	int32 NumObjects() const { return GetStorage().InstanceObjects.Num(); }
+	FConstStructView GetStruct(const int32 Index) const
+	{
+		return GetStorage().GetStruct(Index);
+	}
 
 	/** @return pointer to an instance object   */
-	UObject* GetMutableObject(const int32 Index) { return GetMutableStorage().InstanceObjects[Index]; }
+	UObject* GetMutableObject(const int32 Index)
+	{
+		return GetMutableStorage().GetMutableObject(Index);
+	}
 
 	/** @return const pointer to an instance object   */
-	const UObject* GetObject(const int32 Index) const { return GetStorage().InstanceObjects[Index]; }
+	const UObject* GetObject(const int32 Index) const
+	{
+		return GetStorage().GetObject(Index);
+	}
 
 	/** @return pointer to StateTree execution state, or null if the instance data is not initialized. */
-	const FStateTreeExecutionState* GetExecutionState() const;
-	
-	/** @return array to store unprocessed events. */
-	UE_DEPRECATED(5.2, "Use GetEventQueue() instead.")
-	TArray<FStateTreeEvent>& GetEvents() const;
+	const FStateTreeExecutionState* GetExecutionState() const
+	{
+		return &ExecutionState;
+	}
+
+	FStateTreeExecutionState* GetMutableExecutionState()
+	{
+		return &ExecutionState;
+	}
 
 	/** @return reference to the event queue. */
 	FStateTreeEventQueue& GetMutableEventQueue();
@@ -185,25 +257,47 @@ struct STATETREEMODULE_API FStateTreeInstanceData
 	const FStateTreeInstanceStorage& GetStorage() const;
 
 	int32 GetEstimatedMemoryUsage() const;
-	int32 GetNumItems() const;
 	
 	/** Type traits */
 	bool Identical(const FStateTreeInstanceData* Other, uint32 PortFlags) const;
-	void PostSerialize(const FArchive& Ar);
+
+
+	UE_DEPRECATED(5.4, "Use the structs only Init(), objects should be wrapped in FStateTreeInstanceObjectWrapper.")
+	void Init(UObject& InOwner, TConstArrayView<FInstancedStruct> InStructs, TConstArrayView<const UObject*> InObjects);
+
+	UE_DEPRECATED(5.4, "Use the structs only Init(), objects should be wrapped in FStateTreeInstanceObjectWrapper.")
+	void Init(UObject& InOwner, TConstArrayView<FConstStructView> InStructs, TConstArrayView<const UObject*> InObjects);
+
+	UE_DEPRECATED(5.4, "Use the structs only Init(), objects should be wrapped in FStateTreeInstanceObjectWrapper.")
+	void Append(UObject& InOwner, TConstArrayView<FInstancedStruct> InStructs, TConstArrayView<const UObject*> InObjects);
+
+	UE_DEPRECATED(5.4, "Use the structs only Init(), objects should be wrapped in FStateTreeInstanceObjectWrapper.")
+	void Append(UObject& InOwner, TConstArrayView<FConstStructView> InStructs, TConstArrayView<const UObject*> InObjects);
+
+	UE_DEPRECATED(5.4, "Use the one param ShrinkTo().")
+	void ShrinkTo(const int32 NumStructs, const int32 NumObjects);
+
+	UE_DEPRECATED(5.4, "Use IsValidIndex() instead.")
+	bool IsValidStructIndex(const int32 Index) const { return false; }
+
+	UE_DEPRECATED(5.4, "Use Num() instead.")
+	int32 NumStructs() const { return GetStorage().InstanceStructs.Num(); }
+	
+	UE_DEPRECATED(5.4, "Use Num() instead.")
+	int32 NumObjects() const { return 0; }
+
+	UE_DEPRECATED(5.4, "Use Num() instead.")
+	int32 GetNumItems() const { return 0; }
 
 protected:
 
+	/** Execution state of the state tree instance. */
+	UPROPERTY()
+	FStateTreeExecutionState ExecutionState;
+	
 	/** Storage for the actual instance data, always stores FStateTreeInstanceStorage. */
 	UPROPERTY()
-	FInstancedStruct InstanceStorage;
-
-#if WITH_EDITORONLY_DATA
-	UPROPERTY(meta = (DeprecatedProperty))
-	FInstancedStructContainer InstanceStructs_DEPRECATED;
-
-	UPROPERTY(meta = (DeprecatedProperty))
-	TArray<TObjectPtr<UObject>> InstanceObjects_DEPRECATED;
-#endif // WITH_EDITORONLY_DATA
+	TInstancedStruct<FStateTreeInstanceStorage> InstanceStorage;
 };
 
 template<>
@@ -212,7 +306,6 @@ struct TStructOpsTypeTraits<FStateTreeInstanceData> : public TStructOpsTypeTrait
 	enum
 	{
 		WithIdentical = true,
-		WithPostSerialize = true,
 	};
 };
 
@@ -251,7 +344,7 @@ struct TStateTreeInstanceDataStructRef
 	{
 		const FConstStructView InstanceDataStructView = FConstStructView::template Make(InstanceDataStruct);
 		// Find struct in the instance data.
-		for (int32 Index = 0; Index < Storage.NumStructs(); Index++)
+		for (int32 Index = 0; Index < Storage.Num(); Index++)
 		{
 			if (Storage.GetStruct(Index) == InstanceDataStructView)
 			{
@@ -262,7 +355,7 @@ struct TStateTreeInstanceDataStructRef
 		check(StructIndex != INDEX_NONE);
 	}
 
-	bool IsValid() const { return Storage.IsValidStructIndex(StructIndex); }
+	bool IsValid() const { return Storage.IsValidIndex(StructIndex); }
 
 	T& operator*()
 	{
diff --git a/Engine/Plugins/Runtime/StateTree/Source/StateTreeModule/Public/StateTreeLinker.h b/Engine/Plugins/Runtime/StateTree/Source/StateTreeModule/Public/StateTreeLinker.h
--- a/Engine/Plugins/Runtime/StateTree/Source/StateTreeModule/Public/StateTreeLinker.h
+++ b/Engine/Plugins/Runtime/StateTree/Source/StateTreeModule/Public/StateTreeLinker.h
@@ -25,13 +25,6 @@ struct FStateTreeLinker
 	/** Sets base index for all external data handles. */
 	void SetExternalDataBaseIndex(const int32 InExternalDataBaseIndex) { ExternalDataBaseIndex = InExternalDataBaseIndex; }
 
-	/** Sets currently linked item's instance data type and index. */ 
-	void SetCurrentInstanceDataType(const UStruct* Struct, const int32 Index)
-	{
-		CurrentInstanceStruct = Struct;
-		CurrentInstanceIndex = Index;
-	}
-
 	EStateTreeLinkerStatus GetStatus() const { return Status; }
 	
 	/**
@@ -90,10 +83,9 @@ struct FStateTreeLinker
 		if (Index == INDEX_NONE)
 		{
 			Index = ExternalDataDescs.Add(Desc);
-			check(FStateTreeExternalDataHandle::IsValidIndex(Index + ExternalDataBaseIndex));
-			ExternalDataDescs[Index].Handle.DataViewIndex = FStateTreeIndex16(Index + ExternalDataBaseIndex);
+			ExternalDataDescs[Index].Handle.DataHandle = FStateTreeDataHandle(EStateTreeDataSourceType::ContextData, Index + ExternalDataBaseIndex);
 		}
-		Handle.DataViewIndex = FStateTreeIndex16(Index + ExternalDataBaseIndex);
+		Handle.DataHandle = FStateTreeDataHandle(EStateTreeDataSourceType::ContextData, Index + ExternalDataBaseIndex);
 	}
 
 	/** @return linked external data descriptors. */
@@ -103,8 +95,6 @@ protected:
 
 	const UStateTreeSchema* Schema = nullptr;
 	EStateTreeLinkerStatus Status = EStateTreeLinkerStatus::Succeeded;
-	const UStruct* CurrentInstanceStruct = nullptr;
-	int32 CurrentInstanceIndex = INDEX_NONE;
 	int32 ExternalDataBaseIndex = 0;
 	TArray<FStateTreeExternalDataDesc> ExternalDataDescs;
 };
diff --git a/Engine/Plugins/Runtime/StateTree/Source/StateTreeModule/Public/StateTreeNodeBase.h b/Engine/Plugins/Runtime/StateTree/Source/StateTreeModule/Public/StateTreeNodeBase.h
--- a/Engine/Plugins/Runtime/StateTree/Source/StateTreeModule/Public/StateTreeNodeBase.h
+++ b/Engine/Plugins/Runtime/StateTree/Source/StateTreeModule/Public/StateTreeNodeBase.h
@@ -18,7 +18,13 @@ struct STATETREEMODULE_API FStateTreeNodeBase
 {
 	GENERATED_BODY()
 
+PRAGMA_DISABLE_DEPRECATION_WARNINGS
 	FStateTreeNodeBase() = default;
+	FStateTreeNodeBase(const FStateTreeNodeBase&) = default;
+	FStateTreeNodeBase(FStateTreeNodeBase&&) = default;
+	FStateTreeNodeBase& operator=(const FStateTreeNodeBase&) = default;
+	FStateTreeNodeBase& operator=(FStateTreeNodeBase&&) = default;
+PRAGMA_ENABLE_DEPRECATION_WARNINGS
 
 	virtual ~FStateTreeNodeBase() {}
 
@@ -66,15 +72,25 @@ struct STATETREEMODULE_API FStateTreeNodeBase
 	UPROPERTY()
 	FStateTreeIndex16 BindingsBatch = FStateTreeIndex16::Invalid;
 
-	/** The runtime data's data view index in the StateTreeExecutionContext, and source struct index in property binding. */
+	/** Index of template instance data for the node. Can point to Shared or Default instance data in StateTree depending on node type. */
 	UPROPERTY()
-	FStateTreeIndex16 DataViewIndex = FStateTreeIndex16::Invalid;
+	FStateTreeIndex16 InstanceTemplateIndex = FStateTreeIndex16::Invalid;
 
-	/** Index in runtime instance storage. */
+	/** Data handle to access the instance data. */
 	UPROPERTY()
-	FStateTreeIndex16 InstanceIndex = FStateTreeIndex16::Invalid;
+	FStateTreeDataHandle InstanceDataHandle = FStateTreeDataHandle::Invalid; 
 
-	/** True if the instance is an UObject. */
+#if WITH_EDITORONLY_DATA
+	UE_DEPRECATED(5.4, "InstanceDataHandle is used instead to reference to the instance data.")
 	UPROPERTY()
-	uint8 bInstanceIsObject : 1;
+	FStateTreeIndex16 DataViewIndex_DEPRECATED = FStateTreeIndex16::Invalid;
+
+	UE_DEPRECATED(5.4, "InstanceDataHandle is used instead to reference to the instance data.")
+	UPROPERTY()
+	FStateTreeIndex16 InstanceIndex_DEPRECATED = FStateTreeIndex16::Invalid;
+	
+	UE_DEPRECATED(5.4, "InstanceDataHandle is used to determine if the node has object data.")
+	UPROPERTY()
+	uint8 bInstanceIsObject_DEPRECATED : 1;
+#endif
 };
diff --git a/Engine/Plugins/Runtime/StateTree/Source/StateTreeModule/Public/StateTreePropertyBindings.h b/Engine/Plugins/Runtime/StateTree/Source/StateTreeModule/Public/StateTreePropertyBindings.h
--- a/Engine/Plugins/Runtime/StateTree/Source/StateTreeModule/Public/StateTreePropertyBindings.h
+++ b/Engine/Plugins/Runtime/StateTree/Source/StateTreeModule/Public/StateTreePropertyBindings.h
@@ -124,8 +124,21 @@ struct STATETREEMODULE_API FStateTreeBindableStructDesc
 	FStateTreeBindableStructDesc() = default;
 
 #if WITH_EDITORONLY_DATA
+	FStateTreeBindableStructDesc(const FName InName, const UStruct* InStruct, const FStateTreeDataHandle InDataHandle, const EStateTreeBindableStructSource InDataSource, const FGuid InGuid)
+		: Struct(InStruct)
+		, Name(InName)
+		, DataHandle(InDataHandle)
+		, DataSource(InDataSource)
+		, ID(InGuid)
+	{
+	}
+
+	UE_DEPRECATED(5.4, "Use constructor with DataHandle instead.")
 	FStateTreeBindableStructDesc(const FName InName, const UStruct* InStruct, const EStateTreeBindableStructSource InDataSource, const FGuid InGuid)
-		: Struct(InStruct), Name(InName), DataSource(InDataSource), ID(InGuid)
+		: Struct(InStruct)
+		, Name(InName)
+		, DataSource(InDataSource)
+		, ID(InGuid)
 	{
 	}
 
@@ -147,9 +160,14 @@ struct STATETREEMODULE_API FStateTreeBindableStructDesc
 	UPROPERTY()
 	FName Name;
 
+	/** Runtime data the struct represents. */
+	UPROPERTY()
+	FStateTreeDataHandle DataHandle = FStateTreeDataHandle::Invalid;
+
+	/** Type of the source. */
 	UPROPERTY()
 	EStateTreeBindableStructSource DataSource = EStateTreeBindableStructSource::Context;
-	
+
 #if WITH_EDITORONLY_DATA
 	/** Unique identifier of the struct. */
 	UPROPERTY()
@@ -506,15 +524,24 @@ PRAGMA_DISABLE_DEPRECATION_WARNINGS
 	{
 	}
 
+	FStateTreePropertyPathBinding(const FStateTreeDataHandle InSourceDataHandle, const FStateTreePropertyPath& InSourcePath, const FStateTreePropertyPath& InTargetPath)
+		: SourcePropertyPath(InSourcePath)
+		, TargetPropertyPath(InTargetPath)
+		, SourceDataHandle(InSourceDataHandle)
+	{
+	}
+
+	UE_DEPRECATED(5.4, "Use constructor with DataHandle instead.")
 	FStateTreePropertyPathBinding(const FStateTreeIndex16 InCompiledSourceStructIndex, const FStateTreePropertyPath& InSourcePath, const FStateTreePropertyPath& InTargetPath)
 		: SourcePropertyPath(InSourcePath)
 		, TargetPropertyPath(InTargetPath)
-		, CompiledSourceStructIndex(InCompiledSourceStructIndex)
 	{
 	}
 
 	UE_DEPRECATED(5.3, "Use constructor with FStateTreePropertyPath instead.")
-	FStateTreePropertyPathBinding(const FStateTreeEditorPropertyPath& InSourcePath, const FStateTreeEditorPropertyPath& InTargetPath);
+	FStateTreePropertyPathBinding(const FStateTreeEditorPropertyPath& InSourcePath, const FStateTreeEditorPropertyPath& InTargetPath)
+	{
+	}
 PRAGMA_ENABLE_DEPRECATION_WARNINGS
 
 	void PostSerialize(const FArchive& Ar);
@@ -525,7 +552,11 @@ PRAGMA_ENABLE_DEPRECATION_WARNINGS
 	FStateTreePropertyPath& GetMutableSourcePath() { return SourcePropertyPath; }
 	FStateTreePropertyPath& GetMutableTargetPath() { return TargetPropertyPath; }
 
-	FStateTreeIndex16 GetCompiledSourceStructIndex() const { return CompiledSourceStructIndex; }
+	UE_DEPRECATED(5.4, "Use GetSourceDataHandle instead.")
+	FStateTreeIndex16 GetCompiledSourceStructIndex() const { return {}; }
+
+	void SetSourceDataHandle(const FStateTreeDataHandle NewSourceDataHandle) { SourceDataHandle = NewSourceDataHandle; }
+	FStateTreeDataHandle GetSourceDataHandle() const { return SourceDataHandle; }
 
 private:
 	/** Source property path of the binding */
@@ -536,9 +567,9 @@ private:
 	UPROPERTY()
 	FStateTreePropertyPath TargetPropertyPath;
 
-	/** ID of the struct this property path is relative to. */
+	/** Describes how to get the source data pointer for the binding. */
 	UPROPERTY()
-	FStateTreeIndex16 CompiledSourceStructIndex;
+	FStateTreeDataHandle SourceDataHandle = FStateTreeDataHandle::Invalid;
 
 public:
 #if WITH_EDITORONLY_DATA
@@ -667,6 +698,14 @@ struct STATETREEMODULE_API FStateTreePropertyCopy
 {
 	GENERATED_BODY()
 
+PRAGMA_DISABLE_DEPRECATION_WARNINGS
+	FStateTreePropertyCopy() = default;
+	FStateTreePropertyCopy(const FStateTreePropertyCopy&) = default;
+	FStateTreePropertyCopy(FStateTreePropertyCopy&&) = default;
+	FStateTreePropertyCopy& operator=(const FStateTreePropertyCopy&) = default;
+	FStateTreePropertyCopy& operator=(FStateTreePropertyCopy&&) = default;
+PRAGMA_ENABLE_DEPRECATION_WARNINGS
+
 	/** Source property access. */
 	UPROPERTY()
 	FStateTreePropertyIndirection SourceIndirection;
@@ -681,17 +720,25 @@ struct STATETREEMODULE_API FStateTreePropertyCopy
 	/** Cached pointer to the leaf property of the access. */
 	const FProperty* TargetLeafProperty = nullptr;
 
+	/** Type of the source data, used for validation. */
+	UPROPERTY(Transient)
+	TObjectPtr<const UStruct> SourceStructType = nullptr;
+
 	/** Cached property element size * dim. */
 	UPROPERTY()
 	int32 CopySize = 0;
 
-	/** Index to the struct the source path refers to, sources are stored in FStateTreePropertyBindings. */
+	/** Describes how to get the source data pointer for the copy. */
 	UPROPERTY()
-	FStateTreeIndex16 SourceStructIndex = FStateTreeIndex16::Invalid;
+	FStateTreeDataHandle SourceDataHandle = FStateTreeDataHandle::Invalid;
 
 	/** Type of the copy */
 	UPROPERTY()
 	EStateTreePropertyCopyType Type = EStateTreePropertyCopyType::None;
+
+	UE_DEPRECATED(5.4, "Use SourceDataHandle instead")
+	UPROPERTY()
+	FStateTreeIndex16 SourceStructIndex_DEPRECATED = FStateTreeIndex16::Invalid;
 };
 
 using FStateTreePropCopy UE_DEPRECATED(5.3, "Deprecated struct. Please use FStateTreePropertyCopy instead.") = FStateTreePropertyCopy;
@@ -751,21 +798,39 @@ struct STATETREEMODULE_API FStateTreePropertyBindings
 	 */
 	int32 GetSourceStructNum() const { return SourceStructs.Num(); }
 
-	UE_DEPRECATED(5.3, "Will be removed in a future version.")
-	TArrayView<FStateTreeBindableStructDesc> GetSourceStructs() { return  SourceStructs; };
-
-	UE_DEPRECATED(5.3, "Will be removed in a future version.")
-	TArrayView<FStateTreePropertyCopyBatch> GetCopyBatches() { return CopyBatches; };
-
 	/**
-	 * Copies a batch of properties from source structs to target struct.
-	 * @param SourceStructViews Views to structs where properties are copied from.
-	 * @param TargetBatchIndex Batch index to copy (see FStateTreePropertyBindingCompiler).
-	 * @param TargetStructView View to struct where properties are copied to.
-	 * @return true if all copies succeeded (a copy can fail e.g. if source or destination struct view is invalid).
+	 * Copies a property from Source to Target based on the provided Copy.
+	 * @param Copy Describes which parameter and how is copied.
+	 * @param SourceStructView Pointer and type for the source containing the property to be copied.
+	 * @param TargetStructView Pointer and type for the target containing the property to be copied.
+	 * @return true if the property was copied successfully.
 	 */
-	bool CopyTo(TConstArrayView<FStateTreeDataView> SourceStructViews, const FStateTreeIndex16 TargetBatchIndex, FStateTreeDataView TargetStructView) const;
+	bool CopyProperty(const FStateTreePropertyCopy& Copy, FStateTreeDataView SourceStructView, FStateTreeDataView TargetStructView) const;
+
+	/** @return copy batch at specified index. */
+	const FStateTreePropertyCopyBatch& GetBatch(const FStateTreeIndex16 TargetBatchIndex) const
+	{
+		check(TargetBatchIndex.IsValid());
+		return CopyBatches[TargetBatchIndex.Get()];
+	}
 
+	/** @return All the property copies for a specific batch. */
+	TConstArrayView<FStateTreePropertyCopy> GetBatchCopies(const FStateTreeIndex16 TargetBatchIndex) const
+	{
+		return GetBatchCopies(GetBatch(TargetBatchIndex));
+	}
+
+	/** @return All the property copies for a specific batch. */
+	TConstArrayView<FStateTreePropertyCopy> GetBatchCopies(const FStateTreePropertyCopyBatch& Batch) const
+	{
+		const int32 Count = (int32)Batch.BindingsEnd - (int32)Batch.BindingsBegin;
+		if (Count == 0)
+		{
+			return {};
+		}
+		return MakeArrayView(&PropertyCopies[Batch.BindingsBegin], Count);
+	}
+	
 	/**
 	 * Resets copied properties in TargetStructView. Can be used e.g. to erase UObject references.
 	 * @param TargetBatchIndex Batch index to copy (see FStateTreePropertyBindingCompiler).
@@ -793,8 +858,22 @@ struct STATETREEMODULE_API FStateTreePropertyBindings
 	 */
 	[[nodiscard]] static bool ResolveCopyType(const FStateTreePropertyPathIndirection& SourceIndirection, const FStateTreePropertyPathIndirection& TargetIndirection, FStateTreePropertyCopy& OutCopy);
 
+	
+	UE_DEPRECATED(5.3, "Should not be used, will be removed in a future version.")
+	TArrayView<FStateTreeBindableStructDesc> GetSourceStructs() { return  SourceStructs; };
+
+	UE_DEPRECATED(5.3, "Use GetBatch() instead.")
+	TArrayView<FStateTreePropertyCopyBatch> GetCopyBatches() { return CopyBatches; };
+
+	UE_DEPRECATED(5.4, "Use GetBatchCopies() and Copy() instead.")
+	bool CopyTo(TConstArrayView<FStateTreeDataView> SourceStructViews, const FStateTreeIndex16 TargetBatchIndex, FStateTreeDataView TargetStructView) const
+	{
+		return false;
+	}
+
 private:
 	[[nodiscard]] bool ResolvePath(const UStruct* Struct, const FStateTreePropertyPath& Path, FStateTreePropertyIndirection& OutFirstIndirection, FStateTreePropertyPathIndirection& OutLeafIndirection);
+	const FStateTreeBindableStructDesc* GetSourceDescByHandle(const FStateTreeDataHandle SourceDataHandle);
 
 	void PerformCopy(const FStateTreePropertyCopy& Copy, uint8* SourceAddress, uint8* TargetAddress) const;
 	void PerformResetObjects(const FStateTreePropertyCopy& Copy, uint8* TargetAddress) const;
@@ -812,30 +891,30 @@ private:
 	UPROPERTY()
 	TArray<FStateTreePropertyPathBinding> PropertyPathBindings;
 
+	/** Array of property copies */
+	UPROPERTY(Transient)
+	TArray<FStateTreePropertyCopy> PropertyCopies;
+
+	/** Array of property indirections, indexed by accesses*/
+	UPROPERTY(Transient)
+	TArray<FStateTreePropertyIndirection> PropertyIndirections;
+
+	/** Flag indicating if the properties has been resolved successfully . */
+	bool bBindingsResolved = false;
+
 #if WITH_EDITORONLY_DATA
 PRAGMA_DISABLE_DEPRECATION_WARNINGS
-	/** Array of property bindings, resolved into arrays of copies before use. */
+	
 	UE_DEPRECATED_FORGAME(5.3, "Use PropertyPathBindings instead.")
 	UPROPERTY()
 	TArray<FStateTreePropertyBinding> PropertyBindings_DEPRECATED;
 
-	/** Array of property segments, indexed by property paths. */
 	UE_DEPRECATED_FORGAME(5.3, "Use PropertyPathBindings instead.")
 	UPROPERTY()
 	TArray<FStateTreePropertySegment> PropertySegments_DEPRECATED;
+	
 PRAGMA_ENABLE_DEPRECATION_WARNINGS
 #endif // WITH_EDITORONLY_DATA
-	
-	/** Array of property copies */
-	UPROPERTY(Transient)
-	TArray<FStateTreePropertyCopy> PropertyCopies;
-
-	/** Array of property indirections, indexed by accesses*/
-	UPROPERTY(Transient)
-	TArray<FStateTreePropertyIndirection> PropertyIndirections;
-
-	/** Flag indicating if the properties has been resolved successfully . */
-	bool bBindingsResolved = false;
 
 	friend FStateTreePropertyBindingCompiler;
 	friend UStateTree;
diff --git a/Engine/Plugins/Runtime/StateTree/Source/StateTreeModule/Public/StateTreeTypes.h b/Engine/Plugins/Runtime/StateTree/Source/StateTreeModule/Public/StateTreeTypes.h
--- a/Engine/Plugins/Runtime/StateTree/Source/StateTreeModule/Public/StateTreeTypes.h
+++ b/Engine/Plugins/Runtime/StateTree/Source/StateTreeModule/Public/StateTreeTypes.h
@@ -168,7 +168,12 @@ struct STATETREEMODULE_API FStateTreeStateHandle
 	{
 		return Index >= 0 && Index < (int32)MAX_uint16;
 	}
-	
+
+	friend FORCEINLINE uint32 GetTypeHash(const FStateTreeStateHandle& Handle)
+	{
+		return GetTypeHash(Handle.Index);
+	}
+
 	FStateTreeStateHandle() = default;
 	explicit FStateTreeStateHandle(const uint16 InIndex) : Index(InIndex) {}
 	explicit FStateTreeStateHandle(const int32 InIndex) : Index()
@@ -203,6 +208,182 @@ struct STATETREEMODULE_API FStateTreeStateHandle
 };
 
 
+/** Data type the FStateTreeDataHandle is pointing at. */
+UENUM(BlueprintType)
+enum class EStateTreeDataSourceType : uint8
+{
+	None UMETA(Hidden),
+
+	/** Global Tasks, Evaluators */
+	GlobalInstanceData,
+
+	/** Global Tasks, Evaluators*/
+	GlobalInstanceDataObject,
+
+	/** Active State Tasks */
+	ActiveInstanceData,
+
+	/** Active State Tasks */
+	ActiveInstanceDataObject,
+
+	/** Conditions */
+	SharedInstanceData,
+
+	/** Conditions */
+	SharedInstanceDataObject,
+
+	/** Context Data, External Data, Tree Parameters */
+	ContextData,
+
+	/** Subtree parameters */
+	SubtreeParameterData,
+
+	/** Linked state parameters */
+	LinkedStateParameterData,
+};
+
+/** Handle to a StateTree data */
+USTRUCT(BlueprintType)
+struct STATETREEMODULE_API FStateTreeDataHandle
+{
+	GENERATED_BODY()
+
+	static const FStateTreeDataHandle Invalid;
+	static constexpr uint16 InvalidIndex = MAX_uint16;
+
+	/** @return true if the given index can be represented by the type. */
+	static bool IsValidIndex(const int32 Index)
+	{
+		return Index >= 0 && Index < (int32)InvalidIndex;
+	}
+
+	friend FORCEINLINE uint32 GetTypeHash(const FStateTreeDataHandle& Handle)
+	{
+		uint32 Hash = GetTypeHash(Handle.Source);
+		Hash = HashCombineFast(Hash, GetTypeHash(Handle.Source));
+		Hash = HashCombineFast(Hash, GetTypeHash(Handle.StateHandle));
+		return Hash;
+	}
+	
+	FStateTreeDataHandle() = default;
+	
+	explicit FStateTreeDataHandle(const EStateTreeDataSourceType InSource, const uint16 InIndex, const FStateTreeStateHandle InStateHandle = FStateTreeStateHandle::Invalid)
+		: Source(InSource)
+		, Index(InIndex)
+		, StateHandle(InStateHandle)
+	{
+		// Require valid state for active instance data
+		check(Source != EStateTreeDataSourceType::ActiveInstanceData || (Source == EStateTreeDataSourceType::ActiveInstanceData && StateHandle.IsValid()));
+		check(Source != EStateTreeDataSourceType::ActiveInstanceDataObject || (Source == EStateTreeDataSourceType::ActiveInstanceDataObject && StateHandle.IsValid()));
+	}
+
+	explicit FStateTreeDataHandle(const EStateTreeDataSourceType InSource, const int32 InIndex, const FStateTreeStateHandle InStateHandle = FStateTreeStateHandle::Invalid)
+		: Source(InSource)
+		, StateHandle(InStateHandle)
+	{
+		// Require valid state for active instance data
+		check(Source != EStateTreeDataSourceType::ActiveInstanceData || (Source == EStateTreeDataSourceType::ActiveInstanceData && StateHandle.IsValid()));
+		check(Source != EStateTreeDataSourceType::ActiveInstanceDataObject || (Source == EStateTreeDataSourceType::ActiveInstanceDataObject && StateHandle.IsValid()));
+		check(InIndex == INDEX_NONE || IsValidIndex(InIndex));
+		Index = static_cast<uint16>(InIndex);
+	}
+
+	bool IsValid() const
+	{
+		return Source != EStateTreeDataSourceType::None;
+	}
+
+	void Reset()
+	{
+		Source = EStateTreeDataSourceType::None;
+		Index = InvalidIndex;
+		StateHandle = FStateTreeStateHandle::Invalid;
+	}
+
+	bool operator==(const FStateTreeDataHandle& RHS) const
+	{
+		return Source == RHS.Source && Index == RHS.Index && StateHandle == RHS.StateHandle;
+	}
+	
+	bool operator!=(const FStateTreeDataHandle& RHS) const
+	{
+		return !(*this == RHS);
+	}
+
+	EStateTreeDataSourceType GetSource() const
+	{
+		return Source;
+	}
+
+	int32 GetIndex() const
+	{
+		return Index;
+	}
+
+	FStateTreeStateHandle GetState() const
+	{
+		return StateHandle;
+	}
+
+	bool IsObjectSource() const
+	{
+		return Source == EStateTreeDataSourceType::GlobalInstanceDataObject
+			|| Source == EStateTreeDataSourceType::ActiveInstanceDataObject
+			|| Source == EStateTreeDataSourceType::SharedInstanceDataObject;
+	}
+	
+	FStateTreeDataHandle ToObjectSource() const
+	{
+		switch (Source)
+		{
+		case EStateTreeDataSourceType::GlobalInstanceData:
+			return FStateTreeDataHandle(EStateTreeDataSourceType::GlobalInstanceDataObject, Index, StateHandle);
+		case EStateTreeDataSourceType::ActiveInstanceData:
+			return FStateTreeDataHandle(EStateTreeDataSourceType::ActiveInstanceDataObject, Index, StateHandle);
+		case EStateTreeDataSourceType::SharedInstanceData:
+			return FStateTreeDataHandle(EStateTreeDataSourceType::SharedInstanceDataObject, Index, StateHandle);
+		default:
+			return *this;
+		}
+	}
+	
+	FString Describe() const
+	{
+		switch (Source)
+		{
+		case EStateTreeDataSourceType::None:
+			return TEXT("None");
+		case EStateTreeDataSourceType::ContextData:
+			return FString::Printf(TEXT("Context[%d]"), Index);
+		case EStateTreeDataSourceType::GlobalInstanceData:
+			return FString::Printf(TEXT("Global[%d]"), Index);
+		case EStateTreeDataSourceType::GlobalInstanceDataObject:
+			return FString::Printf(TEXT("GlobalO[%d]"), Index);
+		case EStateTreeDataSourceType::ActiveInstanceData:
+			return FString::Printf(TEXT("Active[%d]"), Index);
+		case EStateTreeDataSourceType::ActiveInstanceDataObject:
+			return FString::Printf(TEXT("ActiveO[%d]"), Index);
+		case EStateTreeDataSourceType::SharedInstanceData:
+			return FString::Printf(TEXT("Shared[%d]"), Index);
+		case EStateTreeDataSourceType::SharedInstanceDataObject:
+			return FString::Printf(TEXT("SharedO[%d]"), Index);
+		default:
+			return TEXT("---");
+		}
+	}
+	
+private:
+	UPROPERTY()
+	EStateTreeDataSourceType Source = EStateTreeDataSourceType::None;
+
+	UPROPERTY()
+	uint16 Index = InvalidIndex;
+
+	UPROPERTY()
+	FStateTreeStateHandle StateHandle = FStateTreeStateHandle::Invalid; 
+};
+
+
 /**
  * Time duration with random variance. Stored compactly as two uint16s, which gives time range of about 650 seconds.
  * The variance is symmetric (+-) around the specified duration.
@@ -386,11 +567,13 @@ struct STATETREEMODULE_API FCompactStateTreeState
 
 	/** Index to state instance data. */
 	UPROPERTY()
-	FStateTreeIndex16 ParameterInstanceIndex = FStateTreeIndex16::Invalid;
+	FStateTreeIndex16 ParameterTemplateIndex = FStateTreeIndex16::Invalid;
 
-	/** Data view index of the input parameters. */
 	UPROPERTY()
-	FStateTreeIndex16 ParameterDataViewIndex = FStateTreeIndex16::Invalid;
+	FStateTreeDataHandle ParameterDataHandle = FStateTreeDataHandle::Invalid;
+
+	UPROPERTY()
+	FStateTreeIndex16 ParameterBindingsBatch = FStateTreeIndex16::Invalid;
 
 	/** Number of enter conditions */
 	UPROPERTY()
@@ -404,13 +587,9 @@ struct STATETREEMODULE_API FCompactStateTreeState
 	UPROPERTY()
 	uint8 TasksNum = 0;
 
-	/** Number of tasks with struct instance data. */
-	UPROPERTY()
-	uint8 TaskInstanceStructNum = 0;
-
-	/** Number of tasks with object instance data. */
+	/** Number of instance data */
 	UPROPERTY()
-	uint8 TaskInstanceObjectNum = 0;
+	uint8 InstanceDataNum = 0;
 
 	/** Type of the state */
 	UPROPERTY()
@@ -434,9 +613,13 @@ struct STATETREEMODULE_API FCompactStateTreeParameters
 {
 	GENERATED_BODY()
 
-	UPROPERTY()
-	FStateTreeIndex16 BindingsBatch = FStateTreeIndex16::Invalid;
+	FCompactStateTreeParameters() = default;
 
+	FCompactStateTreeParameters(const FInstancedPropertyBag& InParameters)
+		: Parameters(InParameters)
+	{
+	}
+	
 	UPROPERTY()
 	FInstancedPropertyBag Parameters;
 };
@@ -775,13 +958,6 @@ protected:
 	uint8* Memory = nullptr;
 };
 
-struct STATETREEMODULE_API FStateTreeTransitionDelayedState
-{
-	FStateTreeIndex16 TransitionIndex = FStateTreeIndex16::Invalid;
-	float TimeLeft = 0.0f;
-};
-
-
 /**
  * Link to another state in StateTree
  */
