StateTree: Changed how instance data is accessed

diff --git a/Engine/Plugins/Runtime/StateTree/Source/StateTreeTestSuite/Private/StateTreeTest.cpp b/Engine/Plugins/Runtime/StateTree/Source/StateTreeTestSuite/Private/StateTreeTest.cpp
--- a/Engine/Plugins/Runtime/StateTree/Source/StateTreeTestSuite/Private/StateTreeTest.cpp
+++ b/Engine/Plugins/Runtime/StateTree/Source/StateTreeTestSuite/Private/StateTreeTest.cpp
@@ -96,6 +96,7 @@ struct FStateTreeTest_MakeAndBakeStateTree : FAITestBase
 		const bool bResult = Compiler.Compile(StateTree);
 
 		AITEST_TRUE("StateTree should get compiled", bResult);
+		AITEST_TRUE("StateTree should be ready to run", StateTree.IsReadyToRun());
 
 		return true;
 	}
@@ -103,6 +104,41 @@ struct FStateTreeTest_MakeAndBakeStateTree : FAITestBase
 IMPLEMENT_AI_INSTANT_TEST(FStateTreeTest_MakeAndBakeStateTree, "System.StateTree.MakeAndBakeStateTree");
 
 
+struct FStateTreeTest_EmptyStateTree : FAITestBase
+{
+	virtual bool InstantTest() override
+	{
+		UStateTree& StateTree = UE::StateTree::Tests::NewStateTree(&GetWorld());
+		UStateTreeEditorData& EditorData = *Cast<UStateTreeEditorData>(StateTree.EditorData);
+		
+		UStateTreeState& Root = EditorData.AddSubTree(FName(TEXT("Root")));
+		Root.AddTransition(EStateTreeTransitionTrigger::OnStateCompleted, EStateTreeTransitionType::Succeeded);
+
+		FStateTreeCompilerLog Log;
+		FStateTreeCompiler Compiler(Log);
+		const bool bResult = Compiler.Compile(StateTree);
+
+		AITEST_TRUE("StateTree should get compiled", bResult);
+
+		EStateTreeRunStatus Status = EStateTreeRunStatus::Unset;
+		FStateTreeInstanceData InstanceData;
+		FTestStateTreeExecutionContext Exec(StateTree, StateTree, InstanceData);
+		const bool bInitSucceeded = Exec.IsValid();
+		AITEST_TRUE("StateTree should init", bInitSucceeded);
+
+		Status = Exec.Start();
+		AITEST_TRUE("StateTree should be running", Status == EStateTreeRunStatus::Running);
+		Exec.LogClear();
+
+		Status = Exec.Tick(0.1f);
+		AITEST_TRUE("StateTree should be completed", Status == EStateTreeRunStatus::Succeeded);
+		Exec.LogClear();
+
+		return true;
+	}
+};
+IMPLEMENT_AI_INSTANT_TEST(FStateTreeTest_EmptyStateTree, "System.StateTree.Empty");
+
 struct FStateTreeTest_Sequence : FAITestBase
 {
 	virtual bool InstantTest() override
@@ -210,7 +246,7 @@ struct FStateTreeTest_Select : FAITestBase
 		AITEST_TRUE("StateTree should be running", Status == EStateTreeRunStatus::Running);
 		Exec.LogClear();
 
-		// Regular tick
+		// Regular tick, no state selection at all.
 		Status = Exec.Tick(0.1f);
 		AITEST_TRUE("StateTree tasks should update in order", Exec.Expect(TaskRoot.GetName(), TickStr).Then(Task1.GetName(), TickStr).Then(Task1A.GetName(), TickStr));
 		AITEST_FALSE("StateTree TaskRoot should not EnterState", Exec.Expect(TaskRoot.GetName(), EnterStateStr));
@@ -427,6 +463,86 @@ struct FStateTreeTest_SubTree : FAITestBase
 };
 IMPLEMENT_AI_INSTANT_TEST(FStateTreeTest_SubTree, "System.StateTree.SubTree");
 
+struct FStateTreeTest_SubTreeCondition : FAITestBase
+{
+	virtual bool InstantTest() override
+	{
+		/*
+		- Root
+			- Linked : Subtree -> Root
+		- SubTree : Task1
+			- ? State1 : Task2 -> Succeeded // condition linked to Task1
+			- State2 : Task3
+		*/
+		
+		UStateTree& StateTree = UE::StateTree::Tests::NewStateTree(&GetWorld());
+		UStateTreeEditorData& EditorData = *Cast<UStateTreeEditorData>(StateTree.EditorData);
+		
+		UStateTreeState& Root = EditorData.AddSubTree(FName(TEXT("Root")));
+		UStateTreeState& Linked = Root.AddChildState(FName(TEXT("Linked")), EStateTreeStateType::Linked);
+		
+		UStateTreeState& SubTree = Root.AddChildState(FName(TEXT("SubTree")), EStateTreeStateType::Subtree);
+		UStateTreeState& State1 = SubTree.AddChildState(FName(TEXT("State1")));
+		UStateTreeState& State2 = SubTree.AddChildState(FName(TEXT("State2")));
+
+		Linked.LinkedSubtree = SubTree.GetLinkToState();
+
+		Linked.AddTransition(EStateTreeTransitionTrigger::OnStateCompleted, EStateTreeTransitionType::GotoState, &Root);
+
+		// SubTask should not complete during the test.
+		TStateTreeEditorNode<FTestTask_Stand>& SubTask = SubTree.AddTask<FTestTask_Stand>(FName(TEXT("SubTask")));
+		SubTask.GetNode().TicksToCompletion = 100;
+
+		TStateTreeEditorNode<FTestTask_Stand>& Task1 = State1.AddTask<FTestTask_Stand>(FName(TEXT("Task1")));
+		Task1.GetNode().TicksToCompletion = 1;
+
+		TStateTreeEditorNode<FTestTask_Stand>& Task2 = State2.AddTask<FTestTask_Stand>(FName(TEXT("Task2")));
+		Task2.GetNode().TicksToCompletion = 1;
+		
+		// Allow to enter State1 if Task1 instance data TicksToCompletion > 0.
+		TStateTreeEditorNode<FStateTreeCompareIntCondition>& IntCond1 = State1.AddEnterCondition<FStateTreeCompareIntCondition>(EGenericAICheck::Greater);
+		EditorData.AddPropertyBinding(SubTask, TEXT("CurrentTick"), IntCond1, TEXT("Left"));
+		IntCond1.GetInstanceData().Right = 0;
+
+		FStateTreeCompilerLog Log;
+		FStateTreeCompiler Compiler(Log);
+		const bool bResult = Compiler.Compile(StateTree);
+		AITEST_TRUE("StateTree should get compiled", bResult);
+
+		EStateTreeRunStatus Status = EStateTreeRunStatus::Unset;
+		FStateTreeInstanceData InstanceData;
+		FTestStateTreeExecutionContext Exec(StateTree, StateTree, InstanceData);
+		const bool bInitSucceeded = Exec.IsValid();
+		AITEST_TRUE("StateTree should init", bInitSucceeded);
+
+		const FString TickStr(TEXT("Tick"));
+		const FString EnterStateStr(TEXT("EnterState"));
+		const FString ExitStateStr(TEXT("ExitState"));
+		const FString StateCompletedStr(TEXT("StateCompleted"));
+
+		// Start and enter state
+		Status = Exec.Start();
+
+		AITEST_TRUE("StateTree Active States should be in Root/Linked/SubTree/State2", Exec.ExpectInActiveStates(Root.Name, Linked.Name, SubTree.Name, State2.Name));
+		AITEST_FALSE("StateTree State1 should not be active", Exec.ExpectInActiveStates(State1.Name)); // Enter condition should prevent to enter State1
+		AITEST_TRUE("StateTree SubTask should enter state", Exec.Expect(SubTask.GetName(), EnterStateStr));
+		AITEST_TRUE("StateTree Task2 should enter state", Exec.Expect(Task2.GetName(), EnterStateStr));
+		AITEST_TRUE("StateTree should be running", Status == EStateTreeRunStatus::Running);
+		Exec.LogClear();
+
+		// Task1 completes, and we should enter State1 since the enter condition now passes.
+		Status = Exec.Tick(0.1f);
+		AITEST_TRUE("StateTree Active States should be in Root/Linked/SubTree/State1", Exec.ExpectInActiveStates(Root.Name, Linked.Name, SubTree.Name, State1.Name));
+		AITEST_FALSE("StateTree State2 should not be active", Exec.ExpectInActiveStates(State2.Name));
+		AITEST_TRUE("StateTree Task1 should enter state", Exec.Expect(Task1.GetName(), EnterStateStr));
+		AITEST_TRUE("StateTree should be running", Status == EStateTreeRunStatus::Running);
+		Exec.LogClear();
+
+		return true;
+	}
+};
+IMPLEMENT_AI_INSTANT_TEST(FStateTreeTest_SubTreeCondition, "System.StateTree.SubTreeCondition");
+
 struct FStateTreeTest_SubTree_CascadedSucceeded : FAITestBase
 {
 	virtual bool InstantTest() override
@@ -603,6 +719,17 @@ struct FStateTreeTest_TransitionPriority : FAITestBase
 	{
 		UStateTree& StateTree = UE::StateTree::Tests::NewStateTree(&GetWorld());
 		UStateTreeEditorData& EditorData = *Cast<UStateTreeEditorData>(StateTree.EditorData);
+
+		/*
+			- Root
+				- State1 : Task1 -> Succeeded
+					- State1A : Task1A -> Next
+					- State1B : Task1B -> Next
+					- State1C : Task1C
+		
+			Task1A completed first, transitioning to State1B.
+			Task1, Task1B, and Task1C complete at the same time, we should take the transition on the first completed state (State1).
+		*/
 		
 		UStateTreeState& Root = EditorData.AddSubTree(FName(TEXT("Root")));
 		UStateTreeState& State1 = Root.AddChildState(FName(TEXT("State1")));
@@ -965,8 +1092,11 @@ struct FStateTreeTest_TransitionDelay : FAITestBase
 		AITEST_TRUE("StateTree Task0 should tick", Exec.Expect(Task0.GetName(), TickStr));
 		Exec.LogClear();
 
+		// Should have execution frames
+		AITEST_TRUE("Should have active frames", InstanceData.GetExecutionState()->ActiveFrames.Num() > 0);
+
 		// Should have delayed transitions
-		const int32 NumDelayedTransitions0 = InstanceData.GetStruct(0).Get<const FStateTreeExecutionState>().DelayedTransitions.Num();
+		const int32 NumDelayedTransitions0 = InstanceData.GetExecutionState()->DelayedTransitions.Num();
 		AITEST_EQUAL("Should have a delayed transition", NumDelayedTransitions0, 1);
 
 		// Tick and expect a delayed transition. 
@@ -974,7 +1104,7 @@ struct FStateTreeTest_TransitionDelay : FAITestBase
 		AITEST_TRUE("StateTree Task0 should tick", Exec.Expect(Task0.GetName(), TickStr));
 		Exec.LogClear();
 
-		const int32 NumDelayedTransitions1 = InstanceData.GetStruct(0).Get<const FStateTreeExecutionState>().DelayedTransitions.Num();
+		const int32 NumDelayedTransitions1 = InstanceData.GetExecutionState()->DelayedTransitions.Num();
 		AITEST_EQUAL("Should have a delayed transition", NumDelayedTransitions1, 1);
 
 		// Should complete delayed transition.
@@ -1381,12 +1511,14 @@ struct FStateTreeTest_BindingsCompiler : FAITestBase
 		SourceADesc.Name = FName(TEXT("SourceA"));
 		SourceADesc.Struct = TBaseStructure<FStateTreeTest_PropertyCopy>::Get();
 		SourceADesc.DataSource = EStateTreeBindableStructSource::Parameter;
+		SourceADesc.DataHandle = FStateTreeDataHandle(EStateTreeDataSourceType::ContextData, 0); // Used as index to SourceViews below.
 		SourceADesc.ID = FGuid::NewGuid();
 
 		FStateTreeBindableStructDesc SourceBDesc;
 		SourceBDesc.Name = FName(TEXT("SourceB"));
 		SourceBDesc.Struct = TBaseStructure<FStateTreeTest_PropertyCopy>::Get();
 		SourceBDesc.DataSource = EStateTreeBindableStructSource::Parameter;
+		SourceBDesc.DataHandle = FStateTreeDataHandle(EStateTreeDataSourceType::ContextData, 1); // Used as index to SourceViews below.
 		SourceBDesc.ID = FGuid::NewGuid();
 
 		FStateTreeBindableStructDesc TargetDesc;
@@ -1445,8 +1577,12 @@ struct FStateTreeTest_BindingsCompiler : FAITestBase
 		SourceViews[SourceAIndex] = FStateTreeDataView(FStructView::Make(SourceA));
 		SourceViews[SourceBIndex] = FStateTreeDataView(FStructView::Make(SourceB));
 		FStateTreeDataView TargetView(FStructView::Make(Target));
-		
-		const bool bCopyResult = Bindings.CopyTo(SourceViews, FStateTreeIndex16(CopyBatchIndex), TargetView);
+
+		bool bCopyResult = true;
+		for (const FStateTreePropertyCopy& Copy : Bindings.GetBatchCopies(FStateTreeIndex16(CopyBatchIndex)))
+		{
+			bCopyResult &= Bindings.CopyProperty(Copy, SourceViews[Copy.SourceDataHandle.GetIndex()], TargetView);
+		}
 		AITEST_TRUE("CopyTo should succeed", bCopyResult);
 
 		// Due to binding sorting, we expect them to executed in this order (sorted based on target access, earliest to latest)
@@ -1480,6 +1616,7 @@ struct FStateTreeTest_CopyObjects : FAITestBase
 		SourceDesc.Name = FName(TEXT("Source"));
 		SourceDesc.Struct = TBaseStructure<FStateTreeTest_PropertyCopyObjects>::Get();
 		SourceDesc.DataSource = EStateTreeBindableStructSource::Parameter;
+		SourceDesc.DataHandle = FStateTreeDataHandle(EStateTreeDataSourceType::ContextData, 0); // Used as index to SourceViews below.
 		SourceDesc.ID = FGuid::NewGuid();
 
 		FStateTreeBindableStructDesc TargetADesc;
@@ -1551,9 +1688,13 @@ struct FStateTreeTest_CopyObjects : FAITestBase
 		TArray<FStateTreeDataView> SourceViews;
 		SourceViews.SetNum(Bindings.GetSourceStructNum());
 		SourceViews[SourceIndex] = FStateTreeDataView(FStructView::Make(Source));
-		
+
 		FStateTreeTest_PropertyCopyObjects TargetA;
-		const bool bCopyResultA = Bindings.CopyTo(SourceViews, FStateTreeIndex16(TargetACopyBatchIndex), FStructView::Make(TargetA));
+		bool bCopyResultA = true;
+		for (const FStateTreePropertyCopy& Copy : Bindings.GetBatchCopies(FStateTreeIndex16(TargetACopyBatchIndex)))
+		{
+			bCopyResultA &= Bindings.CopyProperty(Copy, SourceViews[Copy.SourceDataHandle.GetIndex()], FStructView::Make(TargetA));
+		}
 		AITEST_TRUE("CopyTo should succeed", bCopyResultA);
 
 		AITEST_TRUE("Expect TargetA.Object == Source.Object", TargetA.Object == Source.Object);
@@ -1565,7 +1706,11 @@ struct FStateTreeTest_CopyObjects : FAITestBase
 		TargetA.Object = nullptr;
 		
 		FStateTreeTest_PropertyCopyObjects TargetB;
-		const bool bCopyResultB = Bindings.CopyTo(SourceViews, FStateTreeIndex16(TargetBCopyBatchIndex), FStructView::Make(TargetB));
+		bool bCopyResultB = true;
+		for (const FStateTreePropertyCopy& Copy : Bindings.GetBatchCopies(FStateTreeIndex16(TargetBCopyBatchIndex)))
+		{
+			bCopyResultB &= Bindings.CopyProperty(Copy, SourceViews[Copy.SourceDataHandle.GetIndex()], FStructView::Make(TargetB));
+		}
 		AITEST_TRUE("CopyTo should succeed", bCopyResultB);
 
 		AITEST_TRUE("Expect TargetB.Object == Source.SoftObject", TSoftObjectPtr<UObject>(TargetB.Object) == Source.SoftObject);
@@ -1582,7 +1727,13 @@ struct FStateTreeTest_CopyObjects : FAITestBase
 		CollectGarbage(GARBAGE_COLLECTION_KEEPFLAGS);
 
 		FStateTreeTest_PropertyCopyObjects TargetC;
-		const bool bCopyResultC = Bindings.CopyTo(SourceViews, FStateTreeIndex16(TargetACopyBatchIndex), FStructView::Make(TargetC));
+		bool bCopyResultC = true;
+		for (const FStateTreePropertyCopy& Copy : Bindings.GetBatchCopies(FStateTreeIndex16(TargetACopyBatchIndex)))
+		{
+			bCopyResultB &= Bindings.CopyProperty(Copy, SourceViews[Copy.SourceDataHandle.GetIndex()], FStructView::Make(TargetC));
+		}
+
+		
 		AITEST_TRUE("CopyTo should succeed", bCopyResultC);
 		AITEST_TRUE("Expect TargetC.SoftObject == Source.SoftObject after GC", TargetC.SoftObject == Source.SoftObject);
 
@@ -1742,7 +1893,7 @@ struct FStateTreeTest_InfiniteLoop : FAITestBase
 
 		GetTestRunner().AddExpectedError(TEXT("Loop detected when trying to select state"), EAutomationExpectedErrorFlags::Contains, 1);
 		GetTestRunner().AddExpectedError(TEXT("Failed to select initial state"), EAutomationExpectedErrorFlags::Contains, 1);
-
+		
 		Status = Exec.Start();
 		AITEST_EQUAL("Start should fail", Status, EStateTreeRunStatus::Failed);
 		Exec.LogClear();
