[Backout] - CL29805076
Original CL Desc
-----------------------------------------------------------------
StateTree: Changed how instance data is accessed
- Remove DataView concept
- Added FStateTreeExecutionFrame, which defines an active branch of a specific state tree that is running
- Added FStateTreeDataHandle, which allows frame relative access of instance data (replaces the DataViews)
- Changed linked states to spawn new frames (will later extend to handle linked trees on other assets)
- Bumped StateTree data version number, requires recomplation of StateTrees

diff --git a/Engine/Plugins/Runtime/StateTree/Source/StateTreeEditorModule/Private/StateTreePropertyBindingCompiler.cpp b/Engine/Plugins/Runtime/StateTree/Source/StateTreeEditorModule/Private/StateTreePropertyBindingCompiler.cpp
--- a/Engine/Plugins/Runtime/StateTree/Source/StateTreeEditorModule/Private/StateTreePropertyBindingCompiler.cpp
+++ b/Engine/Plugins/Runtime/StateTree/Source/StateTreeEditorModule/Private/StateTreePropertyBindingCompiler.cpp
@@ -14,9 +14,7 @@ bool FStateTreePropertyBindingCompiler::Init(FStateTreePropertyBindings& InPrope
 	Log = &InLog;
 	PropertyBindings = &InPropertyBindings;
 	PropertyBindings->Reset();
-
 	SourceStructs.Reset();
-	
 	return true;
 }
 
@@ -42,21 +40,26 @@ bool FStateTreePropertyBindingCompiler::CompileBatch(const FStateTreeBindableStr
 			continue;
 		}
 		// Source must be in the source array
-		const FStateTreeBindableStructDesc* SourceStruct = GetSourceStructDescByID(Binding.GetSourcePath().GetStructID());
-		if (!SourceStruct)
+		const FGuid SourceStructID = Binding.GetSourcePath().GetStructID();
+		const int32 SourceStructIdx = SourceStructs.IndexOfByPredicate([SourceStructID](const FStateTreeBindableStructDesc& Struct)
+			{
+				return (Struct.ID == SourceStructID);
+			});
+		if (SourceStructIdx == INDEX_NONE)
 		{
 			Log->Reportf(EMessageSeverity::Error, TargetStruct,
 				TEXT("Could not find a binding source."));
 			return false;
 		}
+		const FStateTreeBindableStructDesc& SourceStruct = SourceStructs[SourceStructIdx];
 
 		FString Error;
 		TArray<FStateTreePropertyPathIndirection> SourceIndirections;
 		TArray<FStateTreePropertyPathIndirection> TargetIndirections;
 		
-		if (!Binding.GetSourcePath().ResolveIndirections(SourceStruct->Struct, SourceIndirections, &Error))
+		if (!Binding.GetSourcePath().ResolveIndirections(SourceStruct.Struct, SourceIndirections, &Error))
 		{
-			Log->Reportf(EMessageSeverity::Error, TargetStruct, TEXT("Resolving path in %s: %s"), *SourceStruct->ToString(), *Error);
+			Log->Reportf(EMessageSeverity::Error, TargetStruct, TEXT("Resolving path in %s: %s"), *SourceStruct.ToString(), *Error);
 			return false;
 		}
 
@@ -67,19 +70,25 @@ bool FStateTreePropertyBindingCompiler::CompileBatch(const FStateTreeBindableStr
 		}
 
 		FStateTreePropertyCopy DummyCopy;
-		FStateTreePropertyPathIndirection LastSourceIndirection = !SourceIndirections.IsEmpty() ? SourceIndirections.Last() : FStateTreePropertyPathIndirection(SourceStruct->Struct);
+		FStateTreePropertyPathIndirection LastSourceIndirection = !SourceIndirections.IsEmpty() ? SourceIndirections.Last() : FStateTreePropertyPathIndirection(SourceStruct.Struct);
 		FStateTreePropertyPathIndirection LastTargetIndirection = !TargetIndirections.IsEmpty() ? TargetIndirections.Last() : FStateTreePropertyPathIndirection(TargetStruct.Struct);
 		if (!PropertyBindings->ResolveCopyType(LastSourceIndirection, LastTargetIndirection, DummyCopy))
 		{
 			Log->Reportf(EMessageSeverity::Error, TargetStruct,
 			TEXT("Cannot copy properties between %s and %s, properties are incompatible."),
-				*UE::StateTree::GetDescAndPathAsString(*SourceStruct, Binding.GetSourcePath()),
+				*UE::StateTree::GetDescAndPathAsString(SourceStruct, Binding.GetSourcePath()),
 				*UE::StateTree::GetDescAndPathAsString(TargetStruct, Binding.GetTargetPath()));
 			return false;
 		}
 
+		if (const auto Validation = UE::StateTree::Compiler::IsValidIndex16(SourceStructIdx); Validation.DidFail())
+		{
+			Validation.Log(*Log, TEXT("Source Struct Index"), SourceStruct);
+			return false;
+		}
+
 		FSortedBinding& NewBinding = NewBindings.AddDefaulted_GetRef();
-		NewBinding.Binding = FStateTreePropertyPathBinding(SourceStruct->DataHandle, Binding.GetSourcePath(), Binding.GetTargetPath());
+		NewBinding.Binding = FStateTreePropertyPathBinding(FStateTreeIndex16(SourceStructIdx), Binding.GetSourcePath(), Binding.GetTargetPath());
 		NewBinding.TargetIndirections = TargetIndirections;
 	}
 
@@ -141,8 +150,6 @@ int32 FStateTreePropertyBindingCompiler::AddSourceStruct(const FStateTreeBindabl
 			*SourceStruct.ToString(), *ExistingStruct->ToString(), *ExistingStruct->ID.ToString());
 	}
 	
-	UE_CLOG(!SourceStruct.DataHandle.IsValid(), LogStateTree, Error, TEXT("%s does not have a valid data handle."), *SourceStruct.ToString());
-	
 	SourceStructs.Add(SourceStruct);
 	return SourceStructs.Num() - 1;
 }
diff --git a/Engine/Plugins/Runtime/StateTree/Source/StateTreeEditorModule/Public/StateTreeCompiler.h b/Engine/Plugins/Runtime/StateTree/Source/StateTreeEditorModule/Public/StateTreeCompiler.h
--- a/Engine/Plugins/Runtime/StateTree/Source/StateTreeEditorModule/Public/StateTreeCompiler.h
+++ b/Engine/Plugins/Runtime/StateTree/Source/StateTreeEditorModule/Public/StateTreeCompiler.h
@@ -50,8 +50,8 @@ private:
 	
 	bool CreateConditions(UStateTreeState& State, TConstArrayView<FStateTreeEditorNode> Conditions);
 	bool CreateCondition(UStateTreeState& State, const FStateTreeEditorNode& CondNode, const EStateTreeConditionOperand Operand, const int8 DeltaIndent);
-	bool CreateTask(UStateTreeState* State, const FStateTreeEditorNode& TaskNode, const FStateTreeDataHandle TaskDataHandle);
-	bool CreateEvaluator(const FStateTreeEditorNode& EvalNode, const FStateTreeDataHandle EvalDataHandle);
+	bool CreateTask(UStateTreeState* State, const FStateTreeEditorNode& TaskNode);
+	bool CreateEvaluator(const FStateTreeEditorNode& EvalNode);
 	bool GetAndValidateBindings(const FStateTreeBindableStructDesc& TargetStruct, FStateTreeDataView TargetValue, TArray<FStateTreePropertyPathBinding>& OutBindings) const;
 	bool IsPropertyAnyEnum(const FStateTreeBindableStructDesc& Struct, FStateTreePropertyPath Path) const;
 	bool ValidateStructRef(const FStateTreeBindableStructDesc& SourceStruct, FStateTreePropertyPath SourcePath,
@@ -60,6 +60,8 @@ private:
 
 	void InstantiateStructSubobjects(FStructView Struct);
 	
+	FStateTreeDataView GetBindingSourceValue(const int32 SourceIndex);
+	
 	FStateTreeCompilerLog& Log;
 	UStateTree* StateTree = nullptr;
 	UStateTreeEditorData* EditorData = nullptr;
@@ -72,6 +74,8 @@ private:
 	TArray<FInstancedStruct> Nodes;
 	TArray<FInstancedStruct> InstanceStructs;
 	TArray<FInstancedStruct> SharedInstanceStructs;
+	TArray<TObjectPtr<UObject>> InstanceObjects;
+	TArray<TObjectPtr<UObject>> SharedInstanceObjects;
 	
 	FStateTreePropertyBindingCompiler BindingsCompiler;
 };
diff --git a/Engine/Plugins/Runtime/StateTree/Source/StateTreeEditorModule/Public/StateTreePropertyBindingCompiler.h b/Engine/Plugins/Runtime/StateTree/Source/StateTreeEditorModule/Public/StateTreePropertyBindingCompiler.h
--- a/Engine/Plugins/Runtime/StateTree/Source/StateTreeEditorModule/Public/StateTreePropertyBindingCompiler.h
+++ b/Engine/Plugins/Runtime/StateTree/Source/StateTreeEditorModule/Public/StateTreePropertyBindingCompiler.h
@@ -46,21 +46,14 @@ struct STATETREEEDITORMODULE_API FStateTreePropertyBindingCompiler
 	int32 AddSourceStruct(const FStateTreeBindableStructDesc& SourceStruct);
 
 	/** @return Index of a source struct by specified ID, or INDEX_NONE if not found. */
-	UE_DEPRECATED(5.4, "Use GetSourceStructDescByID() instead.")
 	int32 GetSourceStructIndexByID(const FGuid& ID) const;
 
 	/** @return Reference to a source struct based on ID. */
-	UE_DEPRECATED(5.4, "Use GetSourceStructDescByID() instead.")
 	const FStateTreeBindableStructDesc& GetSourceStructDesc(const int32 Index) const
 	{
 		return SourceStructs[Index];
 	}
 
-	const FStateTreeBindableStructDesc* GetSourceStructDescByID(const FGuid& ID) const
-	{
-		return SourceStructs.FindByPredicate([ID](const FStateTreeBindableStructDesc& Structs) { return (Structs.ID == ID); });
-	}
-
 PRAGMA_DISABLE_DEPRECATION_WARNINGS
 	/**
 	 * Resolves a string based property path in specified struct into segments of property names and access types.
diff --git a/Engine/Plugins/Runtime/StateTree/Source/StateTreeModule/Private/StateTree.cpp b/Engine/Plugins/Runtime/StateTree/Source/StateTreeModule/Private/StateTree.cpp
--- a/Engine/Plugins/Runtime/StateTree/Source/StateTreeModule/Private/StateTree.cpp
+++ b/Engine/Plugins/Runtime/StateTree/Source/StateTreeModule/Private/StateTree.cpp
@@ -12,7 +12,6 @@
 #include "Logging/LogScopedVerbosityOverride.h"
 #include "Misc/DataValidation.h"
 #include "StructUtilsDelegates.h"
-#include "Misc/EnumerateRange.h"
 #if WITH_EDITOR
 #include "Engine/UserDefinedStruct.h"
 #endif
@@ -24,7 +23,7 @@ FCustomVersionRegistration GRegisterStateTreeCustomVersion(FStateTreeCustomVersi
 bool UStateTree::IsReadyToRun() const
 {
 	// Valid tree must have at least one state and valid instance data.
-	return States.Num() > 0 && bIsLinked && PropertyBindings.IsValid();
+	return States.Num() > 0 && bIsLinked;
 }
 
 FConstStructView UStateTree::GetNode(const int32 NodeIndex) const
@@ -135,6 +134,8 @@ void UStateTree::ResetCompiled()
 	IDToStateMappings.Reset();
 	IDToNodeMappings.Reset();
 	IDToTransitionMappings.Reset();
+
+	ParametersDataViewIndex = FStateTreeIndex8::Invalid;
 	
 	EvaluatorsBegin = 0;
 	EvaluatorsNum = 0;
@@ -366,7 +367,7 @@ void UStateTree::ResetLinked()
 	bIsLinked = false;
 	ExternalDataDescs.Reset();
 	ExternalDataBaseIndex = 0;
-	NumContextDataViews = 0;
+	NumDataViews = 0;
 
 	FWriteScopeLock WriteLock(PerThreadSharedInstanceDataLock);
 	PerThreadSharedInstanceData.Reset();
@@ -378,17 +379,16 @@ bool UStateTree::Link()
 	// This data will be used to allocate runtime instance on all StateTree users.
 	ResetLinked();
 
-	ExternalDataBaseIndex = NumContextData;
-
 	// Resolves nodes references to other StateTree data
 	FStateTreeLinker Linker(Schema);
-	Linker.SetExternalDataBaseIndex(ExternalDataBaseIndex);
+	Linker.SetExternalDataBaseIndex(PropertyBindings.GetSourceStructNum());
 
 	for (int32 Index = 0; Index < Nodes.Num(); Index++)
 	{
 		FStructView Node = Nodes[Index];
 		if (FStateTreeNodeBase* NodePtr = Node.GetPtr<FStateTreeNodeBase>())
 		{
+			Linker.SetCurrentInstanceDataType(NodePtr->GetInstanceDataType(), NodePtr->DataViewIndex.Get());
 			const bool bLinkSucceeded = NodePtr->Link(Linker);
 			if (!bLinkSucceeded || Linker.GetStatus() == EStateTreeLinkerStatus::Failed)
 			{
@@ -398,8 +398,9 @@ bool UStateTree::Link()
 		}
 	}
 
+	ExternalDataBaseIndex = PropertyBindings.GetSourceStructNum();
 	ExternalDataDescs = Linker.GetExternalDataDescs();
-	NumContextDataViews = ExternalDataBaseIndex + ExternalDataDescs.Num();
+	NumDataViews = ExternalDataBaseIndex + ExternalDataDescs.Num();
 
 	if (States.Num() > 0 && Nodes.Num() > 0)
 	{
@@ -418,29 +419,29 @@ bool UStateTree::Link()
 				return false;
 			}
 		}
-	}
 
-	if (!DefaultInstanceData.AreAllInstancesValid())
-	{
-		UE_LOG(LogStateTree, Error, TEXT("%s: State Tree asset was not properly loaded (missing instance data). See log for loading failures, or recompile the StateTree asset."), *GetFullName());
-		return false;
-	}
+		if (!DefaultInstanceData.AreAllInstancesValid())
+		{
+			UE_LOG(LogStateTree, Error, TEXT("%s: State Tree asset was not properly loaded (missing instance data). See log for loading failures, or recompile the StateTree asset."), *GetFullName());
+			return false;
+		}
 
-	if (!SharedInstanceData.AreAllInstancesValid())
-	{
-		UE_LOG(LogStateTree, Error, TEXT("%s: State Tree asset was not properly loaded (missing shared instance data). See log for loading failures, or recompile the StateTree asset."), *GetFullName());
-		return false;
-	}
-	
-	if (!PatchBindings())
-	{
-		return false;
-	}
+		if (!SharedInstanceData.AreAllInstancesValid())
+		{
+			UE_LOG(LogStateTree, Error, TEXT("%s: State Tree asset was not properly loaded (missing shared instance data). See log for loading failures, or recompile the StateTree asset."), *GetFullName());
+			return false;
+		}
 
-	// Resolves property paths used by bindings a store property pointers
-	if (!PropertyBindings.ResolvePaths())
-	{
-		return false;
+		if (!PatchBindings())
+		{
+			return false;
+		}
+
+		// Resolves property paths used by bindings a store property pointers
+		if (!PropertyBindings.ResolvePaths())
+		{
+			return false;
+		}
 	}
 
 	// Link succeeded, setup tree to be ready to run
@@ -455,22 +456,6 @@ bool UStateTree::PatchBindings()
 	const TArrayView<FStateTreePropertyCopyBatch> CopyBatches = PropertyBindings.CopyBatches;
 	const TArrayView<FStateTreePropertyPathBinding> PropertyPathBindings = PropertyBindings.PropertyPathBindings;
 
-	// Make mapping from data handle to source struct.
-	TMap<FStateTreeDataHandle, int32> SourceStructByHandle;
-	for (TConstEnumerateRef<FStateTreeBindableStructDesc> SourceStruct : EnumerateRange(SourceStructs))
-	{
-		SourceStructByHandle.Add(SourceStruct->DataHandle, SourceStruct.GetIndex());
-	}
-
-	auto GetSourceStructByHandle = [&SourceStructByHandle, &SourceStructs](const FStateTreeDataHandle DataHandle) -> FStateTreeBindableStructDesc*
-	{
-		if (int32* Index = SourceStructByHandle.Find(DataHandle))
-		{
-			return &SourceStructs[*Index];
-		}
-		return nullptr;
-	};
-	
 	// Reconcile out of date classes.
 	for (FStateTreeBindableStructDesc& SourceStruct : SourceStructs)
 	{
@@ -514,94 +499,83 @@ bool UStateTree::PatchBindings()
 	}
 
 	// Update property bag structs before resolving binding.
-	if (FStateTreeBindableStructDesc* RootParamsDesc = GetSourceStructByHandle(ParametersDataHandle))
+	if (ParametersDataViewIndex.IsValid() && SourceStructs.IsValidIndex(ParametersDataViewIndex.Get()))
 	{
-		RootParamsDesc->Struct = Parameters.GetPropertyBagStruct();
+		SourceStructs[ParametersDataViewIndex.Get()].Struct = Parameters.GetPropertyBagStruct();
 	}
 
-	// Refresh state parameter descs and bindings batches.
 	for (const FCompactStateTreeState& State : States)
 	{
-		if (State.Type == EStateTreeStateType::Subtree
-			||State.Type == EStateTreeStateType::Linked)
+		if (State.Type == EStateTreeStateType::Subtree)
 		{
-			if (!State.ParameterTemplateIndex.IsValid())
+			if (State.ParameterInstanceIndex.IsValid() == false)
 			{
 				UE_LOG(LogStateTree, Error, TEXT("%s: Data for state '%s' is malformed. Please recompile the StateTree asset."), *GetFullName(), *State.Name.ToString());
 				return false;
 			}
 
 			// Subtree is a bind source, update bag struct.
-			const FCompactStateTreeParameters& Params = DefaultInstanceData.GetMutableStruct(State.ParameterTemplateIndex.Get()).Get<FCompactStateTreeParameters>();
-			FStateTreeBindableStructDesc* Desc = GetSourceStructByHandle(State.ParameterDataHandle);
-			if (!Desc)
+			if (State.ParameterDataViewIndex.IsValid())
 			{
-				UE_LOG(LogStateTree, Error, TEXT("%s: Data for state '%s' is malformed. Please recompile the StateTree asset."), *GetFullName(), *State.Name.ToString());
-				return false;
-			}
-			Desc->Struct = Params.Parameters.GetPropertyBagStruct();
-
-			if (State.ParameterBindingsBatch.IsValid())
-			{
-				FStateTreePropertyCopyBatch& Batch = CopyBatches[State.ParameterBindingsBatch.Get()];
-				Batch.TargetStruct.Struct = Params.Parameters.GetPropertyBagStruct();
+				const FCompactStateTreeParameters& Params = DefaultInstanceData.GetMutableStruct(State.ParameterInstanceIndex.Get()).Get<FCompactStateTreeParameters>();
+				FStateTreeBindableStructDesc& Desc = SourceStructs[State.ParameterDataViewIndex.Get()];
+				Desc.Struct = Params.Parameters.GetPropertyBagStruct();
 			}
 		}
-	}
-
-	// Check linked state property bags consistency
-	for (const FCompactStateTreeState& State : States)
-	{
-		if (State.Type == EStateTreeStateType::Linked && State.LinkedState.IsValid())
+		else if (State.Type == EStateTreeStateType::Linked && State.LinkedState.IsValid())
 		{
 			const FCompactStateTreeState& LinkedState = States[State.LinkedState.Index];
 
-			if (State.ParameterTemplateIndex.IsValid() == false
-				|| LinkedState.ParameterTemplateIndex.IsValid() == false)
+			if (State.ParameterInstanceIndex.IsValid() == false
+				|| LinkedState.ParameterInstanceIndex.IsValid() == false)
 			{
 				UE_LOG(LogStateTree, Error, TEXT("%s: Data for state '%s' is malformed. Please recompile the StateTree asset."), *GetFullName(), *State.Name.ToString());
 				return false;
 			}
 
+			const FCompactStateTreeParameters& Params = DefaultInstanceData.GetMutableStruct(State.ParameterInstanceIndex.Get()).Get<FCompactStateTreeParameters>();
+
 			// Check that the bag in linked state matches.
-			const FCompactStateTreeParameters& Params = DefaultInstanceData.GetMutableStruct(State.ParameterTemplateIndex.Get()).Get<FCompactStateTreeParameters>();
-			const FCompactStateTreeParameters& LinkedStateParams = DefaultInstanceData.GetMutableStruct(LinkedState.ParameterTemplateIndex.Get()).Get<FCompactStateTreeParameters>();
+			const FCompactStateTreeParameters& LinkedStateParams = DefaultInstanceData.GetMutableStruct(LinkedState.ParameterInstanceIndex.Get()).Get<FCompactStateTreeParameters>();
 
 			if (LinkedStateParams.Parameters.GetPropertyBagStruct() != Params.Parameters.GetPropertyBagStruct())
 			{
 				UE_LOG(LogStateTree, Error, TEXT("%s: The parameters on state '%s' does not match the linked state parameters in state '%s'. Please recompile the StateTree asset."), *GetFullName(), *State.Name.ToString(), *LinkedState.Name.ToString());
 				return false;
 			}
+
+			if (Params.BindingsBatch.IsValid())
+			{
+				FStateTreePropertyCopyBatch& Batch = CopyBatches[Params.BindingsBatch.Get()];
+				Batch.TargetStruct.Struct = Params.Parameters.GetPropertyBagStruct();
+			}
 		}
 	}
 
-
-	TMap<FStateTreeDataHandle, FStateTreeDataView> DataViews;
-	TMap<FStateTreeIndex16, FStateTreeDataView> BindingBatchDataView;
+	TArray<FStateTreeDataView> DataViews;
+	DataViews.SetNum(NumDataViews);
 
 	// Tree parameters
-	DataViews.Add(ParametersDataHandle, Parameters.GetMutableValue());
+	DataViews[ParametersDataViewIndex.Get()] = Parameters.GetMutableValue();
 
 	// Setup data views for context data. Since the external data is passed at runtime, we can only provide the type.
 	for (const FStateTreeExternalDataDesc& DataDesc : ContextDataDescs)
 	{
-		DataViews.Add(DataDesc.Handle.DataHandle, FStateTreeDataView(DataDesc.Struct, nullptr));
+		DataViews[DataDesc.Handle.DataViewIndex.Get()] = FStateTreeDataView(DataDesc.Struct, nullptr);
 	}
 	
 	// Setup data views for state parameters.
 	for (FCompactStateTreeState& State : States)
 	{
-		if (State.Type == EStateTreeStateType::Subtree
-			|| State.Type == EStateTreeStateType::Linked)
+		// Subtree state parameters can be a binding source, but not a target.
+		// Linked state parameters can only by a binding target, so not handled here.
+		if (State.Type == EStateTreeStateType::Subtree)
 		{
-			if (State.ParameterDataHandle.IsValid())
+			if (State.ParameterDataViewIndex.IsValid())
 			{
-				FCompactStateTreeParameters& Params = DefaultInstanceData.GetMutableStruct(State.ParameterTemplateIndex.Get()).Get<FCompactStateTreeParameters>();
-				DataViews.Add(State.ParameterDataHandle, Params.Parameters.GetMutableValue());
-				if (State.ParameterBindingsBatch.IsValid())
-				{
-					BindingBatchDataView.Add(State.ParameterBindingsBatch, Params.Parameters.GetMutableValue());
-				}
+				FCompactStateTreeParameters& Params = DefaultInstanceData.GetMutableStruct(State.ParameterInstanceIndex.Get()).Get<FCompactStateTreeParameters>();
+				ensureMsgf(Params.Parameters.IsValid(), TEXT("Expecting valid state parameter value when parameter view is set."));
+				DataViews[State.ParameterDataViewIndex.Get()] = Params.Parameters.GetMutableValue();
 			}
 		}
 	}
@@ -609,52 +583,66 @@ bool UStateTree::PatchBindings()
 	// Setup data views for all nodes.
 	for (FConstStructView NodeView : Nodes)
 	{
-		const FStateTreeNodeBase& Node = NodeView.Get<const FStateTreeNodeBase>();
-		
-		FStateTreeInstanceData* SourceInstanceData = &DefaultInstanceData;
-		if (NodeView.GetPtr<const FStateTreeConditionBase>())
+		if (const FStateTreeConditionBase* Condition = NodeView.GetPtr<const FStateTreeConditionBase>())
 		{
-			// Conditions are stored in shared instance data.
-			SourceInstanceData = &SharedInstanceData;
+			if (Condition->bInstanceIsObject)
+			{
+				DataViews[Condition->DataViewIndex.Get()] = SharedInstanceData.GetMutableObject(Condition->InstanceIndex.Get());
+			}
+			else
+			{
+				DataViews[Condition->DataViewIndex.Get()] = SharedInstanceData.GetMutableStruct(Condition->InstanceIndex.Get());
+			}
 		}
-
-		FStateTreeDataView NodeDataView = Node.InstanceDataHandle.IsObjectSource()
-						? FStateTreeDataView(SourceInstanceData->GetMutableObject(Node.InstanceTemplateIndex.Get()))
-						: FStateTreeDataView(SourceInstanceData->GetMutableStruct(Node.InstanceTemplateIndex.Get()));
-
-		DataViews.Add(Node.InstanceDataHandle, NodeDataView);
-
-		if (Node.BindingsBatch.IsValid())
+		else
 		{
-			BindingBatchDataView.Add(Node.BindingsBatch, NodeDataView);
+			const FStateTreeNodeBase& Node = NodeView.Get<const FStateTreeNodeBase>();
+			if (Node.bInstanceIsObject)
+			{
+				DataViews[Node.DataViewIndex.Get()] = DefaultInstanceData.GetMutableObject(Node.InstanceIndex.Get());
+			}
+			else
+			{
+				DataViews[Node.DataViewIndex.Get()] = DefaultInstanceData.GetMutableStruct(Node.InstanceIndex.Get());
+			}
 		}
 	}
-	
-	auto GetDataSourceView = [&DataViews](const FStateTreeDataHandle Handle) -> FStateTreeDataView
+
+	for (int32 BatchIndex = 0; BatchIndex < CopyBatches.Num(); ++BatchIndex)
 	{
-		if (const FStateTreeDataView* ViewPtr = DataViews.Find(Handle))
+		const FStateTreePropertyCopyBatch& Batch = CopyBatches[BatchIndex];
+
+		// Find data view for the binding target.
+		FStateTreeDataView TargetView;
+
+		// Target can be one of the nodes.
+		for (FConstStructView NodeView : GetNodes())
 		{
-			return *ViewPtr;
+			const FStateTreeNodeBase& Node = NodeView.Get<const FStateTreeNodeBase>();
+			if (Node.BindingsBatch.AsInt32() == BatchIndex)
+			{
+				TargetView = DataViews[Node.DataViewIndex.Get()];
+				break;
+			}
 		}
-		return FStateTreeDataView();
-	};
 
-	auto GetBindingBatchDataView = [&BindingBatchDataView](const FStateTreeIndex16 Index) -> FStateTreeDataView
-	{
-		if (const FStateTreeDataView* ViewPtr = BindingBatchDataView.Find(Index))
+		// Target can be a linked state parameters (subtree state params can only be a binding source).
+		if (!TargetView.IsValid())
 		{
-			return *ViewPtr;
+			for (const FCompactStateTreeState& State : GetStates())
+			{
+				if (State.Type == EStateTreeStateType::Linked)
+				{
+					const FCompactStateTreeParameters& Params = DefaultInstanceData.GetStruct(State.ParameterInstanceIndex.Get()).Get<const FCompactStateTreeParameters>();
+					if (Params.BindingsBatch.Get() == BatchIndex)
+					{
+						const FCompactStateTreeState& LinkedState = States[State.LinkedState.Index];
+						TargetView = DataViews[LinkedState.ParameterDataViewIndex.Get()];
+					}
+				}
+			}
 		}
-		return FStateTreeDataView();
-	};
 
-
-	for (int32 BatchIndex = 0; BatchIndex < CopyBatches.Num(); ++BatchIndex)
-	{
-		const FStateTreePropertyCopyBatch& Batch = CopyBatches[BatchIndex];
-
-		// Find data view for the binding target.
-		FStateTreeDataView TargetView = GetBindingBatchDataView(FStateTreeIndex16(BatchIndex));
 		if (!TargetView.IsValid())
 		{
 			UE_LOG(LogStateTree, Error, TEXT("%hs: Invalid target struct when trying to bind to '%s'."), __FUNCTION__, *Batch.TargetStruct.Name.ToString());
@@ -665,9 +653,8 @@ bool UStateTree::PatchBindings()
 		for (int32 Index = Batch.BindingsBegin; Index != Batch.BindingsEnd; Index++)
 		{
 			FStateTreePropertyPathBinding& Binding = PropertyPathBindings[Index];
-			FStateTreeDataView SourceView = GetDataSourceView(Binding.GetSourceDataHandle());
-			
-			if (!Binding.GetMutableSourcePath().UpdateSegmentsFromValue(SourceView, &ErrorMsg))
+
+			if (!Binding.GetMutableSourcePath().UpdateSegmentsFromValue(DataViews[Binding.GetCompiledSourceStructIndex().Get()], &ErrorMsg))
 			{
 				UE_LOG(LogStateTree, Error, TEXT("%hs: Failed to update source instance structs for property binding '%s'. Reason: %s"), __FUNCTION__, *Binding.GetTargetPath().ToString(), *ErrorMsg);
 				return false;
@@ -767,28 +754,24 @@ TArray<FStateTreeMemoryUsage> UStateTree::CalculateEstimatedMemoryUsage() const
 		{
 			const int32 LinkedUsageIndex = GetUsageIndexForState(CompactState.LinkedState);
 			StateLinks.Emplace(ParentUsageIndex, LinkedUsageIndex);
-		}
-		
-		if (CompactState.Type == EStateTreeStateType::Linked
-			|| CompactState.Type == EStateTreeStateType::Subtree)
-		{
+
 			MemUsage.NodeCount++;
-			MemUsage.AddUsage(DefaultInstanceData.GetStruct(CompactState.ParameterTemplateIndex.Get()));
+			MemUsage.AddUsage(DefaultInstanceData.GetStruct(CompactState.ParameterInstanceIndex.Get()));
 		}
 		
 		for (int32 TaskIndex = CompactState.TasksBegin; TaskIndex < (CompactState.TasksBegin + CompactState.TasksNum); TaskIndex++)
 		{
 			if (const FStateTreeTaskBase* Task = Nodes[TaskIndex].GetPtr<const FStateTreeTaskBase>())
 			{
-				if (Task->InstanceDataHandle.IsObjectSource())
+				if (Task->bInstanceIsObject)
 				{
 					MemUsage.NodeCount++;
-					MemUsage.AddUsage(DefaultInstanceData.GetObject(Task->InstanceTemplateIndex.Get()));
+					MemUsage.AddUsage(DefaultInstanceData.GetObject(Task->InstanceIndex.Get()));
 				}
 				else
 				{
 					MemUsage.NodeCount++;
-					MemUsage.AddUsage(DefaultInstanceData.GetStruct(Task->InstanceTemplateIndex.Get()));
+					MemUsage.AddUsage(DefaultInstanceData.GetStruct(Task->InstanceIndex.Get()));
 				}
 			}
 		}
@@ -837,13 +820,13 @@ TArray<FStateTreeMemoryUsage> UStateTree::CalculateEstimatedMemoryUsage() const
 	for (int32 EvalIndex = EvaluatorsBegin; EvalIndex < (EvaluatorsBegin + EvaluatorsNum); EvalIndex++)
 	{
 		const FStateTreeEvaluatorBase& Eval = Nodes[EvalIndex].Get<const FStateTreeEvaluatorBase>();
-		if (Eval.InstanceDataHandle.IsObjectSource())
+		if (Eval.bInstanceIsObject == false)
 		{
-			EvalMemUsage.AddUsage(DefaultInstanceData.GetObject(Eval.InstanceTemplateIndex.Get()));
+			EvalMemUsage.AddUsage(DefaultInstanceData.GetStruct(Eval.InstanceIndex.Get()));
 		}
 		else
 		{
-			EvalMemUsage.AddUsage(DefaultInstanceData.GetStruct(Eval.InstanceTemplateIndex.Get()));
+			EvalMemUsage.AddUsage(DefaultInstanceData.GetObject(Eval.InstanceIndex.Get()));
 		}
 		EvalMemUsage.NodeCount++;
 	}
@@ -853,13 +836,13 @@ TArray<FStateTreeMemoryUsage> UStateTree::CalculateEstimatedMemoryUsage() const
 	for (int32 TaskIndex = GlobalTasksBegin; TaskIndex < (GlobalTasksBegin + GlobalTasksNum); TaskIndex++)
 	{
 		const FStateTreeTaskBase& Task = Nodes[TaskIndex].Get<const FStateTreeTaskBase>();
-		if (Task.InstanceDataHandle.IsObjectSource())
+		if (Task.bInstanceIsObject == false)
 		{
-			GlobalTaskMemUsage.AddUsage(DefaultInstanceData.GetObject(Task.InstanceTemplateIndex.Get()));
+			GlobalTaskMemUsage.AddUsage(DefaultInstanceData.GetStruct(Task.InstanceIndex.Get()));
 		}
 		else
 		{
-			GlobalTaskMemUsage.AddUsage(DefaultInstanceData.GetStruct(Task.InstanceTemplateIndex.Get()));
+			GlobalTaskMemUsage.AddUsage(DefaultInstanceData.GetObject(Task.InstanceIndex.Get()));
 		}
 		GlobalTaskMemUsage.NodeCount++;
 	}
@@ -907,7 +890,7 @@ TArray<FStateTreeMemoryUsage> UStateTree::CalculateEstimatedMemoryUsage() const
 	if (SharedInstanceData.IsValid())
 	{
 		FStateTreeMemoryUsage& SharedMemUsage = MemoryUsages[SharedMemUsageIndex];
-		SharedMemUsage.NodeCount = SharedInstanceData.Num();
+		SharedMemUsage.NodeCount = SharedInstanceData.GetNumItems();
 		SharedMemUsage.EstimatedMemoryUsage = SharedInstanceData.GetEstimatedMemoryUsage();
 	}
 
diff --git a/Engine/Plugins/Runtime/StateTree/Source/StateTreeModule/Private/StateTreeTypes.cpp b/Engine/Plugins/Runtime/StateTree/Source/StateTreeModule/Private/StateTreeTypes.cpp
--- a/Engine/Plugins/Runtime/StateTree/Source/StateTreeModule/Private/StateTreeTypes.cpp
+++ b/Engine/Plugins/Runtime/StateTree/Source/StateTreeModule/Private/StateTreeTypes.cpp
@@ -13,8 +13,6 @@ const FStateTreeStateHandle FStateTreeStateHandle::Failed = FStateTreeStateHandl
 const FStateTreeStateHandle FStateTreeStateHandle::Stopped = FStateTreeStateHandle(StoppedIndex);
 const FStateTreeStateHandle FStateTreeStateHandle::Root = FStateTreeStateHandle(0);
 
-const FStateTreeDataHandle FStateTreeDataHandle::Invalid = FStateTreeDataHandle();
-
 const FStateTreeIndex16 FStateTreeIndex16::Invalid = FStateTreeIndex16();
 const FStateTreeIndex8 FStateTreeIndex8::Invalid = FStateTreeIndex8();
 
diff --git a/Engine/Plugins/Runtime/StateTree/Source/StateTreeModule/Public/StateTree.h b/Engine/Plugins/Runtime/StateTree/Source/StateTreeModule/Public/StateTree.h
--- a/Engine/Plugins/Runtime/StateTree/Source/StateTreeModule/Public/StateTree.h
+++ b/Engine/Plugins/Runtime/StateTree/Source/StateTreeModule/Public/StateTree.h
@@ -37,8 +37,6 @@ struct STATETREEMODULE_API FStateTreeCustomVersion
 		ChangedBindingsRepresentation,
 		// Added guid to transitions
 		AddedTransitionIds,
-		// Added data handles
-		AddedDataHandlesIds,
 
 		// -----<new versions can be added above this line>-------------------------------------------------
 		VersionPlusOne,
@@ -92,8 +90,8 @@ public:
 	/** @return Shared instance data. */
 	TSharedPtr<FStateTreeInstanceData> GetSharedInstanceData() const;
 
-	/** @return Number of context data views required for StateTree execution (Tree params, context data, External data). */
-	int32 GetNumContextDataViews() const { return NumContextDataViews; }
+	/** @return Number of data views required for StateTree execution (Evaluators, Tasks, Conditions, External data). */
+	int32 GetNumDataViews() const { return NumDataViews; }
 
 	/** @return List of external data required by the state tree */
 	TConstArrayView<FStateTreeExternalDataDesc> GetExternalDataDescs() const { return ExternalDataDescs; }
@@ -143,9 +141,6 @@ public:
 	/** @return Id of the transition matching a given runtime transition index; invalid Id if transition not found. */
 	FGuid GetTransitionIdFromIndex(const FStateTreeIndex16 Index) const;	
 
-	UE_DEPRECATED(5.4, "Replaced with GetNumContextDataViews() which contains context data and external data only.")
-	int32 GetNumDataViews() const { return 0; }
-
 #if WITH_EDITOR
 	/** Resets the compiled data to empty. */
 	void ResetCompiled();
@@ -262,16 +257,8 @@ private:
 
 	/** Data view index of the tree Parameters */
 	UPROPERTY()
-	FStateTreeDataHandle ParametersDataHandle = FStateTreeDataHandle::Invalid; 
-
-	/** Number of context data. */
-	UPROPERTY()
-	uint16 NumContextData = 0;
+	FStateTreeIndex8 ParametersDataViewIndex = FStateTreeIndex8::Invalid;
 
-	/** Number of global instance data. */
-	UPROPERTY()
-	uint16 NumGlobalInstanceData = 0;
-	
 	/** Index of first evaluator in Nodes. */
 	UPROPERTY()
 	uint16 EvaluatorsBegin = 0;
@@ -302,10 +289,10 @@ private:
 	UPROPERTY(Transient)
 	int32 ExternalDataBaseIndex = 0;
 
-	/** Total number of context data views, created during linking. */
+	/** Total number of data views, created during linking. */
 	UPROPERTY(Transient)
-	int32 NumContextDataViews = 0;
-	
+	int32 NumDataViews = 0;
+
 	/** True if the StateTree was linked successfully. */
 	bool bIsLinked = false;
 
diff --git a/Engine/Plugins/Runtime/StateTree/Source/StateTreeModule/Public/StateTreeExecutionTypes.h b/Engine/Plugins/Runtime/StateTree/Source/StateTreeModule/Public/StateTreeExecutionTypes.h
--- a/Engine/Plugins/Runtime/StateTree/Source/StateTreeModule/Public/StateTreeExecutionTypes.h
+++ b/Engine/Plugins/Runtime/StateTree/Source/StateTreeModule/Public/StateTreeExecutionTypes.h
@@ -5,8 +5,6 @@
 
 #include "StateTreeExecutionTypes.generated.h"
 
-class UStateTree;
-
 /**
  * Enumeration for the different update phases.
  * This is used as context information when tracing debug events.
@@ -92,31 +90,16 @@ struct STATETREEMODULE_API FStateTreeExternalDataHandle
 {
 	GENERATED_BODY()
 
-	PRAGMA_DISABLE_DEPRECATION_WARNINGS
-	FStateTreeExternalDataHandle() = default;
-	FStateTreeExternalDataHandle(const FStateTreeExternalDataHandle& Other) = default;
-	FStateTreeExternalDataHandle(FStateTreeExternalDataHandle&& Other) = default;
-	FStateTreeExternalDataHandle& operator=(FStateTreeExternalDataHandle const& Other) = default;
-	FStateTreeExternalDataHandle& operator=(FStateTreeExternalDataHandle&& Other) = default;
-	PRAGMA_ENABLE_DEPRECATION_WARNINGS
-
 	static const FStateTreeExternalDataHandle Invalid;
 	
-	bool IsValid() const { return DataHandle.IsValid(); }
-
-	UE_DEPRECATED(5.4, "Index is deprecated, use DataHandle instead.")
-	static bool IsValidIndex(const int32 Index) { return FStateTreeDataHandle::IsValidIndex(Index); }
+	static bool IsValidIndex(const int32 Index) { return FStateTreeIndex16::IsValidIndex(Index); }
+	bool IsValid() const { return DataViewIndex.IsValid(); }
 
 	UPROPERTY()
-	FStateTreeDataHandle DataHandle = FStateTreeDataHandle::Invalid;
-
-#if WITH_EDITORONLY_DATA
-	UE_DEPRECATED(5.4, "Use DataHandle instead.")
-	UPROPERTY()
-	FStateTreeIndex16 DataViewIndex_DEPRECATED = FStateTreeIndex16::Invalid;
-#endif // WITH_EDITORONLY_DATA
+	FStateTreeIndex16 DataViewIndex = FStateTreeIndex16::Invalid;
 };
 
+
 /**
  * Handle to access an external struct or object.
  * This reference handle can be used in StateTree tasks and evaluators to have quick access to external data.
@@ -218,17 +201,9 @@ struct STATETREEMODULE_API FStateTreeTransitionRequest
 	}
 
 	/** Source state of the transition. Filled in by the StateTree execution context. */
-	UPROPERTY()
+	UPROPERTY(EditDefaultsOnly, Category = "Default")
 	FStateTreeStateHandle SourceState;
 
-	/** StateTree asset that was active when the transition was requested. Filled in by the StateTree execution context. */
-	UPROPERTY()
-	TObjectPtr<const UStateTree> SourceStateTree = nullptr;
-
-	/** Root state the execution frame where the transition was requested. Filled in by the StateTree execution context. */
-	UPROPERTY()
-	FStateTreeStateHandle SourceRootState = FStateTreeStateHandle::Invalid;
-
 	/** Target state of the transition. */
 	UPROPERTY(EditDefaultsOnly, Category = "Default")
 	FStateTreeStateHandle TargetState;
@@ -436,7 +411,6 @@ struct STATETREEMODULE_API FStateTreeTransitionSource
 		*this = {};
 	}
 
-	/** Describes where the transition originated. */
 	EStateTreeTransitionSourceType SourceType = EStateTreeTransitionSourceType::Unset;
 
 	/* Index of the transition if from predefined asset transitions, invalid otherwise */
@@ -450,6 +424,63 @@ struct STATETREEMODULE_API FStateTreeTransitionSource
 };
 
 
+/**
+ * Describes a state tree transition. Source is the state where the transition started, Target describes the state where the transition pointed at,
+ * and Next describes the selected state. The reason Transition and Next are different is that Transition state can be a selector state,
+ * in which case the children will be visited until a leaf state is found, which will be the next state.
+ */
+USTRUCT(BlueprintType)
+struct STATETREEMODULE_API FStateTreeTransitionResult
+{
+	GENERATED_BODY()
+
+	FStateTreeTransitionResult() = default;
+
+	void Reset()
+	{
+		CurrentActiveStates.Reset();
+		CurrentRunStatus = EStateTreeRunStatus::Unset;
+		TargetState = FStateTreeStateHandle::Invalid;
+		NextActiveStates.Reset();
+		CurrentState = FStateTreeStateHandle::Invalid;
+		ChangeType = EStateTreeStateChangeType::Changed;
+		Priority = EStateTreeTransitionPriority::None;
+	}
+	
+	/** Current active states, where the transition started. */
+	UPROPERTY(EditDefaultsOnly, Category = "Default", BlueprintReadOnly)
+	FStateTreeActiveStates CurrentActiveStates;
+
+	/** Current Run status. */
+	UPROPERTY(EditDefaultsOnly, Category = "Default", BlueprintReadOnly)
+	EStateTreeRunStatus CurrentRunStatus = EStateTreeRunStatus::Unset;
+
+	/** Transition source state */
+	UPROPERTY(EditDefaultsOnly, Category = "Default", BlueprintReadOnly)
+	FStateTreeStateHandle SourceState = FStateTreeStateHandle::Invalid;
+
+	/** Transition target state */
+	UPROPERTY(EditDefaultsOnly, Category = "Default", BlueprintReadOnly)
+	FStateTreeStateHandle TargetState = FStateTreeStateHandle::Invalid;
+
+	/** States selected as result of the transition. */
+	UPROPERTY(EditDefaultsOnly, Category = "Default", BlueprintReadOnly)
+	FStateTreeActiveStates NextActiveStates;
+
+	/** The current state being executed. On enter/exit callbacks this is the state of the task. */
+	UPROPERTY(EditDefaultsOnly, Category = "Default", BlueprintReadOnly)
+	FStateTreeStateHandle CurrentState = FStateTreeStateHandle::Invalid;
+
+	/** If the change type is Sustained, then the CurrentState was reselected, or if Changed then the state was just activated. */
+	UPROPERTY(EditDefaultsOnly, Category = "Default", BlueprintReadOnly)
+	EStateTreeStateChangeType ChangeType = EStateTreeStateChangeType::Changed; 
+
+	/** Priority of the transition that caused the state change. */
+	UPROPERTY(EditDefaultsOnly, Category = "Default", BlueprintReadOnly)
+	EStateTreeTransitionPriority Priority = EStateTreeTransitionPriority::None;
+};
+
+
 #if WITH_STATETREE_DEBUGGER
 struct STATETREEMODULE_API FStateTreeInstanceDebugId
 {
@@ -490,235 +521,52 @@ struct STATETREEMODULE_API FStateTreeInstanceDebugId
 };
 #endif // WITH_STATETREE_DEBUGGER
 
-/** Describes current state of a delayed transition. */
-USTRUCT()
-struct STATETREEMODULE_API FStateTreeTransitionDelayedState
-{
-	GENERATED_BODY()
-
-	UPROPERTY()
-	TObjectPtr<const UStateTree> StateTree = nullptr;
-
-	UPROPERTY()
-	FStateTreeIndex16 TransitionIndex = FStateTreeIndex16::Invalid;
-
-	UPROPERTY()
-	float TimeLeft = 0.0f;
-};
-
-/** Describes an active branch of a State Tree. */
-USTRUCT(BlueprintType)
-struct STATETREEMODULE_API FStateTreeExecutionFrame
-{
-	GENERATED_BODY()
-
-	bool IsSameFrame(const FStateTreeExecutionFrame& OtherFrame) const
-	{
-		return StateTree == OtherFrame.StateTree && RootState == OtherFrame.RootState;
-	}
-	
-	/** The State Tree used for ticking this frame. */
-	UPROPERTY()
-	TObjectPtr<const UStateTree> StateTree = nullptr;
-
-	/** The root state of the frame (e.g. Root state or a subtree). */
-	UPROPERTY()
-	FStateTreeStateHandle RootState = FStateTreeStateHandle::Root; 
-	
-	/** Active states in this frame */
-	UPROPERTY()
-	FStateTreeActiveStates ActiveStates;
-
-	/** Index within the instance data to the first global instance data (e.g. global tasks) */
-	UPROPERTY()
-	FStateTreeIndex16 GlobalInstanceIndexBase = FStateTreeIndex16::Invalid;
-
-	/** Index within the instance data to the first active state's instance data (e.g. tasks) */
-	UPROPERTY()
-	FStateTreeIndex16 ActiveInstanceIndexBase = FStateTreeIndex16::Invalid;
-
-	/** Index in the active instance data for the parameters of the root state. */
-	UPROPERTY()
-	FStateTreeIndex16 StateParameterDataIndex = FStateTreeIndex16::Invalid;
-
-	/** If true, the global tasks of the State Tree should be handle in this frame. */
-	UPROPERTY()
-	uint8 bIsGlobalFrame : 1 = false;
-};
-
-/** Describes the execution state of the current State Tree instance. */
 USTRUCT()
 struct STATETREEMODULE_API FStateTreeExecutionState
 {
 	GENERATED_BODY()
 
-PRAGMA_DISABLE_DEPRECATION_WARNINGS
-	FStateTreeExecutionState() = default;
-	FStateTreeExecutionState(const FStateTreeExecutionState&) = default;
-	FStateTreeExecutionState(FStateTreeExecutionState&&) = default;
-	FStateTreeExecutionState& operator=(const FStateTreeExecutionState&) = default;
-	FStateTreeExecutionState& operator=(FStateTreeExecutionState&&) = default;
-PRAGMA_ENABLE_DEPRECATION_WARNINGS
-	
-	void Reset()
-	{
-		ActiveFrames.Reset();
-#if WITH_STATETREE_DEBUGGER
-		InstanceDebugId.Reset();
-#endif
-		EnterStateFailedTaskIndex = FStateTreeIndex16::Invalid;
-		LastTickStatus = EStateTreeRunStatus::Failed;
-		TreeRunStatus = EStateTreeRunStatus::Unset;
-		RequestedStop = EStateTreeRunStatus::Unset;
-		CurrentPhase = EStateTreeUpdatePhase::Unset;
-		CompletedStateHandle = FStateTreeStateHandle::Invalid;
-		StateChangeCount = 0;
-	}
-
 	/** @returns Delayed transition state for a specific transition, or nullptr if it does not exists. */
-	FStateTreeTransitionDelayedState* FindDelayedTransition(const UStateTree* OwnerStateTree, const FStateTreeIndex16 TransitionIndex)
+	FStateTreeTransitionDelayedState* FindDelayedTransition(const FStateTreeIndex16 TransitionIndex)
 	{
-		return DelayedTransitions.FindByPredicate([OwnerStateTree, TransitionIndex](const FStateTreeTransitionDelayedState& TransitionState)
-		{
-			return TransitionState.StateTree == OwnerStateTree && TransitionState.TransitionIndex == TransitionIndex;
-		});
+		return DelayedTransitions.FindByPredicate([TransitionIndex](const FStateTreeTransitionDelayedState& State){ return State.TransitionIndex == TransitionIndex; });
 	}
 
-	/** Currently active frames (and states) */
-	UPROPERTY()
-	TArray<FStateTreeExecutionFrame> ActiveFrames;
-
-	/** Pending delayed transitions. */
-	UPROPERTY()
-	TArray<FStateTreeTransitionDelayedState> DelayedTransitions;
+	/** Currently active states */
+	FStateTreeActiveStates ActiveStates;
 
 #if WITH_STATETREE_DEBUGGER
 	/** Id for the active instance used for debugging. */
 	mutable FStateTreeInstanceDebugId InstanceDebugId;
 #endif
 
-	/** The index of the task that failed during enter state. Exit state uses it to call ExitState() symmetrically. */
-	UPROPERTY()
-	FStateTreeIndex16 EnterStateFailedFrameIndex = FStateTreeIndex16::Invalid;
+	/** Index of the first task struct in the currently initialized instance data. */
+	FStateTreeIndex16 FirstTaskStructIndex = FStateTreeIndex16::Invalid;
+	
+	/** Index of the first task object in the currently initialized instance data. */
+	FStateTreeIndex16 FirstTaskObjectIndex = FStateTreeIndex16::Invalid;
 
-	/** The index of the frame that failed during enter state. Exit state uses it to call ExitState() symmetrically. */
-	UPROPERTY()
+	/** The index of the task that failed during enter state. Exit state uses it to call ExitState() symmetrically. */
 	FStateTreeIndex16 EnterStateFailedTaskIndex = FStateTreeIndex16::Invalid;
 
 	/** Result of last tick */
-	UPROPERTY()
 	EStateTreeRunStatus LastTickStatus = EStateTreeRunStatus::Failed;
 
 	/** Running status of the instance */
-	UPROPERTY()
 	EStateTreeRunStatus TreeRunStatus = EStateTreeRunStatus::Unset;
 
 	/** Completion status stored if Stop was called during the Tick and needed to be deferred. */
-	UPROPERTY()
 	EStateTreeRunStatus RequestedStop = EStateTreeRunStatus::Unset;
 
 	/** Current update phase used to validate reentrant calls to the main entry points of the execution context (i.e. Start, Stop, Tick). */
-	UPROPERTY()
 	EStateTreeUpdatePhase CurrentPhase = EStateTreeUpdatePhase::Unset;
 
 	/** Handle of the state that was first to report state completed (success or failure), used to trigger completion transitions. */
-	UPROPERTY()
-	FStateTreeIndex16 CompletedFrameIndex = FStateTreeIndex16::Invalid; 
-	
-	UPROPERTY()
 	FStateTreeStateHandle CompletedStateHandle = FStateTreeStateHandle::Invalid;
 
 	/** Number of times a new state has been changed. */
-	UPROPERTY()
 	uint16 StateChangeCount = 0;
 
-#if WITH_EDITORONLY_DATA
-	UE_DEPRECATED(5.3, "Use DataHandle instead.")
-	UPROPERTY()
-	FStateTreeActiveStates CurrentActiveStates_DEPRECATED;
-	
-	/** Index of the first task struct in the currently initialized instance data. */
-	UE_DEPRECATED(5.4, "Superceded by State Tree Data Handles.")
-	FStateTreeIndex16 FirstTaskStructIndex_DEPRECATED = FStateTreeIndex16::Invalid;
-	
-	/** Index of the first task object in the currently initialized instance data. */
-	UE_DEPRECATED(5.4, "Superceded by State Tree Data Handles.")
-	FStateTreeIndex16 FirstTaskObjectIndex_DEPRECATED = FStateTreeIndex16::Invalid;
-#endif	
+	/** Running time of the delayed transition */
+	TArray<FStateTreeTransitionDelayedState> DelayedTransitions;
 };
-
-/**
- * Describes a state tree transition. Source is the state where the transition started, Target describes the state where the transition pointed at,
- * and Next describes the selected state. The reason Transition and Next are different is that Transition state can be a selector state,
- * in which case the children will be visited until a leaf state is found, which will be the next state.
- */
-USTRUCT(BlueprintType)
-struct STATETREEMODULE_API FStateTreeTransitionResult
-{
-	GENERATED_BODY()
-
-PRAGMA_DISABLE_DEPRECATION_WARNINGS
-	FStateTreeTransitionResult() = default;
-	FStateTreeTransitionResult(const FStateTreeTransitionResult&) = default;
-	FStateTreeTransitionResult(FStateTreeTransitionResult&&) = default;
-	FStateTreeTransitionResult& operator=(const FStateTreeTransitionResult&) = default;
-	FStateTreeTransitionResult& operator=(FStateTreeTransitionResult&&) = default;
-PRAGMA_ENABLE_DEPRECATION_WARNINGS
-
-	
-	void Reset()
-	{
-		NextActiveFrames.Reset();
-		CurrentRunStatus = EStateTreeRunStatus::Unset;
-		TargetState = FStateTreeStateHandle::Invalid;
-		CurrentState = FStateTreeStateHandle::Invalid;
-		ChangeType = EStateTreeStateChangeType::Changed;
-		Priority = EStateTreeTransitionPriority::None;
-	}
-
-	/** States selected as result of the transition. */
-	UPROPERTY(EditDefaultsOnly, Category = "Default", BlueprintReadOnly)
-	TArray<FStateTreeExecutionFrame> NextActiveFrames;
-
-	/** Current Run status. */
-	UPROPERTY(EditDefaultsOnly, Category = "Default", BlueprintReadOnly)
-	EStateTreeRunStatus CurrentRunStatus = EStateTreeRunStatus::Unset;
-
-	/** Transition source state */
-	UPROPERTY(EditDefaultsOnly, Category = "Default", BlueprintReadOnly)
-	FStateTreeStateHandle SourceState = FStateTreeStateHandle::Invalid;
-
-	/** Transition target state */
-	UPROPERTY(EditDefaultsOnly, Category = "Default", BlueprintReadOnly)
-	FStateTreeStateHandle TargetState = FStateTreeStateHandle::Invalid;
-
-	/** The current state being executed. On enter/exit callbacks this is the state of the task. */
-	UPROPERTY(EditDefaultsOnly, Category = "Default", BlueprintReadOnly)
-	FStateTreeStateHandle CurrentState = FStateTreeStateHandle::Invalid;
-	
-	/** If the change type is Sustained, then the CurrentState was reselected, or if Changed then the state was just activated. */
-	UPROPERTY(EditDefaultsOnly, Category = "Default", BlueprintReadOnly)
-	EStateTreeStateChangeType ChangeType = EStateTreeStateChangeType::Changed; 
-
-	/** Priority of the transition that caused the state change. */
-	UPROPERTY(EditDefaultsOnly, Category = "Default", BlueprintReadOnly)
-	EStateTreeTransitionPriority Priority = EStateTreeTransitionPriority::None;
-
-	/** StateTree asset that was active when the transition was requested. */
-	UPROPERTY(EditDefaultsOnly, Category = "Default", BlueprintReadOnly)
-	TObjectPtr<const UStateTree> SourceStateTree = nullptr;
-
-	/** Root state the execution frame where the transition was requested. */
-	UPROPERTY(EditDefaultsOnly, Category = "Default", BlueprintReadOnly)
-	FStateTreeStateHandle SourceRootState = FStateTreeStateHandle::Invalid;
-	
-#if WITH_EDITORONLY_DATA
-	UE_DEPRECATED(5.4, "Use the ActiveFrames on FStateTreeExecutionState instead.")
-	UPROPERTY()
-	FStateTreeActiveStates CurrentActiveStates_DEPRECATED;
-
-	UE_DEPRECATED(5.4, "Use the NextActiveFrames instead.")
-	UPROPERTY()
-	FStateTreeActiveStates NextActiveStates_DEPRECATED;
-#endif	
-};
\ No newline at end of file
diff --git a/Engine/Plugins/Runtime/StateTree/Source/StateTreeModule/Public/StateTreeIndexTypes.h b/Engine/Plugins/Runtime/StateTree/Source/StateTreeModule/Public/StateTreeIndexTypes.h
--- a/Engine/Plugins/Runtime/StateTree/Source/StateTreeModule/Public/StateTreeIndexTypes.h
+++ b/Engine/Plugins/Runtime/StateTree/Source/StateTreeModule/Public/StateTreeIndexTypes.h
@@ -13,11 +13,6 @@ struct STATETREEMODULE_API FStateTreeIndex16
 	static constexpr uint16 InvalidValue = MAX_uint16;
 	static const FStateTreeIndex16 Invalid;
 
-	friend FORCEINLINE uint32 GetTypeHash(const FStateTreeIndex16 Index)
-	{
-		return GetTypeHash(Index.Value);
-	}
-
 	/** @return true if the given index can be represented by the type. */
 	static bool IsValidIndex(const int32 Index)
 	{
diff --git a/Engine/Plugins/Runtime/StateTree/Source/StateTreeModule/Public/StateTreeInstanceData.h b/Engine/Plugins/Runtime/StateTree/Source/StateTreeModule/Public/StateTreeInstanceData.h
--- a/Engine/Plugins/Runtime/StateTree/Source/StateTreeModule/Public/StateTreeInstanceData.h
+++ b/Engine/Plugins/Runtime/StateTree/Source/StateTreeModule/Public/StateTreeInstanceData.h
@@ -11,19 +11,6 @@
 struct FStateTreeTransitionRequest;
 struct FStateTreeExecutionState;
 
-/** Wrapper class to store an object amongst the structs. */
-USTRUCT()
-struct STATETREEMODULE_API FStateTreeInstanceObjectWrapper
-{
-	GENERATED_BODY()
-
-	FStateTreeInstanceObjectWrapper() = default;
-	FStateTreeInstanceObjectWrapper(UObject* Object) : InstanceObject(Object) {}
-	
-	UPROPERTY()
-	TObjectPtr<UObject> InstanceObject = nullptr;
-};
-
 /**
  * State Tree instance data is used to store the runtime state of a State Tree. It is used together with FStateTreeExecution context to tick the state tree.
  * You are supposed to use FStateTreeInstanceData as a property to store the instance data. That ensures that any UObject references will get GC'd correctly.
@@ -41,17 +28,35 @@ struct STATETREEMODULE_API FStateTreeInstanceStorage
 {
 	GENERATED_BODY()
 
+	/** @return true if the instance is correctly initialized. */
+	bool IsValid() const;
+
+	/** @return Number of items in the instance data. */
+	int32 NumStructs() const { return InstanceStructs.Num(); }
+
+	/** @return true if the specified index is valid index into the struct data container. */
+	bool IsValidStructIndex(const int32 Index) const { return InstanceStructs.IsValidIndex(Index); }
+	
+	/** @return mutable view to the struct at specified index. */
+	FStructView GetMutableStruct(const int32 Index) { return InstanceStructs[Index]; }
+
+	/** @return const view to the struct at specified index. */
+	FConstStructView GetStruct(const int32 Index) const { return InstanceStructs[Index]; }
+
+	/** @return number of instance objects */
+	int32 NumObjects() const { return InstanceObjects.Num(); }
+
+	/** @return pointer to an instance object   */
+	UObject* GetMutableObject(const int32 Index) const { return InstanceObjects[Index]; }
+
+	/** @return const pointer to an instance object   */
+	const UObject* GetObject(const int32 Index) const { return InstanceObjects[Index]; }
+
 	/** @return reference to the event queue. */
-	FStateTreeEventQueue& GetMutableEventQueue()
-	{
-		return EventQueue;
-	}
+	FStateTreeEventQueue& GetMutableEventQueue() { return EventQueue; }
 
 	/** @return reference to the event queue. */
-	const FStateTreeEventQueue& GetEventQueue() const
-	{
-		return EventQueue;
-	}
+	const FStateTreeEventQueue& GetEventQueue() const { return EventQueue; }
 
 	/** 
 	 * Buffers a transition request to be sent to the State Tree.
@@ -63,7 +68,7 @@ struct STATETREEMODULE_API FStateTreeInstanceStorage
 	/** @return currently pending transition requests. */
 	TConstArrayView<FStateTreeTransitionRequest> GetTransitionRequests() const
 	{
-		return TransitionRequests;
+		return TransitionRequests;		
 	}
 
 	/** Reset all pending transition requests. */
@@ -72,67 +77,15 @@ struct STATETREEMODULE_API FStateTreeInstanceStorage
 	/** @return true if all instances are valid. */
 	bool AreAllInstancesValid() const;
 
-	/** @return number of items in the storage. */
-	int32 Num() const
-	{
-		return InstanceStructs.Num();
-	}
-
-	/** @return true if the index can be used to get data. */
-	bool IsValidIndex(const int32 Index) const
-	{
-		return InstanceStructs.IsValidIndex(Index);
-	}
-
-	/** @return true if item at the specified index is object type. */
-	bool IsObject(const int32 Index) const
-	{
-		return InstanceStructs[Index].GetScriptStruct() == TBaseStructure<FStateTreeInstanceObjectWrapper>::Get();
-	}
-
-	/** @return specified item as struct. */
-	FConstStructView GetStruct(const int32 Index) const
-	{
-		return InstanceStructs[Index];
-	}
-	
-	/** @return specified item as mutable struct. */
-	FStructView GetMutableStruct(const int32 Index)
-	{
-		return InstanceStructs[Index];
-	}
-
-	/** @return specified item as object, will check() if the item is not an object. */
-	const UObject* GetObject(const int32 Index) const
-	{
-		const FStateTreeInstanceObjectWrapper& Wrapper = InstanceStructs[Index].Get<const FStateTreeInstanceObjectWrapper>();
-		return Wrapper.InstanceObject;
-	}
-
-	/** @return specified item as mutable Object, will check() if the item is not an object. */
-	UObject* GetMutableObject(const int32 Index) const
-	{
-		const FStateTreeInstanceObjectWrapper& Wrapper = InstanceStructs[Index].Get<const FStateTreeInstanceObjectWrapper>();
-		return Wrapper.InstanceObject;
-	}
-
-	UE_DEPRECATED(5.4, "Use Num() instead.")
-	int32 NumStructs() const { return 0; }
-
-	UE_DEPRECATED(5.4, "Use IsValidIndex() instead.")
-	bool IsValidStructIndex(const int32 Index) const { return false; }
-	
-	UE_DEPRECATED(5.4, "Use Num() instead.")
-	int32 NumObjects() const { return 0; }
-
-	UE_DEPRECATED(5.4, "Not used anymore, since ExecutionState is stored separately.")
-	bool IsValid() const { return false; }
-
 protected:
 	/** Struct instances */
 	UPROPERTY()
 	FInstancedStructContainer InstanceStructs;
 
+	/** Object instances. */
+	UPROPERTY()
+	TArray<TObjectPtr<UObject>> InstanceObjects;
+
 	/** Events */
 	UPROPERTY()
 	FStateTreeEventQueue EventQueue;
@@ -159,18 +112,18 @@ struct STATETREEMODULE_API FStateTreeInstanceData
 
 	FStateTreeInstanceData();
 	~FStateTreeInstanceData();
-	
+
 	/** Initializes the array with specified items. */
-	void Init(UObject& InOwner, TConstArrayView<FInstancedStruct> InStructs);
-	void Init(UObject& InOwner, TConstArrayView<FConstStructView> InStructs);
+	void Init(UObject& InOwner, TConstArrayView<FInstancedStruct> InStructs, TConstArrayView<const UObject*> InObjects);
+	void Init(UObject& InOwner, TConstArrayView<FConstStructView> InStructs, TConstArrayView<const UObject*> InObjects);
 
 	/** Appends new items to the instance. */
-	void Append(UObject& InOwner, TConstArrayView<FInstancedStruct> InStructs);
-	void Append(UObject& InOwner, TConstArrayView<FConstStructView> InStructs);
+	void Append(UObject& InOwner, TConstArrayView<FInstancedStruct> InStructs, TConstArrayView<const UObject*> InObjects);
+	void Append(UObject& InOwner, TConstArrayView<FConstStructView> InStructs, TConstArrayView<const UObject*> InObjects);
 
 	/** Shrinks the array sizes to specified lengths. Sizes must be small or equal than current size. */
-	void ShrinkTo(const int32 Num);
-	
+	void ShrinkTo(const int32 NumStructs, const int32 NumObjects);
+
 	/** Shares the layout from another instance data, and copies the data over. */
 	void CopyFrom(UObject& InOwner, const FStateTreeInstanceData& InOther);
 
@@ -181,57 +134,32 @@ struct STATETREEMODULE_API FStateTreeInstanceData
 	bool IsValid() const;
 
 	/** @return Number of items in the instance data. */
-	int32 Num() const
-	{
-		return GetStorage().Num();
-	}
-
-	/** @return true if the specified index is valid index into the instance data container. */
-	bool IsValidIndex(const int32 Index) const
-	{
-		return GetStorage().IsValidIndex(Index);
-	}
-
-	/** @return true if the data at specified index is object. */
-	bool IsObject(const int32 Index) const
-	{
-		return GetStorage().IsObject(Index);
-	}
+	int32 NumStructs() const { return GetStorage().InstanceStructs.Num(); }
 
+	/** @return true if the specified index is valid index into the struct data container. */
+	bool IsValidStructIndex(const int32 Index) const { return GetStorage().InstanceStructs.IsValidIndex(Index); }
+	
 	/** @return mutable view to the struct at specified index. */
-	FStructView GetMutableStruct(const int32 Index)
-	{
-		return GetMutableStorage().GetMutableStruct(Index);
-	}
+	FStructView GetMutableStruct(const int32 Index) { return GetMutableStorage().InstanceStructs[Index]; }
 
 	/** @return const view to the struct at specified index. */
-	FConstStructView GetStruct(const int32 Index) const
-	{
-		return GetStorage().GetStruct(Index);
-	}
+	FConstStructView GetStruct(const int32 Index) const { return GetStorage().InstanceStructs[Index]; }
+
+	/** @return number of instance objects */
+	int32 NumObjects() const { return GetStorage().InstanceObjects.Num(); }
 
 	/** @return pointer to an instance object   */
-	UObject* GetMutableObject(const int32 Index)
-	{
-		return GetMutableStorage().GetMutableObject(Index);
-	}
+	UObject* GetMutableObject(const int32 Index) { return GetMutableStorage().InstanceObjects[Index]; }
 
 	/** @return const pointer to an instance object   */
-	const UObject* GetObject(const int32 Index) const
-	{
-		return GetStorage().GetObject(Index);
-	}
+	const UObject* GetObject(const int32 Index) const { return GetStorage().InstanceObjects[Index]; }
 
 	/** @return pointer to StateTree execution state, or null if the instance data is not initialized. */
-	const FStateTreeExecutionState* GetExecutionState() const
-	{
-		return &ExecutionState;
-	}
-
-	FStateTreeExecutionState* GetMutableExecutionState()
-	{
-		return &ExecutionState;
-	}
+	const FStateTreeExecutionState* GetExecutionState() const;
+	
+	/** @return array to store unprocessed events. */
+	UE_DEPRECATED(5.2, "Use GetEventQueue() instead.")
+	TArray<FStateTreeEvent>& GetEvents() const;
 
 	/** @return reference to the event queue. */
 	FStateTreeEventQueue& GetMutableEventQueue();
@@ -257,47 +185,25 @@ struct STATETREEMODULE_API FStateTreeInstanceData
 	const FStateTreeInstanceStorage& GetStorage() const;
 
 	int32 GetEstimatedMemoryUsage() const;
+	int32 GetNumItems() const;
 	
 	/** Type traits */
 	bool Identical(const FStateTreeInstanceData* Other, uint32 PortFlags) const;
-
-
-	UE_DEPRECATED(5.4, "Use the structs only Init(), objects should be wrapped in FStateTreeInstanceObjectWrapper.")
-	void Init(UObject& InOwner, TConstArrayView<FInstancedStruct> InStructs, TConstArrayView<const UObject*> InObjects);
-
-	UE_DEPRECATED(5.4, "Use the structs only Init(), objects should be wrapped in FStateTreeInstanceObjectWrapper.")
-	void Init(UObject& InOwner, TConstArrayView<FConstStructView> InStructs, TConstArrayView<const UObject*> InObjects);
-
-	UE_DEPRECATED(5.4, "Use the structs only Init(), objects should be wrapped in FStateTreeInstanceObjectWrapper.")
-	void Append(UObject& InOwner, TConstArrayView<FInstancedStruct> InStructs, TConstArrayView<const UObject*> InObjects);
-
-	UE_DEPRECATED(5.4, "Use the structs only Init(), objects should be wrapped in FStateTreeInstanceObjectWrapper.")
-	void Append(UObject& InOwner, TConstArrayView<FConstStructView> InStructs, TConstArrayView<const UObject*> InObjects);
-
-	UE_DEPRECATED(5.4, "Use the one param ShrinkTo().")
-	void ShrinkTo(const int32 NumStructs, const int32 NumObjects);
-
-	UE_DEPRECATED(5.4, "Use IsValidIndex() instead.")
-	bool IsValidStructIndex(const int32 Index) const { return false; }
-
-	UE_DEPRECATED(5.4, "Use Num() instead.")
-	int32 NumStructs() const { return GetStorage().InstanceStructs.Num(); }
-	
-	UE_DEPRECATED(5.4, "Use Num() instead.")
-	int32 NumObjects() const { return 0; }
-
-	UE_DEPRECATED(5.4, "Use Num() instead.")
-	int32 GetNumItems() const { return 0; }
+	void PostSerialize(const FArchive& Ar);
 
 protected:
 
-	/** Execution state of the state tree instance. */
-	UPROPERTY()
-	FStateTreeExecutionState ExecutionState;
-	
 	/** Storage for the actual instance data, always stores FStateTreeInstanceStorage. */
 	UPROPERTY()
-	TInstancedStruct<FStateTreeInstanceStorage> InstanceStorage;
+	FInstancedStruct InstanceStorage;
+
+#if WITH_EDITORONLY_DATA
+	UPROPERTY(meta = (DeprecatedProperty))
+	FInstancedStructContainer InstanceStructs_DEPRECATED;
+
+	UPROPERTY(meta = (DeprecatedProperty))
+	TArray<TObjectPtr<UObject>> InstanceObjects_DEPRECATED;
+#endif // WITH_EDITORONLY_DATA
 };
 
 template<>
@@ -306,6 +212,7 @@ struct TStructOpsTypeTraits<FStateTreeInstanceData> : public TStructOpsTypeTrait
 	enum
 	{
 		WithIdentical = true,
+		WithPostSerialize = true,
 	};
 };
 
@@ -344,7 +251,7 @@ struct TStateTreeInstanceDataStructRef
 	{
 		const FConstStructView InstanceDataStructView = FConstStructView::template Make(InstanceDataStruct);
 		// Find struct in the instance data.
-		for (int32 Index = 0; Index < Storage.Num(); Index++)
+		for (int32 Index = 0; Index < Storage.NumStructs(); Index++)
 		{
 			if (Storage.GetStruct(Index) == InstanceDataStructView)
 			{
@@ -355,7 +262,7 @@ struct TStateTreeInstanceDataStructRef
 		check(StructIndex != INDEX_NONE);
 	}
 
-	bool IsValid() const { return Storage.IsValidIndex(StructIndex); }
+	bool IsValid() const { return Storage.IsValidStructIndex(StructIndex); }
 
 	T& operator*()
 	{
diff --git a/Engine/Plugins/Runtime/StateTree/Source/StateTreeModule/Public/StateTreeNodeBase.h b/Engine/Plugins/Runtime/StateTree/Source/StateTreeModule/Public/StateTreeNodeBase.h
--- a/Engine/Plugins/Runtime/StateTree/Source/StateTreeModule/Public/StateTreeNodeBase.h
+++ b/Engine/Plugins/Runtime/StateTree/Source/StateTreeModule/Public/StateTreeNodeBase.h
@@ -18,13 +18,7 @@ struct STATETREEMODULE_API FStateTreeNodeBase
 {
 	GENERATED_BODY()
 
-PRAGMA_DISABLE_DEPRECATION_WARNINGS
 	FStateTreeNodeBase() = default;
-	FStateTreeNodeBase(const FStateTreeNodeBase&) = default;
-	FStateTreeNodeBase(FStateTreeNodeBase&&) = default;
-	FStateTreeNodeBase& operator=(const FStateTreeNodeBase&) = default;
-	FStateTreeNodeBase& operator=(FStateTreeNodeBase&&) = default;
-PRAGMA_ENABLE_DEPRECATION_WARNINGS
 
 	virtual ~FStateTreeNodeBase() {}
 
@@ -72,25 +66,15 @@ PRAGMA_ENABLE_DEPRECATION_WARNINGS
 	UPROPERTY()
 	FStateTreeIndex16 BindingsBatch = FStateTreeIndex16::Invalid;
 
-	/** Index of template instance data for the node. Can point to Shared or Default instance data in StateTree depending on node type. */
+	/** The runtime data's data view index in the StateTreeExecutionContext, and source struct index in property binding. */
 	UPROPERTY()
-	FStateTreeIndex16 InstanceTemplateIndex = FStateTreeIndex16::Invalid;
+	FStateTreeIndex16 DataViewIndex = FStateTreeIndex16::Invalid;
 
-	/** Data handle to access the instance data. */
+	/** Index in runtime instance storage. */
 	UPROPERTY()
-	FStateTreeDataHandle InstanceDataHandle = FStateTreeDataHandle::Invalid; 
+	FStateTreeIndex16 InstanceIndex = FStateTreeIndex16::Invalid;
 
-#if WITH_EDITORONLY_DATA
-	UE_DEPRECATED(5.4, "InstanceDataHandle is used instead to reference to the instance data.")
+	/** True if the instance is an UObject. */
 	UPROPERTY()
-	FStateTreeIndex16 DataViewIndex_DEPRECATED = FStateTreeIndex16::Invalid;
-
-	UE_DEPRECATED(5.4, "InstanceDataHandle is used instead to reference to the instance data.")
-	UPROPERTY()
-	FStateTreeIndex16 InstanceIndex_DEPRECATED = FStateTreeIndex16::Invalid;
-	
-	UE_DEPRECATED(5.4, "InstanceDataHandle is used to determine if the node has object data.")
-	UPROPERTY()
-	uint8 bInstanceIsObject_DEPRECATED : 1;
-#endif
+	uint8 bInstanceIsObject : 1;
 };
diff --git a/Engine/Plugins/Runtime/StateTree/Source/StateTreeModule/Public/StateTreeTypes.h b/Engine/Plugins/Runtime/StateTree/Source/StateTreeModule/Public/StateTreeTypes.h
--- a/Engine/Plugins/Runtime/StateTree/Source/StateTreeModule/Public/StateTreeTypes.h
+++ b/Engine/Plugins/Runtime/StateTree/Source/StateTreeModule/Public/StateTreeTypes.h
@@ -168,12 +168,7 @@ struct STATETREEMODULE_API FStateTreeStateHandle
 	{
 		return Index >= 0 && Index < (int32)MAX_uint16;
 	}
-
-	friend FORCEINLINE uint32 GetTypeHash(const FStateTreeStateHandle& Handle)
-	{
-		return GetTypeHash(Handle.Index);
-	}
-
+	
 	FStateTreeStateHandle() = default;
 	explicit FStateTreeStateHandle(const uint16 InIndex) : Index(InIndex) {}
 	explicit FStateTreeStateHandle(const int32 InIndex) : Index()
@@ -208,182 +203,6 @@ struct STATETREEMODULE_API FStateTreeStateHandle
 };
 
 
-/** Data type the FStateTreeDataHandle is pointing at. */
-UENUM(BlueprintType)
-enum class EStateTreeDataSourceType : uint8
-{
-	None UMETA(Hidden),
-
-	/** Global Tasks, Evaluators */
-	GlobalInstanceData,
-
-	/** Global Tasks, Evaluators*/
-	GlobalInstanceDataObject,
-
-	/** Active State Tasks */
-	ActiveInstanceData,
-
-	/** Active State Tasks */
-	ActiveInstanceDataObject,
-
-	/** Conditions */
-	SharedInstanceData,
-
-	/** Conditions */
-	SharedInstanceDataObject,
-
-	/** Context Data, External Data, Tree Parameters */
-	ContextData,
-
-	/** Subtree parameters */
-	SubtreeParameterData,
-
-	/** Linked state parameters */
-	LinkedStateParameterData,
-};
-
-/** Handle to a StateTree data */
-USTRUCT(BlueprintType)
-struct STATETREEMODULE_API FStateTreeDataHandle
-{
-	GENERATED_BODY()
-
-	static const FStateTreeDataHandle Invalid;
-	static constexpr uint16 InvalidIndex = MAX_uint16;
-
-	/** @return true if the given index can be represented by the type. */
-	static bool IsValidIndex(const int32 Index)
-	{
-		return Index >= 0 && Index < (int32)InvalidIndex;
-	}
-
-	friend FORCEINLINE uint32 GetTypeHash(const FStateTreeDataHandle& Handle)
-	{
-		uint32 Hash = GetTypeHash(Handle.Source);
-		Hash = HashCombineFast(Hash, GetTypeHash(Handle.Source));
-		Hash = HashCombineFast(Hash, GetTypeHash(Handle.StateHandle));
-		return Hash;
-	}
-	
-	FStateTreeDataHandle() = default;
-	
-	explicit FStateTreeDataHandle(const EStateTreeDataSourceType InSource, const uint16 InIndex, const FStateTreeStateHandle InStateHandle = FStateTreeStateHandle::Invalid)
-		: Source(InSource)
-		, Index(InIndex)
-		, StateHandle(InStateHandle)
-	{
-		// Require valid state for active instance data
-		check(Source != EStateTreeDataSourceType::ActiveInstanceData || (Source == EStateTreeDataSourceType::ActiveInstanceData && StateHandle.IsValid()));
-		check(Source != EStateTreeDataSourceType::ActiveInstanceDataObject || (Source == EStateTreeDataSourceType::ActiveInstanceDataObject && StateHandle.IsValid()));
-	}
-
-	explicit FStateTreeDataHandle(const EStateTreeDataSourceType InSource, const int32 InIndex, const FStateTreeStateHandle InStateHandle = FStateTreeStateHandle::Invalid)
-		: Source(InSource)
-		, StateHandle(InStateHandle)
-	{
-		// Require valid state for active instance data
-		check(Source != EStateTreeDataSourceType::ActiveInstanceData || (Source == EStateTreeDataSourceType::ActiveInstanceData && StateHandle.IsValid()));
-		check(Source != EStateTreeDataSourceType::ActiveInstanceDataObject || (Source == EStateTreeDataSourceType::ActiveInstanceDataObject && StateHandle.IsValid()));
-		check(InIndex == INDEX_NONE || IsValidIndex(InIndex));
-		Index = static_cast<uint16>(InIndex);
-	}
-
-	bool IsValid() const
-	{
-		return Source != EStateTreeDataSourceType::None;
-	}
-
-	void Reset()
-	{
-		Source = EStateTreeDataSourceType::None;
-		Index = InvalidIndex;
-		StateHandle = FStateTreeStateHandle::Invalid;
-	}
-
-	bool operator==(const FStateTreeDataHandle& RHS) const
-	{
-		return Source == RHS.Source && Index == RHS.Index && StateHandle == RHS.StateHandle;
-	}
-	
-	bool operator!=(const FStateTreeDataHandle& RHS) const
-	{
-		return !(*this == RHS);
-	}
-
-	EStateTreeDataSourceType GetSource() const
-	{
-		return Source;
-	}
-
-	int32 GetIndex() const
-	{
-		return Index;
-	}
-
-	FStateTreeStateHandle GetState() const
-	{
-		return StateHandle;
-	}
-
-	bool IsObjectSource() const
-	{
-		return Source == EStateTreeDataSourceType::GlobalInstanceDataObject
-			|| Source == EStateTreeDataSourceType::ActiveInstanceDataObject
-			|| Source == EStateTreeDataSourceType::SharedInstanceDataObject;
-	}
-	
-	FStateTreeDataHandle ToObjectSource() const
-	{
-		switch (Source)
-		{
-		case EStateTreeDataSourceType::GlobalInstanceData:
-			return FStateTreeDataHandle(EStateTreeDataSourceType::GlobalInstanceDataObject, Index, StateHandle);
-		case EStateTreeDataSourceType::ActiveInstanceData:
-			return FStateTreeDataHandle(EStateTreeDataSourceType::ActiveInstanceDataObject, Index, StateHandle);
-		case EStateTreeDataSourceType::SharedInstanceData:
-			return FStateTreeDataHandle(EStateTreeDataSourceType::SharedInstanceDataObject, Index, StateHandle);
-		default:
-			return *this;
-		}
-	}
-	
-	FString Describe() const
-	{
-		switch (Source)
-		{
-		case EStateTreeDataSourceType::None:
-			return TEXT("None");
-		case EStateTreeDataSourceType::ContextData:
-			return FString::Printf(TEXT("Context[%d]"), Index);
-		case EStateTreeDataSourceType::GlobalInstanceData:
-			return FString::Printf(TEXT("Global[%d]"), Index);
-		case EStateTreeDataSourceType::GlobalInstanceDataObject:
-			return FString::Printf(TEXT("GlobalO[%d]"), Index);
-		case EStateTreeDataSourceType::ActiveInstanceData:
-			return FString::Printf(TEXT("Active[%d]"), Index);
-		case EStateTreeDataSourceType::ActiveInstanceDataObject:
-			return FString::Printf(TEXT("ActiveO[%d]"), Index);
-		case EStateTreeDataSourceType::SharedInstanceData:
-			return FString::Printf(TEXT("Shared[%d]"), Index);
-		case EStateTreeDataSourceType::SharedInstanceDataObject:
-			return FString::Printf(TEXT("SharedO[%d]"), Index);
-		default:
-			return TEXT("---");
-		}
-	}
-	
-private:
-	UPROPERTY()
-	EStateTreeDataSourceType Source = EStateTreeDataSourceType::None;
-
-	UPROPERTY()
-	uint16 Index = InvalidIndex;
-
-	UPROPERTY()
-	FStateTreeStateHandle StateHandle = FStateTreeStateHandle::Invalid; 
-};
-
-
 /**
  * Time duration with random variance. Stored compactly as two uint16s, which gives time range of about 650 seconds.
  * The variance is symmetric (+-) around the specified duration.
@@ -567,13 +386,11 @@ struct STATETREEMODULE_API FCompactStateTreeState
 
 	/** Index to state instance data. */
 	UPROPERTY()
-	FStateTreeIndex16 ParameterTemplateIndex = FStateTreeIndex16::Invalid;
+	FStateTreeIndex16 ParameterInstanceIndex = FStateTreeIndex16::Invalid;
 
+	/** Data view index of the input parameters. */
 	UPROPERTY()
-	FStateTreeDataHandle ParameterDataHandle = FStateTreeDataHandle::Invalid;
-
-	UPROPERTY()
-	FStateTreeIndex16 ParameterBindingsBatch = FStateTreeIndex16::Invalid;
+	FStateTreeIndex16 ParameterDataViewIndex = FStateTreeIndex16::Invalid;
 
 	/** Number of enter conditions */
 	UPROPERTY()
@@ -587,9 +404,13 @@ struct STATETREEMODULE_API FCompactStateTreeState
 	UPROPERTY()
 	uint8 TasksNum = 0;
 
-	/** Number of instance data */
+	/** Number of tasks with struct instance data. */
+	UPROPERTY()
+	uint8 TaskInstanceStructNum = 0;
+
+	/** Number of tasks with object instance data. */
 	UPROPERTY()
-	uint8 InstanceDataNum = 0;
+	uint8 TaskInstanceObjectNum = 0;
 
 	/** Type of the state */
 	UPROPERTY()
@@ -613,13 +434,9 @@ struct STATETREEMODULE_API FCompactStateTreeParameters
 {
 	GENERATED_BODY()
 
-	FCompactStateTreeParameters() = default;
+	UPROPERTY()
+	FStateTreeIndex16 BindingsBatch = FStateTreeIndex16::Invalid;
 
-	FCompactStateTreeParameters(const FInstancedPropertyBag& InParameters)
-		: Parameters(InParameters)
-	{
-	}
-	
 	UPROPERTY()
 	FInstancedPropertyBag Parameters;
 };
@@ -958,6 +775,13 @@ protected:
 	uint8* Memory = nullptr;
 };
 
+struct STATETREEMODULE_API FStateTreeTransitionDelayedState
+{
+	FStateTreeIndex16 TransitionIndex = FStateTreeIndex16::Invalid;
+	float TimeLeft = 0.0f;
+};
+
+
 /**
  * Link to another state in StateTree
  */
