StateTree: Automatic recovery from older that with duplicate node IDs

diff --git a/Engine/Plugins/Runtime/StateTree/Source/StateTreeEditorModule/Private/StateTreeEditorData.cpp b/Engine/Plugins/Runtime/StateTree/Source/StateTreeEditorModule/Private/StateTreeEditorData.cpp
--- a/Engine/Plugins/Runtime/StateTree/Source/StateTreeEditorModule/Private/StateTreeEditorData.cpp
+++ b/Engine/Plugins/Runtime/StateTree/Source/StateTreeEditorModule/Private/StateTreeEditorData.cpp
@@ -127,6 +127,7 @@ void UStateTreeEditorData::PostLoad()
 	Super::PostLoad();
 	ReparentStates();
 	FixObjectNodes();
+	FixDuplicateIDs();
 	UpdateBindingsInstanceStructs();
 }
 
@@ -540,6 +541,131 @@ void UStateTreeEditorData::FixObjectNodes()
 	}
 }
 
+void UStateTreeEditorData::FixDuplicateIDs()
+{
+	// Around version 5.1-5.3 we had issue that copy/paste or some duplication methods could create nodes with duplicate IDs.
+	// This code tries to fix that, it looks for duplicates, makes them unique, and duplicates the bindings when ID changes.
+	TSet<FGuid> FoundNodeIDs;
+
+	// Evaluators
+	for (FStateTreeEditorNode& Node : Evaluators)
+	{
+		if (const FStateTreeEvaluatorBase* Evaluator = Node.Node.GetPtr<FStateTreeEvaluatorBase>())
+		{
+			const FGuid OldID = Node.ID; 
+			if (FoundNodeIDs.Contains(Node.ID))
+			{
+				Node.ID = FGuid::NewGuid();
+				UE_LOG(LogStateTreeEditor, Log, TEXT("%s: Found Evaluator '%s' with duplicate ID, changing ID:%s to ID:%s."),
+					*GetFullName(), *Node.GetName().ToString(), *OldID.ToString(), *Node.ID.ToString());
+				EditorBindings.CopyBindings(OldID, Node.ID);
+			}
+			FoundNodeIDs.Add(Node.ID);
+		}
+	}
+	
+	// Global Tasks
+	for (FStateTreeEditorNode& Node : GlobalTasks)
+	{
+		if (const FStateTreeTaskBase* Task = Node.Node.GetPtr<FStateTreeTaskBase>())
+		{
+			const FGuid OldID = Node.ID; 
+			if (FoundNodeIDs.Contains(Node.ID))
+			{
+				Node.ID = FGuid::NewGuid();
+				UE_LOG(LogStateTreeEditor, Log, TEXT("%s: Found GlobalTask '%s' with duplicate ID, changing ID:%s to ID:%s."),
+					*GetFullName(), *Node.GetName().ToString(), *OldID.ToString(), *Node.ID.ToString());
+				EditorBindings.CopyBindings(OldID, Node.ID);
+			}
+			FoundNodeIDs.Add(Node.ID);
+		}
+	}
+	
+	VisitHierarchy([&FoundNodeIDs, &EditorBindings = EditorBindings, &Self = *this](UStateTreeState& State, UStateTreeState* ParentState)
+	{
+		// Enter conditions
+		for (FStateTreeEditorNode& Node : State.EnterConditions)
+		{
+			if (const FStateTreeConditionBase* Cond = Node.Node.GetPtr<FStateTreeConditionBase>())
+			{
+				const FGuid OldID = Node.ID;
+				
+				bool bIsAlreadyInSet = false;
+				FoundNodeIDs.Add(Node.ID, &bIsAlreadyInSet);
+				if (bIsAlreadyInSet)
+				{
+					Node.ID = FGuid::NewGuid();
+					UE_LOG(LogStateTreeEditor, Log, TEXT("%s: Found Enter Condition '%s' with duplicate ID on state '%s', changing ID:%s to ID:%s."),
+						*Self.GetFullName(), *Node.GetName().ToString(), *GetNameSafe(&State), *OldID.ToString(), *Node.ID.ToString());
+					EditorBindings.CopyBindings(OldID, Node.ID);
+				}
+			}
+		}
+
+		// Tasks
+		for (FStateTreeEditorNode& Node : State.Tasks)
+		{
+			if (const FStateTreeTaskBase* Task = Node.Node.GetPtr<FStateTreeTaskBase>())
+			{
+				const FGuid OldID = Node.ID;
+				
+				bool bIsAlreadyInSet = false;
+				FoundNodeIDs.Add(Node.ID, &bIsAlreadyInSet);
+				if (bIsAlreadyInSet)
+				{
+					Node.ID = FGuid::NewGuid();
+					UE_LOG(LogStateTreeEditor, Log, TEXT("%s: Found Task '%s' with duplicate ID on state '%s', changing ID:%s to ID:%s."),
+						*Self.GetFullName(), *Node.GetName().ToString(), *GetNameSafe(&State), *OldID.ToString(), *Node.ID.ToString());
+					EditorBindings.CopyBindings(OldID, Node.ID);
+				}
+			}
+		}
+
+		if (FStateTreeTaskBase* Task = State.SingleTask.Node.GetMutablePtr<FStateTreeTaskBase>())
+		{
+			const FGuid OldID = State.SingleTask.ID;
+
+			bool bIsAlreadyInSet = false;
+			FoundNodeIDs.Add(State.SingleTask.ID, &bIsAlreadyInSet);
+			if (bIsAlreadyInSet)
+			{
+				State.SingleTask.ID = FGuid::NewGuid();
+				UE_LOG(LogStateTreeEditor, Log, TEXT("%s: Found enter condition '%s' with duplicate ID on state '%s', changing ID:%s to ID:%s."),
+					*Self.GetFullName(), *State.SingleTask.GetName().ToString(), *GetNameSafe(&State), *OldID.ToString(), *State.SingleTask.ID.ToString());
+				EditorBindings.CopyBindings(OldID, State.SingleTask.ID);
+			}
+		}
+
+		// Transitions
+		for (FStateTreeTransition& Transition : State.Transitions)
+		{
+			for (FStateTreeEditorNode& Node : Transition.Conditions)
+			{
+				if (const FStateTreeConditionBase* Cond = Node.Node.GetPtr<FStateTreeConditionBase>())
+				{
+					const FGuid OldID = Node.ID; 
+					bool bIsAlreadyInSet = false;
+					FoundNodeIDs.Add(Node.ID, &bIsAlreadyInSet);
+					if (bIsAlreadyInSet)
+					{
+						Node.ID = FGuid::NewGuid();
+						UE_LOG(LogStateTreeEditor, Log, TEXT("%s: Found transition condition '%s' with duplicate ID on state '%s', changing ID:%s to ID:%s."),
+							*Self.GetFullName(), *Node.GetName().ToString(), *GetNameSafe(&State), *OldID.ToString(), *Node.ID.ToString());
+						EditorBindings.CopyBindings(OldID, Node.ID);
+					}
+				}
+			}
+		}
+		
+		return EStateTreeVisitor::Continue;
+	});
+
+	// It is possible that the user has changed the node type so some of the bindings might not make sense anymore, clean them up.
+	TMap<FGuid, const FStateTreeDataView> AllValues;
+	GetAllStructValues(AllValues);
+	EditorBindings.RemoveUnusedBindings(AllValues);
+}
+
 void UStateTreeEditorData::UpdateBindingsInstanceStructs()
 {
 	TMap<FGuid, const FStateTreeDataView> AllValues;
diff --git a/Engine/Plugins/Runtime/StateTree/Source/StateTreeEditorModule/Public/StateTreeEditorData.h b/Engine/Plugins/Runtime/StateTree/Source/StateTreeEditorModule/Public/StateTreeEditorData.h
--- a/Engine/Plugins/Runtime/StateTree/Source/StateTreeEditorModule/Public/StateTreeEditorData.h
+++ b/Engine/Plugins/Runtime/StateTree/Source/StateTreeEditorModule/Public/StateTreeEditorData.h
@@ -249,6 +249,7 @@ PRAGMA_ENABLE_DEPRECATION_WARNINGS
 private:
 	void FixObjectInstance(TSet<UObject*>& SeenObjects, UObject& Outer, FStateTreeEditorNode& Node);
 	void FixObjectNodes();
+	void FixDuplicateIDs();
 	void UpdateBindingsInstanceStructs();
 
 #if WITH_EDITORONLY_DATA
diff --git a/Engine/Plugins/Runtime/StateTree/Source/StateTreeModule/Private/StateTree.cpp b/Engine/Plugins/Runtime/StateTree/Source/StateTreeModule/Private/StateTree.cpp
--- a/Engine/Plugins/Runtime/StateTree/Source/StateTreeModule/Private/StateTree.cpp
+++ b/Engine/Plugins/Runtime/StateTree/Source/StateTreeModule/Private/StateTree.cpp
@@ -305,6 +305,12 @@ void UStateTree::PostLoad()
 	if (CurrentVersion < FStateTreeCustomVersion::LatestVersion)
 	{
 #if WITH_EDITOR
+		if (EditorData)
+		{
+			// Make sure all the fix up logic in the editor data has had chance to happen.
+			EditorData->ConditionalPostLoad();
+		}
+		
 		// Compiled data is in older format, try to compile the StateTree.
 		if (UE::StateTree::Delegates::OnRequestCompile.IsBound())
 		{
