[StateTreeDebugger] minor changes to improve timeline visualization and to properly handle unbound delegates.

diff --git a/Engine/Plugins/Runtime/StateTree/Source/StateTreeEditorModule/Private/Debugger/SStateTreeDebuggerEventTimelineView.cpp b/Engine/Plugins/Runtime/StateTree/Source/StateTreeEditorModule/Private/Debugger/SStateTreeDebuggerEventTimelineView.cpp
--- a/Engine/Plugins/Runtime/StateTree/Source/StateTreeEditorModule/Private/Debugger/SStateTreeDebuggerEventTimelineView.cpp
+++ b/Engine/Plugins/Runtime/StateTree/Source/StateTreeEditorModule/Private/Debugger/SStateTreeDebuggerEventTimelineView.cpp
@@ -55,10 +55,13 @@ int32 SStateTreeDebuggerEventTimelineView::PaintEvents(const FGeometry& Allotted
 				const float Y = static_cast<float>(AllottedGeometry.Size.Y - EventSize.Y) / 2.f;
 
 				// window bar
+				constexpr float HSizeReduction = 2.f;
+				constexpr float VSizeReduction = 2.f;
+				constexpr float VerticalOffset = 0.5f * VSizeReduction;
 				FSlateDrawElement::MakeBox
 				(OutDrawElements,
 					LayerId++,
-					AllottedGeometry.ToPaintGeometry(FVector2D(XEnd - XStart, EventSize.Y - 2), FSlateLayoutTransform(FVector2D(XStart, Y + 1))),
+					AllottedGeometry.ToPaintGeometry(FVector2f((XEnd - XStart) - HSizeReduction, EventSize.Y - VSizeReduction), FSlateLayoutTransform(FVector2f(XStart, Y + VerticalOffset))),
 					Brush,
 					ESlateDrawEffect::None,
 					Window.Color);
diff --git a/Engine/Plugins/Runtime/StateTree/Source/StateTreeModule/Private/Debugger/StateTreeDebugger.cpp b/Engine/Plugins/Runtime/StateTree/Source/StateTreeModule/Private/Debugger/StateTreeDebugger.cpp
--- a/Engine/Plugins/Runtime/StateTree/Source/StateTreeModule/Private/Debugger/StateTreeDebugger.cpp
+++ b/Engine/Plugins/Runtime/StateTree/Source/StateTreeModule/Private/Debugger/StateTreeDebugger.cpp
@@ -113,7 +113,7 @@ void FStateTreeDebugger::Tick(const float DeltaTime)
 	
 	UpdateInstances();
 
-	if (bSessionAnalysisPaused == false)
+	if (bSessionAnalysisPaused == false && StateTreeAsset.IsValid())
 	{
 		SyncToCurrentSessionDuration();
 	}
@@ -358,7 +358,7 @@ void FStateTreeDebugger::SetScrubStateCollectionIndex(const int32 EventCollectio
 {
 	ScrubState.SetEventCollectionIndex(EventCollectionIndex);
 
-	OnScrubStateChanged.Execute(ScrubState);
+	OnScrubStateChanged.ExecuteIfBound(ScrubState);
 
 	RefreshActiveStates();
 }
@@ -463,7 +463,7 @@ void FStateTreeDebugger::SetScrubTime(const double ScrubTime)
 {
 	if (ScrubState.SetScrubTime(ScrubTime))
 	{
-		OnScrubStateChanged.Execute(ScrubState);
+		OnScrubStateChanged.ExecuteIfBound(ScrubState);
 
 		RefreshActiveStates();
 	}
