[State Tree] Add an option to retain property values when the class of a node changes

diff --git a/Engine/Plugins/Runtime/StateTree/Source/StateTreeEditorModule/Private/Customizations/StateTreeEditorNodeDetails.cpp b/Engine/Plugins/Runtime/StateTree/Source/StateTreeEditorModule/Private/Customizations/StateTreeEditorNodeDetails.cpp
--- a/Engine/Plugins/Runtime/StateTree/Source/StateTreeEditorModule/Private/Customizations/StateTreeEditorNodeDetails.cpp
+++ b/Engine/Plugins/Runtime/StateTree/Source/StateTreeEditorModule/Private/Customizations/StateTreeEditorNodeDetails.cpp
@@ -30,6 +30,7 @@
 #include "Widgets/Input/SSearchBox.h"
 #include "EditorFontGlyphs.h"
 #include "StateTreeEditorNodeUtils.h"
+#include "StateTreeEditorSettings.h"
 #include "Framework/Notifications/NotificationManager.h"
 #include "Widgets/Notifications/SNotificationList.h"
 #include "Debugger/StateTreeDebuggerUIExtensions.h"
@@ -289,6 +290,125 @@ namespace UE::StateTreeEditor::Internal
 			}
 		}
 	}
+
+	struct FNodeRetainPropertyData
+	{
+		FStateTreeNodeBase* NodeBase = nullptr;
+		const UScriptStruct* NodeBaseStruct = nullptr;
+		const UStruct* InstanceStruct = nullptr;
+		void* InstanceData = nullptr;
+	};
+
+	FNodeRetainPropertyData GetNodeData(FStateTreeEditorNode& EditorNode)
+	{
+		FNodeRetainPropertyData Data;
+		Data.NodeBase = EditorNode.Node.GetMutablePtr<FStateTreeNodeBase>();
+
+		if (Data.NodeBase)
+		{
+			Data.NodeBaseStruct = EditorNode.Node.GetScriptStruct();
+			if (const UStruct* InstanceDataType = Data.NodeBase->GetInstanceDataType())
+			{
+				if (InstanceDataType->IsA<UScriptStruct>())
+				{
+					Data.InstanceStruct = EditorNode.Instance.GetScriptStruct();
+					Data.InstanceData = EditorNode.Instance.GetMutableMemory();
+				}
+				else if (InstanceDataType->IsA<UClass>())
+				{
+					Data.InstanceStruct = EditorNode.InstanceObject.GetClass();
+					Data.InstanceData = EditorNode.InstanceObject;
+				}
+			}
+		}
+
+		return Data;
+	}
+
+	void CopyPropertyValues(const UStruct* OldStruct, const void* OldData, const UStruct* NewStruct, void* NewData)
+	{
+		for (TFieldIterator<FProperty> It(OldStruct); It; ++It)
+		{
+			const FProperty* OldProperty = *It;
+			const FProperty* NewProperty = NewStruct->FindPropertyByName(OldProperty->GetFName());
+			if (!NewProperty)
+			{
+				// Let's check if we have the same property present but with(out) the 'b' prefix
+				const FBoolProperty* BoolProperty = ExactCastField<const FBoolProperty>(OldProperty);
+				if (!BoolProperty)
+					continue;
+
+				FString String = OldProperty->GetName();
+				if (String.IsEmpty())
+					continue;
+
+				if (String[0] == TEXT('b'))
+					String.RightChopInline(1, false);
+				else
+					String.InsertAt(0, TEXT('b'));
+
+				NewProperty = NewStruct->FindPropertyByName(FName(String));
+			}
+
+			constexpr uint64 WantedFlags = CPF_Edit;
+			constexpr uint64 UnwantedFlags = CPF_DisableEditOnInstance | CPF_EditConst;
+
+			if (NewProperty
+				&& OldProperty->HasAllPropertyFlags(WantedFlags)
+				&& NewProperty->HasAllPropertyFlags(WantedFlags)
+				&& !OldProperty->HasAnyPropertyFlags(UnwantedFlags)
+				&& !NewProperty->HasAnyPropertyFlags(UnwantedFlags)
+				&& NewProperty->SameType(OldProperty))
+			{
+				OldProperty->CopyCompleteValue(
+					NewProperty->ContainerPtrToValuePtr<void>(NewData),
+					OldProperty->ContainerPtrToValuePtr<void>(OldData)
+				);
+			}
+		}
+	}
+
+	void RetainProperties(FStateTreeEditorNode& OldNode, FStateTreeEditorNode& NewNode)
+	{
+		const FNodeRetainPropertyData OldNodeData = GetNodeData(OldNode);
+		const FNodeRetainPropertyData NewNodeData = GetNodeData(NewNode);
+
+		if (OldNodeData.NodeBase && NewNodeData.NodeBase)
+		{
+			// Copy node -> node
+			CopyPropertyValues(
+				OldNodeData.NodeBaseStruct, OldNodeData.NodeBase,
+				NewNodeData.NodeBaseStruct, NewNodeData.NodeBase
+			);
+
+			if (OldNodeData.InstanceStruct && OldNodeData.InstanceData)
+			{
+				// Copy instance data -> node
+				CopyPropertyValues(
+					OldNodeData.InstanceStruct, OldNodeData.InstanceData,
+					NewNodeData.NodeBaseStruct, NewNodeData.NodeBase
+				);
+
+				if (NewNodeData.InstanceStruct && NewNodeData.InstanceData)
+				{
+					// Copy instance data -> instance data
+					CopyPropertyValues(
+						OldNodeData.InstanceStruct, OldNodeData.InstanceData,
+						NewNodeData.InstanceStruct, NewNodeData.InstanceData
+					);
+				}
+			}
+
+			if (NewNodeData.InstanceStruct && NewNodeData.InstanceData)
+			{
+				// Copy node -> instance data
+				CopyPropertyValues(
+					OldNodeData.NodeBaseStruct, OldNodeData.NodeBase,
+					NewNodeData.InstanceStruct, NewNodeData.InstanceData
+				);
+			}
+		}
+	}
 } // UE::StateTreeEditor::Internal
 
 // Customized version of FInstancedStructDataDetails used to hide bindable properties.
@@ -1859,7 +1979,10 @@ void FStateTreeEditorNodeDetails::OnStructPicked(const UScriptStruct* InStruct)
 			if (UObject* Outer = OuterObjects[Index])
 			{
 				if (FStateTreeEditorNode* Node = static_cast<FStateTreeEditorNode*>(RawNodeData[Index]))
-					{
+				{
+					const bool bRetainProperties = InStruct && UStateTreeEditorSettings::Get().bRetainNodePropertyValues;
+					FStateTreeEditorNode OldNode = bRetainProperties ? *Node : FStateTreeEditorNode();
+
 					Node->Reset();
 					
 					if (InStruct)
@@ -1913,6 +2036,11 @@ void FStateTreeEditorNodeDetails::OnStructPicked(const UScriptStruct* InStruct)
 								Node->InstanceObject = NewObject<UObject>(Outer, InstanceClass);
 							}
 						}
+
+						if (bRetainProperties)
+						{
+							UE::StateTreeEditor::Internal::RetainProperties(OldNode, *Node);
+						}
 					}
 				}
 			}
@@ -1954,6 +2082,9 @@ void FStateTreeEditorNodeDetails::OnClassPicked(const UClass* InClass) const
 			{
 				if (FStateTreeEditorNode* Node = static_cast<FStateTreeEditorNode*>(RawNodeData[Index]))
 				{
+					bool bRetainProperties = InClass && UStateTreeEditorSettings::Get().bRetainNodePropertyValues;
+					FStateTreeEditorNode OldNode = bRetainProperties ? *Node : FStateTreeEditorNode();
+
 					Node->Reset();
 
 					if (InClass && InClass->IsChildOf(UStateTreeTaskBlueprintBase::StaticClass()))
@@ -1989,6 +2120,16 @@ void FStateTreeEditorNodeDetails::OnClassPicked(const UClass* InClass) const
 
 						Node->ID = FGuid::NewGuid();
 					}
+					else
+					{
+						// Not retaining properties if we haven't initialized a new node
+						bRetainProperties = false;
+					}
+
+					if (bRetainProperties)
+					{
+						UE::StateTreeEditor::Internal::RetainProperties(OldNode, *Node);
+					}
 				}
 			}
 		}
diff --git a/Engine/Plugins/Runtime/StateTree/Source/StateTreeEditorModule/Public/StateTreeEditorSettings.h b/Engine/Plugins/Runtime/StateTree/Source/StateTreeEditorModule/Public/StateTreeEditorSettings.h
--- a/Engine/Plugins/Runtime/StateTree/Source/StateTreeEditorModule/Public/StateTreeEditorSettings.h
+++ b/Engine/Plugins/Runtime/StateTree/Source/StateTreeEditorModule/Public/StateTreeEditorSettings.h
@@ -32,6 +32,13 @@ public:
 	UPROPERTY(EditAnywhere, config, Category = "Debugger")
 	bool bShouldDebuggerResetDataOnNewPIESession = false;
 
+	/**
+	 * If enabled, changing the class of a node will try to copy over values of properties with the same name and type.
+	 * i.e. if you change one condition for another, and both have a "Target" BB key selector, it'll be kept.
+	 */
+	UPROPERTY(EditAnywhere, config, Experimental, Category = "Experimental")
+	bool bRetainNodePropertyValues = false;
+
 protected:
 #if WITH_EDITOR
 	virtual FText GetSectionText() const override;
