StateTree: remove redundand check for running state when stopping the StateTree.

diff --git a/Engine/Plugins/Runtime/StateTree/Source/StateTreeModule/Private/StateTreeExecutionContext.cpp b/Engine/Plugins/Runtime/StateTree/Source/StateTreeModule/Private/StateTreeExecutionContext.cpp
--- a/Engine/Plugins/Runtime/StateTree/Source/StateTreeModule/Private/StateTreeExecutionContext.cpp
+++ b/Engine/Plugins/Runtime/StateTree/Source/StateTreeModule/Private/StateTreeExecutionContext.cpp
@@ -365,10 +365,7 @@ EStateTreeRunStatus FStateTreeExecutionContext::Stop(EStateTreeRunStatus Complet
 		NewFrame.RootState = FStateTreeStateHandle::Root;
 		NewFrame.ActiveStates = {};
 		
-		if (Exec.TreeRunStatus == EStateTreeRunStatus::Running)
-		{
-			ExitState(Transition);
-		}
+		ExitState(Transition);
 
 		// Stop evaluators and global tasks.
 		StopEvaluatorsAndGlobalTasks(CompletionStatus);
