StateTree: Fixed crash caused by GetDataView() with SharedInstanceDataObject returning the wrapper struct and not the actual object.

diff --git a/Engine/Plugins/Runtime/StateTree/Source/StateTreeModule/Private/StateTreeExecutionContext.cpp b/Engine/Plugins/Runtime/StateTree/Source/StateTreeModule/Private/StateTreeExecutionContext.cpp
--- a/Engine/Plugins/Runtime/StateTree/Source/StateTreeModule/Private/StateTreeExecutionContext.cpp
+++ b/Engine/Plugins/Runtime/StateTree/Source/StateTreeModule/Private/StateTreeExecutionContext.cpp
@@ -839,7 +839,7 @@ FStateTreeDataView FStateTreeExecutionContext::GetDataView(const FStateTreeExecu
 		return CurrentlyProcessedSharedInstanceStorage->GetMutableStruct(Handle.GetIndex());
 	case EStateTreeDataSourceType::SharedInstanceDataObject:
 		check(CurrentlyProcessedSharedInstanceStorage);
-		return CurrentlyProcessedSharedInstanceStorage->GetMutableStruct(Handle.GetIndex());
+		return CurrentlyProcessedSharedInstanceStorage->GetMutableObject(Handle.GetIndex());
 
 	case EStateTreeDataSourceType::ContextData:
 		return ContextDataViews[Handle.GetIndex()];
