StateTree: Post edit change events for nodes and their instance data

diff --git a/Engine/Plugins/Runtime/StateTree/Source/StateTreeEditorModule/Private/StateTreeEditorData.cpp b/Engine/Plugins/Runtime/StateTree/Source/StateTreeEditorModule/Private/StateTreeEditorData.cpp
--- a/Engine/Plugins/Runtime/StateTree/Source/StateTreeEditorModule/Private/StateTreeEditorData.cpp
+++ b/Engine/Plugins/Runtime/StateTree/Source/StateTreeEditorModule/Private/StateTreeEditorData.cpp
@@ -8,6 +8,7 @@
 #include "StateTreeTaskBase.h"
 #include "Algo/LevenshteinDistance.h"
 #include "StateTreeEditorModule.h"
+#include "StateTreePropertyHelpers.h"
 
 #if WITH_EDITOR
 #include "Engine/UserDefinedStruct.h"
@@ -203,6 +204,8 @@ void UStateTreeEditorData::PostEditChangeChainProperty(FPropertyChangedChainEven
 			}
 		}
 	}
+
+	UE::StateTree::PropertyHelpers::DispatchPostEditToNodes(*this, PropertyChangedEvent);
 }
 #endif // WITH_EDITOR
 
diff --git a/Engine/Plugins/Runtime/StateTree/Source/StateTreeEditorModule/Private/StateTreePropertyHelpers.cpp b/Engine/Plugins/Runtime/StateTree/Source/StateTreeEditorModule/Private/StateTreePropertyHelpers.cpp
new file mode 100644
--- /dev/null
+++ b/Engine/Plugins/Runtime/StateTree/Source/StateTreeEditorModule/Private/StateTreePropertyHelpers.cpp
@@ -0,0 +1,206 @@
+﻿// Copyright Epic Games, Inc. All Rights Reserved.
+
+#include "StateTreePropertyHelpers.h"
+#include "StateTreeEditorNode.h"
+
+namespace UE::StateTree::PropertyHelpers
+{
+
+void DispatchPostEditToNodes(UObject& Owner, FPropertyChangedChainEvent& InPropertyChangedEvent)
+{
+	// Walk back from the changed property and look for first FStateTreeEditorNode, and call the node specific post edit methods.
+	
+	const TDoubleLinkedList<FProperty*>::TDoubleLinkedListNode* CurrentPropNode = InPropertyChangedEvent.PropertyChain.GetHead();
+	const FProperty* HeadProperty = CurrentPropNode->GetValue();
+	check(HeadProperty);
+	if (HeadProperty->GetOwnerClass() != Owner.GetClass())
+	{
+		return;
+	}
+	
+	uint8* CurrentAddress = reinterpret_cast<uint8*>(&Owner);
+	while (CurrentPropNode)
+	{
+		const FProperty* CurrentProperty = CurrentPropNode->GetValue();
+		check(CurrentProperty);
+		CurrentAddress = CurrentAddress + CurrentProperty->GetOffset_ForInternal();
+
+		while (const FArrayProperty* ArrayProperty = CastField<FArrayProperty>(CurrentProperty))
+		{
+			FScriptArrayHelper Helper(ArrayProperty, CurrentAddress);
+			const int32 Index = InPropertyChangedEvent.GetArrayIndex(ArrayProperty->GetName());
+			if (Index == INDEX_NONE)
+			{
+				return;
+			}
+
+			CurrentAddress = Helper.GetRawPtr(Index);
+			CurrentProperty = ArrayProperty->Inner;
+		}
+
+		if (const FStructProperty* StructProperty = CastField<FStructProperty>(CurrentProperty))
+		{
+			if (StructProperty->Struct == FStateTreeEditorNode::StaticStruct())
+			{
+				FStateTreeEditorNode& EditorNode = *reinterpret_cast<FStateTreeEditorNode*>(CurrentAddress);
+				if (FStateTreeNodeBase* StateTreeNode = EditorNode.Node.GetMutablePtr<FStateTreeNodeBase>())
+				{
+					// Check that the path contains EditorNode's: Node, Instance or Instance Object
+					if (const TDoubleLinkedList<FProperty*>::TDoubleLinkedListNode* EditorNodeMemberPropNode = CurrentPropNode->GetNextNode())
+					{
+						// Check that we have a changed property on one of the above properties.
+						if (const TDoubleLinkedList<FProperty*>::TDoubleLinkedListNode* ActiveMemberPropNode = EditorNodeMemberPropNode->GetNextNode()) 
+						{
+							// Update the event
+							const FProperty* EditorNodeChildMember = EditorNodeMemberPropNode->GetValue();
+							check(EditorNodeChildMember);
+
+							// Take copy of the event, we'll modify it.
+							FEditPropertyChain PropertyChainCopy;
+							for (const TDoubleLinkedList<FProperty*>::TDoubleLinkedListNode* Node = InPropertyChangedEvent.PropertyChain.GetHead(); Node->GetNextNode(); Node = Node->GetNextNode())
+							{
+								PropertyChainCopy.AddTail(Node->GetValue());
+							}
+							FPropertyChangedChainEvent PropertyChangedEvent(PropertyChainCopy, InPropertyChangedEvent);
+
+							PropertyChangedEvent.SetActiveMemberProperty(ActiveMemberPropNode->GetValue());
+							PropertyChangedEvent.PropertyChain.SetActiveMemberPropertyNode(PropertyChangedEvent.MemberProperty);
+
+							// To be consistent with the other property chain callbacks, do not cross object boundary.
+							const TDoubleLinkedList<FProperty*>::TDoubleLinkedListNode* ActivePropNode = ActiveMemberPropNode;
+							while (ActivePropNode->GetNextNode())
+							{
+								if (CastField<FObjectProperty>(ActivePropNode->GetValue()))
+								{
+									break;
+								}
+								ActivePropNode = ActivePropNode->GetNextNode();
+							}
+							
+							PropertyChangedEvent.Property = ActivePropNode->GetValue();
+							PropertyChangedEvent.PropertyChain.SetActivePropertyNode(PropertyChangedEvent.Property);
+
+							if (EditorNodeChildMember->GetFName() == GET_MEMBER_NAME_CHECKED(FStateTreeEditorNode, Node))
+							{
+								StateTreeNode->PostEditNodeChangeChainProperty(PropertyChangedEvent, EditorNode.GetInstance());
+							}
+							else if (EditorNodeChildMember->GetFName() == GET_MEMBER_NAME_CHECKED(FStateTreeEditorNode, Instance))
+							{
+								if (EditorNode.Instance.IsValid())
+								{
+									StateTreeNode->PostEditInstanceDataChangeChainProperty(PropertyChangedEvent, FStateTreeDataView(EditorNode.Instance));
+								}
+							}
+							else if (EditorNodeChildMember->GetFName() == GET_MEMBER_NAME_CHECKED(FStateTreeEditorNode, InstanceObject))
+							{
+								if (EditorNode.InstanceObject)
+								{
+									StateTreeNode->PostEditInstanceDataChangeChainProperty(PropertyChangedEvent, FStateTreeDataView(EditorNode.InstanceObject));
+								}
+							}
+						}
+					}
+				}
+
+				break;
+			}
+
+			CurrentPropNode = CurrentPropNode->GetNextNode();
+		}
+		else
+		{
+			break;
+		}
+	}
+}
+
+}; // UE::StateTree::PropertyHelpers
+
+// ------------------------------------------------------------------------------
+// FStateTreeEditPropertyPath
+// ------------------------------------------------------------------------------
+FStateTreeEditPropertyPath::FStateTreeEditPropertyPath(const UStruct* BaseStruct, const FString& InPath)
+{
+	TArray<FString> PathSegments;
+	InPath.ParseIntoArray(PathSegments, TEXT("."));
+
+	const UStruct* CurrBase = BaseStruct;
+	for (const FString& Segment : PathSegments)
+	{
+		const FName PropertyName(Segment);
+		if (const FProperty* Property = CurrBase->FindPropertyByName(PropertyName))
+		{
+			Path.Emplace(Property, PropertyName);
+
+			if (const FArrayProperty* ArrayProperty = CastField<FArrayProperty>(Property))
+			{
+				Property = ArrayProperty->Inner;
+			}
+
+			if (const FStructProperty* StructProperty = CastField<FStructProperty>(Property))
+			{
+				CurrBase = StructProperty->Struct;
+			}
+			else if (const FObjectProperty* ObjectProperty = CastField<FObjectProperty>(Property))
+			{
+				CurrBase = ObjectProperty->PropertyClass;
+			}
+		}
+		else
+		{
+			checkf(false, TEXT("Path %s id not part of type %s."), *InPath, *GetNameSafe(BaseStruct));
+			Path.Reset();
+			break;
+		}
+	}
+}
+
+FStateTreeEditPropertyPath::FStateTreeEditPropertyPath(const FPropertyChangedChainEvent& PropertyChangedEvent)
+{
+	FEditPropertyChain::TDoubleLinkedListNode* PropertyNode = PropertyChangedEvent.PropertyChain.GetActiveMemberNode();
+	while (PropertyNode != nullptr)
+	{
+		if (FProperty* Property = PropertyNode->GetValue())
+		{
+			const FName PropertyName = Property->GetFName(); 
+			const int32 ArrayIndex = PropertyChangedEvent.GetArrayIndex(PropertyName.ToString());
+			Path.Emplace(Property, PropertyName, ArrayIndex);
+		}
+		PropertyNode = PropertyNode->GetNextNode();
+	}
+}
+
+bool FStateTreeEditPropertyPath::ContainsPath(const FStateTreeEditPropertyPath& InPath) const
+{
+	if (InPath.Path.Num() > Path.Num())
+    {
+    	return false;
+    }
+
+    for (TConstEnumerateRef<FStateTreeEditPropertySegment> Segment : EnumerateRange(InPath.Path))
+    {
+    	if (Segment->PropertyName != Path[Segment.GetIndex()].PropertyName)
+    	{
+    		return false;
+    	}
+    }
+    return true;
+}
+
+/** @return true if the property path is exactly the specified path. */
+bool FStateTreeEditPropertyPath::IsPathExact(const FStateTreeEditPropertyPath& InPath) const
+{
+	if (InPath.Path.Num() != Path.Num())
+	{
+		return false;
+	}
+
+	for (TConstEnumerateRef<FStateTreeEditPropertySegment> Segment : EnumerateRange(InPath.Path))
+	{
+		if (Segment->PropertyName != Path[Segment.GetIndex()].PropertyName)
+		{
+			return false;
+		}
+	}
+	return true;
+}
diff --git a/Engine/Plugins/Runtime/StateTree/Source/StateTreeEditorModule/Private/StateTreePropertyHelpers.h b/Engine/Plugins/Runtime/StateTree/Source/StateTreeEditorModule/Private/StateTreePropertyHelpers.h
--- a/Engine/Plugins/Runtime/StateTree/Source/StateTreeEditorModule/Private/StateTreePropertyHelpers.h
+++ b/Engine/Plugins/Runtime/StateTree/Source/StateTreeEditorModule/Private/StateTreePropertyHelpers.h
@@ -13,6 +13,12 @@ struct FSlateBrush;
 
 namespace UE::StateTree::PropertyHelpers {
 
+/**
+ * Dispatches PostEditChange to all FState
+ * Assumes property chain head is member property of Owner. 
+ */
+void DispatchPostEditToNodes(UObject& Owner, FPropertyChangedChainEvent& PropertyChangedEvent);
+
 /**
  * Gets a struct value from property handle, checks type before access. Expects T is struct.
  * @param ValueProperty Handle to property where value is got from.
@@ -118,7 +124,7 @@ FPropertyAccess::Result SetStructValue(const TSharedPtr<IPropertyHandle>& ValueP
 	return FPropertyAccess::Success;
 }
 
-} // UE::StateTree::PropertyHelpers
+}; // UE::StateTree::PropertyHelpers
 
 /**
  * Helper class to deal with relative property paths in PostEditChangeChainProperty().
@@ -145,93 +151,16 @@ public:
 	FStateTreeEditPropertyPath() = default;
 
 	/** Makes property path relative to BaseStruct. Checks if the path is not part of the type. */
-	explicit FStateTreeEditPropertyPath(const UStruct* BaseStruct, const FString& InPath)
-	{
-		TArray<FString> PathSegments;
-		InPath.ParseIntoArray(PathSegments, TEXT("."));
-
-		const UStruct* CurrBase = BaseStruct;
-		for (const FString& Segment : PathSegments)
-		{
-			const FName PropertyName(Segment);
-			if (const FProperty* Property = CurrBase->FindPropertyByName(PropertyName))
-			{
-				Path.Emplace(Property, PropertyName);
-
-				if (const FArrayProperty* ArrayProperty = CastField<FArrayProperty>(Property))
-				{
-					Property = ArrayProperty->Inner;
-				}
-
-				if (const FStructProperty* StructProperty = CastField<FStructProperty>(Property))
-				{
-					CurrBase = StructProperty->Struct;
-				}
-				else if (const FObjectProperty* ObjectProperty = CastField<FObjectProperty>(Property))
-				{
-					CurrBase = ObjectProperty->PropertyClass;
-				}
-			}
-			else
-			{
-				checkf(false, TEXT("Path %s id not part of type %s."), *InPath, *GetNameSafe(BaseStruct));
-				Path.Reset();
-				break;
-			}
-		}
-	}
+	explicit FStateTreeEditPropertyPath(const UStruct* BaseStruct, const FString& InPath);
 
 	/** Makes property path from property change event. */
-	explicit FStateTreeEditPropertyPath(const FPropertyChangedChainEvent& PropertyChangedEvent)
-	{
-		FEditPropertyChain::TDoubleLinkedListNode* PropertyNode = PropertyChangedEvent.PropertyChain.GetActiveMemberNode();
-		while (PropertyNode != nullptr)
-		{
-			if (FProperty* Property = PropertyNode->GetValue())
-			{
-				const FName PropertyName = Property->GetFName(); 
-				const int32 ArrayIndex = PropertyChangedEvent.GetArrayIndex(PropertyName.ToString());
-				Path.Emplace(Property, PropertyName, ArrayIndex);
-			}
-			PropertyNode = PropertyNode->GetNextNode();
-		}
-	}
+	explicit FStateTreeEditPropertyPath(const FPropertyChangedChainEvent& PropertyChangedEvent);
 
 	/** @return true if the property path contains specified path. */
-	bool ContainsPath(const FStateTreeEditPropertyPath& InPath) const
-	{
-		if (InPath.Path.Num() > Path.Num())
-    	{
-    		return false;
-    	}
-
-    	for (TConstEnumerateRef<FStateTreeEditPropertySegment> Segment : EnumerateRange(InPath.Path))
-    	{
-    		if (Segment->PropertyName != Path[Segment.GetIndex()].PropertyName)
-    		{
-    			return false;
-    		}
-    	}
-    	return true;
-	}
+	bool ContainsPath(const FStateTreeEditPropertyPath& InPath) const;
 
 	/** @return true if the property path is exactly the specified path. */
-	bool IsPathExact(const FStateTreeEditPropertyPath& InPath) const
-	{
-		if (InPath.Path.Num() != Path.Num())
-		{
-			return false;
-		}
-
-		for (TConstEnumerateRef<FStateTreeEditPropertySegment> Segment : EnumerateRange(InPath.Path))
-		{
-			if (Segment->PropertyName != Path[Segment.GetIndex()].PropertyName)
-			{
-				return false;
-			}
-		}
-		return true;
-	}
+	bool IsPathExact(const FStateTreeEditPropertyPath& InPath) const;
 
 	/** @return array index at specified property, or INDEX_NONE, if the property is not array or property not found.  */
 	int32 GetPropertyArrayIndex(const FStateTreeEditPropertyPath& InPath) const
diff --git a/Engine/Plugins/Runtime/StateTree/Source/StateTreeEditorModule/Private/StateTreeState.cpp b/Engine/Plugins/Runtime/StateTree/Source/StateTreeEditorModule/Private/StateTreeState.cpp
--- a/Engine/Plugins/Runtime/StateTree/Source/StateTreeEditorModule/Private/StateTreeState.cpp
+++ b/Engine/Plugins/Runtime/StateTree/Source/StateTreeEditorModule/Private/StateTreeState.cpp
@@ -275,6 +275,7 @@ void UStateTreeState::PostEditChangeChainProperty(FPropertyChangedChainEvent& Pr
 		}
 	}
 
+	UE::StateTree::PropertyHelpers::DispatchPostEditToNodes(*this, PropertyChangedEvent);
 }
 
 void UStateTreeState::PostLoad()
diff --git a/Engine/Plugins/Runtime/StateTree/Source/StateTreeModule/Public/StateTreeNodeBase.h b/Engine/Plugins/Runtime/StateTree/Source/StateTreeModule/Public/StateTreeNodeBase.h
--- a/Engine/Plugins/Runtime/StateTree/Source/StateTreeModule/Public/StateTreeNodeBase.h
+++ b/Engine/Plugins/Runtime/StateTree/Source/StateTreeModule/Public/StateTreeNodeBase.h
@@ -1,4 +1,4 @@
-﻿// Copyright Epic Games, Inc. All Rights Reserved.
+// Copyright Epic Games, Inc. All Rights Reserved.
 
 #pragma once
 
@@ -62,6 +62,20 @@ PRAGMA_ENABLE_DEPRECATION_WARNINGS
 	 * @param BindingLookup Reference to binding lookup which can be used to reason about property paths.
 	 */
 	virtual void OnBindingChanged(const FGuid& ID, FStateTreeDataView InstanceData, const FStateTreePropertyPath& SourcePath, const FStateTreePropertyPath& TargetPath, const IStateTreeBindingLookup& BindingLookup) {}
+
+	/**
+	 * Called when a property of the node has been modified externally
+	 * @param PropertyChangedEvent The event for the changed property. PropertyChain's active properties are set relative to node.
+	 * @param InstanceData view to the instance data, can be struct or class.
+	 */
+	virtual void PostEditNodeChangeChainProperty(const FPropertyChangedChainEvent& PropertyChangedEvent, FStateTreeDataView InstanceDataView) {}
+
+	/**
+	 * Called when a property of node's instance data has been modified externally
+	 * @param PropertyChangedEvent The event for the changed property. PropertyChain's active properties are set relative to instance data.
+	 * @param InstanceData view to the instance data, can be struct or class.
+	 */
+	virtual void PostEditInstanceDataChangeChainProperty(const FPropertyChangedChainEvent& PropertyChangedEvent, FStateTreeDataView InstanceDataView) {}
 #endif
 
 	/** Name of the node. */
