StateTree: Fix crash trying to access invalid instance data
- During state selection keep track if the current state we're expanding is part of the active states, if not prevent access to instance data

diff --git a/Engine/Plugins/Runtime/StateTree/Source/StateTreeModule/Private/StateTreeExecutionContext.cpp b/Engine/Plugins/Runtime/StateTree/Source/StateTreeModule/Private/StateTreeExecutionContext.cpp
--- a/Engine/Plugins/Runtime/StateTree/Source/StateTreeModule/Private/StateTreeExecutionContext.cpp
+++ b/Engine/Plugins/Runtime/StateTree/Source/StateTreeModule/Private/StateTreeExecutionContext.cpp
@@ -1793,7 +1793,6 @@ void FStateTreeExecutionContext::StopEvaluatorsAndGlobalTasks(const EStateTreeRu
 	}
 }
 
-
 EStateTreeRunStatus FStateTreeExecutionContext::StartTemporaryEvaluatorsAndGlobalTasks(const FStateTreeExecutionFrame* CurrentParentFrame, const FStateTreeExecutionFrame& CurrentFrame)
 {
 	if (!CurrentFrame.bIsGlobalFrame)
@@ -2675,16 +2674,20 @@ bool FStateTreeExecutionContext::SelectState(const FStateTreeExecutionFrame& Cur
 	}
 
 	// Copy common frames over.
+	// ReferenceCurrentFrame is the original of the last copied frame. It will be used to keep track if we are following the current active frames and states.
+	const FStateTreeExecutionFrame* CurrentFrameInActiveFrames  = nullptr;
 	if (CurrentFrameIndex != INDEX_NONE)
 	{
 		const int32 NumCommonFrames = CurrentFrameIndex + 1;
 		OutNextActiveFrames = MakeArrayView(Exec.ActiveFrames.GetData(), NumCommonFrames);
+		CurrentFrameInActiveFrames  = &Exec.ActiveFrames[CurrentFrameIndex];
 	}
 	else if (CurrentStateTreeIndex != INDEX_NONE)
 	{
 		// If we could not find a common frame, we assume that we jumped to different subtree in same asset.
 		const int32 NumCommonFrames = CurrentStateTreeIndex + 1;
 		OutNextActiveFrames = MakeArrayView(Exec.ActiveFrames.GetData(), NumCommonFrames);
+		CurrentFrameInActiveFrames  = &Exec.ActiveFrames[CurrentStateTreeIndex];
 	}
 	else
 	{
@@ -2713,7 +2716,7 @@ bool FStateTreeExecutionContext::SelectState(const FStateTreeExecutionFrame& Cur
 	
 	// We take copy of the last frame and assign it later, as SelectStateInternal() might change the array and invalidate the pointer.
 	const FStateTreeExecutionFrame* CurrentParentFrame = LastFrameIndex > 0 ? &OutNextActiveFrames[LastFrameIndex - 1] : nullptr; 
-	if (SelectStateInternal(CurrentParentFrame, OutNextActiveFrames[LastFrameIndex], NextState, OutNextActiveFrames))
+	if (SelectStateInternal(CurrentParentFrame, OutNextActiveFrames[LastFrameIndex], CurrentFrameInActiveFrames , NextState, OutNextActiveFrames))
 	{
 		return true;
 	}
@@ -2738,7 +2741,7 @@ bool FStateTreeExecutionContext::SelectState(const FStateTreeExecutionFrame& Cur
 	
 				// We take copy of the last frame and assign it later, as SelectStateInternal() might change the array and invalidate the pointer.
 				CurrentParentFrame = LastFrameIndex > 0 ? &OutNextActiveFrames[LastFrameIndex - 1] : nullptr; 
-				if (SelectStateInternal(CurrentParentFrame, OutNextActiveFrames[LastFrameIndex], ChildStateHandle, OutNextActiveFrames))
+				if (SelectStateInternal(CurrentParentFrame, OutNextActiveFrames[LastFrameIndex], CurrentFrameInActiveFrames , ChildStateHandle, OutNextActiveFrames))
 				{
 					return true;
 				}
@@ -2752,6 +2755,7 @@ bool FStateTreeExecutionContext::SelectState(const FStateTreeExecutionFrame& Cur
 bool FStateTreeExecutionContext::SelectStateInternal(
 	const FStateTreeExecutionFrame* CurrentParentFrame,
 	FStateTreeExecutionFrame& CurrentFrame,
+	const FStateTreeExecutionFrame* CurrentFrameInActiveFrames,
 	const FStateTreeStateHandle NextStateHandle,
 	TArray<FStateTreeExecutionFrame, TFixedAllocator<MaxExecutionFrames>>& OutNextActiveFrames)
 {
@@ -2801,6 +2805,20 @@ bool FStateTreeExecutionContext::SelectStateInternal(
 			return false;
 		}
 
+		// Check if we're still tracking on the current active frame and state.
+		// If not, mark the ActiveInstanceIndexBase as invalid because we no longer can access the instance data during selection.
+		const FStateTreeIndex16 PrevActiveInstanceIndexBase = CurrentFrame.ActiveInstanceIndexBase; 
+		if (CurrentFrame.ActiveInstanceIndexBase.IsValid()
+			&& CurrentFrameInActiveFrames)
+		{
+			const int32 CurrentStateIndex = CurrentFrame.ActiveStates.Num() - 1;
+			const FStateTreeStateHandle MatchingActiveHandle = CurrentFrameInActiveFrames->ActiveStates.GetStateSafe(CurrentStateIndex);
+			if (MatchingActiveHandle != NextStateHandle)
+			{
+				CurrentFrame.ActiveInstanceIndexBase = FStateTreeIndex16();
+			}
+		}
+		
 		if (NextState.ParameterDataHandle.IsValid())
 		{
 			// Instantiate state parameters if not done yet.
@@ -2880,7 +2898,7 @@ bool FStateTreeExecutionContext::SelectStateInternal(
 				OutNextActiveFrames.Push(NewFrame);
 
 				// If State is linked, proceed to the linked state.
-				if (SelectStateInternal(&CurrentFrame, OutNextActiveFrames.Last(), NewFrame.RootState, OutNextActiveFrames))
+				if (SelectStateInternal(&CurrentFrame, OutNextActiveFrames.Last(), ExistingFrame, NewFrame.RootState, OutNextActiveFrames))
 				{
 					return true;
 				}
@@ -2955,7 +2973,7 @@ bool FStateTreeExecutionContext::SelectStateInternal(
 				OutNextActiveFrames.Push(NewFrame);
 
 				// If State is linked, proceed to the linked state.
-				if (SelectStateInternal(&CurrentFrame, OutNextActiveFrames.Last(), NewFrame.RootState, OutNextActiveFrames))
+				if (SelectStateInternal(&CurrentFrame, OutNextActiveFrames.Last(), ExistingFrame, NewFrame.RootState, OutNextActiveFrames))
 				{
 					return true;
 				}
@@ -3054,7 +3072,7 @@ bool FStateTreeExecutionContext::SelectStateInternal(
 				// If the state has children, proceed to select children.
 				for (uint16 ChildState = NextState.ChildrenBegin; ChildState < NextState.ChildrenEnd; ChildState = CurrentStateTree->States[ChildState].GetNextSibling())
 				{
-					if (SelectStateInternal(CurrentParentFrame, CurrentFrame, FStateTreeStateHandle(ChildState), OutNextActiveFrames))
+					if (SelectStateInternal(CurrentParentFrame, CurrentFrame, CurrentFrameInActiveFrames, FStateTreeStateHandle(ChildState), OutNextActiveFrames))
 					{
 						// Selection succeeded
 						return true;
@@ -3068,7 +3086,9 @@ bool FStateTreeExecutionContext::SelectStateInternal(
 				return true;
 			}
 		}
-		
+
+		// State could not be selected, restore.
+		CurrentFrame.ActiveInstanceIndexBase = PrevActiveInstanceIndexBase;
 		CurrentFrame.ActiveStates.Pop();
 	}
 
diff --git a/Engine/Plugins/Runtime/StateTree/Source/StateTreeModule/Public/StateTreeExecutionContext.h b/Engine/Plugins/Runtime/StateTree/Source/StateTreeModule/Public/StateTreeExecutionContext.h
--- a/Engine/Plugins/Runtime/StateTree/Source/StateTreeModule/Public/StateTreeExecutionContext.h
+++ b/Engine/Plugins/Runtime/StateTree/Source/StateTreeModule/Public/StateTreeExecutionContext.h
@@ -407,6 +407,7 @@ protected:
 	bool SelectStateInternal(
 		const FStateTreeExecutionFrame* CurrentParentFrame,
 		FStateTreeExecutionFrame& CurrentFrame,
+		const FStateTreeExecutionFrame* CurrentFrameInActiveFrames,
 		const FStateTreeStateHandle NextStateHandle,
 		TArray<FStateTreeExecutionFrame, TFixedAllocator<MaxExecutionFrames>>& OutNextActiveFrames);
 
@@ -454,7 +455,6 @@ protected:
 
 	/** Copies a batch of properties to the data in TargetView. This version validates the data handles and looks up temporary instances. */
 	bool CopyBatchWithValidation(const FStateTreeExecutionFrame* ParentFrame, const FStateTreeExecutionFrame& CurrentFrame, const FStateTreeDataView TargetView, const FStateTreeIndex16 BindingsBatch) const;
-
 	
 	/** Owner of the instance data. */
 	UObject& Owner;
