StateTree: Fix Mac compiler not being able to figure out implicit cast from TObjectPtr<> to UObject

diff --git a/Engine/Plugins/Runtime/StateTree/Source/StateTreeEditorModule/Private/SStateTreeViewRow.cpp b/Engine/Plugins/Runtime/StateTree/Source/StateTreeEditorModule/Private/SStateTreeViewRow.cpp
--- a/Engine/Plugins/Runtime/StateTree/Source/StateTreeEditorModule/Private/SStateTreeViewRow.cpp
+++ b/Engine/Plugins/Runtime/StateTree/Source/StateTreeEditorModule/Private/SStateTreeViewRow.cpp
@@ -917,7 +917,7 @@ FText SStateTreeViewRow::GetLinkedStateDesc() const
 	}
 	else if (State->Type == EStateTreeStateType::LinkedAsset)
 	{
-		return FText::FromString(*GetNameSafe(State->LinkedAsset));
+		return FText::FromString(GetNameSafe(State->LinkedAsset.Get()));
 	}
 	
 	return FText::GetEmpty();
