StateTree: Fixed crash when removing a task from state

diff --git a/Engine/Plugins/Runtime/StateTree/Source/StateTreeEditorModule/Private/StateTreePropertyHelpers.cpp b/Engine/Plugins/Runtime/StateTree/Source/StateTreeEditorModule/Private/StateTreePropertyHelpers.cpp
--- a/Engine/Plugins/Runtime/StateTree/Source/StateTreeEditorModule/Private/StateTreePropertyHelpers.cpp
+++ b/Engine/Plugins/Runtime/StateTree/Source/StateTreeEditorModule/Private/StateTreePropertyHelpers.cpp
@@ -1,4 +1,4 @@
-﻿// Copyright Epic Games, Inc. All Rights Reserved.
+// Copyright Epic Games, Inc. All Rights Reserved.
 
 #include "StateTreePropertyHelpers.h"
 #include "StateTreeEditorNode.h"
@@ -29,7 +29,7 @@ void DispatchPostEditToNodes(UObject& Owner, FPropertyChangedChainEvent& InPrope
 		{
 			FScriptArrayHelper Helper(ArrayProperty, CurrentAddress);
 			const int32 Index = InPropertyChangedEvent.GetArrayIndex(ArrayProperty->GetName());
-			if (Index == INDEX_NONE)
+			if (!Helper.IsValidIndex(Index))
 			{
 				return;
 			}
