StateTree: Enabled parameters for all state types
- Allow each state to have parameters (not super useful yet)
- Made property bindings for parameters to retain binding when a param is renamed (keeps track of the property bag ID)
- Fixed access of temporary instance data with EStateTreeDataSourceType::StateParameterData
- Added comments for UStateTreeState properties
- Changed property order of UStateTreeState so that display data comes first, then state settings, and finally parameters
- Fixed FStateTreeStateParametersInstanceDataDetails to allow property editing and value binding at the same time

diff --git a/Engine/Plugins/Experimental/StructUtils/Source/StructUtils/Private/PropertyBag.cpp b/Engine/Plugins/Experimental/StructUtils/Source/StructUtils/Private/PropertyBag.cpp
--- a/Engine/Plugins/Experimental/StructUtils/Source/StructUtils/Private/PropertyBag.cpp
+++ b/Engine/Plugins/Experimental/StructUtils/Source/StructUtils/Private/PropertyBag.cpp
@@ -2867,3 +2867,7 @@ const FPropertyBagPropertyDesc* UPropertyBag::FindPropertyDescByName(const FName
 	return PropertyDescs.FindByPredicate([&Name](const FPropertyBagPropertyDesc& Desc) { return Desc.Name == Name; });
 }
 
+const FPropertyBagPropertyDesc* UPropertyBag::FindPropertyDescByPropertyName(const FName PropertyName) const
+{
+	return PropertyDescs.FindByPredicate([&PropertyName](const FPropertyBagPropertyDesc& Desc) { return Desc.CachedProperty && Desc.CachedProperty->GetFName() == PropertyName; });
+}
diff --git a/Engine/Plugins/Experimental/StructUtils/Source/StructUtils/Public/PropertyBag.h b/Engine/Plugins/Experimental/StructUtils/Source/StructUtils/Public/PropertyBag.h
--- a/Engine/Plugins/Experimental/StructUtils/Source/StructUtils/Public/PropertyBag.h
+++ b/Engine/Plugins/Experimental/StructUtils/Source/StructUtils/Public/PropertyBag.h
@@ -874,6 +874,9 @@ public:
 	/** @return property description based on name. */
 	const FPropertyBagPropertyDesc* FindPropertyDescByName(const FName Name) const;
 
+	/** @return property description based on the created property name. The name can be different from the descriptor name due to name sanitization. */
+	const FPropertyBagPropertyDesc* FindPropertyDescByPropertyName(const FName PropertyName) const;
+
 #if WITH_ENGINE && WITH_EDITOR
 	/** @return true if any of the properties on the bag has type of the specified user defined struct. */
 	bool ContainsUserDefinedStruct(const UUserDefinedStruct* UserDefinedStruct) const;
diff --git a/Engine/Plugins/Runtime/StateTree/Config/BaseStateTree.ini b/Engine/Plugins/Runtime/StateTree/Config/BaseStateTree.ini
--- a/Engine/Plugins/Runtime/StateTree/Config/BaseStateTree.ini
+++ b/Engine/Plugins/Runtime/StateTree/Config/BaseStateTree.ini
@@ -63,3 +63,5 @@
 
 +StructRedirects=(OldName="StateTreeEditorPropertyBinding",NewName="/script/StateTreeModule.StateTreePropertyPathBinding")
 +StructRedirects=(OldName="StateTreePropCopyBatch",NewName="/script/StateTreeModule.StateTreePropertyCopyBatch")
+
++EnumRedirects=(OldName="/Script/StateTreeModule.EStateTreeDataSourceType",ValueChanges=(("LinkedStateParameterData", "StateParameterData"))
diff --git a/Engine/Plugins/Runtime/StateTree/Source/StateTreeEditorModule/Private/Customizations/StateTreeStateDetails.cpp b/Engine/Plugins/Runtime/StateTree/Source/StateTreeEditorModule/Private/Customizations/StateTreeStateDetails.cpp
--- a/Engine/Plugins/Runtime/StateTree/Source/StateTreeEditorModule/Private/Customizations/StateTreeStateDetails.cpp
+++ b/Engine/Plugins/Runtime/StateTree/Source/StateTreeEditorModule/Private/Customizations/StateTreeStateDetails.cpp
@@ -73,11 +73,6 @@ void FStateTreeStateDetails::CustomizeDetails(IDetailLayoutBuilder& DetailBuilde
 		LinkedAssetProperty->MarkHiddenByCustomization();
 	}
 	
-	if (!(StateType == EStateTreeStateType::Subtree || StateType == EStateTreeStateType::Linked || StateType == EStateTreeStateType::LinkedAsset))
-	{
-		ParametersProperty->MarkHiddenByCustomization();
-	}
-
 	if (StateType == EStateTreeStateType::Linked || StateType == EStateTreeStateType::LinkedAsset)
 	{
 		SelectionBehaviorProperty->MarkHiddenByCustomization();
diff --git a/Engine/Plugins/Runtime/StateTree/Source/StateTreeEditorModule/Private/Customizations/StateTreeStateParametersDetails.cpp b/Engine/Plugins/Runtime/StateTree/Source/StateTreeEditorModule/Private/Customizations/StateTreeStateParametersDetails.cpp
--- a/Engine/Plugins/Runtime/StateTree/Source/StateTreeEditorModule/Private/Customizations/StateTreeStateParametersDetails.cpp
+++ b/Engine/Plugins/Runtime/StateTree/Source/StateTreeEditorModule/Private/Customizations/StateTreeStateParametersDetails.cpp
@@ -19,65 +19,75 @@
 class FStateTreeStateParametersInstanceDataDetails : public FPropertyBagInstanceDataDetails
 {
 public:
-	FStateTreeStateParametersInstanceDataDetails(TSharedPtr<IPropertyHandle> InStructProperty, TSharedPtr<IPropertyUtilities> InPropUtils, const bool bInFixedLayout, FGuid InID, UStateTreeEditorData* InEditorData)
+	FStateTreeStateParametersInstanceDataDetails(
+		const TSharedPtr<IPropertyHandle>& InStructProperty,
+		const TSharedPtr<IPropertyUtilities>& InPropUtils,
+		const bool bInFixedLayout,
+		FGuid InID,
+		TWeakObjectPtr<UStateTreeEditorData> InEditorData,
+		TWeakObjectPtr<UStateTreeState> InState)
 		: FPropertyBagInstanceDataDetails(InStructProperty, InPropUtils, bInFixedLayout)
-		, EditorData(InEditorData)
+		, WeakEditorData(InEditorData)
+		, WeakState(InState)
+		, ID(InID)
 	{
-		EditorPropBindings = EditorData ? EditorData->GetPropertyEditorBindings() : nullptr;
-		ID = InID;
 	}
-
+	
 	virtual void OnChildRowAdded(IDetailPropertyRow& ChildRow) override
 	{
-		if (!EditorPropBindings
-			|| !bFixedLayout		// No binding for parameter definitions.
-			|| !ID.IsValid())
+		FPropertyBagInstanceDataDetails::OnChildRowAdded(ChildRow);
+
+		EStateTreeStateType Type = EStateTreeStateType::State;
+		if (const UStateTreeState* State = WeakState.Get())
 		{
-			FPropertyBagInstanceDataDetails::OnChildRowAdded(ChildRow);
-			return;
+			Type = State->Type;
 		}
 
-		TSharedPtr<IPropertyHandle> ChildPropHandle = ChildRow.GetPropertyHandle();
-		const FProperty* Property = ChildPropHandle->GetProperty();
+		// Subtree parameters cannot be bound to, they are provided from the linked state.
+		const bool bAllowBinding = Type != EStateTreeStateType::Subtree && ID.IsValid(); 
 
-		// Set the category to Parameter so that the binding extension will pick it up.
-		static const FName CategoryName(TEXT("Category"));
-		ChildPropHandle->SetInstanceMetaData(CategoryName, TEXT("Parameter"));
-	
-		// Conditionally control visibility of the value field of bound properties.
-
-		// Pass the node ID to binding extension. Since the properties are added using AddChildStructure(), we break the hierarchy and cannot access parent.
-		ChildPropHandle->SetInstanceMetaData(UE::StateTree::PropertyBinding::StateTreeNodeIDName, LexToString(ID));
-
-		FStateTreePropertyPath Path(ID, *Property->GetFName().ToString());
-		TSharedPtr<SWidget> NameWidget;
-		TSharedPtr<SWidget> ValueWidget;
-		FDetailWidgetRow Row;
-		ChildRow.GetDefaultWidgets(NameWidget, ValueWidget, Row);
-
-		auto IsValueVisible = TAttribute<EVisibility>::Create([Path, this]() -> EVisibility
-			{
-				return EditorPropBindings->HasPropertyBinding(Path) ? EVisibility::Collapsed : EVisibility::Visible;
-			});
-
-		ChildRow
-			.CustomWidget(/*bShowChildren*/true)
-			.NameContent()
-			[
-				NameWidget.ToSharedRef()
-			]
-			.ValueContent()
-			[
-				SNew(SBox)
+		if (bAllowBinding)
+		{
+			const TSharedPtr<IPropertyHandle> ChildPropHandle = ChildRow.GetPropertyHandle();
+			const FProperty* Property = ChildPropHandle->GetProperty();
+
+			// Set the category to Parameter so that the binding extension will pick it up.
+			static const FName CategoryName(TEXT("Category"));
+			ChildPropHandle->SetInstanceMetaData(CategoryName, TEXT("Parameter"));
+		
+			// Conditionally control visibility of the value field of bound properties.
+
+			// Pass the node ID to binding extension. Since the properties are added using AddChildStructure(), we break the hierarchy and cannot access parent.
+			ChildPropHandle->SetInstanceMetaData(UE::StateTree::PropertyBinding::StateTreeNodeIDName, LexToString(ID));
+
+			FStateTreePropertyPath Path(ID, *Property->GetFName().ToString());
+			
+			const auto IsValueVisible = TAttribute<EVisibility>::Create([Path, WeakEditorData = WeakEditorData]() -> EVisibility
+				{
+					bool bHasBinding = false;
+					if (UStateTreeEditorData* EditorData = WeakEditorData.Get())
+					{
+						if (const FStateTreeEditorPropertyBindings* EditorPropBindings = EditorData->GetPropertyEditorBindings())
+						{
+							bHasBinding = EditorPropBindings->HasPropertyBinding(Path); 
+						}
+					}
+
+					return bHasBinding ? EVisibility::Collapsed : EVisibility::Visible;
+				});
+
+			FDetailWidgetDecl* ValueWidgetDecl = ChildRow.CustomValueWidget();
+			const TSharedRef<SBox> WrappedValueWidget = SNew(SBox)
 				.Visibility(IsValueVisible)
 				[
-					ValueWidget.ToSharedRef()
-				]
-			];
+					ValueWidgetDecl->Widget
+				];
+			ValueWidgetDecl->Widget = WrappedValueWidget;
+		}
 	}
 
-	UStateTreeEditorData* EditorData;
-	FStateTreeEditorPropertyBindings* EditorPropBindings;
+	TWeakObjectPtr<UStateTreeEditorData> WeakEditorData;
+	TWeakObjectPtr<UStateTreeState> WeakState;
 	FGuid ID;
 };
 
@@ -132,7 +142,7 @@ void FStateTreeStateParametersDetails::CustomizeChildren(TSharedRef<class IPrope
 	UE::StateTree::PropertyHelpers::GetStructValue<FGuid>(IDProperty, ID);
 
 	// Show the Value (FInstancedStruct) as child rows.
-	TSharedRef<FStateTreeStateParametersInstanceDataDetails> InstanceDetails = MakeShareable(new FStateTreeStateParametersInstanceDataDetails(ParametersProperty, PropUtils, bFixedLayout, ID, EditorData));
+	TSharedRef<FStateTreeStateParametersInstanceDataDetails> InstanceDetails = MakeShareable(new FStateTreeStateParametersInstanceDataDetails(ParametersProperty, PropUtils, bFixedLayout, ID, WeakEditorData, WeakState));
 	StructBuilder.AddCustomBuilder(InstanceDetails);
 }
 
@@ -140,19 +150,22 @@ void FStateTreeStateParametersDetails::FindOuterObjects()
 {
 	check(StructProperty);
 	
-	EditorData = nullptr;
-	StateTree = nullptr;
+	WeakEditorData = nullptr;
+	WeakStateTree = nullptr;
+	WeakState = nullptr;
 
 	TArray<UObject*> OuterObjects;
 	StructProperty->GetOuterObjects(OuterObjects);
 	for (UObject* Outer : OuterObjects)
 	{
 		UStateTreeEditorData* OuterEditorData = Outer->GetTypedOuter<UStateTreeEditorData>();
+		UStateTreeState* OuterState = Outer->GetTypedOuter<UStateTreeState>();
 		UStateTree* OuterStateTree = OuterEditorData ? OuterEditorData->GetTypedOuter<UStateTree>() : nullptr;
-		if (OuterEditorData && OuterStateTree)
+		if (OuterEditorData && OuterStateTree && OuterState)
 		{
-			StateTree = OuterStateTree;
-			EditorData = OuterEditorData;
+			WeakStateTree = OuterStateTree;
+			WeakEditorData = OuterEditorData;
+			WeakState = OuterState;
 			break;
 		}
 	}
diff --git a/Engine/Plugins/Runtime/StateTree/Source/StateTreeEditorModule/Private/Customizations/StateTreeStateParametersDetails.h b/Engine/Plugins/Runtime/StateTree/Source/StateTreeEditorModule/Private/Customizations/StateTreeStateParametersDetails.h
--- a/Engine/Plugins/Runtime/StateTree/Source/StateTreeEditorModule/Private/Customizations/StateTreeStateParametersDetails.h
+++ b/Engine/Plugins/Runtime/StateTree/Source/StateTreeEditorModule/Private/Customizations/StateTreeStateParametersDetails.h
@@ -7,6 +7,7 @@
 class IPropertyHandle;
 class IPropertyUtilities;
 class UStateTree;
+class UStateTreeState;
 class UStateTreeEditorData;
 
 /**
@@ -36,6 +37,7 @@ private:
 
 	bool bFixedLayout = false;
 
-	UStateTreeEditorData* EditorData = nullptr;
-	UStateTree* StateTree = nullptr;
+	TWeakObjectPtr<UStateTreeEditorData> WeakEditorData = nullptr;
+	TWeakObjectPtr<UStateTree> WeakStateTree = nullptr;
+	TWeakObjectPtr<UStateTreeState> WeakState = nullptr;
 };
diff --git a/Engine/Plugins/Runtime/StateTree/Source/StateTreeEditorModule/Private/StateTreeCompiler.cpp b/Engine/Plugins/Runtime/StateTree/Source/StateTreeEditorModule/Private/StateTreeCompiler.cpp
--- a/Engine/Plugins/Runtime/StateTree/Source/StateTreeEditorModule/Private/StateTreeCompiler.cpp
+++ b/Engine/Plugins/Runtime/StateTree/Source/StateTreeEditorModule/Private/StateTreeCompiler.cpp
@@ -606,75 +606,69 @@ bool FStateTreeCompiler::CreateStateTasksAndParameters()
 		FStateTreeCompilerLogStateScope LogStateScope(State, Log);
 
 		// Create parameters
-		if (State->Type == EStateTreeStateType::Linked
-			|| State->Type == EStateTreeStateType::LinkedAsset
-			|| State->Type == EStateTreeStateType::Subtree)
-		{
-			// Each state has their parameters as instance data.
-			FInstancedStruct& Instance = InstanceStructs.AddDefaulted_GetRef();
-			Instance.InitializeAs<FCompactStateTreeParameters>(State->Parameters.Parameters);
-			FCompactStateTreeParameters& CompactStateTreeParameters = Instance.GetMutable<FCompactStateTreeParameters>(); 
+		
+		// Each state has their parameters as instance data.
+		FInstancedStruct& Instance = InstanceStructs.AddDefaulted_GetRef();
+		Instance.InitializeAs<FCompactStateTreeParameters>(State->Parameters.Parameters);
+		FCompactStateTreeParameters& CompactStateTreeParameters = Instance.GetMutable<FCompactStateTreeParameters>(); 
 			
-			const int32 InstanceIndex = InstanceStructs.Num() - 1;
-			if (const auto Validation = UE::StateTree::Compiler::IsValidIndex16(InstanceIndex); Validation.DidFail())
-			{
-				Validation.Log(Log, TEXT("InstanceIndex"));
-				return false;
-			}
-			CompactState.ParameterTemplateIndex = FStateTreeIndex16(InstanceIndex);
+		const int32 InstanceIndex = InstanceStructs.Num() - 1;
+		if (const auto Validation = UE::StateTree::Compiler::IsValidIndex16(InstanceIndex); Validation.DidFail())
+		{
+			Validation.Log(Log, TEXT("InstanceIndex"));
+			return false;
+		}
+		CompactState.ParameterTemplateIndex = FStateTreeIndex16(InstanceIndex);
 
-			if (State->Type == EStateTreeStateType::Linked || State->Type == EStateTreeStateType::LinkedAsset)
-			{
-				CompactState.ParameterDataHandle = FStateTreeDataHandle(EStateTreeDataSourceType::LinkedStateParameterData, InstanceDataIndex++, CompactStateHandle);
-			}
-			else if (State->Type == EStateTreeStateType::Subtree)
-			{
-				// Subtrees have special indirection, which is used to access the parameters passed to the state.
-				// If subtree is called via linked state, the linked state's parameters are used, otherwise the instance data.
-				CompactState.ParameterDataHandle = FStateTreeDataHandle(EStateTreeDataSourceType::SubtreeParameterData, InstanceDataIndex++, CompactStateHandle);
-			}
+		if (State->Type == EStateTreeStateType::Subtree)
+		{
+			CompactState.ParameterDataHandle = FStateTreeDataHandle(EStateTreeDataSourceType::SubtreeParameterData, InstanceDataIndex++, CompactStateHandle);
+		}
+		else
+		{
+			CompactState.ParameterDataHandle = FStateTreeDataHandle(EStateTreeDataSourceType::StateParameterData, InstanceDataIndex++, CompactStateHandle);
+		}
 
-			// @todo: We should be able to skip empty parameter data.
+		// @todo: We should be able to skip empty parameter data.
 
-			// Binding target
-			FStateTreeBindableStructDesc LinkedParamsDesc = {
-				State->Name,
-				State->Parameters.Parameters.GetPropertyBagStruct(),
-				CompactState.ParameterDataHandle,
-				EStateTreeBindableStructSource::State,
-				State->Parameters.ID
-			};
+		// Binding target
+		FStateTreeBindableStructDesc LinkedParamsDesc = {
+			State->Name,
+			State->Parameters.Parameters.GetPropertyBagStruct(),
+			CompactState.ParameterDataHandle,
+			EStateTreeBindableStructSource::State,
+			State->Parameters.ID
+		};
 
-			if (!UE::StateTree::Compiler::ValidateNoLevelActorReferences(Log, LinkedParamsDesc, FStateTreeDataView(), FStateTreeDataView(CompactStateTreeParameters.Parameters.GetMutableValue())))
-			{
-				return false;
-			}
+		if (!UE::StateTree::Compiler::ValidateNoLevelActorReferences(Log, LinkedParamsDesc, FStateTreeDataView(), FStateTreeDataView(CompactStateTreeParameters.Parameters.GetMutableValue())))
+		{
+			return false;
+		}
 
-			// Add as binding source.
-			BindingsCompiler.AddSourceStruct(LinkedParamsDesc);
+		// Add as binding source.
+		BindingsCompiler.AddSourceStruct(LinkedParamsDesc);
 
-			// Check that the bindings for this struct are still all valid.
-			TArray<FStateTreePropertyPathBinding> Bindings;
-			if (!GetAndValidateBindings(LinkedParamsDesc, FStateTreeDataView(CompactStateTreeParameters.Parameters.GetMutableValue()), Bindings))
-			{
-				return false;
-			}
+		// Check that the bindings for this struct are still all valid.
+		TArray<FStateTreePropertyPathBinding> Bindings;
+		if (!GetAndValidateBindings(LinkedParamsDesc, FStateTreeDataView(CompactStateTreeParameters.Parameters.GetMutableValue()), Bindings))
+		{
+			return false;
+		}
 
-			int32 BatchIndex = INDEX_NONE;
-			if (!BindingsCompiler.CompileBatch(LinkedParamsDesc, Bindings, BatchIndex))
-			{
-				return false;
-			}
+		int32 BatchIndex = INDEX_NONE;
+		if (!BindingsCompiler.CompileBatch(LinkedParamsDesc, Bindings, BatchIndex))
+		{
+			return false;
+		}
 			
-			if (const auto Validation = UE::StateTree::Compiler::IsValidIndex16(BatchIndex); Validation.DidFail())
-			{
-				Validation.Log(Log, TEXT("BatchIndex"), LinkedParamsDesc);
-				return false;
-			}
-
-			CompactState.ParameterBindingsBatch = FStateTreeIndex16(BatchIndex);
+		if (const auto Validation = UE::StateTree::Compiler::IsValidIndex16(BatchIndex); Validation.DidFail())
+		{
+			Validation.Log(Log, TEXT("BatchIndex"), LinkedParamsDesc);
+			return false;
 		}
 
+		CompactState.ParameterBindingsBatch = FStateTreeIndex16(BatchIndex);
+
 		// Create tasks
 		const int32 TasksBegin = Nodes.Num();
 		if (const auto Validation = UE::StateTree::Compiler::IsValidCount16(TasksBegin); Validation.DidFail())
diff --git a/Engine/Plugins/Runtime/StateTree/Source/StateTreeEditorModule/Private/StateTreeEditorData.cpp b/Engine/Plugins/Runtime/StateTree/Source/StateTreeEditorModule/Private/StateTreeEditorData.cpp
--- a/Engine/Plugins/Runtime/StateTree/Source/StateTreeEditorModule/Private/StateTreeEditorData.cpp
+++ b/Engine/Plugins/Runtime/StateTree/Source/StateTreeEditorModule/Private/StateTreeEditorData.cpp
@@ -36,6 +36,7 @@ void UStateTreeEditorData::PostInitProperties()
 #if WITH_EDITOR
 	OnObjectsReinstancedHandle = FCoreUObjectDelegates::OnObjectsReinstanced.AddUObject(this, &UStateTreeEditorData::OnObjectsReinstanced);
 	OnUserDefinedStructReinstancedHandle = UE::StructUtils::Delegates::OnUserDefinedStructReinstanced.AddUObject(this, &UStateTreeEditorData::OnUserDefinedStructReinstanced);
+	OnParametersChangedHandle = UE::StateTree::Delegates::OnParametersChanged.AddUObject(this, &UStateTreeEditorData::OnParametersChanged);
 #endif
 }
 
@@ -53,6 +54,11 @@ void UStateTreeEditorData::BeginDestroy()
 		UE::StructUtils::Delegates::OnUserDefinedStructReinstanced.Remove(OnUserDefinedStructReinstancedHandle);
 		OnUserDefinedStructReinstancedHandle.Reset();
 	}
+	if (OnParametersChangedHandle.IsValid())
+	{
+		UE::StateTree::Delegates::OnParametersChanged.Remove(OnParametersChangedHandle);
+		OnParametersChangedHandle.Reset();
+	}
 	
 	Super::BeginDestroy();
 }
@@ -122,6 +128,17 @@ void UStateTreeEditorData::OnUserDefinedStructReinstanced(const UUserDefinedStru
 	}
 }
 
+void UStateTreeEditorData::OnParametersChanged(const UStateTree& StateTree)
+{
+	if (const UStateTree* OwnerStateTree = GetTypedOuter<UStateTree>())
+	{
+		if (OwnerStateTree == &StateTree)
+		{
+			UpdateBindingsInstanceStructs();
+		}
+	}
+}
+
 
 void UStateTreeEditorData::PostLoad()
 {
@@ -753,8 +770,7 @@ EStateTreeVisitor UStateTreeEditorData::VisitStateNodes(const UStateTreeState& S
 	if (bContinue)
 	{
 		// Bindable state parameters
-		if ((State.Type == EStateTreeStateType::Linked
-				|| State.Type == EStateTreeStateType::LinkedAsset)
+		if (State.Type != EStateTreeStateType::Subtree
 			&& State.Parameters.Parameters.IsValid())
 		{
 			if (InFunc(&State, State.Parameters.ID, State.Name, EStateTreeNodeType::StateParameters, nullptr, State.Parameters.Parameters.GetPropertyBagStruct()) == EStateTreeVisitor::Break)
@@ -774,11 +790,8 @@ EStateTreeVisitor UStateTreeEditorData::VisitStateNodes(const UStateTreeState& S
 
 	if (bContinue)
 	{
-		// Bindable state parameters for subtree or linked tree.
-		if ((State.Type == EStateTreeStateType::Subtree
-				|| State.Type == EStateTreeStateType::Linked
-				|| State.Type == EStateTreeStateType::LinkedAsset)
-			&& State.Parameters.Parameters.IsValid())
+		// Bindable state parameters
+		if (State.Parameters.Parameters.IsValid())
 		{
 			FStateTreeBindableStructDesc Desc;
 			Desc.Struct = State.Parameters.Parameters.GetPropertyBagStruct();
diff --git a/Engine/Plugins/Runtime/StateTree/Source/StateTreeEditorModule/Private/StateTreePropertyHelpers.cpp b/Engine/Plugins/Runtime/StateTree/Source/StateTreeEditorModule/Private/StateTreePropertyHelpers.cpp
--- a/Engine/Plugins/Runtime/StateTree/Source/StateTreeEditorModule/Private/StateTreePropertyHelpers.cpp
+++ b/Engine/Plugins/Runtime/StateTree/Source/StateTreeEditorModule/Private/StateTreePropertyHelpers.cpp
@@ -170,6 +170,20 @@ FStateTreeEditPropertyPath::FStateTreeEditPropertyPath(const FPropertyChangedCha
 	}
 }
 
+FStateTreeEditPropertyPath::FStateTreeEditPropertyPath(const FEditPropertyChain& PropertyChain)
+{
+	FEditPropertyChain::TDoubleLinkedListNode* PropertyNode = PropertyChain.GetActiveMemberNode();
+	while (PropertyNode != nullptr)
+	{
+		if (FProperty* Property = PropertyNode->GetValue())
+		{
+			const FName PropertyName = Property->GetFName(); 
+			Path.Emplace(Property, PropertyName, INDEX_NONE);
+		}
+		PropertyNode = PropertyNode->GetNextNode();
+	}
+}
+
 bool FStateTreeEditPropertyPath::ContainsPath(const FStateTreeEditPropertyPath& InPath) const
 {
 	if (InPath.Path.Num() > Path.Num())
diff --git a/Engine/Plugins/Runtime/StateTree/Source/StateTreeEditorModule/Private/StateTreePropertyHelpers.h b/Engine/Plugins/Runtime/StateTree/Source/StateTreeEditorModule/Private/StateTreePropertyHelpers.h
--- a/Engine/Plugins/Runtime/StateTree/Source/StateTreeEditorModule/Private/StateTreePropertyHelpers.h
+++ b/Engine/Plugins/Runtime/StateTree/Source/StateTreeEditorModule/Private/StateTreePropertyHelpers.h
@@ -156,6 +156,9 @@ public:
 	/** Makes property path from property change event. */
 	explicit FStateTreeEditPropertyPath(const FPropertyChangedChainEvent& PropertyChangedEvent);
 
+	/** Makes property path from property chain. */
+	explicit FStateTreeEditPropertyPath(const FEditPropertyChain& PropertyChain);
+
 	/** @return true if the property path contains specified path. */
 	bool ContainsPath(const FStateTreeEditPropertyPath& InPath) const;
 
diff --git a/Engine/Plugins/Runtime/StateTree/Source/StateTreeEditorModule/Private/StateTreeState.cpp b/Engine/Plugins/Runtime/StateTree/Source/StateTreeEditorModule/Private/StateTreeState.cpp
--- a/Engine/Plugins/Runtime/StateTree/Source/StateTreeEditorModule/Private/StateTreeState.cpp
+++ b/Engine/Plugins/Runtime/StateTree/Source/StateTreeEditorModule/Private/StateTreeState.cpp
@@ -56,6 +56,25 @@ void UStateTreeState::OnTreeCompiled(const UStateTree& StateTree)
 	}
 }
 
+void UStateTreeState::PreEditChange(FEditPropertyChain& PropertyAboutToChange)
+{
+	Super::PreEditChange(PropertyAboutToChange);
+
+	const FStateTreeEditPropertyPath PropertyChainPath(PropertyAboutToChange);
+
+	static const FStateTreeEditPropertyPath StateTypePath(UStateTreeState::StaticClass(), TEXT("Type"));
+
+	if (PropertyChainPath.IsPathExact(StateTypePath))
+	{
+		// If transitioning from linked state, reset the parameters
+		if (Type == EStateTreeStateType::Linked
+			|| Type == EStateTreeStateType::LinkedAsset)
+		{
+			Parameters.Reset();
+		}
+	}
+}
+
 void UStateTreeState::PostEditChangeChainProperty(FPropertyChangedChainEvent& PropertyChangedEvent)
 {
 	Super::PostEditChangeChainProperty(PropertyChangedEvent);
@@ -101,7 +120,7 @@ void UStateTreeState::PostEditChangeChainProperty(FPropertyChangedChainEvent& Pr
 	}
 
 	// Broadcast selection type changes so that the UI can update.
-	if (SelectionBehaviorPath.IsPathExact(StateTypePath))
+	if (ChangePropertyPath.IsPathExact(SelectionBehaviorPath))
 	{
 		const UStateTree* StateTree = GetTypedOuter<UStateTree>();
 		if (ensure(StateTree))
@@ -136,15 +155,10 @@ void UStateTreeState::PostEditChangeChainProperty(FPropertyChangedChainEvent& Pr
 			UpdateParametersFromLinkedSubtree();
 			SelectionBehavior = EStateTreeStateSelectionBehavior::TrySelectChildrenInOrder;
 		}
-		else if (Type == EStateTreeStateType::Subtree)
-		{
-			// Subtree parameter layout can be edited
-			Parameters.bFixedLayout = false;
-		}
 		else
 		{
-			// Other state types do not have parameters, so reset them.
-			Parameters.Reset();
+			// Other layouts can be edited
+			Parameters.bFixedLayout = false;
 		}
 	}
 
@@ -168,7 +182,8 @@ void UStateTreeState::PostEditChangeChainProperty(FPropertyChangedChainEvent& Pr
 	// Broadcast subtree parameter layout edits so that the linked states can adapt.
 	if (ChangePropertyPath.IsPathExact(StateParametersPath))
 	{
-		if (Type == EStateTreeStateType::Subtree)
+		if (!(Type == EStateTreeStateType::Linked
+				|| Type == EStateTreeStateType::LinkedAsset))
 		{
 			const UStateTree* StateTree = GetTypedOuter<UStateTree>();
 			if (ensure(StateTree))
diff --git a/Engine/Plugins/Runtime/StateTree/Source/StateTreeEditorModule/Public/StateTreeEditorData.h b/Engine/Plugins/Runtime/StateTree/Source/StateTreeEditorModule/Public/StateTreeEditorData.h
--- a/Engine/Plugins/Runtime/StateTree/Source/StateTreeEditorModule/Public/StateTreeEditorData.h
+++ b/Engine/Plugins/Runtime/StateTree/Source/StateTreeEditorModule/Public/StateTreeEditorData.h
@@ -65,6 +65,7 @@ public:
 	using FReplacementObjectMap = TMap<UObject*, UObject*>;
 	void OnObjectsReinstanced(const FReplacementObjectMap& ObjectMap);
 	void OnUserDefinedStructReinstanced(const UUserDefinedStruct& UserDefinedStruct);
+	void OnParametersChanged(const UStateTree& StateTree);
 	virtual void BeginDestroy() override;
 	virtual void PostLoad() override;
 	virtual void PostEditChangeChainProperty(FPropertyChangedChainEvent& PropertyChangedEvent) override;
@@ -255,6 +256,7 @@ private:
 #if WITH_EDITORONLY_DATA
 	FDelegateHandle OnObjectsReinstancedHandle;
 	FDelegateHandle OnUserDefinedStructReinstancedHandle;
+	FDelegateHandle OnParametersChangedHandle;
 #endif
 
 public:
diff --git a/Engine/Plugins/Runtime/StateTree/Source/StateTreeEditorModule/Public/StateTreeState.h b/Engine/Plugins/Runtime/StateTree/Source/StateTreeEditorModule/Public/StateTreeState.h
--- a/Engine/Plugins/Runtime/StateTree/Source/StateTreeEditorModule/Public/StateTreeState.h
+++ b/Engine/Plugins/Runtime/StateTree/Source/StateTreeEditorModule/Public/StateTreeState.h
@@ -113,6 +113,7 @@ public:
 	virtual ~UStateTreeState() override;
 
 	virtual void PostInitProperties() override;
+	virtual void PreEditChange(FEditPropertyChain& PropertyAboutToChange) override;
 	virtual void PostEditChangeChainProperty(FPropertyChangedChainEvent& PropertyChangedEvent) override;
 	virtual void PostLoad() override;;
 	void UpdateParametersFromLinkedSubtree();
@@ -195,30 +196,37 @@ public:
 
 	// ~StateTree Builder API
 
+	/** Display name of the State */
 	UPROPERTY(EditDefaultsOnly, Category = "State")
 	FName Name;
 
+	/** Display color of the State */
+	UPROPERTY(EditDefaultsOnly, Category = "State", DisplayName = "Color")
+	FStateTreeEditorColorRef ColorRef;
+
+	/** Type the State, allows e.g. states to be linked to other States. */
 	UPROPERTY(EditDefaultsOnly, Category = "State")
 	EStateTreeStateType Type = EStateTreeStateType::State;
 
+	/** How to treat child states when this State is selected.  */
 	UPROPERTY(EditDefaultsOnly, Category = "State")
 	EStateTreeStateSelectionBehavior SelectionBehavior = EStateTreeStateSelectionBehavior::TrySelectChildrenInOrder;
-	
+
+	/** Subtree to run as extension of this State. */
 	UPROPERTY(EditDefaultsOnly, Category = "State", Meta=(DirectStatesOnly, SubtreesOnly))
 	FStateTreeStateLink LinkedSubtree;
 
+	/** Another State Tree asset to run as extension of this State. */
 	UPROPERTY(EditDefaultsOnly, Category = "State")
 	TObjectPtr<UStateTree> LinkedAsset = nullptr;
-	
+
+	/** Parameters of this state. If the state is linked to another state or asset, the parameters are for the linked state. */
 	UPROPERTY(EditDefaultsOnly, Category = "State")
 	FStateTreeStateParameters Parameters;
 
 	UPROPERTY(EditDefaultsOnly, Category = "State", meta = (IgnoreForMemberInitializationTest))
 	FGuid ID;
 
-	UPROPERTY(EditDefaultsOnly, Category = "State", DisplayName = "Color")
-	FStateTreeEditorColorRef ColorRef;
-
 	UPROPERTY(EditDefaultsOnly, Category = "Enter Conditions", meta = (BaseStruct = "/Script/StateTreeModule.StateTreeConditionBase", BaseClass = "/Script/StateTreeModule.StateTreeConditionBlueprintBase"))
 	TArray<FStateTreeEditorNode> EnterConditions;
 
diff --git a/Engine/Plugins/Runtime/StateTree/Source/StateTreeModule/Private/StateTree.cpp b/Engine/Plugins/Runtime/StateTree/Source/StateTreeModule/Private/StateTree.cpp
--- a/Engine/Plugins/Runtime/StateTree/Source/StateTreeModule/Private/StateTree.cpp
+++ b/Engine/Plugins/Runtime/StateTree/Source/StateTreeModule/Private/StateTree.cpp
@@ -528,6 +528,7 @@ bool UStateTree::PatchBindings()
 	// Refresh state parameter descs and bindings batches.
 	for (const FCompactStateTreeState& State : States)
 	{
+		// For subtrees and linked states, the parameters must exists.
 		if (State.Type == EStateTreeStateType::Subtree
 			|| State.Type == EStateTreeStateType::Linked
 			|| State.Type == EStateTreeStateType::LinkedAsset)
@@ -537,7 +538,10 @@ bool UStateTree::PatchBindings()
 				UE_LOG(LogStateTree, Error, TEXT("%s: Data for state '%s' is malformed. Please recompile the StateTree asset."), *GetFullName(), *State.Name.ToString());
 				return false;
 			}
+		}
 
+		if (State.ParameterTemplateIndex.IsValid())
+		{
 			// Subtree is a bind source, update bag struct.
 			const FCompactStateTreeParameters& Params = DefaultInstanceData.GetMutableStruct(State.ParameterTemplateIndex.Get()).Get<FCompactStateTreeParameters>();
 			FStateTreeBindableStructDesc* Desc = GetSourceStructByHandle(State.ParameterDataHandle);
@@ -611,18 +615,13 @@ bool UStateTree::PatchBindings()
 	// Setup data views for state parameters.
 	for (FCompactStateTreeState& State : States)
 	{
-		if (State.Type == EStateTreeStateType::Subtree
-			|| State.Type == EStateTreeStateType::Linked
-			|| State.Type == EStateTreeStateType::LinkedAsset)
+		if (State.ParameterDataHandle.IsValid())
 		{
-			if (State.ParameterDataHandle.IsValid())
+			FCompactStateTreeParameters& Params = DefaultInstanceData.GetMutableStruct(State.ParameterTemplateIndex.Get()).Get<FCompactStateTreeParameters>();
+			DataViews.Add(State.ParameterDataHandle, Params.Parameters.GetMutableValue());
+			if (State.ParameterBindingsBatch.IsValid())
 			{
-				FCompactStateTreeParameters& Params = DefaultInstanceData.GetMutableStruct(State.ParameterTemplateIndex.Get()).Get<FCompactStateTreeParameters>();
-				DataViews.Add(State.ParameterDataHandle, Params.Parameters.GetMutableValue());
-				if (State.ParameterBindingsBatch.IsValid())
-				{
-					BindingBatchDataView.Add(State.ParameterBindingsBatch, Params.Parameters.GetMutableValue());
-				}
+				BindingBatchDataView.Add(State.ParameterBindingsBatch, Params.Parameters.GetMutableValue());
 			}
 		}
 	}
@@ -790,9 +789,7 @@ TArray<FStateTreeMemoryUsage> UStateTree::CalculateEstimatedMemoryUsage() const
 			StateLinks.Emplace(ParentUsageIndex, LinkedUsageIndex);
 		}
 		
-		if (CompactState.Type == EStateTreeStateType::Linked
-			|| CompactState.Type == EStateTreeStateType::LinkedAsset
-			|| CompactState.Type == EStateTreeStateType::Subtree)
+		if (CompactState.ParameterTemplateIndex.IsValid())
 		{
 			MemUsage.NodeCount++;
 			MemUsage.AddUsage(DefaultInstanceData.GetStruct(CompactState.ParameterTemplateIndex.Get()));
diff --git a/Engine/Plugins/Runtime/StateTree/Source/StateTreeModule/Private/StateTreeExecutionContext.cpp b/Engine/Plugins/Runtime/StateTree/Source/StateTreeModule/Private/StateTreeExecutionContext.cpp
--- a/Engine/Plugins/Runtime/StateTree/Source/StateTreeModule/Private/StateTreeExecutionContext.cpp
+++ b/Engine/Plugins/Runtime/StateTree/Source/StateTreeModule/Private/StateTreeExecutionContext.cpp
@@ -776,25 +776,10 @@ void FStateTreeExecutionContext::UpdateInstanceData(TConstArrayView<FStateTreeEx
 			InstanceStructs.AddDefaulted(State.InstanceDataNum);
 			TempInstanceStructs.AddZeroed(State.InstanceDataNum);
 
-			if (State.Type == EStateTreeStateType::Linked
-				|| State.Type == EStateTreeStateType::LinkedAsset)
-			{
-				// Linked state's instance data is the parameters.
-				check(State.ParameterDataHandle.IsValid());
-				const FConstStructView ParamsInstanceData = NextFrame.StateTree->DefaultInstanceData.GetStruct(State.ParameterTemplateIndex.Get());
-				InstanceStructs[BaseIndex + State.ParameterDataHandle.GetIndex()] = ParamsInstanceData;
-
-				// Store the index of the parameter data, so that we can point the linked state to it.
-				check(State.ParameterDataHandle.GetSource() == EStateTreeDataSourceType::LinkedStateParameterData);
-				checkf(!NextStateParameterDataHandle.IsValid(), TEXT("NextStateParameterDataIndex not should be set yet when we encounter a linked state."));
-				NextStateParameterDataHandle = State.ParameterDataHandle;
-				
-				const FCompactStateTreeParameters* Params = ParamsInstanceData.GetPtr<const FCompactStateTreeParameters>();
-				NextStateParameterDataStruct = Params ? Params->Parameters.GetPropertyBagStruct() : nullptr;
-			}
-			else if (State.Type == EStateTreeStateType::Subtree)
+			if (State.Type == EStateTreeStateType::Subtree)
 			{
 				check(State.ParameterDataHandle.IsValid());
+				check(State.ParameterTemplateIndex.IsValid());
 				const FConstStructView ParamsInstanceData = NextFrame.StateTree->DefaultInstanceData.GetStruct(State.ParameterTemplateIndex.Get());
 				if (!NextStateParameterDataHandle.IsValid())
 				{
@@ -813,6 +798,29 @@ void FStateTreeExecutionContext::UpdateInstanceData(TConstArrayView<FStateTreeEx
 					NextStateParameterDataHandle = FStateTreeDataHandle::Invalid; // Mark as used.
 				}
 			}
+			else
+			{
+				if (State.ParameterTemplateIndex.IsValid())
+				{
+					// Linked state's instance data is the parameters.
+					check(State.ParameterDataHandle.IsValid());
+					const FConstStructView ParamsInstanceData = NextFrame.StateTree->DefaultInstanceData.GetStruct(State.ParameterTemplateIndex.Get());
+					InstanceStructs[BaseIndex + State.ParameterDataHandle.GetIndex()] = ParamsInstanceData;
+
+					if (State.Type == EStateTreeStateType::Linked
+						|| State.Type == EStateTreeStateType::LinkedAsset)
+					{
+						// Store the index of the parameter data, so that we can point the linked state to it.
+						check(State.ParameterDataHandle.GetSource() == EStateTreeDataSourceType::StateParameterData);
+						checkf(!NextStateParameterDataHandle.IsValid(), TEXT("NextStateParameterDataIndex not should be set yet when we encounter a linked state."));
+						NextStateParameterDataHandle = State.ParameterDataHandle;
+					
+						const FCompactStateTreeParameters* Params = ParamsInstanceData.GetPtr<const FCompactStateTreeParameters>();
+						NextStateParameterDataStruct = Params ? Params->Parameters.GetPropertyBagStruct() : nullptr;
+					}
+				}
+			}
+			
 			if (!bAreCommon && State.ParameterDataHandle.IsValid())
 			{
 				TempInstanceStructs[BaseIndex + State.ParameterDataHandle.GetIndex()] = FindInstanceTempData(NextFrame, State.ParameterDataHandle);
@@ -913,7 +921,8 @@ FStateTreeDataView FStateTreeExecutionContext::GetDataView(const FStateTreeExecu
 			check(ParentFrame);
 			return GetDataView(nullptr, *ParentFrame, CurrentFrame.StateParameterDataHandle);
 		}
-	case EStateTreeDataSourceType::LinkedStateParameterData:
+
+	case EStateTreeDataSourceType::StateParameterData:
 		{
 			FCompactStateTreeParameters& Params = InstanceDataStorage->GetMutableStruct(CurrentFrame.ActiveInstanceIndexBase.Get() + Handle.GetIndex()).Get<FCompactStateTreeParameters>();
 			return Params.Parameters.GetMutableValue();
@@ -969,7 +978,7 @@ bool FStateTreeExecutionContext::IsHandleSourceValid(const FStateTreeExecutionFr
 		return ParentFrame
 			&& IsHandleSourceValid(nullptr, *ParentFrame, CurrentFrame.StateParameterDataHandle);
 
-	case EStateTreeDataSourceType::LinkedStateParameterData:
+	case EStateTreeDataSourceType::StateParameterData:
 		return CurrentFrame.ActiveInstanceIndexBase.IsValid()
 			&& CurrentFrame.ActiveStates.Contains(Handle.GetState())
 			&& InstanceDataStorage->IsValidIndex(CurrentFrame.ActiveInstanceIndexBase.Get() + Handle.GetIndex());
@@ -994,7 +1003,6 @@ FStateTreeDataView FStateTreeExecutionContext::GetDataViewOrTemporary(const FSta
 	{
 	case EStateTreeDataSourceType::GlobalInstanceData:
 	case EStateTreeDataSourceType::ActiveInstanceData:
-	case EStateTreeDataSourceType::LinkedStateParameterData:
 		return InstanceDataStorage->GetMutableTemporaryStruct(CurrentFrame, Handle);
 		
 	case EStateTreeDataSourceType::GlobalInstanceDataObject:
@@ -1021,6 +1029,15 @@ FStateTreeDataView FStateTreeExecutionContext::GetDataViewOrTemporary(const FSta
 		}
 		break;
 
+	case EStateTreeDataSourceType::StateParameterData:
+		{
+			if (FCompactStateTreeParameters* Params = InstanceDataStorage->GetMutableTemporaryStruct(CurrentFrame, Handle).GetPtr<FCompactStateTreeParameters>())
+			{
+				return Params->Parameters.GetMutableValue();
+			}
+		}
+		break;
+
 	default:
 		return {};
 	}
diff --git a/Engine/Plugins/Runtime/StateTree/Source/StateTreeModule/Private/StateTreePropertyBindings.cpp b/Engine/Plugins/Runtime/StateTree/Source/StateTreeModule/Private/StateTreePropertyBindings.cpp
--- a/Engine/Plugins/Runtime/StateTree/Source/StateTreeModule/Private/StateTreePropertyBindings.cpp
+++ b/Engine/Plugins/Runtime/StateTree/Source/StateTreeModule/Private/StateTreePropertyBindings.cpp
@@ -3,6 +3,7 @@
 #include "UObject/EnumProperty.h"
 #include "Misc/EnumerateRange.h"
 #include "PropertyPathHelpers.h"
+#include "PropertyBag.h"
 
 #if WITH_EDITOR
 #include "UObject/CoreRedirects.h"
@@ -1425,6 +1426,17 @@ bool FStateTreePropertyPath::ResolveIndirectionsWithValue(const FStateTreeDataVi
 							Property = FoundProperty;
 						}
 					}
+					else if (const UPropertyBag* PropertyBag = Cast<UPropertyBag>(CurrentStruct))
+					{
+						if (const FPropertyBagPropertyDesc* Desc = PropertyBag->FindPropertyDescByID(Segment->GetPropertyGuid()))
+						{
+							if (Desc->CachedProperty)
+							{
+								RedirectedName = Desc->CachedProperty->GetFName();
+								Property = Desc->CachedProperty;
+							}
+						}
+					}
 				}
 				else
 				{
@@ -1457,6 +1469,13 @@ bool FStateTreePropertyPath::ResolveIndirectionsWithValue(const FStateTreeDataVi
 					// Parse Guid from UDS property name.
 					PropertyGuid = FStructureEditorUtils::GetGuidFromPropertyName(PropertyName);
 				}
+				else if (const UPropertyBag* PropertyBag = Cast<UPropertyBag>(CurrentStruct))
+				{
+					if (const FPropertyBagPropertyDesc* Desc = PropertyBag->FindPropertyDescByPropertyName(PropertyName))
+					{
+						PropertyGuid = Desc->ID;
+					}
+				}
 			}
 		}
 #endif // WITH_EDITORONLY_DATA
diff --git a/Engine/Plugins/Runtime/StateTree/Source/StateTreeModule/Public/StateTreeTypes.h b/Engine/Plugins/Runtime/StateTree/Source/StateTreeModule/Public/StateTreeTypes.h
--- a/Engine/Plugins/Runtime/StateTree/Source/StateTreeModule/Public/StateTreeTypes.h
+++ b/Engine/Plugins/Runtime/StateTree/Source/StateTreeModule/Public/StateTreeTypes.h
@@ -243,11 +243,11 @@ enum class EStateTreeDataSourceType : uint8
 	/** Global parameters */
 	GlobalParameterData,
 
-	/** Subtree parameters */
+	/** Parameters for subtree (may resolve to a linked state's parameters or default params) */
 	SubtreeParameterData,
 
-	/** Linked state parameters */
-	LinkedStateParameterData,
+	/** Parameters for regular and linked states */
+	StateParameterData,
 };
 
 /** Handle to a StateTree data */
@@ -380,7 +380,7 @@ struct STATETREEMODULE_API FStateTreeDataHandle
 			return FString::Printf(TEXT("GlobalParam[%d]"), Index);
 		case EStateTreeDataSourceType::SubtreeParameterData:
 			return FString::Printf(TEXT("SubtreeParam[%d]"), Index);
-		case EStateTreeDataSourceType::LinkedStateParameterData:
+		case EStateTreeDataSourceType::StateParameterData:
 			return FString::Printf(TEXT("LinkedParam[%d]"), Index);
 		default:
 			return TEXT("---");
