Fixed typos

diff --git a/Engine/Plugins/Runtime/StateTree/Source/StateTreeModule/Public/StateTreeIndexTypes.h b/Engine/Plugins/Runtime/StateTree/Source/StateTreeModule/Public/StateTreeIndexTypes.h
--- a/Engine/Plugins/Runtime/StateTree/Source/StateTreeModule/Public/StateTreeIndexTypes.h
+++ b/Engine/Plugins/Runtime/StateTree/Source/StateTreeModule/Public/StateTreeIndexTypes.h
@@ -44,7 +44,7 @@ struct STATETREEMODULE_API FStateTreeIndex16
 		Value = InIndex == INDEX_NONE ? InvalidValue : (uint16)InIndex;
 	}
 
-	/** @retrun value of the index or FStateTreeIndex16::InvalidValue (i.e. MAX_uint16) if invalid. */
+	/** @return value of the index or FStateTreeIndex16::InvalidValue (i.e. MAX_uint16) if invalid. */
 	uint16 Get() const { return Value; }
 	
 	/** @return the index value as int32, mapping invalid value to INDEX_NONE. */
@@ -95,7 +95,7 @@ struct STATETREEMODULE_API FStateTreeIndex8
 		Value = InIndex == INDEX_NONE ? InvalidValue : (uint8)InIndex;
 	}
 
-	/** @retrun value of the index. */
+	/** @return value of the index. */
 	uint8 Get() const { return Value; }
 
 	/** @return the index value as int32, mapping invalid value to INDEX_NONE. */
diff --git a/Engine/Plugins/Runtime/WorldConditions/Source/WorldConditions/Public/WorldConditionBase.h b/Engine/Plugins/Runtime/WorldConditions/Source/WorldConditions/Public/WorldConditionBase.h
--- a/Engine/Plugins/Runtime/WorldConditions/Source/WorldConditions/Public/WorldConditionBase.h
+++ b/Engine/Plugins/Runtime/WorldConditions/Source/WorldConditions/Public/WorldConditionBase.h
@@ -128,7 +128,7 @@ struct WORLDCONDITIONS_API FWorldConditionBase
 	virtual TObjectPtr<const UStruct>* GetRuntimeStateType() const { return nullptr; }
 
 #if WITH_EDITOR
-	/** @retrun Description to be shown in the UI. */
+	/** @return Description to be shown in the UI. */
 	virtual FText GetDescription() const;
 #endif
 
diff --git a/Engine/Plugins/Runtime/ZoneGraph/Source/ZoneGraph/Public/ZoneGraphTypes.h b/Engine/Plugins/Runtime/ZoneGraph/Source/ZoneGraph/Public/ZoneGraphTypes.h
--- a/Engine/Plugins/Runtime/ZoneGraph/Source/ZoneGraph/Public/ZoneGraphTypes.h
+++ b/Engine/Plugins/Runtime/ZoneGraph/Source/ZoneGraph/Public/ZoneGraphTypes.h
@@ -959,7 +959,7 @@ struct ZONEGRAPH_API FZoneGraphBuildSettings
 	/** @return tessellation tolerance for specific case, or common tolerance if no match. */
 	float GetLaneTessellationTolerance(const FZoneGraphTagMask LaneTags) const;
 
-	/** @retrun Connection restrictions for specified lane based on PolygonRoutingRules */
+	/** @return Connection restrictions for specified lane based on PolygonRoutingRules */
 	EZoneShapeLaneConnectionRestrictions GetConnectionRestrictions(const FZoneGraphTagMask ZoneTags,
 																   const FZoneLaneProfileRef& SourceLaneProfile, const int32 SourceConnectionCount,
 																   const FZoneLaneProfileRef& DestinationLaneProfile, const int32 DestinationConnectionCount) const;
