StateTree: FIx type on error message.

diff --git a/Engine/Plugins/Runtime/StateTree/Source/StateTreeEditorModule/Private/StateTreeCompiler.cpp b/Engine/Plugins/Runtime/StateTree/Source/StateTreeEditorModule/Private/StateTreeCompiler.cpp
--- a/Engine/Plugins/Runtime/StateTree/Source/StateTreeEditorModule/Private/StateTreeCompiler.cpp
+++ b/Engine/Plugins/Runtime/StateTree/Source/StateTreeEditorModule/Private/StateTreeCompiler.cpp
@@ -1678,7 +1678,7 @@ bool FStateTreeCompiler::GetAndValidateBindings(const FStateTreeBindableStructDe
 				else
 				{
 					Log.Reportf(EMessageSeverity::Error, TargetStruct,
-						TEXT("Cound not find matching Context object for Context property '%s' on %s. Property must have manual binding."),
+						TEXT("Could not find matching Context object for Context property '%s' on '%s'. Property must have manual binding."),
 						*PropertyName.ToString(), *TargetStruct.ToString());
 					bResult = false;
 				}
