[StateTreeDebugger] support for linked tree assets
- added new dedicated section for completed states

diff --git a/Engine/Plugins/Runtime/StateTree/Source/StateTreeEditorModule/Private/Debugger/SStateTreeDebuggerView.cpp b/Engine/Plugins/Runtime/StateTree/Source/StateTreeEditorModule/Private/Debugger/SStateTreeDebuggerView.cpp
--- a/Engine/Plugins/Runtime/StateTree/Source/StateTreeEditorModule/Private/Debugger/SStateTreeDebuggerView.cpp
+++ b/Engine/Plugins/Runtime/StateTree/Source/StateTreeEditorModule/Private/Debugger/SStateTreeDebuggerView.cpp
@@ -181,7 +181,8 @@ void GenerateElementForProperties(const TCHAR* TypeAsText, const TCHAR* ValueAsT
 				// Create Tree element to hold the event
 				const TSharedPtr<FStateTreeDebuggerEventTreeElement> NewChildElement = MakeShareable(new FStateTreeDebuggerEventTreeElement(
 						ParentElement->Frame,
-						FStateTreeTraceEventVariantType(TInPlaceType<FStateTreeTracePropertyEvent>(), PropertyEvent)));
+						FStateTreeTraceEventVariantType(TInPlaceType<FStateTreeTracePropertyEvent>(), PropertyEvent),
+						ParentElement->WeakStateTree.Get()));
 
 				ParentElement->Children.Add(NewChildElement);
 			}
@@ -505,8 +506,7 @@ void SStateTreeDebuggerView::Construct(const FArguments& InArgs, const UStateTre
 	EventsTreeView = SNew(STreeView<TSharedPtr<FStateTreeDebuggerEventTreeElement>>)
 			.OnGenerateRow_Lambda([this](const TSharedPtr<FStateTreeDebuggerEventTreeElement>& InElement, const TSharedRef<STableViewBase>& InOwnerTableView)
 			{
-				check(StateTreeViewModel);
-				return SNew(SStateTreeDebuggerViewRow, InOwnerTableView, InElement, StateTreeViewModel.ToSharedRef());
+				return SNew(SStateTreeDebuggerViewRow, InOwnerTableView, InElement);
 			})
 			.OnGetChildren_Lambda([](const TSharedPtr<const FStateTreeDebuggerEventTreeElement>& InParent, TArray<TSharedPtr<FStateTreeDebuggerEventTreeElement>>& OutChildren)
 			{
@@ -1056,6 +1056,9 @@ void SStateTreeDebuggerView::OnDebuggerScrubStateChanged(const UE::StateTreeDebu
 	const int32 FirstEventIdx = Spans[SpanIdx].EventIdx;
 	const TraceServices::FFrame Frame = Spans[SpanIdx].Frame;
 	const int32 MaxEventIdx = Spans.IsValidIndex(SpanIdx+1) ? Spans[SpanIdx+1].EventIdx : Events.Num();
+
+	const UStateTree* const RootTree = StateTree.Get();
+	const UStateTree* ActiveTree = RootTree;
 	
 	for (int32 EventIdx = FirstEventIdx; EventIdx < MaxEventIdx; EventIdx++)
 	{
@@ -1081,25 +1084,6 @@ void SStateTreeDebuggerView::OnDebuggerScrubStateChanged(const UE::StateTreeDebu
 		{
 			if (PhaseEvent->EventType == EStateTreeTraceEventType::Push)
 			{
-				if (PhaseEvent->Phase != EStateTreeUpdatePhase::Unset)
-				{
-					CustomDescription = UEnum::GetDisplayValueAsText(PhaseEvent->Phase).ToString();
-				}
-
-				if (PhaseEvent->StateHandle.IsValid())
-				{
-					const FCompactStateTreeState* CompactState = StateTree->GetStateFromHandle(PhaseEvent->StateHandle);
-					if (CustomDescription.IsEmpty())
-					{
-						CustomDescription += FString::Printf(TEXT("%s"),
-								CompactState != nullptr ? *CompactState->Name.ToString() : *PhaseEvent->StateHandle.Describe());	
-					}
-					else
-					{
-						CustomDescription += FString::Printf(TEXT(" '%s'"),
-								CompactState != nullptr ? *CompactState->Name.ToString() : *PhaseEvent->StateHandle.Describe());
-					}
-				}
 				bShouldAddToScopeStack = true;
 			}
 			else if (PhaseEvent->EventType == EStateTreeTraceEventType::Pop)
@@ -1107,6 +1091,14 @@ void SStateTreeDebuggerView::OnDebuggerScrubStateChanged(const UE::StateTreeDebu
 				bShouldPopScopeStack = true;
 			}
 		}
+		else if (const FStateTreeTraceInstanceFrameEvent* FrameEvent = Event.TryGet<FStateTreeTraceInstanceFrameEvent>())
+		{
+			ActiveTree = FrameEvent->WeakStateTree.Get();
+			check(ActiveTree);
+
+			// We don't want to create an entry.
+			continue;
+		}
 
 		if (bShouldPopScopeStack)
 		{
@@ -1121,7 +1113,7 @@ void SStateTreeDebuggerView::OnDebuggerScrubStateChanged(const UE::StateTreeDebu
 			continue;
 		}
 
-		const TSharedRef<FStateTreeDebuggerEventTreeElement> NewElement = MakeShareable(new FStateTreeDebuggerEventTreeElement(Frame, Event));
+		const TSharedRef<FStateTreeDebuggerEventTreeElement> NewElement = MakeShareable(new FStateTreeDebuggerEventTreeElement(Frame, Event, ActiveTree));
 		NewElement->Description = CustomDescription;
 
 		TArray<TSharedPtr<FStateTreeDebuggerEventTreeElement>>& TreeElements = ScopeStack.IsEmpty() ? EventsTreeElements : ScopeStack.Top()->Children;
@@ -1163,7 +1155,8 @@ void SStateTreeDebuggerView::GenerateElementsForProperties(const FStateTreeTrace
 				// Create Tree element to hold the event
 				const TSharedPtr<FStateTreeDebuggerEventTreeElement> NewChildElement = MakeShareable(new FStateTreeDebuggerEventTreeElement(
 					ParentElement->Frame,
-					FStateTreeTraceEventVariantType(TInPlaceType<FStateTreeTracePropertyEvent>(), PropertyEvent)));
+					FStateTreeTraceEventVariantType(TInPlaceType<FStateTreeTracePropertyEvent>(), PropertyEvent),
+					ParentElement->WeakStateTree.Get()));
 
 				ParentElement->Children.Add(NewChildElement);
 			};
diff --git a/Engine/Plugins/Runtime/StateTree/Source/StateTreeEditorModule/Private/Debugger/SStateTreeDebuggerViewRow.cpp b/Engine/Plugins/Runtime/StateTree/Source/StateTreeEditorModule/Private/Debugger/SStateTreeDebuggerViewRow.cpp
--- a/Engine/Plugins/Runtime/StateTree/Source/StateTreeEditorModule/Private/Debugger/SStateTreeDebuggerViewRow.cpp
+++ b/Engine/Plugins/Runtime/StateTree/Source/StateTreeEditorModule/Private/Debugger/SStateTreeDebuggerViewRow.cpp
@@ -4,18 +4,16 @@
 
 #include "Debugger/SStateTreeDebuggerViewRow.h"
 #include "StateTreeEditorStyle.h"
-#include "StateTreeViewModel.h"
+#include "StateTree.h"
 #include "Widgets/Images/SImage.h"
 #include "Widgets/Text/STextBlock.h"
 
 
 void SStateTreeDebuggerViewRow::Construct(const FArguments& InArgs,
 										  const TSharedPtr<STableViewBase>& InOwnerTableView,
-										  const TSharedPtr<FStateTreeDebuggerEventTreeElement>& InElement,
-										  const TSharedRef<FStateTreeViewModel>& InStateTreeViewModel)
+										  const TSharedPtr<FStateTreeDebuggerEventTreeElement>& InElement)
 {
 	Item = InElement;
-	StateTreeViewModel = InStateTreeViewModel.ToSharedPtr();
 	STableRow::Construct(InArgs, InOwnerTableView.ToSharedRef());
 
 	const TSharedPtr<SHorizontalBox> HorizontalBox = SNew(SHorizontalBox);
@@ -70,8 +68,9 @@ TSharedPtr<SWidget> SStateTreeDebuggerViewRow::GenerateEventWidget() const
 		const FSlateBrush* Image = nullptr;
 		switch (PhaseEvent->Phase)
 		{
-		case EStateTreeUpdatePhase::EnterStates:	Image = StyleSet.GetBrush("StateTreeEditor.Debugger.State.Enter");	break;
-		case EStateTreeUpdatePhase::ExitStates:		Image = StyleSet.GetBrush("StateTreeEditor.Debugger.State.Exit");	break;
+		case EStateTreeUpdatePhase::EnterStates:	Image = StyleSet.GetBrush("StateTreeEditor.Debugger.State.Enter");		break;
+		case EStateTreeUpdatePhase::ExitStates:		Image = StyleSet.GetBrush("StateTreeEditor.Debugger.State.Exit");		break;
+		case EStateTreeUpdatePhase::StateCompleted:	Image = StyleSet.GetBrush("StateTreeEditor.Debugger.State.Completed");	break;
 		default:
 			return nullptr;
 		}
@@ -86,7 +85,6 @@ TSharedPtr<SWidget> SStateTreeDebuggerViewRow::GenerateEventWidget() const
 		switch (StateEvent->EventType)
 		{
 		case EStateTreeTraceEventType::OnStateSelected:		Image = StyleSet.GetBrush("StateTreeEditor.Debugger.State.Selected");	break;
-		case EStateTreeTraceEventType::OnStateCompleted:	Image = StyleSet.GetBrush("StateTreeEditor.Debugger.State.Completed");	break;
 		default:
 			return nullptr;
 		}
@@ -162,8 +160,7 @@ FText SStateTreeDebuggerViewRow::GetEventDescription() const
 	FString EventDescription;
 	if (Item->Description.IsEmpty())
 	{
-		check(StateTreeViewModel);
-		if (const UStateTree* StateTree = StateTreeViewModel->GetStateTree())
+		if (const UStateTree* StateTree = Item->WeakStateTree.Get())
 		{
 			// Some types have some custom representations so we want to use a more minimal description.
 			if (Item->Event.IsType<FStateTreeTraceStateEvent>()
@@ -198,8 +195,7 @@ FText SStateTreeDebuggerViewRow::GetEventDescription() const
 FText SStateTreeDebuggerViewRow::GetEventTooltip() const
 {
 	FString Tooltip;
-	check(StateTreeViewModel);
-	if (const UStateTree* StateTree = StateTreeViewModel->GetStateTree())
+	if (const UStateTree* StateTree = Item->WeakStateTree.Get())
 	{
 		Visit([&Tooltip, StateTree](auto& TypedEvent)
 			{
diff --git a/Engine/Plugins/Runtime/StateTree/Source/StateTreeEditorModule/Private/Debugger/SStateTreeDebuggerViewRow.h b/Engine/Plugins/Runtime/StateTree/Source/StateTreeEditorModule/Private/Debugger/SStateTreeDebuggerViewRow.h
--- a/Engine/Plugins/Runtime/StateTree/Source/StateTreeEditorModule/Private/Debugger/SStateTreeDebuggerViewRow.h
+++ b/Engine/Plugins/Runtime/StateTree/Source/StateTreeEditorModule/Private/Debugger/SStateTreeDebuggerViewRow.h
@@ -15,8 +15,8 @@ class FStateTreeViewModel;
 /** An item in the StateTreeDebugger trace event tree */
 struct FStateTreeDebuggerEventTreeElement : TSharedFromThis<FStateTreeDebuggerEventTreeElement>
 {
-	explicit FStateTreeDebuggerEventTreeElement(const TraceServices::FFrame& Frame, const FStateTreeTraceEventVariantType& Event)
-		: Frame(Frame), Event(Event)
+	explicit FStateTreeDebuggerEventTreeElement(const TraceServices::FFrame& Frame, const FStateTreeTraceEventVariantType& Event, const UStateTree* StateTree)
+		: Frame(Frame), Event(Event), WeakStateTree(StateTree)
 	{
 	}
 
@@ -24,6 +24,7 @@ struct FStateTreeDebuggerEventTreeElement : TSharedFromThis<FStateTreeDebuggerEv
 	FStateTreeTraceEventVariantType Event;
 	TArray<TSharedPtr<FStateTreeDebuggerEventTreeElement>> Children;
 	FString Description;
+	TWeakObjectPtr<const UStateTree> WeakStateTree;
 };
 
 
@@ -35,16 +36,14 @@ class SStateTreeDebuggerViewRow : public STableRow<TSharedPtr<FStateTreeDebugger
 public:
 	void Construct(const FArguments& InArgs,
 				   const TSharedPtr<STableViewBase>& InOwnerTableView,
-				   const TSharedPtr<FStateTreeDebuggerEventTreeElement>& InElement,
-				   const TSharedRef<FStateTreeViewModel>& InStateTreeViewModel);
+				   const TSharedPtr<FStateTreeDebuggerEventTreeElement>& InElement);
 
 private:
 	TSharedPtr<SWidget> GenerateEventWidget() const;
 	const FTextBlockStyle& GetEventTextStyle() const;
 	FText GetEventDescription() const;
 	FText GetEventTooltip() const;
-	
-	TSharedPtr<FStateTreeViewModel> StateTreeViewModel;
+
 	TSharedPtr<FStateTreeDebuggerEventTreeElement> Item;
 };
 
diff --git a/Engine/Plugins/Runtime/StateTree/Source/StateTreeEditorModule/Private/Debugger/StateTreeDebuggerTrack.cpp b/Engine/Plugins/Runtime/StateTree/Source/StateTreeEditorModule/Private/Debugger/StateTreeDebuggerTrack.cpp
--- a/Engine/Plugins/Runtime/StateTree/Source/StateTreeEditorModule/Private/Debugger/StateTreeDebuggerTrack.cpp
+++ b/Engine/Plugins/Runtime/StateTree/Source/StateTreeEditorModule/Private/Debugger/StateTreeDebuggerTrack.cpp
@@ -65,13 +65,7 @@ bool FStateTreeDebuggerInstanceTrack::UpdateInternal()
 			const uint32 EventIndex = EventCollection.ActiveStatesChanges[StateChangeIndex].EventIndex;
 			const FStateTreeTraceActiveStatesEvent& Event = Events[EventIndex].Get<FStateTreeTraceActiveStatesEvent>();
 				
-			FString StatePath;
-			for (int32 ActiveStateIndex = 0; ActiveStateIndex < Event.ActiveStates.Num(); ActiveStateIndex++)
-			{
-				const FCompactStateTreeState* State = StateTree->GetStateFromHandle(Event.ActiveStates[ActiveStateIndex]);
-				StatePath.Appendf(TEXT("%s%s"), ActiveStateIndex == 0 ? TEXT("") : TEXT("."), State == nullptr ? TEXT("Invalid") : *State->Name.ToString());
-			}
-
+			FString StatePath = Event.GetValueString(*StateTree);
 			UE::StateTreeDebugger::FFrameSpan Span = EventCollection.FrameSpans[SpanIndex];
 			SStateTreeDebuggerEventTimelineView::FTimelineEventData::EventWindow& Window = EventData->Windows.AddDefaulted_GetRef();
 			Window.Color = MakeRandomColor(GetTypeHash(StatePath));
diff --git a/Engine/Plugins/Runtime/StateTree/Source/StateTreeEditorModule/Private/StateTreeViewModel.cpp b/Engine/Plugins/Runtime/StateTree/Source/StateTreeEditorModule/Private/StateTreeViewModel.cpp
--- a/Engine/Plugins/Runtime/StateTree/Source/StateTreeEditorModule/Private/StateTreeViewModel.cpp
+++ b/Engine/Plugins/Runtime/StateTree/Source/StateTreeEditorModule/Private/StateTreeViewModel.cpp
@@ -1048,15 +1048,21 @@ void FStateTreeViewModel::MoveSelectedStates(UStateTreeState* TargetState, const
 void FStateTreeViewModel::BindToDebuggerDelegates()
 {
 #if WITH_STATETREE_DEBUGGER
-	Debugger->OnActiveStatesChanged.BindLambda([this](const TConstArrayView<FStateTreeStateHandle> NewActiveStates)
+	Debugger->OnActiveStatesChanged.BindSPLambda(this, [this](const FStateTreeTraceActiveStates& NewActiveStates)
 	{
 		if (const UStateTree* OuterStateTree = GetStateTree())
 		{
-			ActiveStates.Reset(NewActiveStates.Num());
-
-			for (const FStateTreeStateHandle Handle : NewActiveStates)
+			for (const FStateTreeTraceActiveStates::FAssetActiveStates& AssetActiveStates : NewActiveStates.PerAssetStates)
 			{
-				ActiveStates.Add(OuterStateTree->GetStateIdFromHandle(Handle));
+				// Only track states owned by the StateTree associated to the view model (skip linked assets)
+				if (AssetActiveStates.WeakStateTree == OuterStateTree)
+				{
+					ActiveStates.Reset(AssetActiveStates.ActiveStates.Num());
+					for (const FStateTreeStateHandle Handle : AssetActiveStates.ActiveStates)
+					{
+						ActiveStates.Add(OuterStateTree->GetStateIdFromHandle(Handle));
+					}
+				}
 			}
 		}
 	});
diff --git a/Engine/Plugins/Runtime/StateTree/Source/StateTreeModule/Private/Debugger/StateTreeDebugger.cpp b/Engine/Plugins/Runtime/StateTree/Source/StateTreeModule/Private/Debugger/StateTreeDebugger.cpp
--- a/Engine/Plugins/Runtime/StateTree/Source/StateTreeModule/Private/Debugger/StateTreeDebugger.cpp
+++ b/Engine/Plugins/Runtime/StateTree/Source/StateTreeModule/Private/Debugger/StateTreeDebugger.cpp
@@ -507,7 +507,7 @@ FText FStateTreeDebugger::DescribeInstance(const UE::StateTreeDebugger::FInstanc
 	return FText::FromString(LexToString(InstanceDesc));
 }
 
-void FStateTreeDebugger::SetActiveStates(const TConstArrayView<FStateTreeStateHandle> NewActiveStates)
+void FStateTreeDebugger::SetActiveStates(const FStateTreeTraceActiveStates& NewActiveStates)
 {
 	ActiveStates = NewActiveStates;
 	OnActiveStatesChanged.ExecuteIfBound(ActiveStates);
@@ -515,16 +515,16 @@ void FStateTreeDebugger::SetActiveStates(const TConstArrayView<FStateTreeStateHa
 
 void FStateTreeDebugger::RefreshActiveStates()
 {
-	TArray<FStateTreeStateHandle> NewActiveStates;
-
 	if (ScrubState.IsPointingToValidActiveStates())
 	{
 		const UE::StateTreeDebugger::FInstanceEventCollection& EventCollection = EventCollections[ScrubState.GetEventCollectionIndex()];
 		const int32 EventIndex = EventCollection.ActiveStatesChanges[ScrubState.GetActiveStatesIndex()].EventIndex;
-		NewActiveStates = EventCollection.Events[EventIndex].Get<FStateTreeTraceActiveStatesEvent>().ActiveStates;
+		SetActiveStates(EventCollection.Events[EventIndex].Get<FStateTreeTraceActiveStatesEvent>().ActiveStates);
+	}
+	else
+	{
+		SetActiveStates(FStateTreeTraceActiveStates());	
 	}
-
-	SetActiveStates(NewActiveStates);
 }
 
 bool FStateTreeDebugger::CanStepBackToPreviousStateWithEvents() const
diff --git a/Engine/Plugins/Runtime/StateTree/Source/StateTreeModule/Private/Debugger/StateTreeTrace.cpp b/Engine/Plugins/Runtime/StateTree/Source/StateTreeModule/Private/Debugger/StateTreeTrace.cpp
--- a/Engine/Plugins/Runtime/StateTree/Source/StateTreeModule/Private/Debugger/StateTreeTrace.cpp
+++ b/Engine/Plugins/Runtime/StateTree/Source/StateTreeModule/Private/Debugger/StateTreeTrace.cpp
@@ -23,24 +23,37 @@ UE_TRACE_EVENT_BEGIN(StateTreeDebugger, WorldTimestampEvent)
 	UE_TRACE_EVENT_FIELD(double, WorldTime)
 UE_TRACE_EVENT_END()
 
-UE_TRACE_EVENT_BEGIN(StateTreeDebugger, InstanceEvent)
+UE_TRACE_EVENT_BEGIN(StateTreeDebugger, AssetDebugIdEvent)
 	UE_TRACE_EVENT_FIELD(uint64, Cycle)
 	UE_TRACE_EVENT_FIELD(UE::Trace::WideString, TreeName)
 	UE_TRACE_EVENT_FIELD(UE::Trace::WideString, TreePath)
 	UE_TRACE_EVENT_FIELD(uint32, CompiledDataHash)
+	UE_TRACE_EVENT_FIELD(uint16, AssetDebugId)
+UE_TRACE_EVENT_END()
+
+UE_TRACE_EVENT_BEGIN(StateTreeDebugger, InstanceEvent)
+	UE_TRACE_EVENT_FIELD(uint64, Cycle)
 	UE_TRACE_EVENT_FIELD(uint32, InstanceId)
 	UE_TRACE_EVENT_FIELD(uint32, InstanceSerial)
 	UE_TRACE_EVENT_FIELD(UE::Trace::WideString, InstanceName)
-	UE_TRACE_EVENT_FIELD(uint8, EventType)
+	UE_TRACE_EVENT_FIELD(std::underlying_type_t<EStateTreeTraceEventType>, EventType)
+	UE_TRACE_EVENT_FIELD(uint16, AssetDebugId)
+UE_TRACE_EVENT_END()
+
+UE_TRACE_EVENT_BEGIN(StateTreeDebugger, InstanceFrameEvent)
+	UE_TRACE_EVENT_FIELD(uint64, Cycle)
+	UE_TRACE_EVENT_FIELD(uint32, InstanceId)
+	UE_TRACE_EVENT_FIELD(uint32, InstanceSerial)
+	UE_TRACE_EVENT_FIELD(uint16, AssetDebugId)
 UE_TRACE_EVENT_END()
 
 UE_TRACE_EVENT_BEGIN(StateTreeDebugger, PhaseEvent)
 	UE_TRACE_EVENT_FIELD(uint64, Cycle)
 	UE_TRACE_EVENT_FIELD(uint32, InstanceId)
 	UE_TRACE_EVENT_FIELD(uint32, InstanceSerial)
-	UE_TRACE_EVENT_FIELD(uint16, Phase)
+	UE_TRACE_EVENT_FIELD(std::underlying_type_t<EStateTreeUpdatePhase>, Phase)
 	UE_TRACE_EVENT_FIELD(uint16, StateIndex)
-	UE_TRACE_EVENT_FIELD(uint8, EventType)
+	UE_TRACE_EVENT_FIELD(std::underlying_type_t<EStateTreeTraceEventType>, EventType)
 UE_TRACE_EVENT_END()
 
 UE_TRACE_EVENT_BEGIN(StateTreeDebugger, LogEvent)
@@ -55,8 +68,7 @@ UE_TRACE_EVENT_BEGIN(StateTreeDebugger, StateEvent)
 	UE_TRACE_EVENT_FIELD(uint32, InstanceId)
 	UE_TRACE_EVENT_FIELD(uint32, InstanceSerial)
 	UE_TRACE_EVENT_FIELD(uint16, StateIndex)
-	UE_TRACE_EVENT_FIELD(uint8, EventType)
-	UE_TRACE_EVENT_FIELD(uint8, SelectionBehavior)
+	UE_TRACE_EVENT_FIELD(std::underlying_type_t<EStateTreeTraceEventType>, EventType)
 UE_TRACE_EVENT_END()
 
 UE_TRACE_EVENT_BEGIN(StateTreeDebugger, TaskEvent)
@@ -65,7 +77,7 @@ UE_TRACE_EVENT_BEGIN(StateTreeDebugger, TaskEvent)
 	UE_TRACE_EVENT_FIELD(uint32, InstanceSerial)
 	UE_TRACE_EVENT_FIELD(uint16, NodeIndex)
 	UE_TRACE_EVENT_FIELD(uint8[], DataView)
-	UE_TRACE_EVENT_FIELD(uint8, EventType)
+	UE_TRACE_EVENT_FIELD(std::underlying_type_t<EStateTreeTraceEventType>, EventType)
 	UE_TRACE_EVENT_FIELD(uint8, Status)
 UE_TRACE_EVENT_END()
 
@@ -75,7 +87,7 @@ UE_TRACE_EVENT_BEGIN(StateTreeDebugger, EvaluatorEvent)
 	UE_TRACE_EVENT_FIELD(uint32, InstanceSerial)
 	UE_TRACE_EVENT_FIELD(uint16, NodeIndex)
 	UE_TRACE_EVENT_FIELD(uint8[], DataView)
-	UE_TRACE_EVENT_FIELD(uint8, EventType)
+	UE_TRACE_EVENT_FIELD(std::underlying_type_t<EStateTreeTraceEventType>, EventType)
 UE_TRACE_EVENT_END()
 
 UE_TRACE_EVENT_BEGIN(StateTreeDebugger, TransitionEvent)
@@ -86,7 +98,7 @@ UE_TRACE_EVENT_BEGIN(StateTreeDebugger, TransitionEvent)
 	UE_TRACE_EVENT_FIELD(uint16, TransitionIndex)
 	UE_TRACE_EVENT_FIELD(uint16, TargetStateIndex)
 	UE_TRACE_EVENT_FIELD(uint8, Priority)
-	UE_TRACE_EVENT_FIELD(uint8, EventType)
+	UE_TRACE_EVENT_FIELD(std::underlying_type_t<EStateTreeTraceEventType>, EventType)
 UE_TRACE_EVENT_END()
 
 UE_TRACE_EVENT_BEGIN(StateTreeDebugger, ConditionEvent)
@@ -95,7 +107,7 @@ UE_TRACE_EVENT_BEGIN(StateTreeDebugger, ConditionEvent)
 	UE_TRACE_EVENT_FIELD(uint32, InstanceSerial)
 	UE_TRACE_EVENT_FIELD(uint16, NodeIndex)
 	UE_TRACE_EVENT_FIELD(uint8[], DataView)
-	UE_TRACE_EVENT_FIELD(uint8, EventType)
+	UE_TRACE_EVENT_FIELD(std::underlying_type_t<EStateTreeTraceEventType>, EventType)
 UE_TRACE_EVENT_END()
 
 UE_TRACE_EVENT_BEGIN(StateTreeDebugger, ActiveStatesEvent)
@@ -103,6 +115,7 @@ UE_TRACE_EVENT_BEGIN(StateTreeDebugger, ActiveStatesEvent)
 	UE_TRACE_EVENT_FIELD(uint32, InstanceId)
 	UE_TRACE_EVENT_FIELD(uint32, InstanceSerial)
 	UE_TRACE_EVENT_FIELD(uint16[], ActiveStates)
+	UE_TRACE_EVENT_FIELD(uint16[], AssetDebugIds)
 UE_TRACE_EVENT_END()
 
 namespace UE::StateTreeTrace
@@ -137,20 +150,81 @@ struct FPhaseStack
 	TArray<FPhaseTraceStatusPair> Stack;
 };
 
+/**
+ * Struct to hold data for asset debug id events until we are ready to trace the events (i.e. traces are active and channel is enabled).
+ */
+struct FAssetDebugIdEventBufferedData
+{
+	FAssetDebugIdEventBufferedData() = default;
+	explicit FAssetDebugIdEventBufferedData(const UStateTree* StateTree, const FStateTreeIndex16 AssetDebugId) : WeakStateTree(StateTree), AssetDebugId(AssetDebugId)
+	{
+	}
+
+	void Trace() const
+	{
+		if (ensureMsgf(UE_TRACE_CHANNELEXPR_IS_ENABLED(StateTreeDebugChannel), TEXT("Tracing a buffered data is expected only if channel is enabled.")))
+		{
+			if (const UStateTree* StateTree = WeakStateTree.Get())
+			{
+				OutputAssetDebugIdEvent(StateTree, AssetDebugId);
+			}
+		}
+	}
+
+	TWeakObjectPtr<const UStateTree> WeakStateTree;
+	FStateTreeIndex16 AssetDebugId;
+};
+
 /**
  * Struct to hold data for active states events until we are ready to trace the events (i.e. traces are active and channel is enabled).
  */
 struct FInstanceEventBufferedData
 {
+	struct FActiveStates
+	{
+		FActiveStates() = default;
+		explicit FActiveStates(const TConstArrayView<FStateTreeExecutionFrame> ActiveFrames)
+		{
+			for (const FStateTreeExecutionFrame& Frame : ActiveFrames)
+			{
+				const FStateTreeIndex16 AssetDebugId = FindOrAddDebugIdForAsset(Frame.StateTree.Get());
+
+				const int32 RequiredSize = StatesIndices.Num() + Frame.ActiveStates.Num();
+				StatesIndices.Reserve(RequiredSize);
+				AssetDebugIds.Reserve(RequiredSize);
+
+				for (const FStateTreeStateHandle StateHandle : Frame.ActiveStates)
+				{
+					StatesIndices.Add(StateHandle.Index);
+					AssetDebugIds.Add(AssetDebugId.Get());
+				}
+			}
+		}
+
+		bool IsValid() const { return StatesIndices.Num() > 0 && StatesIndices.Num() == AssetDebugIds.Num(); }
+		void Output(const FStateTreeInstanceDebugId InstanceId) const
+		{
+			UE_TRACE_LOG(StateTreeDebugger, ActiveStatesEvent, StateTreeDebugChannel)
+				<< ActiveStatesEvent.Cycle(FPlatformTime::Cycles64())
+				<< ActiveStatesEvent.InstanceId(InstanceId.Id)
+				<< ActiveStatesEvent.InstanceSerial(InstanceId.SerialNumber)
+				<< ActiveStatesEvent.ActiveStates(StatesIndices.GetData(), StatesIndices.Num())
+				<< ActiveStatesEvent.AssetDebugIds(AssetDebugIds.GetData(), AssetDebugIds.Num());
+		}
+
+		TArray<uint16> StatesIndices;
+		TArray<uint16> AssetDebugIds;
+	};
+
 	FInstanceEventBufferedData() = default;
 	explicit FInstanceEventBufferedData(
 		const double RecordingWorldTime,
-		const UStateTree& StateTree,
+		const UStateTree* StateTree,
 		const FStateTreeInstanceDebugId InstanceId,
 		const FString& InstanceName,
-		const EStateTreeTraceEventType EventType)	
+		const EStateTreeTraceEventType EventType)
 		: InstanceName(InstanceName)
-		, WeakStateTree(&StateTree)
+		, WeakStateTree(StateTree)
 		, InstanceId(InstanceId)
 		, LifetimeRecordingWorldTime(RecordingWorldTime)
 		, EventType(EventType)
@@ -159,31 +233,29 @@ struct FInstanceEventBufferedData
 
 	void Trace() const
 	{
-		if (!ensureMsgf(UE_TRACE_CHANNELEXPR_IS_ENABLED(StateTreeDebugChannel), TEXT("Tracing a buffered data is expected only if channel is enabled.")))
+		if (ensureMsgf(UE_TRACE_CHANNELEXPR_IS_ENABLED(StateTreeDebugChannel), TEXT("Tracing a buffered data is expected only if channel is enabled.")))
 		{
-			return;
-		}
-
-		if (const UStateTree* StateTree = WeakStateTree.Get())
-		{
-			// Force a world time update since we are tracing an event from the past
-			UE_TRACE_LOG(StateTreeDebugger, WorldTimestampEvent, StateTreeDebugChannel)
-				<< WorldTimestampEvent.WorldTime(LifetimeRecordingWorldTime);
-
-			OutputInstanceLifetimeEvent(InstanceId, StateTree, *InstanceName, EventType);
-
-			if (ActiveStates.Num() > 0)
+			if (const UStateTree* StateTree = WeakStateTree.Get())
 			{
 				// Force a world time update since we are tracing an event from the past
 				UE_TRACE_LOG(StateTreeDebugger, WorldTimestampEvent, StateTreeDebugChannel)
-					<< WorldTimestampEvent.WorldTime(ActiveStatesRecordingWorldTime);
+					<< WorldTimestampEvent.WorldTime(LifetimeRecordingWorldTime);
+
+				OutputInstanceLifetimeEvent(InstanceId, StateTree, *InstanceName, EventType);
+
+				if (ActiveStates.IsValid())
+				{
+					// Force a world time update since we are tracing an event from the past
+					UE_TRACE_LOG(StateTreeDebugger, WorldTimestampEvent, StateTreeDebugChannel)
+						<< WorldTimestampEvent.WorldTime(ActiveStatesRecordingWorldTime);
 
-				OutputActiveStatesEventTrace(InstanceId, ActiveStates);
+					ActiveStates.Output(InstanceId);
+				}
 			}
 		}
 	}
 
-	FStateTreeActiveStates ActiveStates;
+	FActiveStates ActiveStates;
 	FString InstanceName;
 	TWeakObjectPtr<const UStateTree> WeakStateTree;
 	FStateTreeInstanceDebugId InstanceId;
@@ -203,8 +275,11 @@ struct FBufferedDataList
 	 */
 	TArray<FPhaseStack> PhaseStacks;
 
+	/** List of asset debug ids events that will be output if channel gets enabled. */
+	TArray<FAssetDebugIdEventBufferedData> AssetDebugIdEvents;
+
 	/** List of lifetime events that will be output if channel gets enabled in the Push - Pop lifetime window of an instance. */
-	TArray<FInstanceEventBufferedData> Events;
+	TArray<FInstanceEventBufferedData> InstanceLifetimeEvents;
 
 	/** Flag use to prevent reentrant calls */
 	bool bFlushing = false;
@@ -218,16 +293,27 @@ struct FBufferedDataList
 
 		TGuardValue<bool> GuardReentry(bFlushing, true);
 
-		// Trace instance lifetime events first since they are required for other event types.
+		// Trace asset events first since they are required for instance lifetime event types.
+		for (const FAssetDebugIdEventBufferedData& AssetDebugIdEventData : AssetDebugIdEvents)
+		{
+			AssetDebugIdEventData.Trace();
+		}
+		AssetDebugIdEvents.Empty();
+
+		// Then trace instance lifetime events since they are required for other event types.
 		// It is also associated to an older world time.
-		for (const FInstanceEventBufferedData& InstanceEventData : Events)
+		for (const FInstanceEventBufferedData& InstanceEventData : InstanceLifetimeEvents)
 		{
 			InstanceEventData.Trace();
 		}
-		Events.Empty();
+		InstanceLifetimeEvents.Empty();
 
 		TraceWorldTime();
-		TraceStackedPhases(InstanceId);
+
+		if (InstanceId.IsValid())
+		{
+			TraceStackedPhases(InstanceId);
+		}
 	}
 
 	/**
@@ -285,11 +371,19 @@ struct FBufferedDataList
  */
 thread_local FBufferedDataList GBufferedEvents;
 
+struct FStateTreeTraceDebugIdPair
+{
+	TWeakObjectPtr<const UStateTree> WeakStateTree;
+	FStateTreeIndex16 Id;
+};
+thread_local uint16 GNextAssetDebugId(1);
+thread_local TArray<FStateTreeTraceDebugIdPair> GStateTreeTraceDebugIds;
+
 /**
  * Pushed or pops an entry on the Phase stack for a given Instance.
  * Will send the Pop events for phases popped if their associated Push events were sent.
  */
-void ProcessPhaseScopeEvent(const FStateTreeInstanceDebugId InstanceId, const EStateTreeUpdatePhase Phase, const EStateTreeTraceEventType EventType, const FStateTreeStateHandle StateHandle)
+void OutputPhaseScopeEvent(const FStateTreeInstanceDebugId InstanceId, const EStateTreeUpdatePhase Phase, const EStateTreeTraceEventType EventType, const FStateTreeStateHandle StateHandle)
 {
 	TArray<FPhaseStack>& PhaseStacks = GBufferedEvents.PhaseStacks;
 	int32 ExistingStackIndex = PhaseStacks.IndexOfByPredicate([InstanceId](const FPhaseStack& PhaseStack){ return PhaseStack.InstanceId == InstanceId; });
@@ -387,6 +481,8 @@ void RegisterGlobalDelegates()
 	GOnPIEStartDelegateHandle = FEditorDelegates::BeginPIE.AddLambda([&LastRecordingWorldTime=GBufferedEvents.TracedRecordingWorldTime](const bool bIsSimulating)
 		{
 			LastRecordingWorldTime = -1;
+			GStateTreeTraceDebugIds.Reset();
+			GNextAssetDebugId = 1;
 		});
 #endif // WITH_EDITOR
 	
@@ -408,6 +504,58 @@ void UnregisterGlobalDelegates()
 	GOnWorldTickStartDelegateHandle.Reset();
 }
 
+FStateTreeIndex16 FindOrAddDebugIdForAsset(const UStateTree* StateTree)
+{
+	FStateTreeIndex16 AssetDebugId;
+	const FStateTreeTraceDebugIdPair* ExistingPair = GStateTreeTraceDebugIds.FindByPredicate([StateTree](const FStateTreeTraceDebugIdPair& Pair)
+	{
+		return Pair.WeakStateTree == StateTree;
+	});
+
+	if (ExistingPair == nullptr)
+	{
+		if (ensure(StateTree != nullptr))
+		{
+			AssetDebugId = FStateTreeIndex16(GNextAssetDebugId++);
+			GStateTreeTraceDebugIds.Emplace(FStateTreeTraceDebugIdPair(StateTree, AssetDebugId));
+
+			OutputAssetDebugIdEvent(StateTree, AssetDebugId);
+		}
+	}
+	else
+	{
+		AssetDebugId = ExistingPair->Id;
+	}
+
+	return AssetDebugId;
+}
+
+void OutputAssetDebugIdEvent(
+	const UStateTree* StateTree,
+	const FStateTreeIndex16 AssetDebugId
+	)
+{
+	if (UE_TRACE_CHANNELEXPR_IS_ENABLED(StateTreeDebugChannel))
+	{
+		TraceBufferedEvents(FStateTreeInstanceDebugId::Invalid);
+
+		check(StateTree);
+		const FString TreeName = StateTree->GetName();
+		const FString TreePath = StateTree->GetPathName();
+
+		UE_TRACE_LOG(StateTreeDebugger, AssetDebugIdEvent, StateTreeDebugChannel)
+			<< AssetDebugIdEvent.Cycle(FPlatformTime::Cycles64())
+			<< AssetDebugIdEvent.TreeName(*TreeName, TreeName.Len())
+			<< AssetDebugIdEvent.TreePath(*TreePath, TreePath.Len())
+			<< AssetDebugIdEvent.CompiledDataHash(StateTree->LastCompiledEditorDataHash)
+			<< AssetDebugIdEvent.AssetDebugId(AssetDebugId.Get());
+	}
+	else
+	{
+		GBufferedEvents.AssetDebugIdEvents.Emplace(StateTree, AssetDebugId);
+	}
+}
+
 void OutputInstanceLifetimeEvent(
 	const FStateTreeInstanceDebugId InstanceId,
 	const UStateTree* StateTree,
@@ -415,38 +563,30 @@ void OutputInstanceLifetimeEvent(
 	const EStateTreeTraceEventType EventType
 	)
 {
-	if (StateTree == nullptr)
-	{
-		return;
-	}
-
 	if (UE_TRACE_CHANNELEXPR_IS_ENABLED(StateTreeDebugChannel))
 	{
 		TraceBufferedEvents(InstanceId);
 
-		const FString TreeName = StateTree->GetName();
-		const FString TreePath = StateTree->GetPathName();
+		const FStateTreeIndex16 AssetDebugId = FindOrAddDebugIdForAsset(StateTree);
 
 		UE_TRACE_LOG(StateTreeDebugger, InstanceEvent, StateTreeDebugChannel)
 			<< InstanceEvent.Cycle(FPlatformTime::Cycles64())
-			<< InstanceEvent.TreeName(*TreeName, TreeName.Len())
-			<< InstanceEvent.TreePath(*TreePath, TreePath.Len())
-			<< InstanceEvent.CompiledDataHash(StateTree->LastCompiledEditorDataHash)
 			<< InstanceEvent.InstanceId(InstanceId.Id)
 			<< InstanceEvent.InstanceSerial(InstanceId.SerialNumber)
 			<< InstanceEvent.InstanceName(InstanceName)
-			<< InstanceEvent.EventType(static_cast<std::underlying_type_t<EStateTreeTraceEventType>>(EventType));
+			<< InstanceEvent.EventType(static_cast<std::underlying_type_t<EStateTreeTraceEventType>>(EventType))
+			<< InstanceEvent.AssetDebugId(AssetDebugId.Get());
 	}
 	else
 	{
 		if (EventType == EStateTreeTraceEventType::Push)
 		{
-			GBufferedEvents.Events.Emplace(GRecordingWorldTime, *StateTree, InstanceId, InstanceName, EventType);
+			GBufferedEvents.InstanceLifetimeEvents.Emplace(GRecordingWorldTime, StateTree, InstanceId, InstanceName, EventType);
 		}
 		else if (EventType == EStateTreeTraceEventType::Pop)
 		{
 			// Remove matching instance events since if it was not sent then no other events were sent between, hence not needed in the trace.
-			GBufferedEvents.Events.SetNum(Algo::StableRemoveIf(GBufferedEvents.Events, 
+			GBufferedEvents.InstanceLifetimeEvents.SetNum(Algo::StableRemoveIf(GBufferedEvents.InstanceLifetimeEvents,
 				[InstanceId](const FInstanceEventBufferedData& BufferedData)
 				{
 					return BufferedData.InstanceId == InstanceId;
@@ -459,6 +599,29 @@ void OutputInstanceLifetimeEvent(
 	}
 }
 
+void OutputInstanceFrameEvent(
+	const FStateTreeInstanceDebugId InstanceId,
+	const FStateTreeExecutionFrame* Frame
+	)
+{
+	check(Frame != nullptr);
+
+	if (UE_TRACE_CHANNELEXPR_IS_ENABLED(StateTreeDebugChannel))
+	{
+		TraceBufferedEvents(InstanceId);
+
+		const FStateTreeIndex16 AssetDebugId = FindOrAddDebugIdForAsset(Frame->StateTree.Get());
+
+		UE_TRACE_LOG(StateTreeDebugger, InstanceFrameEvent, StateTreeDebugChannel)
+			<< InstanceFrameEvent.Cycle(FPlatformTime::Cycles64())
+			<< InstanceFrameEvent.InstanceId(InstanceId.Id)
+			<< InstanceFrameEvent.InstanceSerial(InstanceId.SerialNumber)
+			<< InstanceFrameEvent.AssetDebugId(AssetDebugId.Get());
+	}
+	// No need to buffer since frame event are sent each time a FrameScope is used by the execution context
+	// and we don't expect the trace channel to be enabled/disabled during a single execution context update.
+}
+
 void OutputLogEventTrace(
 	const FStateTreeInstanceDebugId InstanceId,
 	const TCHAR* Fmt, ...
@@ -479,8 +642,7 @@ void OutputLogEventTrace(
 void OutputStateEventTrace(
 	const FStateTreeInstanceDebugId InstanceId,
 	const FStateTreeStateHandle StateHandle,
-	const EStateTreeTraceEventType EventType,
-	const EStateTreeStateSelectionBehavior SelectionBehavior
+	const EStateTreeTraceEventType EventType
 	)
 {
 	TraceBufferedEvents(InstanceId);
@@ -490,8 +652,7 @@ void OutputStateEventTrace(
 		<< StateEvent.InstanceId(InstanceId.Id)
 		<< StateEvent.InstanceSerial(InstanceId.SerialNumber)
 		<< StateEvent.StateIndex(StateHandle.Index)
-		<< StateEvent.EventType(static_cast<std::underlying_type_t<EStateTreeTraceEventType>>(EventType))
-		<< StateEvent.SelectionBehavior(static_cast<std::underlying_type_t<EStateTreeStateSelectionBehavior>>(SelectionBehavior));
+		<< StateEvent.EventType(static_cast<std::underlying_type_t<EStateTreeTraceEventType>>(EventType));
 }
 
 void OutputTaskEventTrace(
@@ -521,7 +682,8 @@ void OutputEvaluatorEventTrace(
 	const FStateTreeInstanceDebugId InstanceId,
 	const FStateTreeIndex16 EvaluatorIdx,
 	const FStateTreeDataView DataView,
-	const EStateTreeTraceEventType EventType)
+	const EStateTreeTraceEventType EventType
+	)
 {
 	FBufferArchive Archive;
 	SerializeDataViewToArchive(Archive, DataView);
@@ -582,37 +744,28 @@ void OutputConditionEventTrace(
 
 void OutputActiveStatesEventTrace(
 	const FStateTreeInstanceDebugId InstanceId,
-	const FStateTreeActiveStates& ActiveStates
+	const TConstArrayView<FStateTreeExecutionFrame> ActiveFrames
 	)
 {
 	if (UE_TRACE_CHANNELEXPR_IS_ENABLED(StateTreeDebugChannel))
 	{
 		TraceBufferedEvents(InstanceId);
 
-		TArray<uint16, TInlineAllocator<FStateTreeActiveStates::MaxStates>> StatesIndices;
-		for (int32 i = 0; i < ActiveStates.Num(); i++)
-		{
-			StatesIndices.Add(ActiveStates[i].Index);
-		}
-
-		UE_TRACE_LOG(StateTreeDebugger, ActiveStatesEvent, StateTreeDebugChannel)
-			<< ActiveStatesEvent.Cycle(FPlatformTime::Cycles64())
-			<< ActiveStatesEvent.InstanceId(InstanceId.Id)
-			<< ActiveStatesEvent.InstanceSerial(InstanceId.SerialNumber)
-			<< ActiveStatesEvent.ActiveStates(StatesIndices.GetData(), StatesIndices.Num());
+		FInstanceEventBufferedData::FActiveStates ActiveStates(ActiveFrames);
+		ActiveStates.Output(InstanceId);
 	}
 	else
 	{
-		FInstanceEventBufferedData* ExisingBufferedData = GBufferedEvents.Events.FindByPredicate([InstanceId](const FInstanceEventBufferedData& BufferedData)
+		FInstanceEventBufferedData* ExisingBufferedData = GBufferedEvents.InstanceLifetimeEvents.FindByPredicate([InstanceId](const FInstanceEventBufferedData& BufferedData)
 			{
 				return BufferedData.InstanceId == InstanceId;
 			});
 
 		// We keep only the most recent active states since this is all we need to know in which state was the instance
-		// when we start receiving the events once the channel is enabled. 
+		// when we start receiving the events once the channel is enabled.
 		if (ExisingBufferedData != nullptr)
 		{
-			ExisingBufferedData->ActiveStates = ActiveStates;
+			ExisingBufferedData->ActiveStates= FInstanceEventBufferedData::FActiveStates(ActiveFrames);
 			ExisingBufferedData->ActiveStatesRecordingWorldTime = GRecordingWorldTime;
 		}
 	}
diff --git a/Engine/Plugins/Runtime/StateTree/Source/StateTreeModule/Private/Debugger/StateTreeTraceAnalyzer.cpp b/Engine/Plugins/Runtime/StateTree/Source/StateTreeModule/Private/Debugger/StateTreeTraceAnalyzer.cpp
--- a/Engine/Plugins/Runtime/StateTree/Source/StateTreeModule/Private/Debugger/StateTreeTraceAnalyzer.cpp
+++ b/Engine/Plugins/Runtime/StateTree/Source/StateTreeModule/Private/Debugger/StateTreeTraceAnalyzer.cpp
@@ -19,8 +19,10 @@ void FStateTreeTraceAnalyzer::OnAnalysisBegin(const FOnAnalysisContext& Context)
 {
 	auto& Builder = Context.InterfaceBuilder;
 
+	Builder.RouteEvent(RouteId_AssetDebugId, "StateTreeDebugger", "AssetDebugIdEvent");
 	Builder.RouteEvent(RouteId_WorldTimestamp, "StateTreeDebugger", "WorldTimestampEvent");
 	Builder.RouteEvent(RouteId_Instance, "StateTreeDebugger", "InstanceEvent");
+	Builder.RouteEvent(RouteId_InstanceFrame, "StateTreeDebugger", "InstanceFrameEvent");
 	Builder.RouteEvent(RouteId_Phase, "StateTreeDebugger", "PhaseEvent");
 	Builder.RouteEvent(RouteId_LogMessage, "StateTreeDebugger", "LogEvent");
 	Builder.RouteEvent(RouteId_State, "StateTreeDebugger", "StateEvent");
@@ -45,7 +47,7 @@ bool FStateTreeTraceAnalyzer::OnEvent(const uint16 RouteId, EStyle Style, const
 			WorldTime = EventData.GetValue<double>("WorldTime");
 			break;
 		}
-	case RouteId_Instance:
+	case RouteId_AssetDebugId:
 		{
 			FString ObjectName, ObjectPathName;
 			EventData.GetString("TreeName", ObjectName);
@@ -65,15 +67,7 @@ bool FStateTreeTraceAnalyzer::OnEvent(const uint16 RouteId, EStyle Style, const
 				const uint32 CompiledDataHash = EventData.GetValue<uint32>("CompiledDataHash");
 				if (StateTree->LastCompiledEditorDataHash == CompiledDataHash)
 				{
-					FString InstanceName;
-					EventData.GetString("InstanceName", InstanceName);
-
-					Provider.AppendInstanceEvent(StateTree,
-						FStateTreeInstanceDebugId(EventData.GetValue<uint32>("InstanceId"), EventData.GetValue<uint32>("InstanceSerial")),
-						*InstanceName,
-						Context.EventTime.AsSeconds(EventData.GetValue<uint64>("Cycle")),
-						WorldTime,
-						EventData.GetValue<EStateTreeTraceEventType>("EventType"));
+					Provider.AppendAssetDebugId(StateTree, FStateTreeIndex16(EventData.GetValue<uint16>("AssetDebugId")));
 				}
 				else
 				{
@@ -84,6 +78,37 @@ bool FStateTreeTraceAnalyzer::OnEvent(const uint16 RouteId, EStyle Style, const
 			{
 				UE_LOG(LogStateTree, Warning, TEXT("Unable to find StateTree asset: %s : %s"), *ObjectPathName, *ObjectName);
 			}
+			break;
+		}
+	case RouteId_Instance:
+		{
+			FString InstanceName;
+			EventData.GetString("InstanceName", InstanceName);
+
+			Provider.AppendInstanceEvent(FStateTreeIndex16(EventData.GetValue<uint16>("AssetDebugId")),
+				FStateTreeInstanceDebugId(EventData.GetValue<uint32>("InstanceId"), EventData.GetValue<uint32>("InstanceSerial")),
+				*InstanceName,
+				Context.EventTime.AsSeconds(EventData.GetValue<uint64>("Cycle")),
+				WorldTime,
+				EventData.GetValue<EStateTreeTraceEventType>("EventType"));
+			break;
+		}
+	case RouteId_InstanceFrame:
+		{
+			TWeakObjectPtr<const UStateTree> WeakStateTree;
+			if (Provider.GetAssetFromDebugId(FStateTreeIndex16(EventData.GetValue<uint16>("AssetDebugId")), WeakStateTree))
+			{
+				const FStateTreeTraceInstanceFrameEvent Event(WorldTime, EStateTreeTraceEventType::Push, WeakStateTree.Get());
+			
+				Provider.AppendEvent(FStateTreeInstanceDebugId(EventData.GetValue<uint32>("InstanceId"), EventData.GetValue<uint32>("InstanceSerial")),
+					Context.EventTime.AsSeconds(EventData.GetValue<uint64>("Cycle")),
+					FStateTreeTraceEventVariantType(TInPlaceType<FStateTreeTraceInstanceFrameEvent>(), Event));
+			}
+			else
+			{
+				UE_LOG(LogStateTree, Error, TEXT("Instance frame event refers to an asset Id that wasn't added previously."));
+			}
+
 			break;
 		}
 	case RouteId_Phase:
@@ -113,8 +138,7 @@ bool FStateTreeTraceAnalyzer::OnEvent(const uint16 RouteId, EStyle Style, const
 		{
 			const FStateTreeTraceStateEvent Event(WorldTime,
 				FStateTreeIndex16(EventData.GetValue<uint16>("StateIndex")),
-				EventData.GetValue<EStateTreeTraceEventType>("EventType"),
-				EventData.GetValue<EStateTreeStateSelectionBehavior>("SelectionBehavior"));
+				EventData.GetValue<EStateTreeTraceEventType>("EventType"));
 
 			Provider.AppendEvent(FStateTreeInstanceDebugId(EventData.GetValue<uint32>("InstanceId"), EventData.GetValue<uint32>("InstanceSerial")),
 				Context.EventTime.AsSeconds(EventData.GetValue<uint64>("Cycle")),
@@ -192,11 +216,43 @@ bool FStateTreeTraceAnalyzer::OnEvent(const uint16 RouteId, EStyle Style, const
 	case RouteId_ActiveStates:
 		{
 			FStateTreeTraceActiveStatesEvent Event(WorldTime);
-			Event.ActiveStates = EventData.GetArrayView<uint16>("ActiveStates");
 
-			Provider.AppendEvent(FStateTreeInstanceDebugId(EventData.GetValue<uint32>("InstanceId"), EventData.GetValue<uint32>("InstanceSerial")),
-				Context.EventTime.AsSeconds(EventData.GetValue<uint64>("Cycle")),
-				FStateTreeTraceEventVariantType(TInPlaceType<FStateTreeTraceActiveStatesEvent>(), Event));
+			TArray<FStateTreeStateHandle> ActiveStates(EventData.GetArrayView<uint16>("ActiveStates"));
+			TArray<uint16> AssetDebugIds(EventData.GetArrayView<uint16>("AssetDebugIds"));
+
+			if (ensureMsgf(ActiveStates.Num() == AssetDebugIds.Num(), TEXT("Each state is expected to have a matching asset id")))
+			{
+				FStateTreeIndex16 LastAssetDebugId;
+
+				FStateTreeTraceActiveStates::FAssetActiveStates* AssetActiveStates = nullptr;
+				for (int Index = 0; Index < ActiveStates.Num(); ++Index)
+				{
+					FStateTreeIndex16 AssetDebugId(AssetDebugIds[Index]);
+					if (AssetDebugId != LastAssetDebugId)
+					{
+						TWeakObjectPtr<const UStateTree> WeakStateTree;
+						if (Provider.GetAssetFromDebugId(AssetDebugId, WeakStateTree))
+						{
+							FStateTreeTraceActiveStates::FAssetActiveStates& NewPair = Event.ActiveStates.PerAssetStates.Emplace_GetRef();
+							NewPair.WeakStateTree = WeakStateTree;
+							AssetActiveStates = &NewPair;
+							LastAssetDebugId = AssetDebugId;	
+						}
+						else
+						{
+							UE_LOG(LogStateTree, Error, TEXT("Instance frame event refers to an asset Id that wasn't added previously."));
+							continue;
+						}
+					}
+
+					check(AssetActiveStates != nullptr);
+					AssetActiveStates->ActiveStates.Push(ActiveStates[Index]);
+				}
+
+				Provider.AppendEvent(FStateTreeInstanceDebugId(EventData.GetValue<uint32>("InstanceId"), EventData.GetValue<uint32>("InstanceSerial")),
+					Context.EventTime.AsSeconds(EventData.GetValue<uint64>("Cycle")),
+					FStateTreeTraceEventVariantType(TInPlaceType<FStateTreeTraceActiveStatesEvent>(), Event));
+			}
 			break;
 		}
 	default:
diff --git a/Engine/Plugins/Runtime/StateTree/Source/StateTreeModule/Private/Debugger/StateTreeTraceProvider.cpp b/Engine/Plugins/Runtime/StateTree/Source/StateTreeModule/Private/Debugger/StateTreeTraceProvider.cpp
--- a/Engine/Plugins/Runtime/StateTree/Source/StateTreeModule/Private/Debugger/StateTreeTraceProvider.cpp
+++ b/Engine/Plugins/Runtime/StateTree/Source/StateTreeModule/Private/Debugger/StateTreeTraceProvider.cpp
@@ -3,7 +3,6 @@
 #if WITH_STATETREE_DEBUGGER
 
 #include "Debugger/StateTreeTraceProvider.h"
-#include "StateTreeTypes.h"
 #include "Debugger/StateTreeDebugger.h"
 
 FName FStateTreeTraceProvider::ProviderName("StateTreeDebuggerProvider");
@@ -79,7 +78,7 @@ void FStateTreeTraceProvider::AppendEvent(const FStateTreeInstanceDebugId InInst
 }
 
 void FStateTreeTraceProvider::AppendInstanceEvent(
-	const UStateTree* InStateTree,
+	const FStateTreeIndex16 AssetDebugId,
 	const FStateTreeInstanceDebugId InInstanceId,
 	const TCHAR* InInstanceName,
 	const double InTime,
@@ -88,12 +87,27 @@ void FStateTreeTraceProvider::AppendInstanceEvent(
 {
 	if (InEventType == EStateTreeTraceEventType::Push)
 	{
-		Descriptors.Emplace(InStateTree, InInstanceId, InInstanceName, TRange<double>(InWorldRecordingTime, std::numeric_limits<double>::max()));
+		TWeakObjectPtr<const UStateTree> WeakStateTree;
+		if (GetAssetFromDebugId(AssetDebugId, WeakStateTree))
+		{
+			if (const UStateTree* StateTree = WeakStateTree.Get())
+			{
+				Descriptors.Emplace(StateTree, InInstanceId, InInstanceName, TRange<double>(InWorldRecordingTime, std::numeric_limits<double>::max()));
 
-		check(InstanceIdToDebuggerEntryTimelines.Find(InInstanceId) == nullptr);
-		InstanceIdToDebuggerEntryTimelines.Add(InInstanceId, EventsTimelines.Num());
+				check(InstanceIdToDebuggerEntryTimelines.Find(InInstanceId) == nullptr);
+				InstanceIdToDebuggerEntryTimelines.Add(InInstanceId, EventsTimelines.Num());
 		
-		EventsTimelines.Emplace(MakeShared<TraceServices::TPointTimeline<FStateTreeTraceEventVariantType>>(Session.GetLinearAllocator()));
+				EventsTimelines.Emplace(MakeShared<TraceServices::TPointTimeline<FStateTreeTraceEventVariantType>>(Session.GetLinearAllocator()));
+			}
+			else
+			{
+				UE_LOG(LogStateTree, Error, TEXT("Instance event refers to an unloaded asset."));	
+			}
+		}
+		else
+		{
+			UE_LOG(LogStateTree, Error, TEXT("Instance event refers to an asset Id that wasn't added previously."));
+		}
 	}
 	else if (InEventType == EStateTreeTraceEventType::Pop)
 	{
@@ -108,6 +122,29 @@ void FStateTreeTraceProvider::AppendInstanceEvent(
 	Session.UpdateDurationSeconds(InTime);
 }
 
+void FStateTreeTraceProvider::AppendAssetDebugId(const UStateTree* InStateTree, const FStateTreeIndex16 AssetDebugId)
+{
+	TWeakObjectPtr<const UStateTree> WeakStateTree;
+	if (ensureMsgf(AssetDebugId.IsValid(), TEXT("Expecting valid asset debug Id."))
+		&& !GetAssetFromDebugId(AssetDebugId, WeakStateTree))
+	{
+		StateTreeAssets.Emplace(FStateTreeDebugIdPair(InStateTree, AssetDebugId));
+	}
+}
+
+bool FStateTreeTraceProvider::GetAssetFromDebugId(const FStateTreeIndex16 AssetDebugId, TWeakObjectPtr<const UStateTree>& WeakStateTree) const
+{
+	verifyf(AssetDebugId.IsValid(), TEXT("Expecting valid asset debug Id."));
+	const FStateTreeDebugIdPair* ExistingPair = StateTreeAssets.FindByPredicate([AssetDebugId](const FStateTreeDebugIdPair& Pair)
+	{
+		return Pair.Id == AssetDebugId;
+	});
+
+	WeakStateTree = ExistingPair ? ExistingPair->WeakStateTree : nullptr; 
+
+	return ExistingPair != nullptr;
+}
+
 void FStateTreeTraceProvider::GetInstances(TArray<UE::StateTreeDebugger::FInstanceDescriptor>& OutInstances) const
 {
 	OutInstances = Descriptors;
diff --git a/Engine/Plugins/Runtime/StateTree/Source/StateTreeModule/Private/Debugger/StateTreeTraceTypes.cpp b/Engine/Plugins/Runtime/StateTree/Source/StateTreeModule/Private/Debugger/StateTreeTraceTypes.cpp
--- a/Engine/Plugins/Runtime/StateTree/Source/StateTreeModule/Private/Debugger/StateTreeTraceTypes.cpp
+++ b/Engine/Plugins/Runtime/StateTree/Source/StateTreeModule/Private/Debugger/StateTreeTraceTypes.cpp
@@ -7,6 +7,25 @@
 
 #if WITH_STATETREE_DEBUGGER
 
+namespace UE::StateTreeTrace
+{
+	FString GetStateName(const UStateTree& StateTree, const FCompactStateTreeState* CompactState)
+	{
+		check(CompactState);
+
+		if (const UStateTree* LinkedAsset = CompactState->LinkedAsset.Get())
+		{
+			return FString::Printf(TEXT("%s > [%s]"), *CompactState->Name.ToString(), *LinkedAsset->GetName());
+		}
+
+		if (const FCompactStateTreeState* LinkedState = StateTree.GetStateFromHandle(CompactState->LinkedState))
+		{
+			return FString::Printf(TEXT("%s > %s"), *CompactState->Name.ToString(), *LinkedState->Name.ToString());
+		}
+
+		return CompactState->Name.ToString();
+	}
+}
 
 //----------------------------------------------------------------------//
 // FStateTreeTracePhaseEvent
@@ -19,12 +38,38 @@ FString FStateTreeTracePhaseEvent::ToFullString(const UStateTree& StateTree) con
 FString FStateTreeTracePhaseEvent::GetValueString(const UStateTree& StateTree) const
 {
 	FStringBuilderBase StrBuilder;
-	StrBuilder.Append(UEnum::GetDisplayValueAsText(Phase).ToString());
+
+	// Override to display only selection behavior
+	if (Phase == EStateTreeUpdatePhase::TrySelectBehavior
+		&& StateHandle.IsValid())
+	{
+		if (const FCompactStateTreeState* CompactState = StateTree.GetStateFromHandle(StateHandle))
+		{
+			return UEnum::GetDisplayValueAsText(CompactState->SelectionBehavior).ToString();
+		}
+
+		return FString::Printf(TEXT("Invalid State Index %s for '%s'"), *StateHandle.Describe(), *StateTree.GetFullName());
+	}
+
+	// Otherwise build either: "<PhaseDescription> '<StateName>'" or "<StateName>"
+	if (Phase != EStateTreeUpdatePhase::Unset)
+	{
+		StrBuilder.Append(UEnum::GetDisplayValueAsText(Phase).ToString());
+	}
 
 	const FCompactStateTreeState* CompactState = StateTree.GetStateFromHandle(StateHandle);
 	if (CompactState != nullptr || StateHandle.IsValid())
 	{
-		StrBuilder.Appendf(TEXT(" '%s'"), CompactState != nullptr ? *CompactState->Name.ToString() : *StateHandle.Describe());
+		if (StrBuilder.Len())
+		{
+			StrBuilder.Appendf(TEXT(" '%s'"),
+				CompactState != nullptr ? *UE::StateTreeTrace::GetStateName(StateTree, CompactState) : *StateHandle.Describe());
+		}
+		else
+		{
+			StrBuilder.Appendf(TEXT("%s"),
+				CompactState != nullptr ? *UE::StateTreeTrace::GetStateName(StateTree, CompactState) : *StateHandle.Describe());
+		}
 	}
 
 	return StrBuilder.ToString();
@@ -85,7 +130,7 @@ FString FStateTreeTraceTransitionEvent::GetValueString(const UStateTree& StateTr
 {
 	const FCompactStateTreeState* CompactState = StateTree.GetStateFromHandle(TransitionSource.TargetState);
 	FStringBuilderBase StrBuilder;
-	StrBuilder.Appendf(TEXT("go to State '%s'"), CompactState != nullptr ? *CompactState->Name.ToString() : *TransitionSource.TargetState.Describe());
+	StrBuilder.Appendf(TEXT("go to State '%s'"), CompactState != nullptr ? *UE::StateTreeTrace::GetStateName(StateTree, CompactState) : *TransitionSource.TargetState.Describe());
 
 	if (TransitionSource.Priority != EStateTreeTransitionPriority::None)
 	{
@@ -172,22 +217,18 @@ FString FStateTreeTraceStateEvent::ToFullString(const UStateTree& StateTree) con
 	const FStateTreeStateHandle StateHandle(Index.Get());
 	if (const FCompactStateTreeState* CompactState = StateTree.GetStateFromHandle(StateHandle))
 	{
-		if (SelectionBehavior != EStateTreeStateSelectionBehavior::None)
+		if (CompactState->SelectionBehavior != EStateTreeStateSelectionBehavior::None)
 		{
 			return FString::Printf(TEXT("%s '%s' (%s)"),
-					*UEnum::GetDisplayValueAsText(EventType).ToString(),
-					*CompactState->Name.ToString(),
-					*UEnum::GetDisplayValueAsText(SelectionBehavior).ToString());
-		}
-		else
-		{
-			return FString::Printf(TEXT("%s '%s'"),
 				*UEnum::GetDisplayValueAsText(EventType).ToString(),
-				*CompactState->Name.ToString());
+				*GetValueString(StateTree),
+				*UEnum::GetDisplayValueAsText(CompactState->SelectionBehavior).ToString());
 		}
 	}
 
-	return FString::Printf(TEXT("Invalid State Index %s for '%s'"), *StateHandle.Describe(), *StateTree.GetFullName());
+	return FString::Printf(TEXT("%s '%s'"),
+		*UEnum::GetDisplayValueAsText(EventType).ToString(),
+		*GetValueString(StateTree));
 }
 
 FString FStateTreeTraceStateEvent::GetValueString(const UStateTree& StateTree) const
@@ -195,7 +236,7 @@ FString FStateTreeTraceStateEvent::GetValueString(const UStateTree& StateTree) c
 	const FStateTreeStateHandle StateHandle(Index.Get());
 	if (const FCompactStateTreeState* CompactState = StateTree.GetStateFromHandle(StateHandle))
 	{
-		return FString::Printf(TEXT("%s"), *CompactState->Name.ToString());
+		return UE::StateTreeTrace::GetStateName(StateTree, CompactState);
 	}
 
 	return FString::Printf(TEXT("Invalid State Index %s for '%s'"), *StateHandle.Describe(), *StateTree.GetFullName());
@@ -203,18 +244,7 @@ FString FStateTreeTraceStateEvent::GetValueString(const UStateTree& StateTree) c
 
 FString FStateTreeTraceStateEvent::GetTypeString(const UStateTree& StateTree) const
 {
-	const FStateTreeStateHandle StateHandle(Index.Get());
-	if (const FCompactStateTreeState* CompactState = StateTree.GetStateFromHandle(StateHandle))
-	{
-		if (SelectionBehavior != EStateTreeStateSelectionBehavior::None)
-		{
-			return FString::Printf(TEXT("%s '%s'"), *UEnum::GetDisplayValueAsText(SelectionBehavior).ToString(), *CompactState->Name.ToString());
-		}
-
-		return FString::Printf(TEXT(""));
-	}
-
-	return FString::Printf(TEXT("Invalid State Index %s for '%s'"), *StateHandle.Describe(), *StateTree.GetFullName());
+	return TEXT("State");
 }
 
 
@@ -223,7 +253,7 @@ FString FStateTreeTraceStateEvent::GetTypeString(const UStateTree& StateTree) co
 //----------------------------------------------------------------------//
 FString FStateTreeTraceTaskEvent::ToFullString(const UStateTree& StateTree) const
 {
-	return FString::Printf(TEXT("%s -> %s"),	*FStateTreeTraceNodeEvent::ToFullString(StateTree), *UEnum::GetDisplayValueAsText(Status).ToString());
+	return FString::Printf(TEXT("%s -> %s"), *FStateTreeTraceNodeEvent::ToFullString(StateTree), *UEnum::GetDisplayValueAsText(Status).ToString());
 }
 
 FString FStateTreeTraceTaskEvent::GetValueString(const UStateTree& StateTree) const
@@ -281,6 +311,7 @@ FString FStateTreeTraceConditionEvent::GetTypeString(const UStateTree& StateTree
 	return FStateTreeTraceNodeEvent::GetTypeString(StateTree);
 }
 
+
 //----------------------------------------------------------------------//
 // FStateTreeTraceActiveStatesEvent
 //----------------------------------------------------------------------//
@@ -291,21 +322,28 @@ FStateTreeTraceActiveStatesEvent::FStateTreeTraceActiveStatesEvent(const double
 
 FString FStateTreeTraceActiveStatesEvent::ToFullString(const UStateTree& StateTree) const
 {
-	if (ActiveStates.Num() > 0)
+	if (ActiveStates.PerAssetStates.Num() > 0)
 	{
 		return FString::Printf(TEXT("%s: %s"), *GetTypeString(StateTree), *GetValueString(StateTree));
 	}
 	return TEXT("No active states");
 }
 
-FString FStateTreeTraceActiveStatesEvent::GetValueString(const UStateTree& StateTree) const
+FString FStateTreeTraceActiveStatesEvent::GetValueString(const UStateTree&) const
 {
 	FStringBuilderBase StatePath;
-	for (int32 i = 0; i < ActiveStates.Num(); i++)
+	for (const FStateTreeTraceActiveStates::FAssetActiveStates& PerAssetStates : ActiveStates.PerAssetStates)
 	{
-		TConstArrayView<FCompactStateTreeState> StatesView = StateTree.GetStates();
-		const FCompactStateTreeState& State = StatesView[ActiveStates[i].Index];
-		StatePath.Appendf(TEXT("%s%s"), i == 0 ? TEXT("") : TEXT("."), *State.Name.ToString());
+		if (const UStateTree* StateTree = PerAssetStates.WeakStateTree.Get())
+		{
+			for (const FStateTreeStateHandle Handle : PerAssetStates.ActiveStates)
+			{
+				const FCompactStateTreeState* State = StateTree->GetStateFromHandle(Handle);
+				StatePath.Appendf(TEXT("%s%s"),
+					StatePath.Len() == 0 ? TEXT("") : TEXT("."),
+					State == nullptr ? TEXT("Invalid") : *UE::StateTreeTrace::GetStateName(*StateTree, State));
+			}
+		}
 	}
 	return StatePath.ToString();
 }
@@ -315,4 +353,23 @@ FString FStateTreeTraceActiveStatesEvent::GetTypeString(const UStateTree& StateT
 	return TEXT("New active states");
 }
 
+
+//----------------------------------------------------------------------//
+// FStateTreeTraceInstanceFrameEvent
+//----------------------------------------------------------------------//
+FString FStateTreeTraceInstanceFrameEvent::ToFullString(const UStateTree& StateTree) const
+{
+	return FString::Printf(TEXT("%s: %s"), *GetTypeString(StateTree), *GetValueString(StateTree));
+}
+
+FString FStateTreeTraceInstanceFrameEvent::GetValueString(const UStateTree& StateTree) const
+{
+	return GetNameSafe(WeakStateTree.Get());
+}
+
+FString FStateTreeTraceInstanceFrameEvent::GetTypeString(const UStateTree& StateTree) const
+{
+	return TEXT("New instance frame");
+}
+
 #endif // WITH_STATETREE_DEBUGGER
diff --git a/Engine/Plugins/Runtime/StateTree/Source/StateTreeModule/Private/StateTreeExecutionContext.cpp b/Engine/Plugins/Runtime/StateTree/Source/StateTreeModule/Private/StateTreeExecutionContext.cpp
--- a/Engine/Plugins/Runtime/StateTree/Source/StateTreeModule/Private/StateTreeExecutionContext.cpp
+++ b/Engine/Plugins/Runtime/StateTree/Source/StateTreeModule/Private/StateTreeExecutionContext.cpp
@@ -38,9 +38,12 @@
 		ON_SCOPE_EXIT { TRACE_STATETREE_PHASE_EVENT(ID_NAME, Phase, EStateTreeTraceEventType::Pop, StateHandle) }
 
 	#define STATETREE_TRACE_INSTANCE_EVENT(EventType)						TRACE_STATETREE_INSTANCE_EVENT(GetInstanceDebugId(), GetStateTree(), *GetInstanceDescription(), EventType);
-	#define STATETREE_TRACE_ACTIVE_STATES_EVENT(ActiveStates)				TRACE_STATETREE_ACTIVE_STATES_EVENT(GetInstanceDebugId(), ActiveStates);
+	#define STATETREE_TRACE_INSTANCE_FRAME_EVENT(InstanceDebugId, Frame)	TRACE_STATETREE_INSTANCE_FRAME_EVENT(InstanceDebugId, Frame);
+	#define STATETREE_TRACE_PHASE_BEGIN(Phase)								TRACE_STATETREE_PHASE_EVENT(GetInstanceDebugId(), Phase, EStateTreeTraceEventType::Push, FStateTreeStateHandle::Invalid)
+	#define STATETREE_TRACE_PHASE_END(Phase)								TRACE_STATETREE_PHASE_EVENT(GetInstanceDebugId(), Phase, EStateTreeTraceEventType::Pop, FStateTreeStateHandle::Invalid)
+	#define STATETREE_TRACE_ACTIVE_STATES_EVENT(ActiveFrames)				TRACE_STATETREE_ACTIVE_STATES_EVENT(GetInstanceDebugId(), ActiveFrames);
 	#define STATETREE_TRACE_LOG_EVENT(Format, ...)							TRACE_STATETREE_LOG_EVENT(GetInstanceDebugId(), Format, ##__VA_ARGS__)
-	#define STATETREE_TRACE_STATE_EVENT(StateHandle, EventType)				TRACE_STATETREE_STATE_EVENT(GetInstanceDebugId(), StateHandle, EventType, EStateTreeStateSelectionBehavior::None);
+	#define STATETREE_TRACE_STATE_EVENT(StateHandle, EventType)				TRACE_STATETREE_STATE_EVENT(GetInstanceDebugId(), StateHandle, EventType);
 	#define STATETREE_TRACE_TASK_EVENT(Index, DataView, EventType, Status)	TRACE_STATETREE_TASK_EVENT(GetInstanceDebugId(), FStateTreeIndex16(Index), DataView, EventType, Status);
 	#define STATETREE_TRACE_EVALUATOR_EVENT(Index, DataView, EventType)		TRACE_STATETREE_EVALUATOR_EVENT(GetInstanceDebugId(), FStateTreeIndex16(Index), DataView, EventType);
 	#define STATETREE_TRACE_CONDITION_EVENT(Index, DataView, EventType)		TRACE_STATETREE_CONDITION_EVENT(GetInstanceDebugId(), FStateTreeIndex16(Index), DataView, EventType);	
@@ -50,7 +53,10 @@
 	#define STATETREE_TRACE_SCOPED_STATE(StateHandle)
 	#define STATETREE_TRACE_SCOPED_STATE_PHASE(StateHandle, Phase)
 	#define STATETREE_TRACE_INSTANCE_EVENT(EventType)
-	#define STATETREE_TRACE_ACTIVE_STATES_EVENT(ActiveStates)
+	#define STATETREE_TRACE_INSTANCE_FRAME_EVENT(InstanceDebugId, Frame)
+	#define STATETREE_TRACE_PHASE_BEGIN(Phase)
+	#define STATETREE_TRACE_PHASE_END(Phase)
+	#define STATETREE_TRACE_ACTIVE_STATES_EVENT(ActiveFrames)
 	#define STATETREE_TRACE_LOG_EVENT(Format, ...)
 	#define STATETREE_TRACE_STATE_EVENT(StateHandle, EventType)
 	#define STATETREE_TRACE_TASK_EVENT(Index, DataView, EventType, Status)
@@ -64,6 +70,35 @@ namespace UE::StateTree
 	constexpr int32 DebugIndentSize = 2;	// Debug printing indent for hierarchical data.
 }; // UE::StateTree
 
+
+FStateTreeExecutionContext::FCurrentlyProcessedFrameScope::FCurrentlyProcessedFrameScope(FStateTreeExecutionContext& InContext, const FStateTreeExecutionFrame* CurrentParentFrame, const FStateTreeExecutionFrame& CurrentFrame): Context(InContext)
+{
+	check(CurrentFrame.StateTree);
+	FStateTreeInstanceStorage* SharedInstanceDataStorage = &CurrentFrame.StateTree->GetSharedInstanceData()->GetMutableStorage();
+
+	SavedFrame = Context.CurrentlyProcessedFrame;
+	SavedParentFrame = Context.CurrentlyProcessedParentFrame;
+	SavedSharedInstanceDataStorage = Context.CurrentlyProcessedSharedInstanceStorage;
+	Context.CurrentlyProcessedFrame = &CurrentFrame;
+	Context.CurrentlyProcessedParentFrame = CurrentParentFrame;
+	Context.CurrentlyProcessedSharedInstanceStorage = SharedInstanceDataStorage;
+	
+	STATETREE_TRACE_INSTANCE_FRAME_EVENT(Context.GetInstanceDebugId(), Context.CurrentlyProcessedFrame);
+}
+
+FStateTreeExecutionContext::FCurrentlyProcessedFrameScope::~FCurrentlyProcessedFrameScope()
+{
+	Context.CurrentlyProcessedFrame = SavedFrame;
+	Context.CurrentlyProcessedParentFrame = SavedParentFrame;
+	Context.CurrentlyProcessedSharedInstanceStorage = SavedSharedInstanceDataStorage;
+
+	if (Context.CurrentlyProcessedFrame)
+	{
+		STATETREE_TRACE_INSTANCE_FRAME_EVENT(Context.GetInstanceDebugId(), Context.CurrentlyProcessedFrame);
+	}
+}
+
+
 FStateTreeExecutionContext::FStateTreeExecutionContext(UObject& InOwner, const UStateTree& InStateTree, FStateTreeInstanceData& InInstanceData, const FOnCollectStateTreeExternalData& InCollectExternalDataDelegate)
 	: Owner(InOwner)
 	, RootStateTree(InStateTree)
@@ -388,7 +423,7 @@ EStateTreeRunStatus FStateTreeExecutionContext::Stop(EStateTreeRunStatus Complet
 	}
 
 	// Trace before resetting the instance data since it is required to provide all the event information
-	STATETREE_TRACE_ACTIVE_STATES_EVENT(FStateTreeActiveStates());
+	STATETREE_TRACE_ACTIVE_STATES_EVENT({});
 	STATETREE_TRACE_INSTANCE_EVENT(EStateTreeTraceEventType::Pop);
 
 	// Destruct all allocated instance data (does not shrink the buffer). This will invalidate Exec too.
@@ -1256,7 +1291,7 @@ EStateTreeRunStatus FStateTreeExecutionContext::EnterState(FStateTreeTransitionR
 	EStateTreeRunStatus Result = EStateTreeRunStatus::Running;
 
 	STATETREE_LOG(Log, TEXT("Enter state '%s' (%d)"), *DebugGetStatePath(Transition.NextActiveFrames), Exec.StateChangeCount);
-	STATETREE_TRACE_SCOPED_PHASE(EStateTreeUpdatePhase::EnterStates);
+	STATETREE_TRACE_PHASE_BEGIN(EStateTreeUpdatePhase::EnterStates);
 
 	// The previous active frames are needed for state enter logic.
 	TArray<FStateTreeExecutionFrame, TConcurrentLinearArrayAllocator<FDefaultBlockAllocationTag>> PreviousActiveFrames;
@@ -1396,18 +1431,9 @@ EStateTreeRunStatus FStateTreeExecutionContext::EnterState(FStateTreeTransitionR
 		}
 	}
 
-#if WITH_STATETREE_DEBUGGER
-	// @todo: implement support for frames in debugger
-	FStateTreeActiveStates AllActiveStates;
-	for (const FStateTreeExecutionFrame& ActiveFrame : Exec.ActiveFrames)
-	{
-		for (const FStateTreeStateHandle StateHandle : ActiveFrame.ActiveStates)
-		{
-			AllActiveStates.Push(StateHandle);	
-		}
-	}
-	STATETREE_TRACE_ACTIVE_STATES_EVENT(AllActiveStates);
-#endif // WITH_STATETREE_DEBUGGER
+	STATETREE_TRACE_PHASE_END(EStateTreeUpdatePhase::EnterStates);
+
+	STATETREE_TRACE_ACTIVE_STATES_EVENT(Exec.ActiveFrames);
 
 	return Result;
 }
@@ -1597,6 +1623,7 @@ void FStateTreeExecutionContext::StateCompleted()
 	}
 
 	STATETREE_LOG(Verbose, TEXT("State Completed %s (%d)"), *UEnum::GetDisplayValueAsText(Exec.LastTickStatus).ToString(), Exec.StateChangeCount);
+	STATETREE_TRACE_SCOPED_PHASE(EStateTreeUpdatePhase::StateCompleted);
 
 	// Call from child towards root to allow to pass results back.
 	// Note: Completed is assumed to be called immediately after tick or enter state, so there's no property copying.
@@ -2948,8 +2975,9 @@ bool FStateTreeExecutionContext::SelectStateInternal(
 	}
 
 	// Check that the state can be entered
-	STATETREE_TRACE_SCOPED_PHASE(EStateTreeUpdatePhase::EnterConditions);
+	STATETREE_TRACE_PHASE_BEGIN(EStateTreeUpdatePhase::EnterConditions);
 	const bool bEnterConditionsPassed = TestAllConditions(CurrentParentFrame, CurrentFrame, NextState.EnterConditionsBegin, NextState.EnterConditionsNum);
+	STATETREE_TRACE_PHASE_END(EStateTreeUpdatePhase::EnterConditions);
 
 	if (bEnterConditionsPassed)
 	{
@@ -3017,7 +3045,7 @@ bool FStateTreeExecutionContext::SelectStateInternal(
 				NewFrame.StateTree = CurrentFrame.StateTree;
 				NewFrame.RootState = NextState.LinkedState;
 				NewFrame.ExternalDataBaseIndex = CurrentFrame.ExternalDataBaseIndex;
-				
+
 				// Check and prevent recursion.
 				const bool bNewFrameAlreadySelected = OutNextActiveFrames.ContainsByPredicate([&NewFrame](const FStateTreeExecutionFrame& Frame) {
 					return Frame.IsSameFrame(NewFrame);
@@ -3168,6 +3196,8 @@ bool FStateTreeExecutionContext::SelectStateInternal(
 		}
 		else if (NextState.SelectionBehavior == EStateTreeStateSelectionBehavior::TryFollowTransitions)
 		{
+			STATETREE_TRACE_SCOPED_STATE_PHASE(NextStateHandle, EStateTreeUpdatePhase::TrySelectBehavior);
+
 			EStateTreeTransitionPriority CurrentPriority = EStateTreeTransitionPriority::None;
 
 			for (uint8 i = 0; i < NextState.TransitionsNum; i++)
@@ -3243,6 +3273,8 @@ bool FStateTreeExecutionContext::SelectStateInternal(
 		{
 			if (NextState.HasChildren())
 			{
+				STATETREE_TRACE_SCOPED_STATE_PHASE(NextStateHandle, EStateTreeUpdatePhase::TrySelectBehavior);
+
 				// If the state has children, proceed to select children.
 				for (uint16 ChildState = NextState.ChildrenBegin; ChildState < NextState.ChildrenEnd; ChildState = CurrentStateTree->States[ChildState].GetNextSibling())
 				{
diff --git a/Engine/Plugins/Runtime/StateTree/Source/StateTreeModule/Public/Debugger/StateTreeDebugger.h b/Engine/Plugins/Runtime/StateTree/Source/StateTreeModule/Public/Debugger/StateTreeDebugger.h
--- a/Engine/Plugins/Runtime/StateTree/Source/StateTreeModule/Public/Debugger/StateTreeDebugger.h
+++ b/Engine/Plugins/Runtime/StateTree/Source/StateTreeModule/Public/Debugger/StateTreeDebugger.h
@@ -27,7 +27,7 @@ class UStateTree;
 
 DECLARE_DELEGATE_OneParam(FOnStateTreeDebuggerScrubStateChanged, const UE::StateTreeDebugger::FScrubState& ScrubState);
 DECLARE_DELEGATE_TwoParams(FOnStateTreeDebuggerBreakpointHit, FStateTreeInstanceDebugId InstanceId, const FStateTreeDebuggerBreakpoint Breakpoint);
-DECLARE_DELEGATE_OneParam(FOnStateTreeDebuggerActiveStatesChanges, TConstArrayView<FStateTreeStateHandle> ActiveStates);
+DECLARE_DELEGATE_OneParam(FOnStateTreeDebuggerActiveStatesChanges, const FStateTreeTraceActiveStates& ActiveStates);
 DECLARE_DELEGATE_OneParam(FOnStateTreeDebuggerNewInstance, FStateTreeInstanceDebugId InstanceId);
 DECLARE_DELEGATE(FOnStateTreeDebuggerNewSession);
 DECLARE_DELEGATE(FOnStateTreeDebuggerDebuggedInstanceSet);
@@ -235,7 +235,7 @@ private:
 
 	void SendNotifications();
 
-	void SetActiveStates(const TConstArrayView<FStateTreeStateHandle> NewActiveStates);
+	void SetActiveStates(const FStateTreeTraceActiveStates& NewActiveStates);
 
 	/**
 	 * Request an analysis session on the latest next available live trace.
@@ -294,7 +294,7 @@ private:
 	TArray<FStateTreeDebuggerBreakpoint> Breakpoints;
 
 	/** List of currently active states in the selected instance */
-	TArray<FStateTreeStateHandle> ActiveStates;
+	FStateTreeTraceActiveStates ActiveStates;
 
 	/**
 	 * When auto-connecting on next live session it is possible that a few frames are required for the tracing session to be accessible and connected to.
diff --git a/Engine/Plugins/Runtime/StateTree/Source/StateTreeModule/Public/Debugger/StateTreeTrace.h b/Engine/Plugins/Runtime/StateTree/Source/StateTreeModule/Public/Debugger/StateTreeTrace.h
--- a/Engine/Plugins/Runtime/StateTree/Source/StateTreeModule/Public/Debugger/StateTreeTrace.h
+++ b/Engine/Plugins/Runtime/StateTree/Source/StateTreeModule/Public/Debugger/StateTreeTrace.h
@@ -5,7 +5,9 @@
 #if WITH_STATETREE_DEBUGGER
 
 #include "Trace/Trace.h"
+#include "Containers/ContainersFwd.h"
 
+struct FStateTreeExecutionFrame;
 class UStateTree;
 struct FStateTreeDataView;
 struct FStateTreeActiveStates;
@@ -24,22 +26,28 @@ namespace UE::StateTreeTrace
 {
 	void RegisterGlobalDelegates();
 	void UnregisterGlobalDelegates();
-	void ProcessPhaseScopeEvent(FStateTreeInstanceDebugId InstanceId, EStateTreeUpdatePhase Phase, EStateTreeTraceEventType EventType, FStateTreeStateHandle StateHandle);
+	FStateTreeIndex16 FindOrAddDebugIdForAsset(const UStateTree* StateTree);
+	void OutputPhaseScopeEvent(FStateTreeInstanceDebugId InstanceId, EStateTreeUpdatePhase Phase, EStateTreeTraceEventType EventType, FStateTreeStateHandle StateHandle);
+	void OutputAssetDebugIdEvent(const UStateTree* StateTree, FStateTreeIndex16 AssetDebugId);
 	void OutputInstanceLifetimeEvent(FStateTreeInstanceDebugId InstanceId, const UStateTree* StateTree, const TCHAR* InstanceName, EStateTreeTraceEventType EventType);
+	void OutputInstanceFrameEvent(FStateTreeInstanceDebugId InstanceId, const FStateTreeExecutionFrame* Frame);
 	void OutputLogEventTrace(FStateTreeInstanceDebugId InstanceId, const TCHAR* Fmt, ...);
-	void OutputStateEventTrace(FStateTreeInstanceDebugId InstanceId, FStateTreeStateHandle StateHandle, EStateTreeTraceEventType EventType, EStateTreeStateSelectionBehavior SelectionBehavior);
+	void OutputStateEventTrace(FStateTreeInstanceDebugId InstanceId, FStateTreeStateHandle StateHandle, EStateTreeTraceEventType EventType);
 	void OutputTaskEventTrace(FStateTreeInstanceDebugId InstanceId, FStateTreeIndex16 TaskIdx, FStateTreeDataView DataView, EStateTreeTraceEventType EventType, EStateTreeRunStatus Status);
 	void OutputEvaluatorEventTrace(FStateTreeInstanceDebugId InstanceId, FStateTreeIndex16 EvaluatorIdx, FStateTreeDataView DataView, EStateTreeTraceEventType EventType);
 	void OutputConditionEventTrace(FStateTreeInstanceDebugId InstanceId, FStateTreeIndex16 ConditionIdx, FStateTreeDataView DataView, EStateTreeTraceEventType EventType);
 	void OutputTransitionEventTrace(FStateTreeInstanceDebugId InstanceId, FStateTreeTransitionSource TransitionSource, EStateTreeTraceEventType EventType);
-	void OutputActiveStatesEventTrace(FStateTreeInstanceDebugId InstanceId, const FStateTreeActiveStates& ActiveStates);
+	void OutputActiveStatesEventTrace(FStateTreeInstanceDebugId InstanceId, TConstArrayView<FStateTreeExecutionFrame> ActiveFrames);
 }
 
 #define TRACE_STATETREE_INSTANCE_EVENT(InstanceID, StateTree, InstanceName, EventType) \
 	UE::StateTreeTrace::OutputInstanceLifetimeEvent(InstanceID, StateTree, InstanceName, EventType);
 
+#define TRACE_STATETREE_INSTANCE_FRAME_EVENT(InstanceID, Frame) \
+	UE::StateTreeTrace::OutputInstanceFrameEvent(InstanceID, Frame);
+
 #define TRACE_STATETREE_PHASE_EVENT(InstanceID, Phase, EventType, StateHandle) \
-	UE::StateTreeTrace::ProcessPhaseScopeEvent(InstanceID, Phase, EventType, StateHandle); \
+	UE::StateTreeTrace::OutputPhaseScopeEvent(InstanceID, Phase, EventType, StateHandle); \
 
 #define TRACE_STATETREE_LOG_EVENT(InstanceId, Format, ...) \
 	if (UE_TRACE_CHANNELEXPR_IS_ENABLED(StateTreeDebugChannel)) \
@@ -47,10 +55,10 @@ namespace UE::StateTreeTrace
 		UE::StateTreeTrace::OutputLogEventTrace(InstanceId, Format, ##__VA_ARGS__); \
 	}
 
-#define TRACE_STATETREE_STATE_EVENT(InstanceId, StateHandle, EventType, SelectionBehavior) \
+#define TRACE_STATETREE_STATE_EVENT(InstanceId, StateHandle, EventType) \
 	if (UE_TRACE_CHANNELEXPR_IS_ENABLED(StateTreeDebugChannel)) \
 	{ \
-		UE::StateTreeTrace::OutputStateEventTrace(InstanceId, StateHandle, EventType, SelectionBehavior); \
+		UE::StateTreeTrace::OutputStateEventTrace(InstanceId, StateHandle, EventType); \
 	}
 
 #define TRACE_STATETREE_TASK_EVENT(InstanceId, TaskIdx, DataView, EventType, Status) \
@@ -77,19 +85,20 @@ namespace UE::StateTreeTrace
 		UE::StateTreeTrace::OutputTransitionEventTrace(InstanceId, TransitionIdx, EventType); \
 	}
 
-#define TRACE_STATETREE_ACTIVE_STATES_EVENT(InstanceId, ActivateStates) \
-		UE::StateTreeTrace::OutputActiveStatesEventTrace(InstanceId, ActivateStates);
+#define TRACE_STATETREE_ACTIVE_STATES_EVENT(InstanceId, ActivateFrames) \
+		UE::StateTreeTrace::OutputActiveStatesEventTrace(InstanceId, ActivateFrames);
 
 #else //STATETREE_DEBUG_TRACE_ENABLED
 
 #define TRACE_STATETREE_INSTANCE_EVENT(InstanceID, StateTree, InstanceName, EventType)
+#define TRACE_STATETREE_INSTANCE_FRAME_EVENT(InstanceID, Frame)
 #define TRACE_STATETREE_PHASE_EVENT(InstanceID, Phase, EventType, StateHandle)
 #define TRACE_STATETREE_LOG_EVENT(InstanceId, Format, ...)
-#define TRACE_STATETREE_STATE_EVENT(InstanceId, StateHandle, EventType, SelectionBehavior)
+#define TRACE_STATETREE_STATE_EVENT(InstanceId, StateHandle, EventType)
 #define TRACE_STATETREE_TASK_EVENT(InstanceId, TaskIdx, DataView, EventType, Status)
 #define TRACE_STATETREE_EVALUATOR_EVENT(InstanceId, EvaluatorIdx, DataView, EventType)
 #define TRACE_STATETREE_CONDITION_EVENT(InstanceId, ConditionIdx, DataView, EventType)
 #define TRACE_STATETREE_TRANSITION_EVENT(InstanceId, TransitionIdx, EventType)
-#define TRACE_STATETREE_ACTIVE_STATES_EVENT(InstanceId, ActivateStates)
+#define TRACE_STATETREE_ACTIVE_STATES_EVENT(InstanceId, ActiveFrames)
 
 #endif // STATETREE_DEBUG_TRACE_ENABLED
diff --git a/Engine/Plugins/Runtime/StateTree/Source/StateTreeModule/Public/Debugger/StateTreeTraceAnalyzer.h b/Engine/Plugins/Runtime/StateTree/Source/StateTreeModule/Public/Debugger/StateTreeTraceAnalyzer.h
--- a/Engine/Plugins/Runtime/StateTree/Source/StateTreeModule/Public/Debugger/StateTreeTraceAnalyzer.h
+++ b/Engine/Plugins/Runtime/StateTree/Source/StateTreeModule/Public/Debugger/StateTreeTraceAnalyzer.h
@@ -22,7 +22,9 @@ private:
 	enum : uint16
 	{
 		RouteId_WorldTimestamp,
+		RouteId_AssetDebugId,
 		RouteId_Instance,
+		RouteId_InstanceFrame,
 		RouteId_Phase,
 		RouteId_LogMessage,
 		RouteId_State,
diff --git a/Engine/Plugins/Runtime/StateTree/Source/StateTreeModule/Public/Debugger/StateTreeTraceProvider.h b/Engine/Plugins/Runtime/StateTree/Source/StateTreeModule/Public/Debugger/StateTreeTraceProvider.h
--- a/Engine/Plugins/Runtime/StateTree/Source/StateTreeModule/Public/Debugger/StateTreeTraceProvider.h
+++ b/Engine/Plugins/Runtime/StateTree/Source/StateTreeModule/Public/Debugger/StateTreeTraceProvider.h
@@ -6,6 +6,7 @@
 
 #include "IStateTreeTraceProvider.h"
 #include "Model/PointTimeline.h"
+#include "StateTreeIndexTypes.h"
 #include "StateTreeTypes.h" // required to compile TMap<FStateTreeInstanceDebugId, ...>
 
 namespace TraceServices { class IAnalysisSession; }
@@ -21,13 +22,16 @@ public:
 	
 	void AppendEvent(FStateTreeInstanceDebugId InInstanceId, double InTime, const FStateTreeTraceEventVariantType& InEvent);
 	void AppendInstanceEvent(
-		const UStateTree* InStateTree,
+		const FStateTreeIndex16 AssetDebugId,
 		const FStateTreeInstanceDebugId InInstanceId,
 		const TCHAR* InInstanceName,
 		double InTime,
 		double InWorldRecordingTime,
 		EStateTreeTraceEventType InEventType);
 
+	void AppendAssetDebugId(const UStateTree* InStateTree, const FStateTreeIndex16 AssetDebugId);
+	bool GetAssetFromDebugId(const FStateTreeIndex16 AssetDebugId, TWeakObjectPtr<const UStateTree>& WeakStateTree) const;
+
 protected:
 	/** IStateTreeDebuggerProvider interface */
 	virtual void GetInstances(TArray<UE::StateTreeDebugger::FInstanceDescriptor>& OutInstances) const override;
@@ -40,5 +44,13 @@ private:
 	TMap<FStateTreeInstanceDebugId, uint32> InstanceIdToDebuggerEntryTimelines;
 	TArray<UE::StateTreeDebugger::FInstanceDescriptor> Descriptors;
 	TArray<TSharedRef<TraceServices::TPointTimeline<FStateTreeTraceEventVariantType>>> EventsTimelines;
+
+	struct FStateTreeDebugIdPair
+	{
+		TWeakObjectPtr<const UStateTree> WeakStateTree;
+		FStateTreeIndex16 Id;
+	};
+
+	TArray<FStateTreeDebugIdPair> StateTreeAssets;
 };
 #endif // WITH_STATETREE_DEBUGGER
\ No newline at end of file
diff --git a/Engine/Plugins/Runtime/StateTree/Source/StateTreeModule/Public/Debugger/StateTreeTraceTypes.h b/Engine/Plugins/Runtime/StateTree/Source/StateTreeModule/Public/Debugger/StateTreeTraceTypes.h
--- a/Engine/Plugins/Runtime/StateTree/Source/StateTreeModule/Public/Debugger/StateTreeTraceTypes.h
+++ b/Engine/Plugins/Runtime/StateTree/Source/StateTreeModule/Public/Debugger/StateTreeTraceTypes.h
@@ -129,9 +129,8 @@ struct FStateTreeTraceNodeEvent : FStateTreeTraceBaseEvent
 
 struct FStateTreeTraceStateEvent : FStateTreeTraceNodeEvent
 {
-	explicit FStateTreeTraceStateEvent(const double RecordingWorldTime, const FStateTreeIndex16 Index, const EStateTreeTraceEventType EventType, const EStateTreeStateSelectionBehavior SelectionBehavior)
+	explicit FStateTreeTraceStateEvent(const double RecordingWorldTime, const FStateTreeIndex16 Index, const EStateTreeTraceEventType EventType)
 		: FStateTreeTraceNodeEvent(RecordingWorldTime, Index, EventType)
-		, SelectionBehavior(SelectionBehavior)
 	{
 	}
 
@@ -139,8 +138,6 @@ struct FStateTreeTraceStateEvent : FStateTreeTraceNodeEvent
 	STATETREEMODULE_API FString GetValueString(const UStateTree& StateTree) const;
 	STATETREEMODULE_API FString GetTypeString(const UStateTree& StateTree) const;
 	STATETREEMODULE_API FStateTreeStateHandle GetStateHandle() const;
-
-	EStateTreeStateSelectionBehavior SelectionBehavior;
 };
 
 struct FStateTreeTraceTaskEvent : FStateTreeTraceNodeEvent
@@ -205,6 +202,17 @@ struct FStateTreeTraceConditionEvent : FStateTreeTraceNodeEvent
 	FString InstanceDataAsText;
 };
 
+struct FStateTreeTraceActiveStates
+{
+	struct FAssetActiveStates
+	{
+		TWeakObjectPtr<const UStateTree> WeakStateTree;
+		TArray<FStateTreeStateHandle> ActiveStates;
+	};
+
+	TArray<FAssetActiveStates> PerAssetStates;
+};
+
 struct FStateTreeTraceActiveStatesEvent : FStateTreeTraceBaseEvent
 {
 	// Intentionally implemented in source file to compile 'TArray<FStateTreeStateHandle>' using only forward declaration.
@@ -214,7 +222,22 @@ struct FStateTreeTraceActiveStatesEvent : FStateTreeTraceBaseEvent
 	STATETREEMODULE_API FString GetValueString(const UStateTree& StateTree) const;
 	STATETREEMODULE_API FString GetTypeString(const UStateTree& StateTree) const;
 
-	TArray<FStateTreeStateHandle> ActiveStates;
+	FStateTreeTraceActiveStates ActiveStates;
+};
+
+struct FStateTreeTraceInstanceFrameEvent : FStateTreeTraceBaseEvent
+{
+	explicit FStateTreeTraceInstanceFrameEvent(const double RecordingWorldTime, const EStateTreeTraceEventType EventType, const UStateTree* StateTree)
+		: FStateTreeTraceBaseEvent(RecordingWorldTime, EventType)
+		, WeakStateTree(StateTree)
+	{
+	}
+
+	STATETREEMODULE_API FString ToFullString(const UStateTree& StateTree) const;
+	STATETREEMODULE_API FString GetValueString(const UStateTree& StateTree) const;
+	STATETREEMODULE_API FString GetTypeString(const UStateTree& StateTree) const;
+
+	TWeakObjectPtr<const UStateTree> WeakStateTree;
 };
 
 /** Type aliases for statetree trace events */
@@ -227,6 +250,7 @@ using FStateTreeTraceEventVariantType = TVariant<FStateTreeTracePhaseEvent,
 												FStateTreeTraceEvaluatorEvent,
 												FStateTreeTraceTransitionEvent,
 												FStateTreeTraceConditionEvent,
-												FStateTreeTraceActiveStatesEvent>;
+												FStateTreeTraceActiveStatesEvent,
+												FStateTreeTraceInstanceFrameEvent>;
 
 #endif // WITH_STATETREE_DEBUGGER
diff --git a/Engine/Plugins/Runtime/StateTree/Source/StateTreeModule/Public/StateTreeExecutionContext.h b/Engine/Plugins/Runtime/StateTree/Source/StateTreeModule/Public/StateTreeExecutionContext.h
--- a/Engine/Plugins/Runtime/StateTree/Source/StateTreeModule/Public/StateTreeExecutionContext.h
+++ b/Engine/Plugins/Runtime/StateTree/Source/StateTreeModule/Public/StateTreeExecutionContext.h
@@ -580,26 +580,9 @@ protected:
 	/** Helper struct to track currently processed frame. */
 	struct FCurrentlyProcessedFrameScope
 	{
-		FCurrentlyProcessedFrameScope(FStateTreeExecutionContext& InContext, const FStateTreeExecutionFrame* CurrentParentFrame, const FStateTreeExecutionFrame& CurrentFrame)
-			: Context(InContext)
-		{
-			check(CurrentFrame.StateTree);
-			FStateTreeInstanceStorage* SharedInstanceDataStorage = &CurrentFrame.StateTree->GetSharedInstanceData()->GetMutableStorage();
-
-			SavedFrame = Context.CurrentlyProcessedFrame;
-			SavedParentFrame = Context.CurrentlyProcessedParentFrame;
-			SavedSharedInstanceDataStorage = Context.CurrentlyProcessedSharedInstanceStorage;
-			Context.CurrentlyProcessedFrame = &CurrentFrame;
-			Context.CurrentlyProcessedParentFrame = CurrentParentFrame;
-			Context.CurrentlyProcessedSharedInstanceStorage = SharedInstanceDataStorage;
-		}
+		FCurrentlyProcessedFrameScope(FStateTreeExecutionContext& InContext, const FStateTreeExecutionFrame* CurrentParentFrame, const FStateTreeExecutionFrame& CurrentFrame);
 
-		~FCurrentlyProcessedFrameScope()
-		{
-			Context.CurrentlyProcessedFrame = SavedFrame;
-			Context.CurrentlyProcessedParentFrame = SavedParentFrame;
-			Context.CurrentlyProcessedSharedInstanceStorage = SavedSharedInstanceDataStorage;
-		}
+		~FCurrentlyProcessedFrameScope();
 
 	private:
 		FStateTreeExecutionContext& Context;
diff --git a/Engine/Plugins/Runtime/StateTree/Source/StateTreeModule/Public/StateTreeExecutionTypes.h b/Engine/Plugins/Runtime/StateTree/Source/StateTreeModule/Public/StateTreeExecutionTypes.h
--- a/Engine/Plugins/Runtime/StateTree/Source/StateTreeModule/Public/StateTreeExecutionTypes.h
+++ b/Engine/Plugins/Runtime/StateTree/Source/StateTreeModule/Public/StateTreeExecutionTypes.h
@@ -26,9 +26,11 @@ enum class EStateTreeUpdatePhase : uint8
 	TickingTasks			UMETA(DisplayName = "Tick Tasks"),
 	TransitionConditions	UMETA(DisplayName = "Transition conditions"),
 	StateSelection			UMETA(DisplayName = "Try Enter"),
+	TrySelectBehavior		UMETA(DisplayName = "Try Select Behavior"),
 	EnterConditions			UMETA(DisplayName = "Enter conditions"),
 	EnterStates				UMETA(DisplayName = "Enter States"),
 	ExitStates				UMETA(DisplayName = "Exit States"),
+	StateCompleted			UMETA(DisplayName = "State(s) Completed")
 };
 
 
