Other fix for CIS issue 408344

diff --git a/Engine/Plugins/Runtime/StateTree/Source/StateTreeModule/Private/Debugger/StateTreeTraceTypes.cpp b/Engine/Plugins/Runtime/StateTree/Source/StateTreeModule/Private/Debugger/StateTreeTraceTypes.cpp
--- a/Engine/Plugins/Runtime/StateTree/Source/StateTreeModule/Private/Debugger/StateTreeTraceTypes.cpp
+++ b/Engine/Plugins/Runtime/StateTree/Source/StateTreeModule/Private/Debugger/StateTreeTraceTypes.cpp
@@ -357,6 +357,12 @@ FString FStateTreeTraceActiveStatesEvent::GetTypeString(const UStateTree& StateT
 //----------------------------------------------------------------------//
 // FStateTreeTraceInstanceFrameEvent
 //----------------------------------------------------------------------//
+FStateTreeTraceInstanceFrameEvent::FStateTreeTraceInstanceFrameEvent(const double RecordingWorldTime, const EStateTreeTraceEventType EventType, const UStateTree* StateTree)
+	: FStateTreeTraceBaseEvent(RecordingWorldTime, EventType)
+	, WeakStateTree(StateTree)
+{
+}
+
 FString FStateTreeTraceInstanceFrameEvent::ToFullString(const UStateTree& StateTree) const
 {
 	return FString::Printf(TEXT("%s: %s"), *GetTypeString(StateTree), *GetValueString(StateTree));
diff --git a/Engine/Plugins/Runtime/StateTree/Source/StateTreeModule/Public/Debugger/StateTreeTraceTypes.h b/Engine/Plugins/Runtime/StateTree/Source/StateTreeModule/Public/Debugger/StateTreeTraceTypes.h
--- a/Engine/Plugins/Runtime/StateTree/Source/StateTreeModule/Public/Debugger/StateTreeTraceTypes.h
+++ b/Engine/Plugins/Runtime/StateTree/Source/StateTreeModule/Public/Debugger/StateTreeTraceTypes.h
@@ -227,11 +227,7 @@ struct FStateTreeTraceActiveStatesEvent : FStateTreeTraceBaseEvent
 
 struct FStateTreeTraceInstanceFrameEvent : FStateTreeTraceBaseEvent
 {
-	explicit FStateTreeTraceInstanceFrameEvent(const double RecordingWorldTime, const EStateTreeTraceEventType EventType, const UStateTree* StateTree)
-		: FStateTreeTraceBaseEvent(RecordingWorldTime, EventType)
-		, WeakStateTree(StateTree)
-	{
-	}
+	explicit FStateTreeTraceInstanceFrameEvent(const double RecordingWorldTime, const EStateTreeTraceEventType EventType, const UStateTree* StateTree);
 
 	STATETREEMODULE_API FString ToFullString(const UStateTree& StateTree) const;
 	STATETREEMODULE_API FString GetValueString(const UStateTree& StateTree) const;
