Fix CIS issue 408392 (seems to not be an issue with clang 17)

diff --git a/Engine/Plugins/Runtime/StateTree/Source/StateTreeModule/Private/Debugger/StateTreeTrace.cpp b/Engine/Plugins/Runtime/StateTree/Source/StateTreeModule/Private/Debugger/StateTreeTrace.cpp
--- a/Engine/Plugins/Runtime/StateTree/Source/StateTreeModule/Private/Debugger/StateTreeTrace.cpp
+++ b/Engine/Plugins/Runtime/StateTree/Source/StateTreeModule/Private/Debugger/StateTreeTrace.cpp
@@ -202,12 +202,12 @@ struct FInstanceEventBufferedData
 		}
 
 		bool IsValid() const { return StatesIndices.Num() > 0 && StatesIndices.Num() == AssetDebugIds.Num(); }
-		void Output(const FStateTreeInstanceDebugId InstanceId) const
+		void Output(const FStateTreeInstanceDebugId InInstanceId) const
 		{
 			UE_TRACE_LOG(StateTreeDebugger, ActiveStatesEvent, StateTreeDebugChannel)
 				<< ActiveStatesEvent.Cycle(FPlatformTime::Cycles64())
-				<< ActiveStatesEvent.InstanceId(InstanceId.Id)
-				<< ActiveStatesEvent.InstanceSerial(InstanceId.SerialNumber)
+				<< ActiveStatesEvent.InstanceId(InInstanceId.Id)
+				<< ActiveStatesEvent.InstanceSerial(InInstanceId.SerialNumber)
 				<< ActiveStatesEvent.ActiveStates(StatesIndices.GetData(), StatesIndices.Num())
 				<< ActiveStatesEvent.AssetDebugIds(AssetDebugIds.GetData(), AssetDebugIds.Num());
 		}
