StateTree: Change StateTree editor transition display to reflect the execution logic
- The execution logic will transition to Root if no state is found, change editor to reflet that
- Change viewmodel from ref to ptr to make it optional (allows to reuse the row for non-editing purposes)

diff --git a/Engine/Plugins/Runtime/StateTree/Source/StateTreeEditorModule/Private/SStateTreeViewRow.cpp b/Engine/Plugins/Runtime/StateTree/Source/StateTreeEditorModule/Private/SStateTreeViewRow.cpp
--- a/Engine/Plugins/Runtime/StateTree/Source/StateTreeEditorModule/Private/SStateTreeViewRow.cpp
+++ b/Engine/Plugins/Runtime/StateTree/Source/StateTreeEditorModule/Private/SStateTreeViewRow.cpp
@@ -35,7 +35,7 @@ namespace UE::StateTree::Editor
 	}
 } // UE:StateTree::Editor
 
-void SStateTreeViewRow::Construct(const FArguments& InArgs, const TSharedRef<STableViewBase>& InOwnerTableView, TWeakObjectPtr<UStateTreeState> InState, const TSharedPtr<SScrollBox>& ViewBox, TSharedRef<FStateTreeViewModel> InStateTreeViewModel)
+void SStateTreeViewRow::Construct(const FArguments& InArgs, const TSharedRef<STableViewBase>& InOwnerTableView, TWeakObjectPtr<UStateTreeState> InState, const TSharedPtr<SScrollBox>& ViewBox, TSharedPtr<FStateTreeViewModel> InStateTreeViewModel)
 {
 	StateTreeViewModel = InStateTreeViewModel;
 	WeakState = InState;
@@ -1043,7 +1043,7 @@ FText SStateTreeViewRow::GetTransitionsDesc(const UStateTreeState& State, const
 		}
 		else
 		{
-			DescItems.Add(LOCTEXT("TransitionActionMissingTransition", "Missing Transition"));
+			DescItems.Add(LOCTEXT("TransitionActionRoot", "[Root]"));
 		}
 	}
 	
@@ -1058,7 +1058,6 @@ FText SStateTreeViewRow::GetTransitionsIcon(const UStateTreeState& State, const
 		IconRightArrow =	1 << 0,
 		IconDownArrow =		1 << 1,
 		IconLevelUp =		1 << 2,
-		IconWarning =		1 << 3,
 	};
 	uint8 IconType = IconNone;
 	
@@ -1123,14 +1122,8 @@ FText SStateTreeViewRow::GetTransitionsIcon(const UStateTreeState& State, const
 		&& IconType == IconNone
 		&& EnumHasAnyFlags(Trigger, EStateTreeTransitionTrigger::OnStateCompleted))
 	{
-		if (HasParentTransitionForTrigger(State, Trigger))
-		{
-			IconType = IconLevelUp;
-		}
-		else
-		{
-			IconType = IconWarning;
-		}
+		// Transition is handled on parent state, or implicit Root.
+		IconType = IconLevelUp;
 	}
 
 	switch (IconType)
@@ -1141,8 +1134,6 @@ FText SStateTreeViewRow::GetTransitionsIcon(const UStateTreeState& State, const
 			return FEditorFontGlyphs::Long_Arrow_Down;
 		case IconLevelUp:
 			return FEditorFontGlyphs::Level_Up;
-		case IconWarning:
-			return FEditorFontGlyphs::Exclamation_Triangle;
 		default:
 			return FText::GetEmpty();
 	}
diff --git a/Engine/Plugins/Runtime/StateTree/Source/StateTreeEditorModule/Private/SStateTreeViewRow.h b/Engine/Plugins/Runtime/StateTree/Source/StateTreeEditorModule/Private/SStateTreeViewRow.h
--- a/Engine/Plugins/Runtime/StateTree/Source/StateTreeEditorModule/Private/SStateTreeViewRow.h
+++ b/Engine/Plugins/Runtime/StateTree/Source/StateTreeEditorModule/Private/SStateTreeViewRow.h
@@ -25,7 +25,7 @@ public:
 	}
 	SLATE_END_ARGS()
 
-	void Construct(const FArguments& InArgs, const TSharedRef<STableViewBase>& InOwnerTableView, TWeakObjectPtr<UStateTreeState> InState, const TSharedPtr<SScrollBox>& ViewBox, TSharedRef<FStateTreeViewModel> InStateTreeViewModel);
+	void Construct(const FArguments& InArgs, const TSharedRef<STableViewBase>& InOwnerTableView, TWeakObjectPtr<UStateTreeState> InState, const TSharedPtr<SScrollBox>& ViewBox, TSharedPtr<FStateTreeViewModel> InStateTreeViewModel);
 	void RequestRename() const;
 
 private:
