StateTree: Fixed comments and restored timing collection for Mass external data collecting

diff --git a/Engine/Plugins/AI/MassAI/Source/MassAIBehavior/Private/MassStateTreeExecutionContext.cpp b/Engine/Plugins/AI/MassAI/Source/MassAIBehavior/Private/MassStateTreeExecutionContext.cpp
--- a/Engine/Plugins/AI/MassAI/Source/MassAIBehavior/Private/MassStateTreeExecutionContext.cpp
+++ b/Engine/Plugins/AI/MassAI/Source/MassAIBehavior/Private/MassStateTreeExecutionContext.cpp
@@ -7,11 +7,14 @@
 #include "MassSignalSubsystem.h"
 #include "MassEntityView.h"
 #include "Engine/World.h"
+#include "ProfilingDebugging/CsvProfiler.h"
 
 namespace UE::MassBehavior
 {
 bool CollectExternalData(const FStateTreeExecutionContext& Context, const UStateTree* StateTree, TArrayView<const FStateTreeExternalDataDesc> ExternalDataDescs, TArrayView<FStateTreeDataView> OutDataViews)
 {
+	CSV_SCOPED_TIMING_STAT_EXCLUSIVE(StateTreeProcessorCollectExternalData);
+	
 	const FMassStateTreeExecutionContext& MassContext = static_cast<const FMassStateTreeExecutionContext&>(Context); 
 	const FMassEntityManager& EntityManager = MassContext.GetEntityManager();
 	const UWorld* World = MassContext.GetWorld();
diff --git a/Engine/Plugins/Runtime/StateTree/Source/StateTreeModule/Public/StateTree.h b/Engine/Plugins/Runtime/StateTree/Source/StateTreeModule/Public/StateTree.h
--- a/Engine/Plugins/Runtime/StateTree/Source/StateTreeModule/Public/StateTree.h
+++ b/Engine/Plugins/Runtime/StateTree/Source/StateTreeModule/Public/StateTree.h
@@ -41,7 +41,7 @@ struct STATETREEMODULE_API FStateTreeCustomVersion
 		AddedDataHandlesIds,
 		// Added linked asset state
 		AddedLinkedAssetState,
-		// Change hoe external data is accessed
+		// Change how external data is accessed
 		ChangedExternalDataAccess,
 
 		// -----<new versions can be added above this line>-------------------------------------------------
diff --git a/Engine/Plugins/Runtime/StateTree/Source/StateTreeModule/Public/StateTreeExecutionContext.h b/Engine/Plugins/Runtime/StateTree/Source/StateTreeModule/Public/StateTreeExecutionContext.h
--- a/Engine/Plugins/Runtime/StateTree/Source/StateTreeModule/Public/StateTreeExecutionContext.h
+++ b/Engine/Plugins/Runtime/StateTree/Source/StateTreeModule/Public/StateTreeExecutionContext.h
@@ -525,11 +525,13 @@ protected:
 	/** Copies a batch of properties to the data in TargetView. This version validates the data handles and looks up temporary instances. */
 	bool CopyBatchWithValidation(const FStateTreeExecutionFrame* ParentFrame, const FStateTreeExecutionFrame& CurrentFrame, const FStateTreeDataView TargetView, const FStateTreeIndex16 BindingsBatch) const;
 
-	/** Collects external data for all StateTrees in active frames. @return true*/
+	/**
+	 * Collects external data for all StateTrees in active frames.
+	 * @returns true if all external data are set successfully. */
 	bool CollectActiveExternalData();
 
 	/**
-	 * Collects external data for specific State Tree asset. If the data is already collect, cached index is returned.
+	 * Collects external data for specific State Tree asset. If the data is already collected, cached index is returned.
 	 * @returns index in ContextAndExternalDataViews for the first external data.
 	 */
 	FStateTreeIndex16 CollectExternalData(const UStateTree* StateTree);
