[StatetreeDebugger] fixed empty lines in some node details

diff --git a/Engine/Plugins/Runtime/StateTree/Source/StateTreeEditorModule/Private/Debugger/SStateTreeDebuggerView.cpp b/Engine/Plugins/Runtime/StateTree/Source/StateTreeEditorModule/Private/Debugger/SStateTreeDebuggerView.cpp
--- a/Engine/Plugins/Runtime/StateTree/Source/StateTreeEditorModule/Private/Debugger/SStateTreeDebuggerView.cpp
+++ b/Engine/Plugins/Runtime/StateTree/Source/StateTreeEditorModule/Private/Debugger/SStateTreeDebuggerView.cpp
@@ -1193,7 +1193,11 @@ void SStateTreeDebuggerView::GenerateElementsForProperties(const FStateTreeTrace
 				NextToken = ViewIt+1;
 				NestedCount += LocalNestedCount;
 			}
-			CreatePropertyElement(FStringView(NextToken, UE_PTRDIFF_TO_INT32(ViewIt - NextToken)), NestedCount);
+
+			if (ViewIt != NextToken)
+			{
+				CreatePropertyElement(FStringView(NextToken, UE_PTRDIFF_TO_INT32(ViewIt - NextToken)), NestedCount);
+			}
 		}
 		else
 		{
