StateTree: Clear struct references on FStateTreePropertyBindings::PerformResetObjects()

diff --git a/Engine/Plugins/Runtime/StateTree/Source/StateTreeModule/Private/StateTreePropertyBindings.cpp b/Engine/Plugins/Runtime/StateTree/Source/StateTreeModule/Private/StateTreePropertyBindings.cpp
--- a/Engine/Plugins/Runtime/StateTree/Source/StateTreeModule/Private/StateTreePropertyBindings.cpp
+++ b/Engine/Plugins/Runtime/StateTree/Source/StateTreeModule/Private/StateTreePropertyBindings.cpp
@@ -1055,6 +1055,9 @@ void FStateTreePropertyBindings::PerformResetObjects(const FStateTreePropertyCop
 	case EStateTreePropertyCopyType::CopyObject:
 		static_cast<const FObjectPropertyBase*>(Copy.TargetLeafProperty)->SetObjectPropertyValue(TargetAddress, nullptr);
 		break;
+	case EStateTreePropertyCopyType::StructReference:
+		reinterpret_cast<FStateTreeStructRef*>(TargetAddress)->Set(FStructView());
+		break;
 	case EStateTreePropertyCopyType::CopyName:
 		break;
 	case EStateTreePropertyCopyType::CopyFixedArray:
