StateTree: Update state details after adding a globa task/eval so that bindings will pick up the new task

diff --git a/Engine/Plugins/Runtime/StateTree/Source/StateTreeEditorModule/Private/StateTreeEditor.cpp b/Engine/Plugins/Runtime/StateTree/Source/StateTreeEditorModule/Private/StateTreeEditor.cpp
--- a/Engine/Plugins/Runtime/StateTree/Source/StateTreeEditorModule/Private/StateTreeEditor.cpp
+++ b/Engine/Plugins/Runtime/StateTree/Source/StateTreeEditorModule/Private/StateTreeEditor.cpp
@@ -255,7 +255,8 @@ void FStateTreeEditor::InitEditor( const EToolkitMode::Type Mode, const TSharedP
 
 	UE::StateTree::Delegates::OnIdentifierChanged.AddSP(this, &FStateTreeEditor::OnIdentifierChanged);
 	UE::StateTree::Delegates::OnSchemaChanged.AddSP(this, &FStateTreeEditor::OnSchemaChanged);
-	UE::StateTree::Delegates::OnParametersChanged.AddSP(this, &FStateTreeEditor::OnParametersChanged);
+	UE::StateTree::Delegates::OnParametersChanged.AddSP(this, &FStateTreeEditor::OnRefreshDetailsView);
+	UE::StateTree::Delegates::OnGlobalDataChanged.AddSP(this, &FStateTreeEditor::OnRefreshDetailsView);
 	UE::StateTree::Delegates::OnStateParametersChanged.AddSP(this, &FStateTreeEditor::OnStateParametersChanged);
 }
 
@@ -501,7 +502,7 @@ void FStateTreeEditor::OnSchemaChanged(const UStateTree& InStateTree)
 	}
 }
 
-void FStateTreeEditor::OnParametersChanged(const UStateTree& InStateTree)
+void FStateTreeEditor::OnRefreshDetailsView(const UStateTree& InStateTree)
 {
 	if (StateTree == &InStateTree)
 	{
diff --git a/Engine/Plugins/Runtime/StateTree/Source/StateTreeEditorModule/Private/StateTreeEditor.h b/Engine/Plugins/Runtime/StateTree/Source/StateTreeEditorModule/Private/StateTreeEditor.h
--- a/Engine/Plugins/Runtime/StateTree/Source/StateTreeEditorModule/Private/StateTreeEditor.h
+++ b/Engine/Plugins/Runtime/StateTree/Source/StateTreeEditorModule/Private/StateTreeEditor.h
@@ -87,7 +87,7 @@ private:
 	void HandleModelAssetChanged();
 	void OnIdentifierChanged(const UStateTree& StateTree);
 	void OnSchemaChanged(const UStateTree& StateTree);
-	void OnParametersChanged(const UStateTree& StateTree);
+	void OnRefreshDetailsView(const UStateTree& StateTree);
 	void OnStateParametersChanged(const UStateTree& StateTree, const FGuid StateID);
 
 	FText GetStatisticsText() const;
diff --git a/Engine/Plugins/Runtime/StateTree/Source/StateTreeEditorModule/Private/StateTreeEditorData.cpp b/Engine/Plugins/Runtime/StateTree/Source/StateTreeEditorModule/Private/StateTreeEditorData.cpp
--- a/Engine/Plugins/Runtime/StateTree/Source/StateTreeEditorModule/Private/StateTreeEditorData.cpp
+++ b/Engine/Plugins/Runtime/StateTree/Source/StateTreeEditorModule/Private/StateTreeEditorData.cpp
@@ -220,6 +220,14 @@ void UStateTreeEditorData::PostEditChangeChainProperty(FPropertyChangedChainEven
 				EditorBindings.RemoveUnusedBindings(AllStructValues);
 			}
 		}
+
+		// Notify that the global data changed (will need to update binding widgets, etc)
+		if (MemberName == GET_MEMBER_NAME_CHECKED(UStateTreeEditorData, Evaluators)
+			|| MemberName == GET_MEMBER_NAME_CHECKED(UStateTreeEditorData, GlobalTasks))
+		{
+			UE::StateTree::Delegates::OnGlobalDataChanged.Broadcast(*StateTree);
+		}
+
 	}
 
 	UE::StateTree::PropertyHelpers::DispatchPostEditToNodes(*this, PropertyChangedEvent);
diff --git a/Engine/Plugins/Runtime/StateTree/Source/StateTreeModule/Private/StateTreeDelegates.cpp b/Engine/Plugins/Runtime/StateTree/Source/StateTreeModule/Private/StateTreeDelegates.cpp
--- a/Engine/Plugins/Runtime/StateTree/Source/StateTreeModule/Private/StateTreeDelegates.cpp
+++ b/Engine/Plugins/Runtime/StateTree/Source/StateTreeModule/Private/StateTreeDelegates.cpp
@@ -9,6 +9,7 @@ namespace UE::StateTree::Delegates
 FOnIdentifierChanged OnIdentifierChanged;
 FOnSchemaChanged OnSchemaChanged;
 FOnParametersChanged OnParametersChanged;
+FOnGlobalDataChanged OnGlobalDataChanged;
 FOnStateParametersChanged OnStateParametersChanged;
 FOnBreakpointsChanged OnBreakpointsChanged;
 FOnPostCompile OnPostCompile;
diff --git a/Engine/Plugins/Runtime/StateTree/Source/StateTreeModule/Public/StateTreeDelegates.h b/Engine/Plugins/Runtime/StateTree/Source/StateTreeModule/Public/StateTreeDelegates.h
--- a/Engine/Plugins/Runtime/StateTree/Source/StateTreeModule/Public/StateTreeDelegates.h
+++ b/Engine/Plugins/Runtime/StateTree/Source/StateTreeModule/Public/StateTreeDelegates.h
@@ -38,6 +38,13 @@ extern STATETREEMODULE_API FOnParametersChanged OnParametersChanged;
 DECLARE_MULTICAST_DELEGATE_TwoParams(FOnStateParametersChanged, const UStateTree& /*StateTree*/, const FGuid /*StateID*/);
 extern STATETREEMODULE_API FOnStateParametersChanged OnStateParametersChanged;
 
+/**
+ * Called when Global Tasks or Evaluators of the StateTree EditorData changed.
+ * This should mainly used by the asset editor to maintain consistency in the UI for manipulations on the EditorData.
+ */
+DECLARE_MULTICAST_DELEGATE_OneParam(FOnGlobalDataChanged, const UStateTree& /*StateTree*/);
+extern STATETREEMODULE_API FOnGlobalDataChanged OnGlobalDataChanged;
+
 /**
  * Called when breakpoints of the StateTree EditorData changed.
  * This should mainly used by the asset editor to update the debugger.
