StateTree: Changed guid fixup to be more deterministic.

diff --git a/Engine/Plugins/Runtime/StateTree/Source/StateTreeEditorModule/Private/StateTreeEditorData.cpp b/Engine/Plugins/Runtime/StateTree/Source/StateTreeEditorModule/Private/StateTreeEditorData.cpp
--- a/Engine/Plugins/Runtime/StateTree/Source/StateTreeEditorModule/Private/StateTreeEditorData.cpp
+++ b/Engine/Plugins/Runtime/StateTree/Source/StateTreeEditorModule/Private/StateTreeEditorData.cpp
@@ -576,14 +576,16 @@ void UStateTreeEditorData::FixDuplicateIDs()
 	TSet<FGuid> FoundNodeIDs;
 
 	// Evaluators
-	for (FStateTreeEditorNode& Node : Evaluators)
+	for (int32 Index = 0; Index < Evaluators.Num(); Index++)
 	{
+		FStateTreeEditorNode& Node = Evaluators[Index];
 		if (const FStateTreeEvaluatorBase* Evaluator = Node.Node.GetPtr<FStateTreeEvaluatorBase>())
 		{
 			const FGuid OldID = Node.ID; 
 			if (FoundNodeIDs.Contains(Node.ID))
 			{
-				Node.ID = FGuid::NewGuid();
+				Node.ID = UE::StateTree::PropertyHelpers::MakeDeterministicID(*this, TEXT("Evaluators"), Index);
+				
 				UE_LOG(LogStateTreeEditor, Log, TEXT("%s: Found Evaluator '%s' with duplicate ID, changing ID:%s to ID:%s."),
 					*GetFullName(), *Node.GetName().ToString(), *OldID.ToString(), *Node.ID.ToString());
 				EditorBindings.CopyBindings(OldID, Node.ID);
@@ -593,14 +595,16 @@ void UStateTreeEditorData::FixDuplicateIDs()
 	}
 	
 	// Global Tasks
-	for (FStateTreeEditorNode& Node : GlobalTasks)
+	for (int32 Index = 0; Index < GlobalTasks.Num(); Index++)
 	{
+		FStateTreeEditorNode& Node = GlobalTasks[Index];
 		if (const FStateTreeTaskBase* Task = Node.Node.GetPtr<FStateTreeTaskBase>())
 		{
 			const FGuid OldID = Node.ID; 
 			if (FoundNodeIDs.Contains(Node.ID))
 			{
-				Node.ID = FGuid::NewGuid();
+				Node.ID = UE::StateTree::PropertyHelpers::MakeDeterministicID(*this, TEXT("GlobalTasks"), Index);
+				
 				UE_LOG(LogStateTreeEditor, Log, TEXT("%s: Found GlobalTask '%s' with duplicate ID, changing ID:%s to ID:%s."),
 					*GetFullName(), *Node.GetName().ToString(), *OldID.ToString(), *Node.ID.ToString());
 				EditorBindings.CopyBindings(OldID, Node.ID);
@@ -612,8 +616,9 @@ void UStateTreeEditorData::FixDuplicateIDs()
 	VisitHierarchy([&FoundNodeIDs, &EditorBindings = EditorBindings, &Self = *this](UStateTreeState& State, UStateTreeState* ParentState)
 	{
 		// Enter conditions
-		for (FStateTreeEditorNode& Node : State.EnterConditions)
+		for (int32 Index = 0; Index < State.EnterConditions.Num(); Index++)
 		{
+			FStateTreeEditorNode& Node = State.EnterConditions[Index];
 			if (const FStateTreeConditionBase* Cond = Node.Node.GetPtr<FStateTreeConditionBase>())
 			{
 				const FGuid OldID = Node.ID;
@@ -622,7 +627,8 @@ void UStateTreeEditorData::FixDuplicateIDs()
 				FoundNodeIDs.Add(Node.ID, &bIsAlreadyInSet);
 				if (bIsAlreadyInSet)
 				{
-					Node.ID = FGuid::NewGuid();
+					Node.ID = UE::StateTree::PropertyHelpers::MakeDeterministicID(State, TEXT("EnterConditions"), Index);
+					
 					UE_LOG(LogStateTreeEditor, Log, TEXT("%s: Found Enter Condition '%s' with duplicate ID on state '%s', changing ID:%s to ID:%s."),
 						*Self.GetFullName(), *Node.GetName().ToString(), *GetNameSafe(&State), *OldID.ToString(), *Node.ID.ToString());
 					EditorBindings.CopyBindings(OldID, Node.ID);
@@ -631,8 +637,9 @@ void UStateTreeEditorData::FixDuplicateIDs()
 		}
 
 		// Tasks
-		for (FStateTreeEditorNode& Node : State.Tasks)
+		for (int32 Index = 0; Index < State.Tasks.Num(); Index++)
 		{
+			FStateTreeEditorNode& Node = State.Tasks[Index];
 			if (const FStateTreeTaskBase* Task = Node.Node.GetPtr<FStateTreeTaskBase>())
 			{
 				const FGuid OldID = Node.ID;
@@ -641,7 +648,8 @@ void UStateTreeEditorData::FixDuplicateIDs()
 				FoundNodeIDs.Add(Node.ID, &bIsAlreadyInSet);
 				if (bIsAlreadyInSet)
 				{
-					Node.ID = FGuid::NewGuid();
+					Node.ID = UE::StateTree::PropertyHelpers::MakeDeterministicID(State, TEXT("Tasks"), Index);
+
 					UE_LOG(LogStateTreeEditor, Log, TEXT("%s: Found Task '%s' with duplicate ID on state '%s', changing ID:%s to ID:%s."),
 						*Self.GetFullName(), *Node.GetName().ToString(), *GetNameSafe(&State), *OldID.ToString(), *Node.ID.ToString());
 					EditorBindings.CopyBindings(OldID, Node.ID);
@@ -657,7 +665,8 @@ void UStateTreeEditorData::FixDuplicateIDs()
 			FoundNodeIDs.Add(State.SingleTask.ID, &bIsAlreadyInSet);
 			if (bIsAlreadyInSet)
 			{
-				State.SingleTask.ID = FGuid::NewGuid();
+				State.SingleTask.ID = UE::StateTree::PropertyHelpers::MakeDeterministicID(State, TEXT("SingleTask"), 0);
+
 				UE_LOG(LogStateTreeEditor, Log, TEXT("%s: Found enter condition '%s' with duplicate ID on state '%s', changing ID:%s to ID:%s."),
 					*Self.GetFullName(), *State.SingleTask.GetName().ToString(), *GetNameSafe(&State), *OldID.ToString(), *State.SingleTask.ID.ToString());
 				EditorBindings.CopyBindings(OldID, State.SingleTask.ID);
@@ -665,10 +674,12 @@ void UStateTreeEditorData::FixDuplicateIDs()
 		}
 
 		// Transitions
-		for (FStateTreeTransition& Transition : State.Transitions)
+		for (int32 TransitionIndex = 0; TransitionIndex < State.Transitions.Num(); TransitionIndex++)
 		{
-			for (FStateTreeEditorNode& Node : Transition.Conditions)
+			FStateTreeTransition& Transition = State.Transitions[TransitionIndex];
+			for (int32 Index = 0; Index < Transition.Conditions.Num(); Index++)
 			{
+				FStateTreeEditorNode& Node = Transition.Conditions[Index];
 				if (const FStateTreeConditionBase* Cond = Node.Node.GetPtr<FStateTreeConditionBase>())
 				{
 					const FGuid OldID = Node.ID; 
@@ -676,7 +687,8 @@ void UStateTreeEditorData::FixDuplicateIDs()
 					FoundNodeIDs.Add(Node.ID, &bIsAlreadyInSet);
 					if (bIsAlreadyInSet)
 					{
-						Node.ID = FGuid::NewGuid();
+						Node.ID = UE::StateTree::PropertyHelpers::MakeDeterministicID(State, TEXT("TransitionConditions"), ((uint64)TransitionIndex << 32) | (uint64)Index);
+
 						UE_LOG(LogStateTreeEditor, Log, TEXT("%s: Found transition condition '%s' with duplicate ID on state '%s', changing ID:%s to ID:%s."),
 							*Self.GetFullName(), *Node.GetName().ToString(), *GetNameSafe(&State), *OldID.ToString(), *Node.ID.ToString());
 						EditorBindings.CopyBindings(OldID, Node.ID);
diff --git a/Engine/Plugins/Runtime/StateTree/Source/StateTreeEditorModule/Private/StateTreePropertyHelpers.cpp b/Engine/Plugins/Runtime/StateTree/Source/StateTreeEditorModule/Private/StateTreePropertyHelpers.cpp
--- a/Engine/Plugins/Runtime/StateTree/Source/StateTreeEditorModule/Private/StateTreePropertyHelpers.cpp
+++ b/Engine/Plugins/Runtime/StateTree/Source/StateTreeEditorModule/Private/StateTreePropertyHelpers.cpp
@@ -2,6 +2,8 @@
 
 #include "StateTreePropertyHelpers.h"
 #include "StateTreeEditorNode.h"
+#include "Hash/Blake3.h"
+#include "Misc/StringBuilder.h"
 
 namespace UE::StateTree::PropertyHelpers
 {
@@ -114,6 +116,29 @@ void DispatchPostEditToNodes(UObject& Owner, FPropertyChangedChainEvent& InPrope
 	}
 }
 
+
+FGuid MakeDeterministicID(const UObject& Owner, const FString& PropertyPath, const uint64 Seed)
+{
+	// From FGuid::NewDeterministicGuid(FStringView ObjectPath, uint64 Seed)
+	
+	// Convert the objectpath to utf8 so that whether TCHAR is UTF8 or UTF16 does not alter the hash.
+	TUtf8StringBuilder<1024> Utf8ObjectPath(InPlace, Owner.GetPathName());
+	TUtf8StringBuilder<1024> Utf8PropertyPath(InPlace, PropertyPath);
+
+	FBlake3 Builder;
+
+	// Hash this as the namespace of the Version 3 UUID, to avoid collisions with any other guids created using Blake3.
+	static FGuid BaseVersion(TEXT("bf324a38-a445-45a4-8921-249554b58189"));
+	Builder.Update(&BaseVersion, sizeof(FGuid));
+	Builder.Update(Utf8ObjectPath.GetData(), Utf8ObjectPath.Len() * sizeof(UTF8CHAR));
+	Builder.Update(Utf8PropertyPath.GetData(), Utf8PropertyPath.Len() * sizeof(UTF8CHAR));
+	Builder.Update(&Seed, sizeof(Seed));
+
+	const FBlake3Hash Hash = Builder.Finalize();
+
+	return FGuid::NewGuidFromHash(Hash);
+}
+
 }; // UE::StateTree::PropertyHelpers
 
 // ------------------------------------------------------------------------------
diff --git a/Engine/Plugins/Runtime/StateTree/Source/StateTreeEditorModule/Private/StateTreePropertyHelpers.h b/Engine/Plugins/Runtime/StateTree/Source/StateTreeEditorModule/Private/StateTreePropertyHelpers.h
--- a/Engine/Plugins/Runtime/StateTree/Source/StateTreeEditorModule/Private/StateTreePropertyHelpers.h
+++ b/Engine/Plugins/Runtime/StateTree/Source/StateTreeEditorModule/Private/StateTreePropertyHelpers.h
@@ -19,6 +19,9 @@ namespace UE::StateTree::PropertyHelpers {
  */
 void DispatchPostEditToNodes(UObject& Owner, FPropertyChangedChainEvent& PropertyChangedEvent);
 
+/** Makes deterministic ID from the owners property path, a property path (or any string), and a seed value (e.g. array index). */
+FGuid MakeDeterministicID(const UObject& Owner, const FString& PropertyPath, const uint64 Seed);
+
 /**
  * Gets a struct value from property handle, checks type before access. Expects T is struct.
  * @param ValueProperty Handle to property where value is got from.
