[StateTreeDebugger] added trace event to make it more obvious that global task completion completes the tree

diff --git a/Engine/Plugins/Runtime/StateTree/Source/StateTreeModule/Private/StateTreeExecutionContext.cpp b/Engine/Plugins/Runtime/StateTree/Source/StateTreeModule/Private/StateTreeExecutionContext.cpp
--- a/Engine/Plugins/Runtime/StateTree/Source/StateTreeModule/Private/StateTreeExecutionContext.cpp
+++ b/Engine/Plugins/Runtime/StateTree/Source/StateTreeModule/Private/StateTreeExecutionContext.cpp
@@ -569,7 +569,7 @@ EStateTreeRunStatus FStateTreeExecutionContext::Tick(const float DeltaTime)
 	}
 	else
 	{
-		// If global tasks succeed or fail, stop the tree.
+		STATETREE_TRACE_LOG_EVENT(TEXT("Global tasks completed (%s), stopping the tree"), *UEnum::GetDisplayValueAsText(EvalAndGlobalTaskStatus).ToString());
 		Exec.RequestedStop = EvalAndGlobalTaskStatus;
 	}
 
