[StateTreeDebugger] added trace event for conditions forced to a given result, either from an internal failure to access source data or from the evaluation mode set.

diff --git a/Engine/Plugins/Runtime/StateTree/Source/StateTreeEditorModule/Private/Debugger/SStateTreeDebuggerViewRow.cpp b/Engine/Plugins/Runtime/StateTree/Source/StateTreeEditorModule/Private/Debugger/SStateTreeDebuggerViewRow.cpp
--- a/Engine/Plugins/Runtime/StateTree/Source/StateTreeEditorModule/Private/Debugger/SStateTreeDebuggerViewRow.cpp
+++ b/Engine/Plugins/Runtime/StateTree/Source/StateTreeEditorModule/Private/Debugger/SStateTreeDebuggerViewRow.cpp
@@ -124,10 +124,13 @@ TSharedPtr<SWidget> SStateTreeDebuggerViewRow::GenerateEventWidget() const
 
 		switch (ConditionEvent.EventType)
 		{
-		case EStateTreeTraceEventType::Passed:			Image = StyleSet.GetBrush("StateTreeEditor.Debugger.Condition.Passed");			break;
-		case EStateTreeTraceEventType::Failed:			Image = StyleSet.GetBrush("StateTreeEditor.Debugger.Condition.Failed");			break;
-		case EStateTreeTraceEventType::OnEvaluating:	Image = StyleSet.GetBrush("StateTreeEditor.Debugger.Condition.OnEvaluating");	break;
-		case EStateTreeTraceEventType::OnTransition:	Image = StyleSet.GetBrush("StateTreeEditor.Debugger.Condition.OnTransition");	break;
+		case EStateTreeTraceEventType::Passed:					Image = StyleSet.GetBrush("StateTreeEditor.Debugger.Condition.Passed");			break;
+		case EStateTreeTraceEventType::ForcedSuccess:			Image = StyleSet.GetBrush("StateTreeEditor.Debugger.Condition.Passed");			break;
+		case EStateTreeTraceEventType::Failed:					Image = StyleSet.GetBrush("StateTreeEditor.Debugger.Condition.Failed");			break;
+		case EStateTreeTraceEventType::ForcedFailure:			Image = StyleSet.GetBrush("StateTreeEditor.Debugger.Condition.Failed");			break;
+		case EStateTreeTraceEventType::InternalForcedFailure:	Image = StyleSet.GetBrush("StateTreeEditor.Debugger.Condition.Failed");			break;
+		case EStateTreeTraceEventType::OnEvaluating:			Image = StyleSet.GetBrush("StateTreeEditor.Debugger.Condition.OnEvaluating");	break;
+		case EStateTreeTraceEventType::OnTransition:			Image = StyleSet.GetBrush("StateTreeEditor.Debugger.Condition.OnTransition");	break;
 		default:
 			Image = StyleSet.GetBrush("StateTreeEditor.Debugger.Unset");
 		}
diff --git a/Engine/Plugins/Runtime/StateTree/Source/StateTreeModule/Private/StateTreeExecutionContext.cpp b/Engine/Plugins/Runtime/StateTree/Source/StateTreeModule/Private/StateTreeExecutionContext.cpp
--- a/Engine/Plugins/Runtime/StateTree/Source/StateTreeModule/Private/StateTreeExecutionContext.cpp
+++ b/Engine/Plugins/Runtime/StateTree/Source/StateTreeModule/Private/StateTreeExecutionContext.cpp
@@ -2291,6 +2291,8 @@ bool FStateTreeExecutionContext::TestAllConditions(const FStateTreeExecutionFram
 				if (!CopyBatchWithValidation(CurrentParentFrame, CurrentFrame, ConditionInstanceView, Cond.BindingsBatch))
 				{
 					// If the source data cannot be accessed, the whole expression evaluates to false.
+					STATETREE_TRACE_CONDITION_EVENT(ConditionIndex, ConditionInstanceView, EStateTreeTraceEventType::InternalForcedFailure);
+					STATETREE_TRACE_LOG_EVENT(TEXT("Evaluation forced to false: source data cannot be accessed (e.g. enter conditions trying to access inactive parent state)"));
 					Values[0] = false;
 					break;
 				}
@@ -2307,8 +2309,8 @@ bool FStateTreeExecutionContext::TestAllConditions(const FStateTreeExecutionFram
 		}
 		else
 		{
-			bValue = Cond.EvaluationMode == EStateTreeConditionEvaluationMode::ForcedTrue ? true : /* EStateTreeConditionEvaluationMode::AlwaysFalse */ false;
-			STATETREE_TRACE_CONDITION_EVENT(ConditionIndex, FStateTreeDataView{}, bValue ? EStateTreeTraceEventType::Passed : EStateTreeTraceEventType::Failed);
+			bValue = Cond.EvaluationMode == EStateTreeConditionEvaluationMode::ForcedTrue;
+			STATETREE_TRACE_CONDITION_EVENT(ConditionIndex, FStateTreeDataView{}, bValue ? EStateTreeTraceEventType::ForcedSuccess : EStateTreeTraceEventType::ForcedFailure);
 		}
 
 		const int32 DeltaIndent = Cond.DeltaIndent;
diff --git a/Engine/Plugins/Runtime/StateTree/Source/StateTreeModule/Public/Debugger/StateTreeTraceTypes.h b/Engine/Plugins/Runtime/StateTree/Source/StateTreeModule/Public/Debugger/StateTreeTraceTypes.h
--- a/Engine/Plugins/Runtime/StateTree/Source/StateTreeModule/Public/Debugger/StateTreeTraceTypes.h
+++ b/Engine/Plugins/Runtime/StateTree/Source/StateTreeModule/Public/Debugger/StateTreeTraceTypes.h
@@ -17,23 +17,27 @@ UENUM()
 enum class EStateTreeTraceEventType : uint8
 {
 	Unset,
-	OnEntering			UMETA(DisplayName = "Entering"),
-	OnEntered			UMETA(DisplayName = "Entered"),
-	OnExiting			UMETA(DisplayName = "Exiting"),
-	OnExited			UMETA(DisplayName = "Exited"),
-	Push				UMETA(DisplayName = "Push"),
-	Pop					UMETA(DisplayName = "Pop"),
-	OnStateSelected		UMETA(DisplayName = "Selected"),
-	OnStateCompleted	UMETA(DisplayName = "Completed"),
-	OnTicking			UMETA(DisplayName = "Tick"),
-	OnTaskCompleted		UMETA(DisplayName = "Completed"),
-	OnTicked			UMETA(DisplayName = "Ticked"),
-	Passed				UMETA(DisplayName = "Passed"),
-	Failed				UMETA(DisplayName = "Failed"),
-	OnEvaluating		UMETA(DisplayName = "Evaluating"),
-	OnTransition		UMETA(DisplayName = "Transition"),
-	OnTreeStarted		UMETA(DisplayName = "Tree Started"),
-	OnTreeStopped		UMETA(DisplayName = "Tree Stopped")
+	OnEntering				UMETA(DisplayName = "Entering"),
+	OnEntered				UMETA(DisplayName = "Entered"),
+	OnExiting				UMETA(DisplayName = "Exiting"),
+	OnExited				UMETA(DisplayName = "Exited"),
+	Push					UMETA(DisplayName = "Push"),
+	Pop						UMETA(DisplayName = "Pop"),
+	OnStateSelected			UMETA(DisplayName = "Selected"),
+	OnStateCompleted		UMETA(DisplayName = "Completed"),
+	OnTicking				UMETA(DisplayName = "Tick"),
+	OnTaskCompleted			UMETA(DisplayName = "Completed"),
+	OnTicked				UMETA(DisplayName = "Ticked"),
+	Passed					UMETA(DisplayName = "Passed"),
+	Failed					UMETA(DisplayName = "Failed"),
+	ForcedSuccess			UMETA(DisplayName = "Forced Success"),
+	ForcedFailure			UMETA(DisplayName = "Forced Failure"),
+	InternalForcedFailure	UMETA(DisplayName = "Internal Forced Failure"),
+	OnEvaluating			UMETA(DisplayName = "Evaluating"),
+	OnTransition			UMETA(DisplayName = "Transition"),
+	OnTreeStarted			UMETA(DisplayName = "Tree Started"),
+	OnTreeStopped			UMETA(DisplayName = "Tree Stopped")
+
 };
 
 #if WITH_STATETREE_DEBUGGER
