StateTree: Fixed crash in FStateTreeReference when no parameters are present

diff --git a/Engine/Plugins/Runtime/StateTree/Source/StateTreeModule/Private/StateTreeReference.cpp b/Engine/Plugins/Runtime/StateTree/Source/StateTreeModule/Private/StateTreeReference.cpp
--- a/Engine/Plugins/Runtime/StateTree/Source/StateTreeModule/Private/StateTreeReference.cpp
+++ b/Engine/Plugins/Runtime/StateTree/Source/StateTreeModule/Private/StateTreeReference.cpp
@@ -19,12 +19,14 @@ void FStateTreeReference::SyncParameters()
 		// Remove overrides that do not exists anymore
 		if (!PropertyOverrides.IsEmpty())
 		{
-			const UPropertyBag* Bag = Parameters.GetPropertyBagStruct();
-			for (TArray<FGuid>::TIterator It = PropertyOverrides.CreateIterator(); It; ++It)
+			if (const UPropertyBag* Bag = Parameters.GetPropertyBagStruct())
 			{
-				if (!Bag->FindPropertyDescByID(*It))
+				for (TArray<FGuid>::TIterator It = PropertyOverrides.CreateIterator(); It; ++It)
 				{
-					It.RemoveCurrentSwap();
+					if (!Bag->FindPropertyDescByID(*It))
+					{
+						It.RemoveCurrentSwap();
+					}
 				}
 			}
 		}
