StateTree: Added property references (get pointer to a property based on binding)
- Added FStateTreePropertyRef which allows to get pointer to bindable output properties in the StateTree
- This will eventually supercede FStateTreeStructRef

diff --git a/Engine/Plugins/Runtime/StateTree/Source/StateTreeEditorModule/Private/Customizations/StateTreeBindingExtension.cpp b/Engine/Plugins/Runtime/StateTree/Source/StateTreeEditorModule/Private/Customizations/StateTreeBindingExtension.cpp
--- a/Engine/Plugins/Runtime/StateTree/Source/StateTreeEditorModule/Private/Customizations/StateTreeBindingExtension.cpp
+++ b/Engine/Plugins/Runtime/StateTree/Source/StateTreeEditorModule/Private/Customizations/StateTreeBindingExtension.cpp
@@ -11,6 +11,8 @@
 #include "Styling/AppStyle.h"
 #include "UObject/EnumProperty.h"
 #include "Widgets/Layout/SBox.h"
+#include "StateTreePropertyRef.h"
+#include "StateTreePropertyRefHelpers.h"
 
 #define LOCTEXT_NAMESPACE "StateTreeEditor"
 
@@ -186,7 +188,7 @@ EStateTreePropertyUsage MakeStructPropertyPathFromPropertyHandle(TSharedPtr<cons
 			if (const FString* IDString = CurrentPropertyHandle->GetInstanceMetaData(UE::StateTree::PropertyBinding::StateTreeNodeIDName))
 			{
 				LexFromString(StructID, **IDString);
-				ResultUsage = UE::StateTree::Compiler::GetUsageFromMetaData(Property);
+				ResultUsage = UE::StateTree::GetUsageFromMetaData(Property);
 				break;
 			}
 		}
@@ -390,7 +392,16 @@ struct FCachedBindingData : public TSharedFromThis<FCachedBindingData>
 		check(Schema);
 
 		FEdGraphPinType PinType;
-		Schema->ConvertPropertyToPinType(Property, PinType);
+
+		if (UE::StateTree::PropertyRefHelpers::IsPropertyRef(*Property))
+		{
+			// Use internal type to construct PinType if it's property of PropertyRef type.
+			PinType = UE::StateTree::PropertyRefHelpers::GetPropertyRefInternalTypeAsPin(*Property);
+		}
+		else
+		{
+			Schema->ConvertPropertyToPinType(Property, PinType);
+		}
 
 		if (const FStateTreePropertyPath* SourcePath = EditorBindings->GetPropertyBindingSource(TargetPath))
 		{
@@ -553,6 +564,24 @@ struct FCachedBindingData : public TSharedFromThis<FCachedBindingData>
 		return false;
 	}
 
+	bool CanAcceptPropertyOrChildren(const FProperty* SourceProperty, TConstArrayView<TSharedPtr<FBindingChainElement>> InBindingChain)
+	{
+		ConditionallyUpdateData();
+
+		if (UE::StateTree::PropertyRefHelpers::IsPropertyRef(*PropertyHandle->GetProperty()))
+		{
+			const int32 SourceStructIndex = InBindingChain[0]->ArrayIndex;
+			check(AccessibleStructs.IsValidIndex(SourceStructIndex));
+
+			if (!UE::StateTree::PropertyRefHelpers::IsPropertyAccessibleForPropertyRef(*SourceProperty, InBindingChain, AccessibleStructs[SourceStructIndex]))
+			{
+				return false;
+			}
+		}
+
+		return SourceProperty->HasAnyPropertyFlags(CPF_Edit);
+	}
+
 	static bool ArePropertyAndContextStructCompatible(const UStruct* SourceStruct, const FProperty* TargetProperty)
 	{
 		if (const FStructProperty* TargetStructProperty = CastField<FStructProperty>(TargetProperty))
@@ -619,6 +648,11 @@ struct FCachedBindingData : public TSharedFromThis<FCachedBindingData>
 				}
 			}
 		}
+		else if (TargetStructProperty && TargetStructProperty->Struct == FStateTreePropertyRef::StaticStruct())
+		{
+			check(TargetPropertyValue);
+			bCanBind = UE::StateTree::PropertyRefHelpers::IsPropertyRefCompatibleWithProperty(*TargetStructProperty, *SourceProperty);
+		}
 		else
 		{
 			// Note: We support type promotion here
@@ -798,6 +832,7 @@ void FStateTreeBindingExtension::ExtendWidgetRow(FDetailWidgetRow& InWidgetRow,
 		{
 			EditorBindings = BindingOwner->GetPropertyEditorBindings();
 			BindingOwner->GetAccessibleStructs(TargetPath.GetStructID(), AccessibleStructs);
+
 			for (FStateTreeBindableStructDesc& StructDesc : AccessibleStructs)
 			{
 				const UStruct* Struct = StructDesc.Struct;
@@ -845,10 +880,9 @@ void FStateTreeBindingExtension::ExtendWidgetRow(FDetailWidgetRow& InWidgetRow,
 			return CachedBindingData->CanBindToContextStruct(InStruct);
 		});
 
-	Args.OnCanAcceptPropertyOrChildrenWithBindingChain = FOnCanAcceptPropertyOrChildrenWithBindingChain::CreateLambda([](FProperty* InProperty, TConstArrayView<TSharedPtr<FBindingChainElement>> InBindingChain)
+	Args.OnCanAcceptPropertyOrChildrenWithBindingChain = FOnCanAcceptPropertyOrChildrenWithBindingChain::CreateLambda([CachedBindingData](FProperty* InProperty, TConstArrayView<TSharedPtr<FBindingChainElement>> InBindingChain)
 		{
-			// Make only editor visible properties visible for binding.
-			return InProperty->HasAnyPropertyFlags(CPF_Edit);
+			return CachedBindingData->CanAcceptPropertyOrChildren(InProperty, InBindingChain);
 		});
 
 	Args.OnCanBindToClass = FOnCanBindToClass::CreateLambda([](UClass* InClass)
diff --git a/Engine/Plugins/Runtime/StateTree/Source/StateTreeEditorModule/Private/Customizations/StateTreeEditorNodeDetails.cpp b/Engine/Plugins/Runtime/StateTree/Source/StateTreeEditorModule/Private/Customizations/StateTreeEditorNodeDetails.cpp
--- a/Engine/Plugins/Runtime/StateTree/Source/StateTreeEditorModule/Private/Customizations/StateTreeEditorNodeDetails.cpp
+++ b/Engine/Plugins/Runtime/StateTree/Source/StateTreeEditorModule/Private/Customizations/StateTreeEditorNodeDetails.cpp
@@ -8,6 +8,7 @@
 #include "StateTree.h"
 #include "StateTreeEditor.h"
 #include "StateTreeEditorData.h"
+#include "StateTreePropertyRef.h"
 #include "Widgets/Input/SComboButton.h"
 #include "Framework/MultiBox/MultiBoxBuilder.h"
 #include "Styling/SlateIconFinder.h"
@@ -145,7 +146,7 @@ namespace UE::StateTreeEditor::Internal
 		TSharedPtr<IPropertyHandle> ChildPropHandle = ChildRow.GetPropertyHandle();
 		check(ChildPropHandle.IsValid());
 		
-		const EStateTreePropertyUsage Usage = UE::StateTree::Compiler::GetUsageFromMetaData(ChildPropHandle->GetProperty());
+		const EStateTreePropertyUsage Usage = UE::StateTree::GetUsageFromMetaData(ChildPropHandle->GetProperty());
 		const FProperty* Property = ChildPropHandle->GetProperty();
 		
 		// Conditionally control visibility of the value field of bound properties.
@@ -174,7 +175,17 @@ namespace UE::StateTreeEditor::Internal
 				
 				FEdGraphPinType PinType;
 				const UEdGraphSchema_K2* Schema = GetDefault<UEdGraphSchema_K2>();
-				Schema->ConvertPropertyToPinType(Property, PinType);
+
+				// Show referenced type for property refs.
+				if (UE::StateTree::PropertyRefHelpers::IsPropertyRef(*Property))
+				{
+					// Use internal type to construct PinType if it's property of PropertyRef type.
+					PinType = UE::StateTree::PropertyRefHelpers::GetPropertyRefInternalTypeAsPin(*Property);
+				}
+				else
+				{
+					Schema->ConvertPropertyToPinType(Property, PinType);
+				}
 				
 				const FSlateBrush* Icon = FBlueprintEditorUtils::GetIconFromPin(PinType, true);
 				FText Text = GetPinTypeText(PinType);
@@ -853,7 +864,7 @@ void FStateTreeEditorNodeDetails::CustomizeChildren(TSharedRef<class IPropertyHa
 			{
 				FSortedChild Child;
 				Child.PropertyHandle = ChildHandle;
-				Child.Usage = UE::StateTree::Compiler::GetUsageFromMetaData(Child.PropertyHandle->GetProperty());
+				Child.Usage = UE::StateTree::GetUsageFromMetaData(Child.PropertyHandle->GetProperty());
 
 				// If the property is set to one of these usages, display it even if it is not edit on instance.
 				// It is a common mistake to forget to set the "eye" on these properties it and wonder why it does not show up.
diff --git a/Engine/Plugins/Runtime/StateTree/Source/StateTreeEditorModule/Private/Debugger/SStateTreeDebuggerView.cpp b/Engine/Plugins/Runtime/StateTree/Source/StateTreeEditorModule/Private/Debugger/SStateTreeDebuggerView.cpp
--- a/Engine/Plugins/Runtime/StateTree/Source/StateTreeEditorModule/Private/Debugger/SStateTreeDebuggerView.cpp
+++ b/Engine/Plugins/Runtime/StateTree/Source/StateTreeEditorModule/Private/Debugger/SStateTreeDebuggerView.cpp
@@ -162,7 +162,7 @@ void GenerateElementForProperties(const TCHAR* TypeAsText, const TCHAR* ValueAsT
 		{
 			const FProperty* const Property = PropertyIt.Key();
 			check(Property);
-			const EStateTreePropertyUsage Usage = UE::StateTree::Compiler::GetUsageFromMetaData(Property);
+			const EStateTreePropertyUsage Usage = UE::StateTree::GetUsageFromMetaData(Property);
 
 			// If the property is set to one of these usages, display it even if it is not edit on instance.
 			// It is a common mistake to forget to set the "eye" on these properties it and wonder why it does not show up.
diff --git a/Engine/Plugins/Runtime/StateTree/Source/StateTreeEditorModule/Private/StateTreeCompiler.cpp b/Engine/Plugins/Runtime/StateTree/Source/StateTreeEditorModule/Private/StateTreeCompiler.cpp
--- a/Engine/Plugins/Runtime/StateTree/Source/StateTreeEditorModule/Private/StateTreeCompiler.cpp
+++ b/Engine/Plugins/Runtime/StateTree/Source/StateTreeEditorModule/Private/StateTreeCompiler.cpp
@@ -10,6 +10,9 @@
 #include "StateTreeConditionBase.h"
 #include "Serialization/ArchiveUObject.h"
 #include "GameFramework/Actor.h"
+#include "StateTreePropertyRef.h"
+#include "StateTreePropertyRefHelpers.h"
+#include "StateTreePropertyHelpers.h"
 
 
 namespace UE::StateTree::Compiler
@@ -177,41 +180,6 @@ namespace UE::StateTree::Compiler
 		return Result;
 	}
 
-	EStateTreePropertyUsage GetUsageFromMetaData(const FProperty* Property)
-	{
-		static const FName CategoryName(TEXT("Category"));
-
-		if (Property == nullptr)
-		{
-			return EStateTreePropertyUsage::Invalid;
-		}
-		
-		const FString Category = Property->GetMetaData(CategoryName);
-
-		if (Category == TEXT("Input"))
-		{
-			return EStateTreePropertyUsage::Input;
-		}
-		if (Category == TEXT("Inputs"))
-		{
-			return EStateTreePropertyUsage::Input;
-		}
-		if (Category == TEXT("Output"))
-		{
-			return EStateTreePropertyUsage::Output;
-		}
-		if (Category == TEXT("Outputs"))
-		{
-			return EStateTreePropertyUsage::Output;
-		}
-		if (Category == TEXT("Context"))
-		{
-			return EStateTreePropertyUsage::Context;
-		}
-
-		return EStateTreePropertyUsage::Parameter;
-	}
-
 }; // UE::StateTree::Compiler
 
 bool FStateTreeCompiler::Compile(UStateTree& InStateTree)
@@ -649,14 +617,20 @@ bool FStateTreeCompiler::CreateStateTasksAndParameters()
 		BindingsCompiler.AddSourceStruct(LinkedParamsDesc);
 
 		// Check that the bindings for this struct are still all valid.
-		TArray<FStateTreePropertyPathBinding> Bindings;
-		if (!GetAndValidateBindings(LinkedParamsDesc, FStateTreeDataView(CompactStateTreeParameters.Parameters.GetMutableValue()), Bindings))
+		TArray<FStateTreePropertyPathBinding> CopyBindings;
+		TArray<FStateTreePropertyPathBinding> ReferenceBindings;
+		if (!GetAndValidateBindings(LinkedParamsDesc, FStateTreeDataView(CompactStateTreeParameters.Parameters.GetMutableValue()), CopyBindings, ReferenceBindings))
 		{
 			return false;
 		}
 
 		int32 BatchIndex = INDEX_NONE;
-		if (!BindingsCompiler.CompileBatch(LinkedParamsDesc, Bindings, BatchIndex))
+		if (!BindingsCompiler.CompileBatch(LinkedParamsDesc, CopyBindings, BatchIndex))
+		{
+			return false;
+		}
+
+		if (!BindingsCompiler.CompileReferences(LinkedParamsDesc, ReferenceBindings, FStateTreeDataView(CompactStateTreeParameters.Parameters.GetMutableValue())))
 		{
 			return false;
 		}
@@ -1085,15 +1059,21 @@ bool FStateTreeCompiler::CreateCondition(UStateTreeState& State, const FStateTre
 	BindingsCompiler.AddSourceStruct(StructDesc);
 
 	// Check that the bindings for this struct are still all valid.
-	TArray<FStateTreePropertyPathBinding> Bindings;
-	if (!GetAndValidateBindings(StructDesc, InstanceDataView, Bindings))
+	TArray<FStateTreePropertyPathBinding> CopyBindings;
+	TArray<FStateTreePropertyPathBinding> ReferenceBindings;
+	if (!GetAndValidateBindings(StructDesc, InstanceDataView, CopyBindings, ReferenceBindings))
 	{
 		return false;
 	}
 
 	// Compile batch copy for this struct, we pass in all the bindings, the compiler will pick up the ones for the target structs.
 	int32 BatchIndex = INDEX_NONE;
-	if (!BindingsCompiler.CompileBatch(StructDesc, Bindings, BatchIndex))
+	if (!BindingsCompiler.CompileBatch(StructDesc, CopyBindings, BatchIndex))
+	{
+		return false;
+	}
+
+	if (!BindingsCompiler.CompileReferences(StructDesc, ReferenceBindings, InstanceDataView))
 	{
 		return false;
 	}
@@ -1254,15 +1234,21 @@ bool FStateTreeCompiler::CreateTask(UStateTreeState* State, const FStateTreeEdit
 	BindingsCompiler.AddSourceStruct(StructDesc);
 	
 	// Check that the bindings for this struct are still all valid.
-	TArray<FStateTreePropertyPathBinding> Bindings;
-	if (!GetAndValidateBindings(StructDesc, InstanceDataView, Bindings))
+	TArray<FStateTreePropertyPathBinding> CopyBindings;
+	TArray<FStateTreePropertyPathBinding> ReferenceBindings;
+	if (!GetAndValidateBindings(StructDesc, InstanceDataView, CopyBindings, ReferenceBindings))
 	{
 		return false;
 	}
 
 	// Compile batch copy for this struct, we pass in all the bindings, the compiler will pick up the ones for the target structs.
 	int32 BatchIndex = INDEX_NONE;
-	if (!BindingsCompiler.CompileBatch(StructDesc, Bindings, BatchIndex))
+	if (!BindingsCompiler.CompileBatch(StructDesc, CopyBindings, BatchIndex))
+	{
+		return false;
+	}
+
+	if (!BindingsCompiler.CompileReferences(StructDesc, ReferenceBindings, InstanceDataView))
 	{
 		return false;
 	}
@@ -1361,15 +1347,21 @@ bool FStateTreeCompiler::CreateEvaluator(const FStateTreeEditorNode& EvalNode, c
 	BindingsCompiler.AddSourceStruct(StructDesc);
 
 	// Check that the bindings for this struct are still all valid.
-	TArray<FStateTreePropertyPathBinding> Bindings;
-	if (!GetAndValidateBindings(StructDesc, InstanceDataView, Bindings))
+	TArray<FStateTreePropertyPathBinding> CopyBindings;
+	TArray<FStateTreePropertyPathBinding> ReferenceBindings;
+	if (!GetAndValidateBindings(StructDesc, InstanceDataView, CopyBindings, ReferenceBindings))
 	{
 		return false;
 	}
 
 	// Compile batch copy for this struct, we pass in all the bindings, the compiler will pick up the ones for the target structs.
 	int32 BatchIndex = INDEX_NONE;
-	if (!BindingsCompiler.CompileBatch(StructDesc, Bindings, BatchIndex))
+	if (!BindingsCompiler.CompileBatch(StructDesc, CopyBindings, BatchIndex))
+	{
+		return false;
+	}
+
+	if (!BindingsCompiler.CompileReferences(StructDesc, ReferenceBindings, InstanceDataView))
 	{
 		return false;
 	}
@@ -1384,10 +1376,8 @@ bool FStateTreeCompiler::CreateEvaluator(const FStateTreeEditorNode& EvalNode, c
 	return true;
 }
 
-bool FStateTreeCompiler::IsPropertyAnyEnum(const FStateTreeBindableStructDesc& Struct, FStateTreePropertyPath Path) const
+bool FStateTreeCompiler::IsPropertyOfType(UScriptStruct& Type, const FStateTreeBindableStructDesc& Struct, FStateTreePropertyPath Path) const
 {
-	bool bIsAnyEnum = false;
-
 	TArray<FStateTreePropertyPathIndirection> Indirection;
 	const bool bResolved = Path.ResolveIndirections(Struct.Struct, Indirection);
 	
@@ -1398,11 +1388,11 @@ bool FStateTreeCompiler::IsPropertyAnyEnum(const FStateTreeBindableStructDesc& S
 		{
 			if (const FStructProperty* OwnerStructProperty = CastField<FStructProperty>(OwnerProperty))
 			{
-				bIsAnyEnum = OwnerStructProperty->Struct == TBaseStructure<FStateTreeAnyEnum>::Get();
+				return OwnerStructProperty->Struct == &Type;
 			}
 		}
 	}
-	return bIsAnyEnum;
+	return false;
 }
 
 bool FStateTreeCompiler::ValidateStructRef(const FStateTreeBindableStructDesc& SourceStruct, FStateTreePropertyPath SourcePath,
@@ -1491,11 +1481,12 @@ bool FStateTreeCompiler::ValidateStructRef(const FStateTreeBindableStructDesc& S
 }
 
 
-bool FStateTreeCompiler::GetAndValidateBindings(const FStateTreeBindableStructDesc& TargetStruct, FStateTreeDataView TargetValue, TArray<FStateTreePropertyPathBinding>& OutBindings) const
+bool FStateTreeCompiler::GetAndValidateBindings(const FStateTreeBindableStructDesc& TargetStruct, FStateTreeDataView TargetValue, TArray<FStateTreePropertyPathBinding>& OutCopyBindings, TArray<FStateTreePropertyPathBinding>& OutReferenceBindings) const
 {
 	check(EditorData);
 	
-	OutBindings.Reset();
+	OutCopyBindings.Reset();
+	OutReferenceBindings.Reset();
 
 	// If target struct is not set, nothing to do.
 	if (TargetStruct.Struct == nullptr)
@@ -1586,8 +1577,8 @@ bool FStateTreeCompiler::GetAndValidateBindings(const FStateTreeBindableStructDe
 		// Special case fo AnyEnum. StateTreeBindingExtension allows AnyEnums to bind to other enum types.
 		// The actual copy will be done via potential type promotion copy, into the value property inside the AnyEnum.
 		// We amend the paths here to point to the 'Value' property.
-		const bool bSourceIsAnyEnum = IsPropertyAnyEnum(*SourceStruct, Binding.GetSourcePath());
-		const bool bTargetIsAnyEnum = IsPropertyAnyEnum(TargetStruct, Binding.GetTargetPath());
+		const bool bSourceIsAnyEnum = IsPropertyOfType(*TBaseStructure<FStateTreeAnyEnum>::Get(), *SourceStruct, Binding.GetSourcePath());
+		const bool bTargetIsAnyEnum = IsPropertyOfType(*TBaseStructure<FStateTreeAnyEnum>::Get(), TargetStruct, Binding.GetTargetPath());
 		if (bSourceIsAnyEnum || bTargetIsAnyEnum)
 		{
 			if (bSourceIsAnyEnum)
@@ -1599,9 +1590,16 @@ bool FStateTreeCompiler::GetAndValidateBindings(const FStateTreeBindableStructDe
 				BindingCopy.GetMutableTargetPath().AddPathSegment(GET_MEMBER_NAME_STRING_CHECKED(FStateTreeAnyEnum, Value));
 			}
 		}
-		
-		OutBindings.Add(BindingCopy);
 
+		if (IsPropertyOfType(*FStateTreePropertyRef::StaticStruct(), TargetStruct, Binding.GetTargetPath()))
+		{
+			OutReferenceBindings.Add(BindingCopy);
+		}
+		else
+		{
+			OutCopyBindings.Add(BindingCopy);
+		}
+		
 		// Check if the bindings is for struct ref and validate the types.
 		if (!ValidateStructRef(*SourceStruct, Binding.GetSourcePath(), TargetStruct, Binding.GetTargetPath()))
 		{
@@ -1610,9 +1608,9 @@ bool FStateTreeCompiler::GetAndValidateBindings(const FStateTreeBindableStructDe
 	}
 
 
-	auto IsPropertyBound = [&OutBindings](const FName& PropertyName)
+	auto IsPropertyBound = [](const FName& PropertyName, TConstArrayView<FStateTreePropertyPathBinding> Bindings)
 	{
-		return OutBindings.ContainsByPredicate([&PropertyName](const FStateTreePropertyPathBinding& Binding)
+		return Bindings.ContainsByPredicate([&PropertyName](const FStateTreePropertyPathBinding& Binding)
 			{
 				// We're looping over just the first level of properties on the struct, so we assume that the path is just one item
 				// (or two in case of AnyEnum, because we expand the path to Property.Value, see code above).
@@ -1628,59 +1626,72 @@ bool FStateTreeCompiler::GetAndValidateBindings(const FStateTreeBindableStructDe
 		const FProperty* Property = *It;
 		check(Property);
 		const FName PropertyName = Property->GetFName();
-		const EStateTreePropertyUsage Usage = UE::StateTree::Compiler::GetUsageFromMetaData(Property);
-		if (Usage == EStateTreePropertyUsage::Input)
+		const bool bIsOptional = UE::StateTree::PropertyHelpers::HasOptionalMetadata(*Property);
+
+		if (UE::StateTree::PropertyRefHelpers::IsPropertyRef(*Property))
 		{
-			const bool bIsOptional = Property->HasMetaData(TEXT("Optional"));
-			
-			// Make sure that an Input property is bound unless marked optional.
-			if (bIsOptional == false && !IsPropertyBound(PropertyName))
+			if (bIsOptional == false && !IsPropertyBound(PropertyName, OutReferenceBindings))
 			{
 				Log.Reportf(EMessageSeverity::Error, TargetStruct,
-					TEXT("Input property '%s' on %s is expected to have a binding."),
-					*PropertyName.ToString(), *TargetStruct.ToString());
-				bResult = false;
+						TEXT("Property reference '%s' on % s is expected to have a binding."),
+						*PropertyName.ToString(), *TargetStruct.ToString());
+					bResult = false;
 			}
 		}
-		else if (Usage == EStateTreePropertyUsage::Context)
+		else
 		{
-			// Make sure that an Context property is manually or automatically bound. 
-			const UStruct* ContextObjectType = nullptr; 
-			if (const FStructProperty* StructProperty = CastField<FStructProperty>(Property))
-			{
-				ContextObjectType = StructProperty->Struct;
-			}		
-			else if (const FObjectPropertyBase* ObjectProperty = CastField<FObjectPropertyBase>(Property))
+			const EStateTreePropertyUsage Usage = UE::StateTree::GetUsageFromMetaData(Property);
+			if (Usage == EStateTreePropertyUsage::Input)
 			{
-				ContextObjectType = ObjectProperty->PropertyClass;
-			}
-
-			if (ContextObjectType == nullptr)
-			{
-				Log.Reportf(EMessageSeverity::Error, TargetStruct,
-					TEXT("The type of Context property '%s' on %s is expected to be Object Reference or Struct."),
-					*PropertyName.ToString(), *TargetStruct.ToString());
-				bResult = false;
-				continue;
+				// Make sure that an Input property is bound unless marked optional.
+				if (bIsOptional == false && !IsPropertyBound(PropertyName, OutCopyBindings))
+				{
+					Log.Reportf(EMessageSeverity::Error, TargetStruct,
+						TEXT("Input property '%s' on %s is expected to have a binding."),
+						*PropertyName.ToString(), *TargetStruct.ToString());
+					bResult = false;
+				}
 			}
-
-			const bool bIsBound = IsPropertyBound(PropertyName);
-
-			if (!bIsBound)
+			else if (Usage == EStateTreePropertyUsage::Context)
 			{
-				const FStateTreeBindableStructDesc Desc = EditorData->FindContextData(ContextObjectType, PropertyName.ToString());
-
-				if (Desc.IsValid())
+				// Make sure that an Context property is manually or automatically bound. 
+				const UStruct* ContextObjectType = nullptr; 
+				if (const FStructProperty* StructProperty = CastField<FStructProperty>(Property))
 				{
-					// Add automatic binding to Context data.
-					OutBindings.Emplace(FStateTreePropertyPath(Desc.ID), FStateTreePropertyPath(TargetStruct.ID, PropertyName));
+					ContextObjectType = StructProperty->Struct;
+				}		
+				else if (const FObjectPropertyBase* ObjectProperty = CastField<FObjectPropertyBase>(Property))
+				{
+					ContextObjectType = ObjectProperty->PropertyClass;
 				}
-				else
+
+				if (ContextObjectType == nullptr)
 				{
 					Log.Reportf(EMessageSeverity::Error, TargetStruct,
-						TEXT("Could not find matching Context object for Context property '%s' on '%s'. Property must have manual binding."),
+						TEXT("The type of Context property '%s' on %s is expected to be Object Reference or Struct."),
 						*PropertyName.ToString(), *TargetStruct.ToString());
 					bResult = false;
+					continue;
+				}
+
+				const bool bIsBound = IsPropertyBound(PropertyName, OutCopyBindings);
+
+				if (!bIsBound)
+				{
+					const FStateTreeBindableStructDesc Desc = EditorData->FindContextData(ContextObjectType, PropertyName.ToString());
+
+					if (Desc.IsValid())
+					{
+						// Add automatic binding to Context data.
+						OutCopyBindings.Emplace(FStateTreePropertyPath(Desc.ID), FStateTreePropertyPath(TargetStruct.ID, PropertyName));
+					}
+					else
+					{
+						Log.Reportf(EMessageSeverity::Error, TargetStruct,
+							TEXT("Could not find matching Context object for Context property '%s' on '%s'. Property must have manual binding."),
+							*PropertyName.ToString(), *TargetStruct.ToString());
+						bResult = false;
+					}
 				}
 			}
 		}
diff --git a/Engine/Plugins/Runtime/StateTree/Source/StateTreeEditorModule/Private/StateTreePropertyBindingCompiler.cpp b/Engine/Plugins/Runtime/StateTree/Source/StateTreeEditorModule/Private/StateTreePropertyBindingCompiler.cpp
--- a/Engine/Plugins/Runtime/StateTree/Source/StateTreeEditorModule/Private/StateTreePropertyBindingCompiler.cpp
+++ b/Engine/Plugins/Runtime/StateTree/Source/StateTreeEditorModule/Private/StateTreePropertyBindingCompiler.cpp
@@ -6,6 +6,10 @@
 #include "StateTreeCompilerLog.h"
 #include "StateTreeEditorPropertyBindings.h"
 #include "Misc/EnumerateRange.h"
+#include "StateTreePropertyBindings.h"
+#include "StateTreePropertyRef.h"
+#include "StateTreePropertyRefHelpers.h"
+#include "StateTreePropertyHelpers.h"
 
 #include UE_INLINE_GENERATED_CPP_BY_NAME(StateTreePropertyBindingCompiler)
 
@@ -127,6 +131,121 @@ bool FStateTreePropertyBindingCompiler::CompileBatch(const FStateTreeBindableStr
 	return true;
 }
 
+bool FStateTreePropertyBindingCompiler::CompileReferences(const FStateTreeBindableStructDesc& TargetStruct, TConstArrayView<FStateTreePropertyPathBinding> PropertyReferenceBindings, FStateTreeDataView InstanceDataView)
+{
+	for (const FStateTreePropertyPathBinding& Binding : PropertyReferenceBindings)
+	{
+		if (Binding.GetTargetPath().GetStructID() != TargetStruct.ID)
+		{
+			continue;
+		}
+
+		// Source must be in the source array/
+		const FStateTreeBindableStructDesc* SourceStruct = GetSourceStructDescByID(Binding.GetSourcePath().GetStructID());
+		if (!SourceStruct)
+		{
+			Log->Reportf(EMessageSeverity::Error, TargetStruct,
+				TEXT("Could not find a binding source."));
+			return false;
+		}
+
+		FString Error;
+		TArray<FStateTreePropertyPathIndirection> SourceIndirections;
+		
+		if (!Binding.GetSourcePath().ResolveIndirections(SourceStruct->Struct, SourceIndirections, &Error))
+		{
+			Log->Reportf(EMessageSeverity::Error, TargetStruct, TEXT("Resolving path in %s: %s"), *SourceStruct->ToString(), *Error);
+			return false;
+		}
+
+		if (!UE::StateTree::PropertyRefHelpers::IsPropertyAccessibleForPropertyRef(SourceIndirections, *SourceStruct))
+		{
+			Log->Reportf(EMessageSeverity::Error, TargetStruct,
+					TEXT("%s cannot reference non-output %s "),
+					*UE::StateTree::GetDescAndPathAsString(TargetStruct, Binding.GetTargetPath()),
+					*UE::StateTree::GetDescAndPathAsString(*SourceStruct, Binding.GetSourcePath()));
+			return false;
+		}
+
+		TArray<FStateTreePropertyIndirection> TargetIndirections;
+		FStateTreePropertyIndirection TargetFirstIndirection;
+		FStateTreePropertyPathIndirection TargetLeafIndirection;
+		if (!FStateTreePropertyBindings::ResolvePath(InstanceDataView.GetStruct(), Binding.GetTargetPath(), TargetIndirections, TargetFirstIndirection, TargetLeafIndirection))
+		{
+			Log->Reportf(EMessageSeverity::Error, TargetStruct, TEXT("Resolving path in %s: %s"), *TargetStruct.ToString(), *Error);
+			return false;
+		}
+
+		if (!UE::StateTree::PropertyRefHelpers::IsPropertyRefCompatibleWithProperty(*TargetLeafIndirection.GetProperty(), *SourceIndirections.Last().GetProperty()))
+		{
+			Log->Reportf(EMessageSeverity::Error, TargetStruct,
+				TEXT("%s cannot reference %s, types are incompatible."),		
+				*UE::StateTree::GetDescAndPathAsString(TargetStruct, Binding.GetTargetPath()),
+				*UE::StateTree::GetDescAndPathAsString(*SourceStruct, Binding.GetSourcePath()));
+			return false;
+		}
+
+		FStateTreeIndex16 ReferenceIndex;
+
+		// Reuse the index if another PropertyRef already references the same property.
+		{
+			int32 IndexOfAlreadyExisting = PropertyBindings->PropertyReferencePaths.IndexOfByPredicate([&Binding](const FStateTreePropertyRefPath& RefPath)
+			{
+				return RefPath.GetSourcePath() == Binding.GetSourcePath();
+			});
+
+			if (IndexOfAlreadyExisting != INDEX_NONE)
+			{
+				ReferenceIndex = FStateTreeIndex16(IndexOfAlreadyExisting);
+			}
+		}
+
+		if (!ReferenceIndex.IsValid())
+		{
+			// If referencing another PropertyRef, reuse it's index.
+			if (UE::StateTree::PropertyRefHelpers::IsPropertyRef(*SourceIndirections.Last().GetProperty()))
+			{
+				const FCompiledReference* ReferencedReference = CompiledReferences.FindByPredicate([&Binding](const FCompiledReference& CompiledReference)
+				{
+					return CompiledReference.Path == Binding.GetSourcePath();
+				});
+
+				if (ReferencedReference)
+				{
+					ReferenceIndex = ReferencedReference->Index;
+				}
+				else
+				{
+					if(!UE::StateTree::PropertyHelpers::HasOptionalMetadata(*TargetLeafIndirection.GetProperty()))
+					{
+						Log->Reportf(EMessageSeverity::Error, TargetStruct, TEXT("Referenced %s is not bound"), *UE::StateTree::GetDescAndPathAsString(*SourceStruct, Binding.GetSourcePath()));
+						return false;
+					}
+							
+					return true;
+				}
+			}
+		}
+
+		if (!ReferenceIndex.IsValid())
+		{
+			ReferenceIndex = FStateTreeIndex16(PropertyBindings->PropertyReferencePaths.Num());
+			PropertyBindings->PropertyReferencePaths.Emplace(SourceStruct->DataHandle, Binding.GetSourcePath());
+		}
+
+		// Store index in instance data.
+		uint8* RawData = FStateTreePropertyBindings::GetAddress(InstanceDataView, TargetIndirections, TargetFirstIndirection, TargetLeafIndirection.GetProperty());
+		check(RawData);
+		reinterpret_cast<FStateTreePropertyRef*>(RawData)->RefAccessIndex = ReferenceIndex;
+
+		FCompiledReference& CompiledReference = CompiledReferences.AddDefaulted_GetRef();
+		CompiledReference.Path = Binding.GetTargetPath();
+		CompiledReference.Index = ReferenceIndex;
+	}
+
+	return true;
+}
+
 void FStateTreePropertyBindingCompiler::Finalize()
 {
 	StoreSourceStructs();
diff --git a/Engine/Plugins/Runtime/StateTree/Source/StateTreeEditorModule/Private/StateTreePropertyHelpers.cpp b/Engine/Plugins/Runtime/StateTree/Source/StateTreeEditorModule/Private/StateTreePropertyHelpers.cpp
--- a/Engine/Plugins/Runtime/StateTree/Source/StateTreeEditorModule/Private/StateTreePropertyHelpers.cpp
+++ b/Engine/Plugins/Runtime/StateTree/Source/StateTreeEditorModule/Private/StateTreePropertyHelpers.cpp
@@ -4,6 +4,7 @@
 #include "StateTreeEditorNode.h"
 #include "Hash/Blake3.h"
 #include "Misc/StringBuilder.h"
+#include "UObject/Field.h"
 
 namespace UE::StateTree::PropertyHelpers
 {
@@ -42,7 +43,12 @@ void DispatchPostEditToNodes(UObject& Owner, FPropertyChangedChainEvent& InPrope
 
 		if (const FStructProperty* StructProperty = CastField<FStructProperty>(CurrentProperty))
 		{
-			if (StructProperty->Struct == FStateTreeEditorNode::StaticStruct())
+			if (StructProperty->Struct == FInstancedStruct::StaticStruct())
+			{
+				FInstancedStruct& InstancedStruct = *reinterpret_cast<FInstancedStruct*>(CurrentAddress);
+				CurrentAddress = InstancedStruct.GetMutableMemory();
+			}
+			else if (StructProperty->Struct == FStateTreeEditorNode::StaticStruct())
 			{
 				FStateTreeEditorNode& EditorNode = *reinterpret_cast<FStateTreeEditorNode*>(CurrentAddress);
 				if (FStateTreeNodeBase* StateTreeNode = EditorNode.Node.GetMutablePtr<FStateTreeNodeBase>())
@@ -139,6 +145,11 @@ FGuid MakeDeterministicID(const UObject& Owner, const FString& PropertyPath, con
 	return FGuid::NewGuidFromHash(Hash);
 }
 
+bool HasOptionalMetadata(const FProperty& Property)
+{
+	return Property.HasMetaData(TEXT("Optional"));
+}
+
 }; // UE::StateTree::PropertyHelpers
 
 // ------------------------------------------------------------------------------
diff --git a/Engine/Plugins/Runtime/StateTree/Source/StateTreeEditorModule/Private/StateTreePropertyHelpers.h b/Engine/Plugins/Runtime/StateTree/Source/StateTreeEditorModule/Private/StateTreePropertyHelpers.h
--- a/Engine/Plugins/Runtime/StateTree/Source/StateTreeEditorModule/Private/StateTreePropertyHelpers.h
+++ b/Engine/Plugins/Runtime/StateTree/Source/StateTreeEditorModule/Private/StateTreePropertyHelpers.h
@@ -34,6 +34,10 @@ bool IsScriptStruct(const TSharedPtr<IPropertyHandle>& PropertyHandle)
 	FStructProperty* StructProperty = CastField<FStructProperty>(PropertyHandle->GetProperty());
 	return StructProperty && StructProperty->Struct->IsA(TBaseStructure<T>::Get()->GetClass());
 }
+/**
+ * @return true if provided Property contains "Optional" metadata
+ */
+bool HasOptionalMetadata(const FProperty& Property);
 
 /**
  * Gets a struct value from property handle, checks type before access. Expects T is struct.
diff --git a/Engine/Plugins/Runtime/StateTree/Source/StateTreeEditorModule/Public/StateTreeCompiler.h b/Engine/Plugins/Runtime/StateTree/Source/StateTreeEditorModule/Public/StateTreeCompiler.h
--- a/Engine/Plugins/Runtime/StateTree/Source/StateTreeEditorModule/Public/StateTreeCompiler.h
+++ b/Engine/Plugins/Runtime/StateTree/Source/StateTreeEditorModule/Public/StateTreeCompiler.h
@@ -52,8 +52,8 @@ private:
 	bool CreateCondition(UStateTreeState& State, const FStateTreeEditorNode& CondNode, const EStateTreeConditionOperand Operand, const int8 DeltaIndent);
 	bool CreateTask(UStateTreeState* State, const FStateTreeEditorNode& TaskNode, const FStateTreeDataHandle TaskDataHandle);
 	bool CreateEvaluator(const FStateTreeEditorNode& EvalNode, const FStateTreeDataHandle EvalDataHandle);
-	bool GetAndValidateBindings(const FStateTreeBindableStructDesc& TargetStruct, FStateTreeDataView TargetValue, TArray<FStateTreePropertyPathBinding>& OutBindings) const;
-	bool IsPropertyAnyEnum(const FStateTreeBindableStructDesc& Struct, FStateTreePropertyPath Path) const;
+	bool GetAndValidateBindings(const FStateTreeBindableStructDesc& TargetStruct, FStateTreeDataView TargetValue, TArray<FStateTreePropertyPathBinding>& OutCopyBindings, TArray<FStateTreePropertyPathBinding>& OutReferenceBindings) const;
+	bool IsPropertyOfType(UScriptStruct& Type, const FStateTreeBindableStructDesc& Struct, FStateTreePropertyPath Path) const;
 	bool ValidateStructRef(const FStateTreeBindableStructDesc& SourceStruct, FStateTreePropertyPath SourcePath,
 							const FStateTreeBindableStructDesc& TargetStruct, FStateTreePropertyPath TargetPath) const;
 	bool CompileAndValidateNode(const UStateTreeState* SourceState, const FStateTreeBindableStructDesc& NodeDesc, FStructView NodeView, const FStateTreeDataView InstanceData);
@@ -151,14 +151,6 @@ namespace UE::StateTree::Compiler
 	 */
 	const UScriptStruct* GetBaseStructFromMetaData(const FProperty* Property, FString& OutBaseStructName);
 
-	/**
-	 * Returns property usage based on the Category metadata of given property.
-	 * @param Property Handle to property where value is got from.
-	 * @return found usage type, or EStateTreePropertyUsage::Invalid if not found.
-	 */
-	EStateTreePropertyUsage GetUsageFromMetaData(const FProperty* Property);
-
-
 }; // UE::StateTree::Compiler
 
 #if UE_ENABLE_INCLUDE_ORDER_DEPRECATED_IN_5_2
diff --git a/Engine/Plugins/Runtime/StateTree/Source/StateTreeEditorModule/Public/StateTreePropertyBindingCompiler.h b/Engine/Plugins/Runtime/StateTree/Source/StateTreeEditorModule/Public/StateTreePropertyBindingCompiler.h
--- a/Engine/Plugins/Runtime/StateTree/Source/StateTreeEditorModule/Public/StateTreePropertyBindingCompiler.h
+++ b/Engine/Plugins/Runtime/StateTree/Source/StateTreeEditorModule/Public/StateTreePropertyBindingCompiler.h
@@ -34,6 +34,15 @@ struct STATETREEEDITORMODULE_API FStateTreePropertyBindingCompiler
 	 */
 	[[nodiscard]] bool CompileBatch(const FStateTreeBindableStructDesc& TargetStruct, TConstArrayView<FStateTreePropertyPathBinding> PropertyBindings, int32& OutBatchIndex);
 
+	/**
+	  * Compiles references for selected struct
+	  * @param TargetStruct - Description of the structs which contains the target properties.
+	  * @param PropertyReferenceBindings - Array of bindings to compile, all bindings that point to TargetStructs will be added.
+	  * @param InstanceDataView - view to the instance data
+	  * @return True on success, false on failure.
+	 */
+	[[nodiscard]] bool CompileReferences(const FStateTreeBindableStructDesc& TargetStruct, TConstArrayView<FStateTreePropertyPathBinding> PropertyReferenceBindings, FStateTreeDataView InstanceDataView);
+
 	/** Finalizes compilation, should be called once all batches are compiled. */
 	void Finalize();
 
@@ -96,6 +105,17 @@ protected:
 	UPROPERTY()
 	TArray<FStateTreeBindableStructDesc> SourceStructs;
 
+	/**
+	 * Representation of compiled reference.
+	 */
+	struct FCompiledReference
+	{
+		FStateTreePropertyPath Path;
+		FStateTreeIndex16 Index;
+	};
+
+	TArray<FCompiledReference> CompiledReferences;
+
 	FStateTreePropertyBindings* PropertyBindings = nullptr;
 
 	FStateTreeCompilerLog* Log = nullptr;
diff --git a/Engine/Plugins/Runtime/StateTree/Source/StateTreeModule/Private/StateTreePropertyBindings.cpp b/Engine/Plugins/Runtime/StateTree/Source/StateTreeModule/Private/StateTreePropertyBindings.cpp
--- a/Engine/Plugins/Runtime/StateTree/Source/StateTreeModule/Private/StateTreePropertyBindings.cpp
+++ b/Engine/Plugins/Runtime/StateTree/Source/StateTreeModule/Private/StateTreePropertyBindings.cpp
@@ -4,6 +4,7 @@
 #include "Misc/EnumerateRange.h"
 #include "PropertyPathHelpers.h"
 #include "PropertyBag.h"
+#include "StateTreePropertyRef.h"
 
 #if WITH_EDITOR
 #include "UObject/CoreRedirects.h"
@@ -11,6 +12,7 @@
 #include "Engine/BlueprintGeneratedClass.h"
 #include "Engine/UserDefinedStruct.h"
 #include "Kismet2/StructureEditorUtils.h"
+#include "UObject/Field.h"
 #endif
 
 #include UE_INLINE_GENERATED_CPP_BY_NAME(StateTreePropertyBindings)
@@ -32,6 +34,43 @@ namespace UE::StateTree
 		return Result.ToString();
 	}
 
+#if WITH_EDITOR
+	EStateTreePropertyUsage GetUsageFromMetaData(const FProperty* Property)
+	{
+		static const FName CategoryName(TEXT("Category"));
+
+		if (Property == nullptr)
+		{
+			return EStateTreePropertyUsage::Invalid;
+		}
+		
+		const FString Category = Property->GetMetaData(CategoryName);
+
+		if (Category == TEXT("Input"))
+		{
+			return EStateTreePropertyUsage::Input;
+		}
+		if (Category == TEXT("Inputs"))
+		{
+			return EStateTreePropertyUsage::Input;
+		}
+		if (Category == TEXT("Output"))
+		{
+			return EStateTreePropertyUsage::Output;
+		}
+		if (Category == TEXT("Outputs"))
+		{
+			return EStateTreePropertyUsage::Output;
+		}
+		if (Category == TEXT("Context"))
+		{
+			return EStateTreePropertyUsage::Context;
+		}
+
+		return EStateTreePropertyUsage::Parameter;
+	}
+#endif
+
 } // UE::StateTree
 
 namespace UE::StateTree::Private
@@ -131,6 +170,8 @@ void FStateTreePropertyBindings::Reset()
 	CopyBatches.Reset();
 	PropertyPathBindings.Reset();
 	PropertyCopies.Reset();
+	PropertyAccesses.Reset();
+	PropertyReferencePaths.Reset();
 	PropertyIndirections.Reset();
 	
 	bBindingsResolved = false;
@@ -223,10 +264,30 @@ bool FStateTreePropertyBindings::ResolvePaths()
 		}
 	}
 
+	PropertyAccesses.Reset();
+	PropertyAccesses.Reserve(PropertyReferencePaths.Num());
+
+	for (const FStateTreePropertyRefPath& ReferencePath : PropertyReferencePaths)
+	{
+		FStateTreePropertyAccess& PropertyAccess = PropertyAccesses.AddDefaulted_GetRef();
+		
+		PropertyAccess.SourceDataHandle = ReferencePath.GetSourceDataHandle();
+		const FStateTreeBindableStructDesc* SourceDesc = GetSourceDescByHandle(PropertyAccess.SourceDataHandle);
+		PropertyAccess.SourceStructType = SourceDesc->Struct;
+
+		FStateTreePropertyPathIndirection SourceLeafIndirection;
+		if (!ResolvePath(SourceDesc->Struct, ReferencePath.GetSourcePath(), PropertyAccess.SourceIndirection, SourceLeafIndirection))
+		{
+			bResult = false;
+		}
+
+		PropertyAccess.SourceLeafProperty = SourceLeafIndirection.GetProperty();
+	}
+
 	return bResult;
 }
 
-bool FStateTreePropertyBindings::ResolvePath(const UStruct* Struct, const FStateTreePropertyPath& Path, FStateTreePropertyIndirection& OutFirstIndirection, FStateTreePropertyPathIndirection& OutLeafIndirection)
+bool FStateTreePropertyBindings::ResolvePath(const UStruct* Struct, const FStateTreePropertyPath& Path, TArray<FStateTreePropertyIndirection>& OutIndirections, FStateTreePropertyIndirection& OutFirstIndirection, FStateTreePropertyPathIndirection& OutLeafIndirection)
 {
 	if (!Struct)
  	{
@@ -350,9 +411,9 @@ bool FStateTreePropertyBindings::ResolvePath(const UStruct* Struct, const FState
 		FStateTreePropertyIndirection* PrevIndirection = &OutFirstIndirection;
 		for (int32 Index = 1; Index < TempIndirections.Num(); Index++)
 		{
-			const int32 IndirectionIndex = PropertyIndirections.Num();
+			const int32 IndirectionIndex = OutIndirections.Num();
 			PrevIndirection->NextIndex = FStateTreeIndex16(IndirectionIndex); // Set PrevIndirection before array add, as it can invalidate the pointer.
-			FStateTreePropertyIndirection& NewIndirection = PropertyIndirections.Add_GetRef(TempIndirections[Index]);
+			FStateTreePropertyIndirection& NewIndirection = OutIndirections.Add_GetRef(TempIndirections[Index]);
 			PrevIndirection = &NewIndirection;
 		}
 	}
@@ -628,8 +689,6 @@ bool FStateTreePropertyBindings::ResolveCopyType(const FStateTreePropertyPathInd
 			}
 		}
 	}
-	
-	ensureMsgf(false, TEXT("Couldnt determine property copy type (%s -> %s)"), *SourceProperty->GetNameCPP(), *TargetProperty->GetNameCPP());
 
 	return false;
 }
@@ -752,7 +811,7 @@ EStateTreePropertyAccessCompatibility FStateTreePropertyBindings::GetPropertyCom
 	return EStateTreePropertyAccessCompatibility::Incompatible;
 }
 
-uint8* FStateTreePropertyBindings::GetAddress(FStateTreeDataView InStructView, const FStateTreePropertyIndirection& FirstIndirection, const FProperty* LeafProperty) const
+uint8* FStateTreePropertyBindings::GetAddress(FStateTreeDataView InStructView, TConstArrayView<FStateTreePropertyIndirection> Indirections, const FStateTreePropertyIndirection& FirstIndirection, const FProperty* LeafProperty)
 {
 	uint8* Address = InStructView.GetMutableMemory();
 	if (Address == nullptr)
@@ -845,7 +904,7 @@ uint8* FStateTreePropertyBindings::GetAddress(FStateTreeDataView InStructView, c
 				*StaticEnum<EStateTreePropertyAccessType>()->GetValueAsString(Indirection->Type), *LeafProperty->GetNameCPP());
 		}
 
-		Indirection = Indirection->NextIndex.IsValid() ? &PropertyIndirections[Indirection->NextIndex.Get()] : nullptr;
+		Indirection = Indirection->NextIndex.IsValid() ? &Indirections[Indirection->NextIndex.Get()] : nullptr;
 	}
 
 	return Address;
@@ -1076,6 +1135,21 @@ void FStateTreePropertyBindings::PerformResetObjects(const FStateTreePropertyCop
 	}
 }
 
+const FStateTreePropertyAccess* FStateTreePropertyBindings::GetPropertyAccess(const FStateTreePropertyRef& PropertyRef) const
+{
+	if (!PropertyRef.GetRefAccessIndex().IsValid())
+	{
+		return nullptr;
+	}
+
+	if (!ensure(PropertyAccesses.IsValidIndex(PropertyRef.GetRefAccessIndex().Get())))
+	{
+		return nullptr;
+	}
+
+	return &PropertyAccesses[PropertyRef.GetRefAccessIndex().Get()];
+}
+
 bool FStateTreePropertyBindings::ResetObjects(const FStateTreeIndex16 TargetBatchIndex, FStateTreeDataView TargetStructView) const
 {
 	// This is made ensure so that the programmers have the change to catch it (it's usually programming error not to call ResolvePaths(), and it wont spam log for others.
diff --git a/Engine/Plugins/Runtime/StateTree/Source/StateTreeModule/Private/StateTreePropertyRefHelpers.cpp b/Engine/Plugins/Runtime/StateTree/Source/StateTreeModule/Private/StateTreePropertyRefHelpers.cpp
new file mode 100644
--- /dev/null
+++ b/Engine/Plugins/Runtime/StateTree/Source/StateTreeModule/Private/StateTreePropertyRefHelpers.cpp
@@ -0,0 +1,277 @@
+// Copyright Epic Games, Inc. All Rights Reserved.
+
+#if WITH_EDITOR
+#include "StateTreePropertyRefHelpers.h"
+#include "StateTreePropertyRef.h"
+#include "UObject/TextProperty.h"
+#include "UObject/EnumProperty.h"
+#include "UObject/Class.h"
+#include "IPropertyAccessEditor.h"
+#include "StateTreePropertyBindings.h"
+#include "EdGraph/EdGraphPin.h"
+#include "EdGraphSchema_K2.h"
+
+namespace UE::StateTree::PropertyRefHelpers
+{
+	static const FName BoolName = TEXT("bool");
+	static const FName ByteName = TEXT("byte");
+	static const FName Int32Name = TEXT("int32");
+	static const FName Int64Name = TEXT("int64");
+	static const FName FloatName = TEXT("float");
+	static const FName DoubleName = TEXT("double");
+	static const FName NameName = TEXT("Name");
+	static const FName StringName = TEXT("String");
+	static const FName TextName = TEXT("Text");
+
+	static const FName IsRefToArrayName = TEXT("IsRefToArray");
+	static const FName RefTypeName = TEXT("RefType");
+
+	bool IsPropertyRefCompatibleWithProperty(const FProperty& RefProperty, const FProperty& SourceProperty)
+	{
+		ensure(IsPropertyRef(RefProperty));
+
+		const FProperty* TestProperty = &SourceProperty;
+		const bool bIsTargetRefArray = RefProperty.HasMetaData(IsRefToArrayName);
+
+		if (bIsTargetRefArray)
+		{
+			if (const FArrayProperty* ArrayProperty = CastField<FArrayProperty>(TestProperty))
+			{
+				TestProperty = ArrayProperty->Inner;
+			}
+			else
+			{
+				return false;
+			}
+		}
+
+		const FString& TargetTypeNameStr = RefProperty.GetMetaData(RefTypeName);
+		if (TargetTypeNameStr.IsEmpty())
+		{
+			return false;
+		}
+
+		const FName TargetTypeName = FName(*TargetTypeNameStr);
+
+		const FStructProperty* SourceStructProperty = CastField<FStructProperty>(TestProperty);
+
+		// Compare properties metadata directly if SourceProperty is PropertyRef as well
+		if (SourceStructProperty && SourceStructProperty->Struct == FStateTreePropertyRef::StaticStruct())
+		{
+			const FName SourceTypeName(SourceStructProperty->GetMetaData(RefTypeName));
+			const bool bIsSourceRefArray = SourceStructProperty->GetBoolMetaData(IsRefToArrayName);
+
+			return SourceTypeName == TargetTypeName && bIsSourceRefArray == bIsTargetRefArray;
+		}
+
+		if(TargetTypeName == BoolName)
+		{
+			return CastField<FBoolProperty>(TestProperty) != nullptr;
+		}
+		else if(TargetTypeName == ByteName)
+		{
+			return CastField<FByteProperty>(TestProperty) != nullptr;
+		}
+		else if(TargetTypeName == Int32Name)
+		{
+			return CastField<FIntProperty>(TestProperty) != nullptr;
+		}
+		else if(TargetTypeName == Int64Name)
+		{
+			return CastField<FInt64Property>(TestProperty) != nullptr;
+		}
+		else if(TargetTypeName == FloatName)
+		{
+			return CastField<FFloatProperty>(TestProperty) != nullptr;
+		}
+		else if(TargetTypeName == DoubleName)
+		{
+			return CastField<FDoubleProperty>(TestProperty) != nullptr;
+		}
+		else if(TargetTypeName == NameName)
+		{
+			return CastField<FNameProperty>(TestProperty) != nullptr;
+		}
+		else if(TargetTypeName == StringName)
+		{
+			return CastField<FStrProperty>(TestProperty) != nullptr;
+		}
+		else if(TargetTypeName == TextName)
+		{
+			return CastField<FTextProperty>(TestProperty) != nullptr;
+		}
+		else
+		{
+			UField* TargetRefField = UClass::TryFindTypeSlow<UField>(TargetTypeNameStr);
+			if (!TargetRefField)
+			{
+				TargetRefField = LoadObject<UField>(nullptr, *TargetTypeNameStr);
+			}
+
+			if (SourceStructProperty)
+			{
+				return SourceStructProperty->Struct->IsChildOf(Cast<UStruct>(TargetRefField));
+			}
+
+			if (const FObjectProperty* ObjectProperty = CastField<FObjectProperty>(TestProperty))
+			{
+				return ObjectProperty->PropertyClass->IsChildOf(Cast<UStruct>(TargetRefField));
+			}
+
+			if (const FEnumProperty* EnumProperty = CastField<FEnumProperty>(TestProperty))
+			{
+				return EnumProperty->GetEnum() == TargetRefField;
+			}
+		}
+
+		return false;
+	}
+
+	bool IsPropertyAccessibleForPropertyRef(const FProperty& SourceProperty, FStateTreeBindableStructDesc SourceStruct, bool bIsOutput)
+	{
+		switch (SourceStruct.DataSource)
+		{
+		case EStateTreeBindableStructSource::Parameter:
+		case EStateTreeBindableStructSource::State:
+			return true;
+
+		case EStateTreeBindableStructSource::Context:
+		case EStateTreeBindableStructSource::Condition:
+			return false;
+
+		case EStateTreeBindableStructSource::GlobalTask:
+		case EStateTreeBindableStructSource::Evaluator:
+		case EStateTreeBindableStructSource::Task:
+			return bIsOutput || IsPropertyRef(SourceProperty);
+		default:
+			checkNoEntry();
+		}
+
+		return false;
+	}
+
+	bool IsPropertyAccessibleForPropertyRef(TConstArrayView<FStateTreePropertyPathIndirection> SourcePropertyPathIndirections, FStateTreeBindableStructDesc SourceStruct)
+	{
+		bool bIsOutput = false;
+		for (const FStateTreePropertyPathIndirection& Indirection : SourcePropertyPathIndirections)
+		{
+			if (UE::StateTree::GetUsageFromMetaData(Indirection.GetProperty()) == EStateTreePropertyUsage::Output)
+			{
+				bIsOutput = true;
+				break;
+			}
+		}
+
+		return IsPropertyAccessibleForPropertyRef(*SourcePropertyPathIndirections.Last().GetProperty(), SourceStruct, bIsOutput);
+	}
+
+	bool IsPropertyAccessibleForPropertyRef(const FProperty& SourceProperty, TConstArrayView<TSharedPtr<FBindingChainElement>> BindingChain, FStateTreeBindableStructDesc SourceStruct)
+	{
+		bool bIsOutput = UE::StateTree::GetUsageFromMetaData(&SourceProperty) == EStateTreePropertyUsage::Output;
+		for (TSharedPtr<FBindingChainElement> ChainElement : BindingChain)
+		{
+			if (const FProperty* Property = ChainElement->Field.Get<FProperty>())
+			{
+				if (UE::StateTree::GetUsageFromMetaData(Property) == EStateTreePropertyUsage::Output)
+				{
+					bIsOutput = true;
+					break;
+				}
+			}
+		}
+
+		return IsPropertyAccessibleForPropertyRef(SourceProperty, SourceStruct, bIsOutput);
+	}
+
+	FEdGraphPinType GetPropertyRefInternalTypeAsPin(const FProperty& RefProperty)
+	{
+		ensure(IsPropertyRef(RefProperty));
+
+		FEdGraphPinType PinType;
+		PinType.PinSubCategory = NAME_None;
+
+		PinType.ContainerType = RefProperty.HasMetaData(IsRefToArrayName) ? EPinContainerType::Array : EPinContainerType::None;
+		const FString& TargetTypeNameStr = RefProperty.GetMetaData(RefTypeName);
+		const FName TargetTypeName = FName(*TargetTypeNameStr);
+
+		if(TargetTypeName == BoolName)
+		{
+			PinType.PinCategory = UEdGraphSchema_K2::PC_Boolean;
+		}
+		else if(TargetTypeName == ByteName)
+		{
+			PinType.PinCategory = UEdGraphSchema_K2::PC_Byte;
+		}
+		else if(TargetTypeName == Int32Name)
+		{
+			PinType.PinCategory = UEdGraphSchema_K2::PC_Int;
+		}
+		else if(TargetTypeName == Int64Name)
+		{
+			PinType.PinCategory = UEdGraphSchema_K2::PC_Int64;
+		}
+		else if(TargetTypeName == FloatName)
+		{
+			PinType.PinCategory = UEdGraphSchema_K2::PC_Real;
+			PinType.PinSubCategory = UEdGraphSchema_K2::PC_Float;
+		}
+		else if(TargetTypeName == DoubleName)
+		{
+			PinType.PinCategory = UEdGraphSchema_K2::PC_Real;
+			PinType.PinSubCategory = UEdGraphSchema_K2::PC_Double;
+		}
+		else if(TargetTypeName == NameName)
+		{
+			PinType.PinCategory = UEdGraphSchema_K2::PC_Name;
+		}
+		else if(TargetTypeName == StringName)
+		{
+			PinType.PinCategory = UEdGraphSchema_K2::PC_String;
+		}
+		else if(TargetTypeName == TextName)
+		{
+			PinType.PinCategory = UEdGraphSchema_K2::PC_Text;
+		}
+		else
+		{
+			UField* TargetRefField = UClass::TryFindTypeSlow<UField>(TargetTypeNameStr);
+			if (!TargetRefField)
+			{
+				TargetRefField = LoadObject<UField>(nullptr, *TargetTypeNameStr);
+			}
+
+			if (UStruct* Struct = Cast<UStruct>(TargetRefField))
+			{
+				PinType.PinCategory = UEdGraphSchema_K2::PC_Struct;
+				PinType.PinSubCategoryObject = Struct;
+			}
+			else if (UObject* Object = Cast<UObject>(TargetRefField))
+			{
+				PinType.PinCategory = UEdGraphSchema_K2::PC_Object;
+				PinType.PinSubCategoryObject = Object;
+			}
+			else if (UEnum* Enum = Cast<UEnum>(TargetRefField))
+			{
+				PinType.PinCategory = UEdGraphSchema_K2::PC_Enum;
+				PinType.PinSubCategoryObject = Enum;
+			}
+			else
+			{
+				checkNoEntry();
+			}
+		}
+
+		return PinType;
+	}
+
+	bool IsPropertyRef(const FProperty& Property)
+	{
+		if (const FStructProperty* StructProperty = CastField<FStructProperty>(&Property))
+		{
+			return StructProperty->Struct == FStateTreePropertyRef::StaticStruct();
+		}
+
+		return false;
+	}
+}
+#endif
\ No newline at end of file
diff --git a/Engine/Plugins/Runtime/StateTree/Source/StateTreeModule/Public/StateTreeExecutionContext.h b/Engine/Plugins/Runtime/StateTree/Source/StateTreeModule/Public/StateTreeExecutionContext.h
--- a/Engine/Plugins/Runtime/StateTree/Source/StateTreeModule/Public/StateTreeExecutionContext.h
+++ b/Engine/Plugins/Runtime/StateTree/Source/StateTreeModule/Public/StateTreeExecutionContext.h
@@ -126,6 +126,23 @@ public:
 
 	/** @return True of the the execution context is valid and initialized. */ 
 	bool IsValid() const { return RootStateTree.IsReadyToRun(); }
+
+	/**
+	 * @param PropertyRef Property's reference to get pointer to.
+	 * @return Pointer to referenced property if succeeded.
+	 */
+	template<class T>
+	T* GetMutablePropertyPtr(const FStateTreePropertyRef& PropertyRef)
+	{
+		const FStateTreePropertyBindings& PropertyBindings = CurrentlyProcessedFrame->StateTree->PropertyBindings;
+		if(const FStateTreePropertyAccess* PropertyAccess = PropertyBindings.GetPropertyAccess(PropertyRef))
+		{
+			FStateTreeDataView SourceView = GetDataView(CurrentlyProcessedParentFrame, *CurrentlyProcessedFrame, PropertyAccess->SourceDataHandle);
+			return PropertyBindings.GetMutablePropertyPtr<T>(SourceView, *PropertyAccess);
+		}
+
+		return nullptr;
+	}
 	
 	/** Start executing. */
 	EStateTreeRunStatus Start();
diff --git a/Engine/Plugins/Runtime/StateTree/Source/StateTreeModule/Public/StateTreePropertyBindings.h b/Engine/Plugins/Runtime/StateTree/Source/StateTreeModule/Public/StateTreePropertyBindings.h
--- a/Engine/Plugins/Runtime/StateTree/Source/StateTreeModule/Public/StateTreePropertyBindings.h
+++ b/Engine/Plugins/Runtime/StateTree/Source/StateTreeModule/Public/StateTreePropertyBindings.h
@@ -3,11 +3,13 @@
 #pragma once
 
 #include "StateTreeTypes.h"
+#include "StateTreePropertyRefHelpers.h"
 #include "StateTreePropertyBindings.generated.h"
 
 class FProperty;
 struct FStateTreePropertyPath;
 struct FStateTreePropertyBindingCompiler;
+struct FStateTreePropertyRef;
 class UStateTree;
 
 UENUM()
@@ -594,6 +596,38 @@ struct TStructOpsTypeTraits<FStateTreePropertyPathBinding> : public TStructOpsTy
 
 using FStateTreeEditorPropertyBinding UE_DEPRECATED(5.3, "Deprecated struct. Please use FStateTreePropertyPathBinding instead.") = FStateTreePropertyPathBinding;
 
+/**
+ * Representation of a property reference binding in StateTree.
+ */
+USTRUCT()
+struct STATETREEMODULE_API FStateTreePropertyRefPath
+{
+	GENERATED_BODY()
+
+	FStateTreePropertyRefPath() = default;
+
+	FStateTreePropertyRefPath(FStateTreeDataHandle InSourceDataHandle, const FStateTreePropertyPath& InSourcePath)
+		: SourcePropertyPath(InSourcePath)
+		, SourceDataHandle(InSourceDataHandle)
+	{
+	}
+
+	const FStateTreePropertyPath& GetSourcePath() const { return SourcePropertyPath; }
+
+	FStateTreePropertyPath& GetMutableSourcePath() { return SourcePropertyPath; }
+
+	void SetSourceDataHandle(const FStateTreeDataHandle NewSourceDataHandle) { SourceDataHandle = NewSourceDataHandle; }
+	FStateTreeDataHandle GetSourceDataHandle() const { return SourceDataHandle; }
+
+private:
+	/** Source property path of the reference */
+	UPROPERTY()
+	FStateTreePropertyPath SourcePropertyPath;
+
+	/** Describes how to get the source data pointer */
+	UPROPERTY()
+	FStateTreeDataHandle SourceDataHandle = FStateTreeDataHandle::Invalid;
+};
 
 /**
  * Deprecated. Describes a segment of a property path. Used for storage only.
@@ -768,6 +802,29 @@ struct STATETREEMODULE_API FStateTreePropertyCopyBatch
 
 using FStateTreePropCopyBatch UE_DEPRECATED(5.3, "Deprecated struct. Please use FStateTreePropertyCopy instead.") = FStateTreePropertyCopyBatch;
 
+/**
+ * Describes access to referenced property.
+ */
+USTRUCT()
+struct STATETREEMODULE_API FStateTreePropertyAccess
+{
+	GENERATED_BODY()
+
+	/** Source property access. */
+	UPROPERTY()
+	FStateTreePropertyIndirection SourceIndirection;
+
+	/** Cached pointer to the leaf property of the access. */
+	const FProperty* SourceLeafProperty = nullptr;
+
+	/** Type of the source data, used for validation. */
+	UPROPERTY(Transient)
+	TObjectPtr<const UStruct> SourceStructType = nullptr;
+
+	/** Describes how to get the source data pointer. */
+	UPROPERTY()
+	FStateTreeDataHandle SourceDataHandle = FStateTreeDataHandle::Invalid;
+};
 
 /**
  * Runtime storage and execution of property bindings.
@@ -830,7 +887,31 @@ struct STATETREEMODULE_API FStateTreePropertyBindings
 		}
 		return MakeArrayView(&PropertyCopies[Batch.BindingsBegin], Count);
 	}
-	
+
+	/**
+	 * @return Referenced property access for provided PropertyRef.
+	 */
+	const FStateTreePropertyAccess* GetPropertyAccess(const FStateTreePropertyRef& Reference) const;
+
+	/**
+	 * Pointer to referenced property 
+	 * @param SourceView Data view to referenced property's owner.
+	 * @param PropertyAccess Access to the property for which we want to obtain a pointer.
+	 * @return Pointer to referenced property if it's type match, nullptr otherwise.
+	 */
+	template< class T >
+	T* GetMutablePropertyPtr(FStateTreeDataView SourceView, const FStateTreePropertyAccess& PropertyAccess) const
+	{
+		check(SourceView.GetStruct() == PropertyAccess.SourceStructType);
+
+		if (!UE::StateTree::PropertyRefHelpers::Validator<T>::IsValid(*PropertyAccess.SourceLeafProperty))
+		{
+			return nullptr;
+		}
+
+		return reinterpret_cast<T*>(GetAddress(SourceView, PropertyAccess.SourceIndirection, PropertyAccess.SourceLeafProperty));
+	}
+
 	/**
 	 * Resets copied properties in TargetStructView. Can be used e.g. to erase UObject references.
 	 * @param TargetBatchIndex Batch index to copy (see FStateTreePropertyBindingCompiler).
@@ -872,12 +953,23 @@ struct STATETREEMODULE_API FStateTreePropertyBindings
 	}
 
 private:
-	[[nodiscard]] bool ResolvePath(const UStruct* Struct, const FStateTreePropertyPath& Path, FStateTreePropertyIndirection& OutFirstIndirection, FStateTreePropertyPathIndirection& OutLeafIndirection);
+	[[nodiscard]] static bool ResolvePath(const UStruct* Struct, const FStateTreePropertyPath& Path, TArray<FStateTreePropertyIndirection>& OutIndirections, FStateTreePropertyIndirection& OutFirstIndirection, FStateTreePropertyPathIndirection& OutLeafIndirection);
+	static uint8* GetAddress(FStateTreeDataView InStructView, TConstArrayView<FStateTreePropertyIndirection> Indirections, const FStateTreePropertyIndirection& FirstIndirection, const FProperty* LeafProperty);
+	
+	uint8* GetAddress(FStateTreeDataView InStructView, const FStateTreePropertyIndirection& FirstIndirection, const FProperty* LeafProperty) const
+	{
+		return GetAddress(InStructView, PropertyIndirections, FirstIndirection, LeafProperty);
+	}
+
+	[[nodiscard]] bool ResolvePath(const UStruct* Struct, const FStateTreePropertyPath& Path, FStateTreePropertyIndirection& OutFirstIndirection, FStateTreePropertyPathIndirection& OutLeafIndirection)
+	{
+		return ResolvePath(Struct, Path, PropertyIndirections, OutFirstIndirection, OutLeafIndirection);
+	}
+
 	const FStateTreeBindableStructDesc* GetSourceDescByHandle(const FStateTreeDataHandle SourceDataHandle);
 
 	void PerformCopy(const FStateTreePropertyCopy& Copy, uint8* SourceAddress, uint8* TargetAddress) const;
 	void PerformResetObjects(const FStateTreePropertyCopy& Copy, uint8* TargetAddress) const;
-	uint8* GetAddress(FStateTreeDataView InStructView, const FStateTreePropertyIndirection& FirstIndirection, const FProperty* LeafProperty) const;
 
 	/** Array of expected source structs. */
 	UPROPERTY()
@@ -895,6 +987,14 @@ private:
 	UPROPERTY(Transient)
 	TArray<FStateTreePropertyCopy> PropertyCopies;
 
+	/** Array of referenced property paths */
+	UPROPERTY()
+	TArray<FStateTreePropertyRefPath> PropertyReferencePaths;
+
+	/** Array of individually accessed properties */
+	UPROPERTY()
+	TArray<FStateTreePropertyAccess> PropertyAccesses;
+
 	/** Array of property indirections, indexed by accesses*/
 	UPROPERTY(Transient)
 	TArray<FStateTreePropertyIndirection> PropertyIndirections;
@@ -920,7 +1020,6 @@ PRAGMA_ENABLE_DEPRECATION_WARNINGS
 	friend UStateTree;
 };
 
-
 /**
  * Helper interface to reason about bound properties. The implementation is in the editor plugin.
  */
@@ -952,6 +1051,15 @@ namespace UE::StateTree
 {
 	/** @return desc and path as a display string. */
 	extern STATETREEMODULE_API FString GetDescAndPathAsString(const FStateTreeBindableStructDesc& Desc, const FStateTreePropertyPath& Path);
+
+#if WITH_EDITOR
+	/**
+	 * Returns property usage based on the Category metadata of given property.
+	 * @param Property Handle to property where value is got from.
+	 * @return found usage type, or EStateTreePropertyUsage::Invalid if not found.
+	 */
+	STATETREEMODULE_API EStateTreePropertyUsage GetUsageFromMetaData(const FProperty* Property);
+#endif
 } // UE::StateTree
 
 namespace UE::StateTree::Private
diff --git a/Engine/Plugins/Runtime/StateTree/Source/StateTreeModule/Public/StateTreePropertyRef.h b/Engine/Plugins/Runtime/StateTree/Source/StateTreeModule/Public/StateTreePropertyRef.h
new file mode 100644
--- /dev/null
+++ b/Engine/Plugins/Runtime/StateTree/Source/StateTreeModule/Public/StateTreePropertyRef.h
@@ -0,0 +1,64 @@
+// Copyright Epic Games, Inc. All Rights Reserved.
+
+#pragma once
+
+#include "StateTreeIndexTypes.h"
+#include "StateTreeExecutionContext.h"
+#include "StateTreePropertyRef.generated.h"
+
+/**
+ * Property ref allows to get a pointer to selected property in StateTree.
+ * The expected type of the reference should be set in "RefType" meta specifier.
+ *
+ * Meta specifiers for the type:
+ *  - RefType = "<type>"
+ *		- Specifies the type of property to reference.
+ *		- Supported types are: bool, byte, int32, int64, float, double, Name, String, Text, UObject pointers, and structs.
+ *  - IsRefToArray
+ *		- If specified, the reference is to an TArray<RefType>
+ *  - Optional
+ *		- If specified, the reference can be left unbound, otherwise the compiler report error if the reference is not bound.
+ *
+ * Example:
+ *
+ *  // Reference to float
+ *	UPROPERTY(EditAnywhere, meta = (RefType = "float"))
+ *	FStateTreePropertyRef RefToFloat;
+ * 
+ *  // Reference to FTestStructBase
+ *	UPROPERTY(EditAnywhere, meta = (RefType = "/Script/ModuleName.TestStructBase"))
+ *	FStateTreePropertyRef RefToTest;
+ *
+ *  // Reference to TArray<FTestStructBase>
+ *	UPROPERTY(EditAnywhere, meta = (RefType = "/Script/ModuleName.TestStructBase", IsRefToArray))
+ *	FStateTreePropertyRef RefToArrayOfTests;
+ */
+USTRUCT(BlueprintType)
+struct STATETREEMODULE_API FStateTreePropertyRef
+{
+	GENERATED_BODY()
+
+	FStateTreePropertyRef() = default;
+
+	/** @return pointer to the property if possible, nullptr otherwise. */
+	template<class T>
+	T* GetMutablePtr(FStateTreeExecutionContext& Context) const
+	{
+		return Context.GetMutablePropertyPtr<T>(*this);
+	}
+
+	/**
+	 * Used internally.
+	 * @return index to referenced property access
+	 */
+	FStateTreeIndex16 GetRefAccessIndex() const
+	{
+		return RefAccessIndex;
+	}
+
+private:
+	UPROPERTY()
+	FStateTreeIndex16 RefAccessIndex;
+
+	friend FStateTreePropertyBindingCompiler;
+};
\ No newline at end of file
diff --git a/Engine/Plugins/Runtime/StateTree/Source/StateTreeModule/Public/StateTreePropertyRefHelpers.h b/Engine/Plugins/Runtime/StateTree/Source/StateTreeModule/Public/StateTreePropertyRefHelpers.h
new file mode 100644
--- /dev/null
+++ b/Engine/Plugins/Runtime/StateTree/Source/StateTreeModule/Public/StateTreePropertyRefHelpers.h
@@ -0,0 +1,170 @@
+// Copyright Epic Games, Inc. All Rights Reserved.
+
+#pragma once
+#include "UObject/Field.h"
+#include "UObject/TextProperty.h"
+#include "UObject/EnumProperty.h"
+
+#if WITH_EDITOR
+#include "Containers/ContainersFwd.h"
+#include "Templates/SharedPointerFwd.h"
+#endif
+
+class FProperty;
+struct FBindingChainElement;
+struct FStateTreeBindableStructDesc;
+struct FStateTreePropertyPathIndirection;
+struct FEdGraphPinType;
+
+namespace UE::StateTree::PropertyRefHelpers
+{
+#if WITH_EDITOR
+	/**
+	 * @param RefProperty Property of PropertyRef type
+	 * @param SourceProperty Property to check it's type compatibility
+	 * @return true if SourceProperty type is compatible with PropertyRef
+	 */
+	bool STATETREEMODULE_API IsPropertyRefCompatibleWithProperty(const FProperty& RefProperty, const FProperty& SourceProperty);
+
+	/**
+	 * @param SourcePropertyPathIndirections Path indirections of the referenced property
+	 * @param SourceStruct Bindable owner of referenced property
+	 * @return true if property can be referenced by PropertyRef
+	 */
+	bool STATETREEMODULE_API IsPropertyAccessibleForPropertyRef(TConstArrayView<FStateTreePropertyPathIndirection> SourcePropertyPathIndirections, FStateTreeBindableStructDesc SourceStruct);
+
+	/**
+	 * @param SourceProperty Referenced property
+	 * @param BindingChain Binding chain to referenced property
+	 * @param SourceStruct Bindable owner of referenced property
+	 * @return true if property can be referenced by PropertyRef
+	 */
+	bool STATETREEMODULE_API IsPropertyAccessibleForPropertyRef(const FProperty& SourceProperty, TConstArrayView<TSharedPtr<FBindingChainElement>> BindingChain, FStateTreeBindableStructDesc SourceStruct);
+
+	/**
+	 * @param Property Property to check
+	 * @return true if Property is a PropertyRef
+	 */
+	bool STATETREEMODULE_API IsPropertyRef(const FProperty& Property);
+
+	/**
+	 * @param RefProperty Property of PropertyRef type
+	 * @return PinType for PropertyRef's internal type
+	 */
+	FEdGraphPinType STATETREEMODULE_API GetPropertyRefInternalTypeAsPin(const FProperty& RefProperty);
+#endif
+
+	template<class T, class = void>
+	struct Validator
+	{};
+
+	template<>
+	struct Validator<FBoolProperty::TCppType>
+	{
+		static bool IsValid(const FProperty& Property){ return Property.IsA<FBoolProperty>(); };
+	};
+
+	template<>
+	struct Validator<FByteProperty::TCppType>
+	{
+		static bool IsValid(const FProperty& Property){ return Property.IsA<FByteProperty>(); };
+	};
+
+	template<>
+	struct Validator<FIntProperty::TCppType>
+	{
+		static bool IsValid(const FProperty& Property){ return Property.IsA<FIntProperty>(); };
+	};
+
+	template<>
+	struct Validator<FInt64Property::TCppType>
+	{
+		static bool IsValid(const FProperty& Property){ return Property.IsA<FInt64Property>(); };
+	};
+
+	template<>
+	struct Validator<FFloatProperty::TCppType>
+	{
+		static bool IsValid(const FProperty& Property){ return Property.IsA<FFloatProperty>(); };
+	};
+
+	template<>
+	struct Validator<FDoubleProperty::TCppType>
+	{
+		static bool IsValid(const FProperty& Property){ return Property.IsA<FDoubleProperty>(); };
+	};
+
+	template<>
+	struct Validator<FNameProperty::TCppType>
+	{
+		static bool IsValid(const FProperty& Property){ return Property.IsA<FNameProperty>(); };
+	};
+
+	template<>
+	struct Validator<FStrProperty::TCppType>
+	{
+		static bool IsValid(const FProperty& Property){ return Property.IsA<FStrProperty>(); };
+	};
+
+	template<>
+	struct Validator<FTextProperty::TCppType>
+	{
+		static bool IsValid(const FProperty& Property){ return Property.IsA<FTextProperty>(); };
+	};
+
+	template<class T>
+	struct Validator<T, typename TEnableIf<TIsTArray_V<T>, void>::Type>
+	{
+		static bool IsValid(const FProperty& Property)
+		{
+			if (const FArrayProperty* ArrayProperty = CastField<FArrayProperty>(&Property))
+			{
+				return Validator<typename T::ElementType>::IsValid(*ArrayProperty->Inner);
+			}
+
+			return false;
+		}
+	};
+
+	template<class T>
+	struct Validator<T, decltype(TBaseStructure<T>::Get, void())>
+	{
+		static bool IsValid(const FProperty& Property)
+		{		
+			if (const FStructProperty* StructProperty = CastField<FStructProperty>(&Property))
+			{
+				return StructProperty->Struct == TBaseStructure<T>::Get();
+			}
+
+			return false;
+		}
+	};
+
+	template<class T>
+	struct Validator<T, typename TEnableIf<TIsDerivedFrom<typename TRemovePointer<T>::Type, UObject>::IsDerived>::Type>
+	{
+		static bool IsValid(const FProperty& Property)
+		{		
+			if (const FObjectProperty* ObjectProperty = CastField<FObjectProperty>(&Property))
+			{
+				return ObjectProperty->PropertyClass == TRemovePointer<T>::Type::StaticClass();
+			}
+
+			return false;
+		}
+	};
+
+	template<class T>
+	struct Validator<T, typename TEnableIf<TIsEnum<T>::Value>::Type>
+	{
+		static bool IsValid(const FProperty& Property)
+		{		
+			if (const FEnumProperty* EnumProperty = CastField<FEnumProperty>(&Property))
+			{
+				return EnumProperty->GetEnum() == StaticEnum<T>();
+			}
+
+			return false;
+		}
+	};
+}
\ No newline at end of file
diff --git a/Engine/Plugins/Runtime/StateTree/Source/StateTreeModule/StateTreeModule.Build.cs b/Engine/Plugins/Runtime/StateTree/Source/StateTreeModule/StateTreeModule.Build.cs
--- a/Engine/Plugins/Runtime/StateTree/Source/StateTreeModule/StateTreeModule.Build.cs
+++ b/Engine/Plugins/Runtime/StateTree/Source/StateTreeModule/StateTreeModule.Build.cs
@@ -20,7 +20,7 @@ namespace UnrealBuildTool.Rules
 					"AIModule",
 					"GameplayTags",
 					"StructUtils",
-					"StructUtilsEngine"
+					"StructUtilsEngine",
 				}
 			);
 
@@ -34,7 +34,12 @@ namespace UnrealBuildTool.Rules
 
 			if (Target.bBuildEditor)
 			{
-				PublicDependencyModuleNames.Add("UnrealEd");
+				PublicDependencyModuleNames.AddRange(
+					new string[] {
+						"UnrealEd",
+						"BlueprintGraph",
+					}
+				);
 			}
 
 			if (Target.Platform == UnrealTargetPlatform.Win64 && 
diff --git a/Engine/Plugins/Runtime/StateTree/Source/StateTreeTestSuite/Private/StateTreeTest.cpp b/Engine/Plugins/Runtime/StateTree/Source/StateTreeTestSuite/Private/StateTreeTest.cpp
--- a/Engine/Plugins/Runtime/StateTree/Source/StateTreeTestSuite/Private/StateTreeTest.cpp
+++ b/Engine/Plugins/Runtime/StateTree/Source/StateTreeTestSuite/Private/StateTreeTest.cpp
@@ -32,6 +32,19 @@ namespace UE::StateTree::Tests
 		return *StateTree;
 	}
 
+	FStateTreePropertyPathBinding MakeBinding(const FGuid& SourceID, const FString& Source, const FGuid& TargetID, const FString& Target)
+	{
+		FStateTreePropertyPath SourcePath;
+		SourcePath.FromString(Source);
+		SourcePath.SetStructID(SourceID);
+
+		FStateTreePropertyPath TargetPath;
+		TargetPath.FromString(Target);
+		TargetPath.SetStructID(TargetID);
+
+		return FStateTreePropertyPathBinding(SourcePath, TargetPath);
+	}
+
 	// Helper struct to define some test tags
 	struct FNativeGameplayTags : public FGameplayTagNativeAdder
 	{
@@ -1530,23 +1543,10 @@ struct FStateTreeTest_BindingsCompiler : FAITestBase
 		const int32 SourceAIndex = BindingCompiler.AddSourceStruct(SourceADesc);
 		const int32 SourceBIndex = BindingCompiler.AddSourceStruct(SourceBDesc);
 
-		auto MakeBinding = [](const FGuid& SourceID, const FString& Source, const FGuid& TargetID, const FString& Target)
-		{
-			FStateTreePropertyPath SourcePath;
-			SourcePath.FromString(Source);
-			SourcePath.SetStructID(SourceID);
-
-			FStateTreePropertyPath TargetPath;
-			TargetPath.FromString(Target);
-			TargetPath.SetStructID(TargetID);
-
-			return FStateTreePropertyPathBinding(SourcePath, TargetPath);
-		};
-
 		TArray<FStateTreePropertyPathBinding> PropertyBindings;
-		PropertyBindings.Add(MakeBinding(SourceBDesc.ID, TEXT("Item"), TargetDesc.ID, TEXT("Array[1]")));
-		PropertyBindings.Add(MakeBinding(SourceADesc.ID, TEXT("Item.B"), TargetDesc.ID, TEXT("Array[1].B")));
-		PropertyBindings.Add(MakeBinding(SourceADesc.ID, TEXT("Array"), TargetDesc.ID, TEXT("Array")));
+		PropertyBindings.Add(UE::StateTree::Tests::MakeBinding(SourceBDesc.ID, TEXT("Item"), TargetDesc.ID, TEXT("Array[1]")));
+		PropertyBindings.Add(UE::StateTree::Tests::MakeBinding(SourceADesc.ID, TEXT("Item.B"), TargetDesc.ID, TEXT("Array[1].B")));
+		PropertyBindings.Add(UE::StateTree::Tests::MakeBinding(SourceADesc.ID, TEXT("Array"), TargetDesc.ID, TEXT("Array")));
 
 		int32 CopyBatchIndex = INDEX_NONE;
 		const bool bCompileBatchResult = BindingCompiler.CompileBatch(TargetDesc, PropertyBindings, CopyBatchIndex);
@@ -1633,31 +1633,18 @@ struct FStateTreeTest_CopyObjects : FAITestBase
 
 		const int32 SourceIndex = BindingCompiler.AddSourceStruct(SourceDesc);
 
-		auto MakeBinding = [](const FGuid& SourceID, const FString& Source, const FGuid& TargetID, const FString& Target)
-		{
-			FStateTreePropertyPath SourcePath;
-			SourcePath.FromString(Source);
-			SourcePath.SetStructID(SourceID);
-
-			FStateTreePropertyPath TargetPath;
-			TargetPath.FromString(Target);
-			TargetPath.SetStructID(TargetID);
-
-			return FStateTreePropertyPathBinding(SourcePath, TargetPath);
-		};
-
 		TArray<FStateTreePropertyPathBinding> PropertyBindings;
 		// One-to-one copy from source to target A
-		PropertyBindings.Add(MakeBinding(SourceDesc.ID, TEXT("Object"), TargetADesc.ID, TEXT("Object")));
-		PropertyBindings.Add(MakeBinding(SourceDesc.ID, TEXT("SoftObject"), TargetADesc.ID, TEXT("SoftObject")));
-		PropertyBindings.Add(MakeBinding(SourceDesc.ID, TEXT("Class"), TargetADesc.ID, TEXT("Class")));
-		PropertyBindings.Add(MakeBinding(SourceDesc.ID, TEXT("SoftClass"), TargetADesc.ID, TEXT("SoftClass")));
+		PropertyBindings.Add(UE::StateTree::Tests::MakeBinding(SourceDesc.ID, TEXT("Object"), TargetADesc.ID, TEXT("Object")));
+		PropertyBindings.Add(UE::StateTree::Tests::MakeBinding(SourceDesc.ID, TEXT("SoftObject"), TargetADesc.ID, TEXT("SoftObject")));
+		PropertyBindings.Add(UE::StateTree::Tests::MakeBinding(SourceDesc.ID, TEXT("Class"), TargetADesc.ID, TEXT("Class")));
+		PropertyBindings.Add(UE::StateTree::Tests::MakeBinding(SourceDesc.ID, TEXT("SoftClass"), TargetADesc.ID, TEXT("SoftClass")));
 
 		// Cross copy from source to target B
-		PropertyBindings.Add(MakeBinding(SourceDesc.ID, TEXT("SoftObject"), TargetBDesc.ID, TEXT("Object")));
-		PropertyBindings.Add(MakeBinding(SourceDesc.ID, TEXT("Object"), TargetBDesc.ID, TEXT("SoftObject")));
-		PropertyBindings.Add(MakeBinding(SourceDesc.ID, TEXT("SoftClass"), TargetBDesc.ID, TEXT("Class")));
-		PropertyBindings.Add(MakeBinding(SourceDesc.ID, TEXT("Class"), TargetBDesc.ID, TEXT("SoftClass")));
+		PropertyBindings.Add(UE::StateTree::Tests::MakeBinding(SourceDesc.ID, TEXT("SoftObject"), TargetBDesc.ID, TEXT("Object")));
+		PropertyBindings.Add(UE::StateTree::Tests::MakeBinding(SourceDesc.ID, TEXT("Object"), TargetBDesc.ID, TEXT("SoftObject")));
+		PropertyBindings.Add(UE::StateTree::Tests::MakeBinding(SourceDesc.ID, TEXT("SoftClass"), TargetBDesc.ID, TEXT("Class")));
+		PropertyBindings.Add(UE::StateTree::Tests::MakeBinding(SourceDesc.ID, TEXT("Class"), TargetBDesc.ID, TEXT("SoftClass")));
 		
 		int32 TargetACopyBatchIndex = INDEX_NONE;
 		const bool bCompileBatchResultA = BindingCompiler.CompileBatch(TargetADesc, PropertyBindings, TargetACopyBatchIndex);
@@ -1742,6 +1729,146 @@ struct FStateTreeTest_CopyObjects : FAITestBase
 };
 IMPLEMENT_AI_INSTANT_TEST(FStateTreeTest_CopyObjects, "System.StateTree.CopyObjects");
 
+struct FStateTreeTest_References : FAITestBase
+{
+	virtual bool InstantTest() override
+	{
+		FStateTreeCompilerLog Log;
+		FStateTreePropertyBindings Bindings;
+		FStateTreePropertyBindingCompiler BindingCompiler;
+
+		const bool bInitResult = BindingCompiler.Init(Bindings, Log);
+		AITEST_TRUE("Expect init to succeed", bInitResult);
+
+		FStateTreeBindableStructDesc SourceDesc;
+		SourceDesc.Name = FName(TEXT("Source"));
+		SourceDesc.Struct = TBaseStructure<FStateTreeTest_PropertyRefSourceStruct>::Get();
+		SourceDesc.DataSource = EStateTreeBindableStructSource::Parameter;
+		SourceDesc.DataHandle = FStateTreeDataHandle(EStateTreeDataSourceType::ContextData, 0);
+		SourceDesc.ID = FGuid::NewGuid();
+		BindingCompiler.AddSourceStruct(SourceDesc);
+
+		FStateTreeBindableStructDesc TargetDesc;
+		TargetDesc.Name = FName(TEXT("Target"));
+		TargetDesc.Struct = TBaseStructure<FStateTreeTest_PropertyRefTargetStruct>::Get();
+		TargetDesc.DataSource = EStateTreeBindableStructSource::Parameter;
+		TargetDesc.ID = FGuid::NewGuid();
+		
+		TArray<FStateTreePropertyPathBinding> PropertyBindings;
+		PropertyBindings.Add(UE::StateTree::Tests::MakeBinding(SourceDesc.ID, TEXT("Item"), TargetDesc.ID, TEXT("RefToStruct")));
+		PropertyBindings.Add(UE::StateTree::Tests::MakeBinding(SourceDesc.ID, TEXT("Item.A"), TargetDesc.ID, TEXT("RefToInt")));
+		PropertyBindings.Add(UE::StateTree::Tests::MakeBinding(SourceDesc.ID, TEXT("Array"), TargetDesc.ID, TEXT("RefToStructArray")));
+
+		FStateTreeTest_PropertyRefTargetStruct Target;
+		FStateTreeDataView TargetView(FStructView::Make(Target));
+		const bool bCompileReferencesResult = BindingCompiler.CompileReferences(TargetDesc, PropertyBindings, TargetView);
+		AITEST_TRUE("CompileReferences should succeed", bCompileReferencesResult);	
+
+		BindingCompiler.Finalize();
+
+		const bool bResolveResult = Bindings.ResolvePaths();
+		AITEST_TRUE("ResolvePaths should succeed", bResolveResult);
+
+		FStateTreeTest_PropertyRefSourceStruct Source;
+
+		FStateTreeDataView SourceView = FStateTreeDataView(FStructView::Make(Source));
+
+		{
+			const FStateTreePropertyAccess* PropertyAccess = Bindings.GetPropertyAccess(Target.RefToStruct);
+			AITEST_NOT_NULL("GetPropertyAccess should succeed", PropertyAccess);
+			
+			FStateTreeTest_PropertyStruct* Reference = Bindings.GetMutablePropertyPtr<FStateTreeTest_PropertyStruct>(SourceView, *PropertyAccess);
+			AITEST_EQUAL("Expect RefToStruct to point to SourceA.Item", Reference, &Source.Item);
+		}
+
+		{
+			const FStateTreePropertyAccess* PropertyAccess = Bindings.GetPropertyAccess(Target.RefToInt);
+			AITEST_NOT_NULL("GetPropertyAccess should succeed", PropertyAccess);
+
+			int32* Reference = Bindings.GetMutablePropertyPtr<int32>(SourceView, *PropertyAccess);
+			AITEST_EQUAL("Expect RefToInt to point to SourceA.Item.A", Reference, &Source.Item);
+		}
+
+		{
+			const FStateTreePropertyAccess* PropertyAccess = Bindings.GetPropertyAccess(Target.RefToStructArray);
+			AITEST_NOT_NULL("GetPropertyAccess should succeed", PropertyAccess);
+
+			TArray<FStateTreeTest_PropertyStruct>* Reference = Bindings.GetMutablePropertyPtr<TArray<FStateTreeTest_PropertyStruct>>(SourceView, *PropertyAccess);
+			AITEST_EQUAL("Expect RefToStructArray to point to SourceA.Array", Reference, &Source.Array);
+		}
+		
+		return true;
+	}
+};
+IMPLEMENT_AI_INSTANT_TEST(FStateTreeTest_References, "System.StateTree.References");
+
+struct FStateTreeTest_ReferencesConstness : FAITestBase
+{
+	virtual bool InstantTest() override
+	{
+		FStateTreeCompilerLog Log;
+		FStateTreePropertyBindings Bindings;
+		FStateTreePropertyBindingCompiler BindingCompiler;
+
+		const bool bInitResult = BindingCompiler.Init(Bindings, Log);
+		AITEST_TRUE("Expect init to succeed", bInitResult);
+
+		FStateTreeBindableStructDesc SourceAsTaskDesc;
+		SourceAsTaskDesc.Name = FName(TEXT("SourceTask"));
+		SourceAsTaskDesc.Struct = TBaseStructure<FStateTreeTest_PropertyRefSourceStruct>::Get();
+		SourceAsTaskDesc.DataSource = EStateTreeBindableStructSource::Task;
+		SourceAsTaskDesc.DataHandle = FStateTreeDataHandle(EStateTreeDataSourceType::ContextData, 0);
+		SourceAsTaskDesc.ID = FGuid::NewGuid();
+		BindingCompiler.AddSourceStruct(SourceAsTaskDesc);
+
+		FStateTreeBindableStructDesc SourceAsContextDesc;
+		SourceAsContextDesc.Name = FName(TEXT("SourceContext"));
+		SourceAsContextDesc.Struct = TBaseStructure<FStateTreeTest_PropertyRefSourceStruct>::Get();
+		SourceAsContextDesc.DataSource = EStateTreeBindableStructSource::Context;
+		SourceAsContextDesc.DataHandle = FStateTreeDataHandle(EStateTreeDataSourceType::ContextData, 0);
+		SourceAsContextDesc.ID = FGuid::NewGuid();
+		BindingCompiler.AddSourceStruct(SourceAsContextDesc);
+
+		FStateTreeBindableStructDesc TargetDesc;
+		TargetDesc.Name = FName(TEXT("Target"));
+		TargetDesc.Struct = TBaseStructure<FStateTreeTest_PropertyRefTargetStruct>::Get();
+		TargetDesc.DataSource = EStateTreeBindableStructSource::Parameter;
+		TargetDesc.ID = FGuid::NewGuid();
+		
+		FStateTreePropertyPathBinding TaskPropertyBinding = UE::StateTree::Tests::MakeBinding(SourceAsTaskDesc.ID, TEXT("Item"), TargetDesc.ID, TEXT("RefToStruct"));
+		FStateTreePropertyPathBinding TaskOutputPropertyBinding = UE::StateTree::Tests::MakeBinding(SourceAsTaskDesc.ID, TEXT("OutputItem"), TargetDesc.ID, TEXT("RefToStruct"));
+
+		FStateTreePropertyPathBinding ContextPropertyBinding = UE::StateTree::Tests::MakeBinding(SourceAsTaskDesc.ID, TEXT("Item"), TargetDesc.ID, TEXT("RefToStruct"));
+		FStateTreePropertyPathBinding ContextOutputPropertyBinding = UE::StateTree::Tests::MakeBinding(SourceAsTaskDesc.ID, TEXT("Item"), TargetDesc.ID, TEXT("RefToStruct"));
+
+		FStateTreeTest_PropertyRefTargetStruct Target;
+		FStateTreeDataView TargetView(FStructView::Make(Target));
+
+		{
+			const bool bCompileReferenceResult = BindingCompiler.CompileReferences(TargetDesc, {TaskPropertyBinding}, TargetView);
+			AITEST_FALSE("CompileReferences should fail", bCompileReferenceResult);
+		}
+
+		{
+			const bool bCompileReferenceResult = BindingCompiler.CompileReferences(TargetDesc, {TaskOutputPropertyBinding}, TargetView);
+			AITEST_TRUE("CompileReferences should succeed", bCompileReferenceResult);
+		}
+
+		{
+			const bool bCompileReferenceResult = BindingCompiler.CompileReferences(TargetDesc, {ContextPropertyBinding}, TargetView);
+			AITEST_FALSE("CompileReferences should fail", bCompileReferenceResult);
+		}
+
+		{	
+			const bool bCompileReferenceResult = BindingCompiler.CompileReferences(TargetDesc, {ContextOutputPropertyBinding}, TargetView);
+			AITEST_FALSE("CompileReferences should fail", bCompileReferenceResult);
+		}
+
+		return true;
+	}
+};
+IMPLEMENT_AI_INSTANT_TEST(FStateTreeTest_ReferencesConstness, "System.StateTree.ReferencesConstness");
+
 struct FStateTreeTest_FollowTransitions : FAITestBase
 {
 	virtual bool InstantTest() override
diff --git a/Engine/Plugins/Runtime/StateTree/Source/StateTreeTestSuite/Private/StateTreeTestTypes.h b/Engine/Plugins/Runtime/StateTree/Source/StateTreeTestSuite/Private/StateTreeTestTypes.h
--- a/Engine/Plugins/Runtime/StateTree/Source/StateTreeTestSuite/Private/StateTreeTestTypes.h
+++ b/Engine/Plugins/Runtime/StateTree/Source/StateTreeTestSuite/Private/StateTreeTestTypes.h
@@ -6,6 +6,7 @@
 #include "StateTreeEvaluatorBase.h"
 #include "StateTreeConditionBase.h"
 #include "StateTreeExecutionContext.h"
+#include "StateTreePropertyRef.h"
 #include "StateTreeTestTypes.generated.h"
 
 class UStateTree;
@@ -465,6 +466,36 @@ struct FStateTreeTest_PropertyCopy
 	TArray<FStateTreeTest_PropertyStruct> Array;
 };
 
+USTRUCT()
+struct FStateTreeTest_PropertyRefSourceStruct
+{
+	GENERATED_BODY()
+
+	UPROPERTY(EditAnywhere, Category = "")
+	FStateTreeTest_PropertyStruct Item;
+
+	UPROPERTY(EditAnywhere, Category = "Output")
+	FStateTreeTest_PropertyStruct OutputItem;
+
+	UPROPERTY(EditAnywhere, Category = "")
+	TArray<FStateTreeTest_PropertyStruct> Array;
+};
+
+USTRUCT()
+struct FStateTreeTest_PropertyRefTargetStruct
+{
+	GENERATED_BODY()
+
+	UPROPERTY(EditAnywhere, Category = "", meta = (RefType = "/Script/StateTreeTestSuite.StateTreeTest_PropertyStruct"))
+	FStateTreePropertyRef RefToStruct;
+
+	UPROPERTY(EditAnywhere, Category = "", meta = (RefType = "Int32"))
+	FStateTreePropertyRef RefToInt;
+
+	UPROPERTY(EditAnywhere, Category = "", meta = (RefType = "/Script/StateTreeTestSuite.StateTreeTest_PropertyStruct", IsRefToArray))
+	FStateTreePropertyRef RefToStructArray;
+};
+
 USTRUCT()
 struct FStateTreeTest_PropertyCopyObjects
 {
