Fixed pass-by-value archive in SerializeDataViewToArchive.

diff --git a/Engine/Plugins/Runtime/StateTree/Source/StateTreeModule/Private/Debugger/StateTreeTrace.cpp b/Engine/Plugins/Runtime/StateTree/Source/StateTreeModule/Private/Debugger/StateTreeTrace.cpp
--- a/Engine/Plugins/Runtime/StateTree/Source/StateTreeModule/Private/Debugger/StateTreeTrace.cpp
+++ b/Engine/Plugins/Runtime/StateTree/Source/StateTreeModule/Private/Debugger/StateTreeTrace.cpp
@@ -438,7 +438,7 @@ void TraceBufferedEvents(const FStateTreeInstanceDebugId InstanceId)
 	GBufferedEvents.Flush(InstanceId);
 }
 
-void SerializeDataViewToArchive(FBufferArchive Ar, const FStateTreeDataView DataView)
+void SerializeDataViewToArchive(FBufferArchive& Ar, const FStateTreeDataView DataView)
 {
 	constexpr uint32 PortFlags = 
 		PPF_PropertyWindow // limit to properties visible in Editor 
