StateTree: Added parameter overrides for StateTree states

diff --git a/Engine/Plugins/Runtime/StateTree/Source/StateTreeEditorModule/Private/Customizations/StateTreeStateParametersDetails.cpp b/Engine/Plugins/Runtime/StateTree/Source/StateTreeEditorModule/Private/Customizations/StateTreeStateParametersDetails.cpp
--- a/Engine/Plugins/Runtime/StateTree/Source/StateTreeEditorModule/Private/Customizations/StateTreeStateParametersDetails.cpp
+++ b/Engine/Plugins/Runtime/StateTree/Source/StateTreeEditorModule/Private/Customizations/StateTreeStateParametersDetails.cpp
@@ -21,12 +21,14 @@ class FStateTreeStateParametersInstanceDataDetails : public FPropertyBagInstance
 public:
 	FStateTreeStateParametersInstanceDataDetails(
 		const TSharedPtr<IPropertyHandle>& InStructProperty,
+		const TSharedPtr<IPropertyHandle>& InParametersStructProperty,
 		const TSharedPtr<IPropertyUtilities>& InPropUtils,
 		const bool bInFixedLayout,
 		FGuid InID,
 		TWeakObjectPtr<UStateTreeEditorData> InEditorData,
 		TWeakObjectPtr<UStateTreeState> InState)
-		: FPropertyBagInstanceDataDetails(InStructProperty, InPropUtils, bInFixedLayout)
+		: FPropertyBagInstanceDataDetails(InParametersStructProperty, InPropUtils, bInFixedLayout)
+		, StructProperty(InStructProperty)
 		, WeakEditorData(InEditorData)
 		, WeakState(InState)
 		, ID(InID)
@@ -86,6 +88,63 @@ public:
 		}
 	}
 
+	struct FStateTreeStateOverrideProvider : public IPropertyBagOverrideProvider
+	{
+		FStateTreeStateOverrideProvider(UStateTreeState& InState)
+			: State(InState)
+		{
+		}
+		
+		virtual bool IsPropertyOverridden(const FGuid PropertyID) const override
+		{
+			return State.IsParametersPropertyOverridden(PropertyID);
+		}
+		
+		virtual void SetPropertyOverride(const FGuid PropertyID, const bool bIsOverridden) const override
+		{
+			State.SetParametersPropertyOverridden(PropertyID, bIsOverridden);
+		}
+
+	private:
+		UStateTreeState& State;
+	};
+
+	virtual bool HasPropertyOverrides() const override
+	{
+		if (const UStateTreeState* State = WeakState.Get())
+		{
+			return State->Type == EStateTreeStateType::Linked || State->Type == EStateTreeStateType::LinkedAsset;
+		}
+		return false;
+	}
+
+	virtual void PreChangeOverrides() override
+	{
+		check(StructProperty);
+		StructProperty->NotifyPreChange();
+	}
+
+	virtual void PostChangeOverrides() override
+	{
+		check(StructProperty);
+		StructProperty->NotifyPostChange(EPropertyChangeType::ValueSet);
+		StructProperty->NotifyFinishedChangingProperties();
+	}
+
+	virtual void EnumeratePropertyBags(TSharedPtr<IPropertyHandle> PropertyBagHandle, const EnumeratePropertyBagFuncRef& Func) const override
+	{
+		if (UStateTreeState* State = WeakState.Get())
+		{
+			if (const FInstancedPropertyBag* DefaultParameters = State->GetDefaultParameters())
+			{
+				FInstancedPropertyBag& Parameters = State->Parameters.Parameters;
+				FStateTreeStateOverrideProvider OverrideProvider(*State);
+				Func(*DefaultParameters, Parameters, OverrideProvider);
+			}
+		}
+	}
+
+	TSharedPtr<IPropertyHandle> StructProperty;
 	TWeakObjectPtr<UStateTreeEditorData> WeakEditorData;
 	TWeakObjectPtr<UStateTreeState> WeakState;
 	FGuid ID;
@@ -142,7 +201,7 @@ void FStateTreeStateParametersDetails::CustomizeChildren(TSharedRef<class IPrope
 	UE::StateTree::PropertyHelpers::GetStructValue<FGuid>(IDProperty, ID);
 
 	// Show the Value (FInstancedStruct) as child rows.
-	TSharedRef<FStateTreeStateParametersInstanceDataDetails> InstanceDetails = MakeShareable(new FStateTreeStateParametersInstanceDataDetails(ParametersProperty, PropUtils, bFixedLayout, ID, WeakEditorData, WeakState));
+	TSharedRef<FStateTreeStateParametersInstanceDataDetails> InstanceDetails = MakeShareable(new FStateTreeStateParametersInstanceDataDetails(StructProperty, ParametersProperty, PropUtils, bFixedLayout, ID, WeakEditorData, WeakState));
 	StructBuilder.AddCustomBuilder(InstanceDetails);
 }
 
@@ -158,8 +217,8 @@ void FStateTreeStateParametersDetails::FindOuterObjects()
 	StructProperty->GetOuterObjects(OuterObjects);
 	for (UObject* Outer : OuterObjects)
 	{
+		UStateTreeState* OuterState = Cast<UStateTreeState>(Outer);
 		UStateTreeEditorData* OuterEditorData = Outer->GetTypedOuter<UStateTreeEditorData>();
-		UStateTreeState* OuterState = Outer->GetTypedOuter<UStateTreeState>();
 		UStateTree* OuterStateTree = OuterEditorData ? OuterEditorData->GetTypedOuter<UStateTree>() : nullptr;
 		if (OuterEditorData && OuterStateTree && OuterState)
 		{
diff --git a/Engine/Plugins/Runtime/StateTree/Source/StateTreeEditorModule/Private/StateTreeState.cpp b/Engine/Plugins/Runtime/StateTree/Source/StateTreeEditorModule/Private/StateTreeState.cpp
--- a/Engine/Plugins/Runtime/StateTree/Source/StateTreeEditorModule/Private/StateTreeState.cpp
+++ b/Engine/Plugins/Runtime/StateTree/Source/StateTreeEditorModule/Private/StateTreeState.cpp
@@ -10,6 +10,28 @@
 
 #include UE_INLINE_GENERATED_CPP_BY_NAME(StateTreeState)
 
+
+//////////////////////////////////////////////////////////////////////////
+// FStateTreeStateParameters
+
+void FStateTreeStateParameters::RemoveUnusedOverrides()
+{
+	// Remove overrides that do not exists anymore
+	if (!PropertyOverrides.IsEmpty())
+	{
+		if (const UPropertyBag* Bag = Parameters.GetPropertyBagStruct())
+		{
+			for (TArray<FGuid>::TIterator It = PropertyOverrides.CreateIterator(); It; ++It)
+			{
+				if (!Bag->FindPropertyDescByID(*It))
+				{
+					It.RemoveCurrentSwap();
+				}
+			}
+		}
+	}
+}
+
 //////////////////////////////////////////////////////////////////////////
 // FStateTreeTransition
 
@@ -346,11 +368,70 @@ void UStateTreeState::PostLoad()
 			Transition.ID = FGuid::NewDeterministicGuid(GetPathName(), Index);
 		}
 	}
+
+	if (CurrentVersion < FStateTreeCustomVersion::OverridableStateParameters)
+	{
+		// In earlier versions, all parameters were overwritten.
+		if (const UPropertyBag* Bag = Parameters.Parameters.GetPropertyBagStruct())
+		{
+			for (const FPropertyBagPropertyDesc& Desc : Bag->GetPropertyDescs())
+			{
+				Parameters.PropertyOverrides.Add(Desc.ID);
+			}
+		}
+	}
+	
 #endif // WITH_EDITORONLY_DATA
 
 }
 
 void UStateTreeState::UpdateParametersFromLinkedSubtree()
+{
+	if (const FInstancedPropertyBag* DefaultParameters = GetDefaultParameters())
+	{
+		Parameters.Parameters.MigrateToNewBagInstanceWithOverrides(*DefaultParameters, Parameters.PropertyOverrides);
+		Parameters.RemoveUnusedOverrides();
+	}
+	else
+	{
+		Parameters.Reset();
+	}
+}
+
+void UStateTreeState::SetParametersPropertyOverridden(const FGuid PropertyID, const bool bIsOverridden)
+{
+	if (bIsOverridden)
+	{
+		Parameters.PropertyOverrides.AddUnique(PropertyID);
+	}
+	else
+	{
+		Parameters.PropertyOverrides.Remove(PropertyID);
+		UpdateParametersFromLinkedSubtree();
+
+		// Remove binding when override is removed.
+		if (UStateTreeEditorData* EditorData = GetTypedOuter<UStateTreeEditorData>())
+		{
+			if (FStateTreeEditorPropertyBindings* Bindings = EditorData->GetPropertyEditorBindings())
+			{
+				if (const UPropertyBag* ParametersBag = Parameters.Parameters.GetPropertyBagStruct())
+				{
+					if (const FPropertyBagPropertyDesc* Desc = ParametersBag->FindPropertyDescByID(PropertyID))
+					{
+						check(Desc->CachedProperty);
+
+						EditorData->Modify();
+
+						FStateTreePropertyPath Path(Parameters.ID, Desc->CachedProperty->GetFName());
+						Bindings->RemovePropertyBindings(Path);
+					}
+				}
+			}
+		}
+	}
+}
+
+const FInstancedPropertyBag* UStateTreeState::GetDefaultParameters() const
 {
 	if (Type == EStateTreeStateType::Linked)
 	{
@@ -358,12 +439,7 @@ void UStateTreeState::UpdateParametersFromLinkedSubtree()
 		{
 			if (const UStateTreeState* LinkTargetState = TreeData->GetStateByID(LinkedSubtree.ID))
 			{
-				Parameters.Parameters.MigrateToNewBagInstance(LinkTargetState->Parameters.Parameters);
-			}
-			else
-			{
-				// No state selected, reset. 
-				Parameters.Parameters.Reset();
+				return &LinkTargetState->Parameters.Parameters;
 			}
 		}
 	}
@@ -371,10 +447,11 @@ void UStateTreeState::UpdateParametersFromLinkedSubtree()
 	{
 		if (LinkedAsset)
 		{
-			Parameters.Parameters.MigrateToNewBagInstance(LinkedAsset->GetDefaultParameters());
+			return &LinkedAsset->GetDefaultParameters();
 		}
-		
 	}
+
+	return nullptr;
 }
 
 const UStateTreeState* UStateTreeState::GetRootState() const
diff --git a/Engine/Plugins/Runtime/StateTree/Source/StateTreeEditorModule/Public/StateTreeState.h b/Engine/Plugins/Runtime/StateTree/Source/StateTreeEditorModule/Public/StateTreeState.h
--- a/Engine/Plugins/Runtime/StateTree/Source/StateTreeEditorModule/Public/StateTreeState.h
+++ b/Engine/Plugins/Runtime/StateTree/Source/StateTreeEditorModule/Public/StateTreeState.h
@@ -87,12 +87,20 @@ struct STATETREEEDITORMODULE_API FStateTreeStateParameters
 	void Reset()
 	{
 		Parameters.Reset();
+		PropertyOverrides.Reset();
 		bFixedLayout = false;
 	}
 
+	/** Removes overrides that do appear in Parameters. */
+	void RemoveUnusedOverrides();
+	
 	UPROPERTY(EditDefaultsOnly, Category = Parameters)
 	FInstancedPropertyBag Parameters;
 
+	/** Overrides for parameters. */
+	UPROPERTY()
+	TArray<FGuid> PropertyOverrides;
+
 	UPROPERTY(EditDefaultsOnly, Category = Parameters)
 	bool bFixedLayout = false;
 
@@ -122,6 +130,18 @@ public:
 	const UStateTreeState* GetRootState() const;
 	const UStateTreeState* GetNextSiblingState() const;
 	const UStateTreeState* GetNextSelectableSiblingState() const;
+
+	/** @return true if the property of specified ID is overridden. */
+	bool IsParametersPropertyOverridden(const FGuid PropertyID) const
+	{
+		return Parameters.PropertyOverrides.Contains(PropertyID);
+	}
+
+	/** Sets the override status of specified property by ID. */
+	void SetParametersPropertyOverridden(const FGuid PropertyID, const bool bIsOverridden);
+
+	/** @returns Default parameters from linked state or asset). */
+	const FInstancedPropertyBag* GetDefaultParameters() const;
 	
 	// StateTree Builder API
 	/** @return state link to this state. */
diff --git a/Engine/Plugins/Runtime/StateTree/Source/StateTreeModule/Public/StateTree.h b/Engine/Plugins/Runtime/StateTree/Source/StateTreeModule/Public/StateTree.h
--- a/Engine/Plugins/Runtime/StateTree/Source/StateTreeModule/Public/StateTree.h
+++ b/Engine/Plugins/Runtime/StateTree/Source/StateTreeModule/Public/StateTree.h
@@ -45,6 +45,8 @@ struct STATETREEMODULE_API FStateTreeCustomVersion
 		ChangedExternalDataAccess,
 		// Added override option for parameters
 		OverridableParameters,
+		// Added override option for state parameters
+		OverridableStateParameters,
 
 		// -----<new versions can be added above this line>-------------------------------------------------
 		VersionPlusOne,
