State Tree Details:
- Disable Category Row Buttons if Property Editing is disabled

diff --git a/Engine/Plugins/Runtime/StateTree/Source/StateTreeEditorModule/Private/Customizations/StateTreeEditorDataDetails.cpp b/Engine/Plugins/Runtime/StateTree/Source/StateTreeEditorModule/Private/Customizations/StateTreeEditorDataDetails.cpp
--- a/Engine/Plugins/Runtime/StateTree/Source/StateTreeEditorModule/Private/Customizations/StateTreeEditorDataDetails.cpp
+++ b/Engine/Plugins/Runtime/StateTree/Source/StateTreeEditorModule/Private/Customizations/StateTreeEditorDataDetails.cpp
@@ -142,7 +142,7 @@ void FStateTreeEditorDataDetails::CustomizeDetails(IDetailLayoutBuilder& DetailB
 	ParametersCategory.SetSortOrder(2);
 	{
 		// Show parameters as a category.
-		TSharedPtr<IPropertyUtilities> PropUtils = DetailBuilder.GetPropertyUtilities();
+		TSharedRef<IPropertyUtilities> PropUtils = DetailBuilder.GetPropertyUtilities();
 		TSharedPtr<IPropertyHandle> RootParametersProperty = DetailBuilder.GetProperty(TEXT("RootParameters")); // FStateTreeStateParameters
 		check(RootParametersProperty);
 		RootParametersProperty->MarkHiddenByCustomization();
@@ -150,7 +150,9 @@ void FStateTreeEditorDataDetails::CustomizeDetails(IDetailLayoutBuilder& DetailB
 		TSharedPtr<IPropertyHandle> ParametersProperty = RootParametersProperty->GetChildHandle(TEXT("Parameters")); // FInstancedPropertyBag
 		check(ParametersProperty);
 
-		TSharedRef<SHorizontalBox> HeaderContentWidget = SNew(SHorizontalBox);
+		TSharedRef<SHorizontalBox> HeaderContentWidget = SNew(SHorizontalBox)
+			.IsEnabled(PropUtils, &IPropertyUtilities::IsPropertyEditingEnabled);
+
 		HeaderContentWidget->AddSlot()
 		.HAlign(HAlign_Right)
 		.VAlign(VAlign_Center)
@@ -208,7 +210,9 @@ void FStateTreeEditorDataDetails::MakeArrayCategory(IDetailLayoutBuilder& Detail
 	IDetailCategoryBuilder& Category = DetailBuilder.EditCategory(CategoryName, DisplayName);
 	Category.SetSortOrder(SortOrder);
 
-	TSharedRef<SHorizontalBox> HeaderContentWidget = SNew(SHorizontalBox);
+	TSharedRef<SHorizontalBox> HeaderContentWidget = SNew(SHorizontalBox)
+		.IsEnabled(DetailBuilder.GetPropertyUtilities(), &IPropertyUtilities::IsPropertyEditingEnabled);
+
 	HeaderContentWidget->AddSlot()
 	.HAlign(HAlign_Right)
 	.VAlign(VAlign_Center)
diff --git a/Engine/Plugins/Runtime/StateTree/Source/StateTreeEditorModule/Private/Customizations/StateTreeStateDetails.cpp b/Engine/Plugins/Runtime/StateTree/Source/StateTreeEditorModule/Private/Customizations/StateTreeStateDetails.cpp
--- a/Engine/Plugins/Runtime/StateTree/Source/StateTreeEditorModule/Private/Customizations/StateTreeStateDetails.cpp
+++ b/Engine/Plugins/Runtime/StateTree/Source/StateTreeEditorModule/Private/Customizations/StateTreeStateDetails.cpp
@@ -132,7 +132,9 @@ void FStateTreeStateDetails::MakeArrayCategory(IDetailLayoutBuilder& DetailBuild
 	IDetailCategoryBuilder& Category = DetailBuilder.EditCategory(CategoryName, DisplayName);
 	Category.SetSortOrder(SortOrder);
 
-	const TSharedRef<SHorizontalBox> HeaderContentWidget = SNew(SHorizontalBox);
+	const TSharedRef<SHorizontalBox> HeaderContentWidget = SNew(SHorizontalBox)
+		.IsEnabled(DetailBuilder.GetPropertyUtilities(), &IPropertyUtilities::IsPropertyEditingEnabled);
+
 	HeaderContentWidget->AddSlot()
 	.HAlign(HAlign_Right)
 	.VAlign(VAlign_Center)
diff --git a/Engine/Plugins/Runtime/StateTree/Source/StateTreeEditorModule/Private/Debugger/StateTreeDebuggerUIExtensions.cpp b/Engine/Plugins/Runtime/StateTree/Source/StateTreeEditorModule/Private/Debugger/StateTreeDebuggerUIExtensions.cpp
--- a/Engine/Plugins/Runtime/StateTree/Source/StateTreeEditorModule/Private/Debugger/StateTreeDebuggerUIExtensions.cpp
+++ b/Engine/Plugins/Runtime/StateTree/Source/StateTreeEditorModule/Private/Debugger/StateTreeDebuggerUIExtensions.cpp
@@ -4,6 +4,7 @@
 #include "Customizations/StateTreeEditorNodeUtils.h"
 #include "DetailLayoutBuilder.h"
 #include "Framework/MultiBox/MultiBoxBuilder.h"
+#include "IPropertyUtilities.h"
 #include "PropertyHandle.h"
 #include "ScopedTransaction.h"
 #include "StateTreeConditionBase.h"
@@ -390,6 +391,7 @@ TSharedRef<SWidget> CreateStateWidget(const IDetailLayoutBuilder& DetailBuilder,
 	DebuggingActions.EndSection();
 
 	const TSharedRef<SHorizontalBox> HeaderContentWidget = SNew(SHorizontalBox)
+		.IsEnabled(DetailBuilder.GetPropertyUtilities(), &IPropertyUtilities::IsPropertyEditingEnabled)
 		+ SHorizontalBox::Slot()
 		.FillWidth(1.0f)
 		.HAlign(HAlign_Right)
