[StateTreeTrace] Replaced include by forward declaration now that it is supported for TMap

diff --git a/Engine/Plugins/Runtime/StateTree/Source/StateTreeModule/Public/Debugger/StateTreeTraceProvider.h b/Engine/Plugins/Runtime/StateTree/Source/StateTreeModule/Public/Debugger/StateTreeTraceProvider.h
--- a/Engine/Plugins/Runtime/StateTree/Source/StateTreeModule/Public/Debugger/StateTreeTraceProvider.h
+++ b/Engine/Plugins/Runtime/StateTree/Source/StateTreeModule/Public/Debugger/StateTreeTraceProvider.h
@@ -7,10 +7,10 @@
 #include "IStateTreeTraceProvider.h"
 #include "Model/PointTimeline.h"
 #include "StateTreeIndexTypes.h"
-#include "StateTreeTypes.h" // required to compile TMap<FStateTreeInstanceDebugId, ...>
 
 namespace TraceServices { class IAnalysisSession; }
 class UStateTree;
+struct FStateTreeInstanceDebugId;
 namespace UE::StateTreeDebugger { struct FInstanceDescriptor; }
 
 class FStateTreeTraceProvider : public IStateTreeTraceProvider
