Removing coreminimal.h from some of the most included locations

diff --git a/Engine/Plugins/Runtime/StateTree/Source/StateTreeModule/Public/Debugger/StateTreeDebuggerTypes.h b/Engine/Plugins/Runtime/StateTree/Source/StateTreeModule/Public/Debugger/StateTreeDebuggerTypes.h
--- a/Engine/Plugins/Runtime/StateTree/Source/StateTreeModule/Public/Debugger/StateTreeDebuggerTypes.h
+++ b/Engine/Plugins/Runtime/StateTree/Source/StateTreeModule/Public/Debugger/StateTreeDebuggerTypes.h
@@ -5,6 +5,7 @@
 
 #if WITH_STATETREE_DEBUGGER
 
+#include "Math/Range.h"
 #include "StateTree.h"
 #include "StateTreeTypes.h"
 #include "TraceServices/Model/Frames.h"
diff --git a/Engine/Source/Runtime/Engine/Classes/Engine/AssetUserData.h b/Engine/Source/Runtime/Engine/Classes/Engine/AssetUserData.h
--- a/Engine/Source/Runtime/Engine/Classes/Engine/AssetUserData.h
+++ b/Engine/Source/Runtime/Engine/Classes/Engine/AssetUserData.h
@@ -2,7 +2,9 @@
 
 #pragma once
 
+#if UE_ENABLE_INCLUDE_ORDER_DEPRECATED_IN_5_4
 #include "CoreMinimal.h"
+#endif //UE_ENABLE_INCLUDE_ORDER_DEPRECATED_IN_5_4
 #include "UObject/ObjectMacros.h"
 #include "UObject/Object.h"
 #include "AssetUserData.generated.h"
diff --git a/Engine/Source/Runtime/Engine/Classes/Engine/DataAsset.h b/Engine/Source/Runtime/Engine/Classes/Engine/DataAsset.h
--- a/Engine/Source/Runtime/Engine/Classes/Engine/DataAsset.h
+++ b/Engine/Source/Runtime/Engine/Classes/Engine/DataAsset.h
@@ -2,7 +2,9 @@
 
 #pragma once
 
+#if UE_ENABLE_INCLUDE_ORDER_DEPRECATED_IN_5_4
 #include "CoreMinimal.h"
+#endif //UE_ENABLE_INCLUDE_ORDER_DEPRECATED_IN_5_4
 #include "UObject/ObjectMacros.h"
 #include "UObject/Object.h"
 #include "AssetRegistry/AssetBundleData.h"
diff --git a/Engine/Source/Runtime/Engine/Classes/Engine/DataTable.h b/Engine/Source/Runtime/Engine/Classes/Engine/DataTable.h
--- a/Engine/Source/Runtime/Engine/Classes/Engine/DataTable.h
+++ b/Engine/Source/Runtime/Engine/Classes/Engine/DataTable.h
@@ -2,7 +2,9 @@
 
 #pragma once
 
+#if UE_ENABLE_INCLUDE_ORDER_DEPRECATED_IN_5_4
 #include "CoreMinimal.h"
+#endif //UE_ENABLE_INCLUDE_ORDER_DEPRECATED_IN_5_4
 #include "UObject/ObjectMacros.h"
 #include "UObject/Object.h"
 #include "UObject/Class.h"
diff --git a/Engine/Source/Runtime/GameplayTags/Classes/GameplayTagContainer.h b/Engine/Source/Runtime/GameplayTags/Classes/GameplayTagContainer.h
--- a/Engine/Source/Runtime/GameplayTags/Classes/GameplayTagContainer.h
+++ b/Engine/Source/Runtime/GameplayTags/Classes/GameplayTagContainer.h
@@ -2,7 +2,9 @@
 
 #pragma once
 
+#if UE_ENABLE_INCLUDE_ORDER_DEPRECATED_IN_5_4
 #include "CoreMinimal.h"
+#endif //UE_ENABLE_INCLUDE_ORDER_DEPRECATED_IN_5_4
 #include "Stats/Stats.h"
 #include "UObject/ObjectMacros.h"
 #include "UObject/Object.h"
diff --git a/Engine/Source/Runtime/UMG/Public/Blueprint/UserWidget.h b/Engine/Source/Runtime/UMG/Public/Blueprint/UserWidget.h
--- a/Engine/Source/Runtime/UMG/Public/Blueprint/UserWidget.h
+++ b/Engine/Source/Runtime/UMG/Public/Blueprint/UserWidget.h
@@ -2,7 +2,9 @@
 
 #pragma once
 
+#if UE_ENABLE_INCLUDE_ORDER_DEPRECATED_IN_5_4
 #include "CoreMinimal.h"
+#endif //UE_ENABLE_INCLUDE_ORDER_DEPRECATED_IN_5_4
 #include "Blueprint/WidgetChild.h"
 #include "UObject/ObjectMacros.h"
 #include "UObject/ObjectSaveContext.h"
