[StateTreeTrace] fixed performance issue when buffering lot of events (> 17K in City Sample) by switching array to map.

diff --git a/Engine/Plugins/Runtime/StateTree/Source/StateTreeModule/Private/Debugger/StateTreeTrace.cpp b/Engine/Plugins/Runtime/StateTree/Source/StateTreeModule/Private/Debugger/StateTreeTrace.cpp
--- a/Engine/Plugins/Runtime/StateTree/Source/StateTreeModule/Private/Debugger/StateTreeTrace.cpp
+++ b/Engine/Plugins/Runtime/StateTree/Source/StateTreeModule/Private/Debugger/StateTreeTrace.cpp
@@ -3,7 +3,6 @@
 #if WITH_STATETREE_DEBUGGER
 
 #include "Debugger/StateTreeTrace.h"
-#include "Algo/Removeif.h"
 #include "Debugger/StateTreeDebugger.h"
 #include "Exporters/Exporter.h"
 #include "ObjectTrace.h"
@@ -221,13 +220,11 @@ struct FInstanceEventBufferedData
 		const double RecordingWorldTime,
 		const UStateTree* StateTree,
 		const FStateTreeInstanceDebugId InstanceId,
-		const FString& InstanceName,
-		const EStateTreeTraceEventType EventType)
+		const FString& InstanceName)
 		: InstanceName(InstanceName)
 		, WeakStateTree(StateTree)
 		, InstanceId(InstanceId)
 		, LifetimeRecordingWorldTime(RecordingWorldTime)
-		, EventType(EventType)
 	{
 	}
 
@@ -241,7 +238,7 @@ struct FInstanceEventBufferedData
 				UE_TRACE_LOG(StateTreeDebugger, WorldTimestampEvent, StateTreeDebugChannel)
 					<< WorldTimestampEvent.WorldTime(LifetimeRecordingWorldTime);
 
-				OutputInstanceLifetimeEvent(InstanceId, StateTree, *InstanceName, EventType);
+				OutputInstanceLifetimeEvent(InstanceId, StateTree, *InstanceName, EStateTreeTraceEventType::Push);
 
 				if (ActiveStates.IsValid())
 				{
@@ -261,7 +258,6 @@ struct FInstanceEventBufferedData
 	FStateTreeInstanceDebugId InstanceId;
 	double LifetimeRecordingWorldTime = 0;
 	double ActiveStatesRecordingWorldTime = 0;
-	EStateTreeTraceEventType EventType = EStateTreeTraceEventType::Unset;
 };
 
 /** Struct to keep track of the buffered event data and flush them. */
@@ -279,7 +275,7 @@ struct FBufferedDataList
 	TArray<FAssetDebugIdEventBufferedData> AssetDebugIdEvents;
 
 	/** List of lifetime events that will be output if channel gets enabled in the Push - Pop lifetime window of an instance. */
-	TArray<FInstanceEventBufferedData> InstanceLifetimeEvents;
+	TMap<FStateTreeInstanceDebugId, FInstanceEventBufferedData> InstanceLifetimeEvents;
 
 	/** Flag use to prevent reentrant calls */
 	bool bFlushing = false;
@@ -302,9 +298,9 @@ struct FBufferedDataList
 
 		// Then trace instance lifetime events since they are required for other event types.
 		// It is also associated to an older world time.
-		for (const FInstanceEventBufferedData& InstanceEventData : InstanceLifetimeEvents)
+		for (TPair<FStateTreeInstanceDebugId, FInstanceEventBufferedData>& Pair : InstanceLifetimeEvents)
 		{
-			InstanceEventData.Trace();
+			Pair.Value.Trace();
 		}
 		InstanceLifetimeEvents.Empty();
 
@@ -411,7 +407,7 @@ void OutputPhaseScopeEvent(const FStateTreeInstanceDebugId InstanceId, const ESt
 			// Clear associated InstanceId when removing last entry from the stack.
 			if (PhaseStack.Stack.IsEmpty())
 			{
-				PhaseStacks.RemoveAt(ExistingStackIndex);
+				PhaseStacks.RemoveAt(ExistingStackIndex, /*Count*/1, EAllowShrinking::No);
 			}
 
 			// Phase was previously traced (i.e. other events were traced in that scope so we need to trace the closing (i.e. Pop) event.
@@ -581,16 +577,12 @@ void OutputInstanceLifetimeEvent(
 	{
 		if (EventType == EStateTreeTraceEventType::Push)
 		{
-			GBufferedEvents.InstanceLifetimeEvents.Emplace(GRecordingWorldTime, StateTree, InstanceId, InstanceName, EventType);
+			GBufferedEvents.InstanceLifetimeEvents.Emplace(InstanceId, FInstanceEventBufferedData(GRecordingWorldTime, StateTree, InstanceId, InstanceName));
 		}
 		else if (EventType == EStateTreeTraceEventType::Pop)
 		{
 			// Remove matching instance events since if it was not sent then no other events were sent between, hence not needed in the trace.
-			GBufferedEvents.InstanceLifetimeEvents.SetNum(Algo::StableRemoveIf(GBufferedEvents.InstanceLifetimeEvents,
-				[InstanceId](const FInstanceEventBufferedData& BufferedData)
-				{
-					return BufferedData.InstanceId == InstanceId;
-				}));
+			GBufferedEvents.InstanceLifetimeEvents.Remove(InstanceId);
 		}
 		else
 		{
@@ -751,19 +743,14 @@ void OutputActiveStatesEventTrace(
 	{
 		TraceBufferedEvents(InstanceId);
 
-		FInstanceEventBufferedData::FActiveStates ActiveStates(ActiveFrames);
+		const FInstanceEventBufferedData::FActiveStates ActiveStates(ActiveFrames);
 		ActiveStates.Output(InstanceId);
 	}
 	else
 	{
-		FInstanceEventBufferedData* ExisingBufferedData = GBufferedEvents.InstanceLifetimeEvents.FindByPredicate([InstanceId](const FInstanceEventBufferedData& BufferedData)
-			{
-				return BufferedData.InstanceId == InstanceId;
-			});
-
 		// We keep only the most recent active states since this is all we need to know in which state was the instance
 		// when we start receiving the events once the channel is enabled.
-		if (ExisingBufferedData != nullptr)
+		if (FInstanceEventBufferedData* ExisingBufferedData = GBufferedEvents.InstanceLifetimeEvents.Find(InstanceId))
 		{
 			ExisingBufferedData->ActiveStates= FInstanceEventBufferedData::FActiveStates(ActiveFrames);
 			ExisingBufferedData->ActiveStatesRecordingWorldTime = GRecordingWorldTime;
