StateTree: Small UI fixes
- changed the font to default font to support all supported languages (e.g. chinese was displaying wrong)
- removed bad pattern of using new FSlateRoundedBoxBrush() in editor node tails (leaks memory)
- fixed potentially dangerous use of [this] capture in FStateTreeTransitionDetails

diff --git a/Engine/Plugins/Runtime/StateTree/Source/StateTreeEditorModule/Private/Customizations/StateTreeEditorDataDetails.cpp b/Engine/Plugins/Runtime/StateTree/Source/StateTreeEditorModule/Private/Customizations/StateTreeEditorDataDetails.cpp
--- a/Engine/Plugins/Runtime/StateTree/Source/StateTreeEditorModule/Private/Customizations/StateTreeEditorDataDetails.cpp
+++ b/Engine/Plugins/Runtime/StateTree/Source/StateTreeEditorModule/Private/Customizations/StateTreeEditorDataDetails.cpp
@@ -104,7 +104,7 @@ void FStateTreeEditorDataDetails::CustomizeDetails(IDetailLayoutBuilder& DetailB
 					[
 						SNew(SBorder)
 						.Padding(FMargin(6, 1))
-						.BorderImage(new FSlateRoundedBoxBrush(FStyleColors::Hover, 6.f))
+						.BorderImage(FStateTreeEditorStyle::Get().GetBrush("StateTree.Param.Background"))
 						[
 							SNew(STextBlock)
 							.TextStyle(FStateTreeEditorStyle::Get(), "StateTree.Param.Label")
diff --git a/Engine/Plugins/Runtime/StateTree/Source/StateTreeEditorModule/Private/Customizations/StateTreeEditorNodeDetails.cpp b/Engine/Plugins/Runtime/StateTree/Source/StateTreeEditorModule/Private/Customizations/StateTreeEditorNodeDetails.cpp
--- a/Engine/Plugins/Runtime/StateTree/Source/StateTreeEditorModule/Private/Customizations/StateTreeEditorNodeDetails.cpp
+++ b/Engine/Plugins/Runtime/StateTree/Source/StateTreeEditorModule/Private/Customizations/StateTreeEditorNodeDetails.cpp
@@ -260,7 +260,7 @@ namespace UE::StateTreeEditor::Internal
 						[
 							SNew(SBorder)
 							.Padding(FMargin(6, 1))
-							.BorderImage(new FSlateRoundedBoxBrush(FStyleColors::Hover, 6))
+							.BorderImage(FStateTreeEditorStyle::Get().GetBrush("StateTree.Param.Background"))
 							.Visibility(Label.IsEmpty() ? EVisibility::Collapsed : EVisibility::Visible)
 							[
 								SNew(STextBlock)
diff --git a/Engine/Plugins/Runtime/StateTree/Source/StateTreeEditorModule/Private/Customizations/StateTreeTransitionDetails.cpp b/Engine/Plugins/Runtime/StateTree/Source/StateTreeEditorModule/Private/Customizations/StateTreeTransitionDetails.cpp
--- a/Engine/Plugins/Runtime/StateTree/Source/StateTreeEditorModule/Private/Customizations/StateTreeTransitionDetails.cpp
+++ b/Engine/Plugins/Runtime/StateTree/Source/StateTreeEditorModule/Private/Customizations/StateTreeTransitionDetails.cpp
@@ -95,9 +95,14 @@ void FStateTreeTransitionDetails::CustomizeChildren(TSharedRef<class IPropertyHa
 	check(ConditionsProperty);
 	check(IDProperty);
 
-	auto IsTickOrEventTransition = [this]()
+	TWeakPtr<FStateTreeTransitionDetails> WeakSelf = SharedThis(this);
+	auto IsTickOrEventTransition = [WeakSelf]()
 	{
-		return !EnumHasAnyFlags(GetTrigger(), EStateTreeTransitionTrigger::OnStateCompleted) ? EVisibility::Visible : EVisibility::Collapsed;
+		if (const TSharedPtr<FStateTreeTransitionDetails> Self = WeakSelf.Pin())
+		{
+			return !EnumHasAnyFlags(Self->GetTrigger(), EStateTreeTransitionTrigger::OnStateCompleted) ? EVisibility::Visible : EVisibility::Collapsed;
+		}
+		return EVisibility::Collapsed;
 	};
 
 	if (UE::StateTree::Editor::GbDisplayItemIds)
@@ -110,9 +115,13 @@ void FStateTreeTransitionDetails::CustomizeChildren(TSharedRef<class IPropertyHa
 
 	// Show event only when the trigger is set to Event. 
 	StructBuilder.AddProperty(EventTagProperty.ToSharedRef())
-		.Visibility(TAttribute<EVisibility>::Create(TAttribute<EVisibility>::FGetter::CreateLambda([this]()
+		.Visibility(TAttribute<EVisibility>::Create(TAttribute<EVisibility>::FGetter::CreateLambda([WeakSelf]()
 		{
-			return (GetTrigger() == EStateTreeTransitionTrigger::OnEvent) ? EVisibility::Visible : EVisibility::Collapsed;
+			if (const TSharedPtr<FStateTreeTransitionDetails> Self = WeakSelf.Pin())
+			{
+				return (Self->GetTrigger() == EStateTreeTransitionTrigger::OnEvent) ? EVisibility::Visible : EVisibility::Collapsed;
+			}
+			return EVisibility::Collapsed;
 		})));
 
 	// State
diff --git a/Engine/Plugins/Runtime/StateTree/Source/StateTreeEditorModule/Private/StateTreeEditorStyle.cpp b/Engine/Plugins/Runtime/StateTree/Source/StateTreeEditorModule/Private/StateTreeEditorStyle.cpp
--- a/Engine/Plugins/Runtime/StateTree/Source/StateTreeEditorModule/Private/StateTreeEditorStyle.cpp
+++ b/Engine/Plugins/Runtime/StateTree/Source/StateTreeEditorModule/Private/StateTreeEditorStyle.cpp
@@ -12,8 +12,6 @@
 #include "Styling/SlateStyleMacros.h"
 
 
-#define CORE_FONT( RelativePath, ... ) FSlateFontInfo( RootToCoreContentDir( RelativePath, TEXT( ".ttf" ) ), __VA_ARGS__ )
-
 namespace UE::StateTree::Editor
 {
 
@@ -64,13 +62,13 @@ FStateTreeEditorStyle::FStateTreeEditorStyle()
 		Set("StateTree.DetailsIcon", StateDetailsIcon);
 
 		const FTextBlockStyle StateTitle = FTextBlockStyle(NormalText)
-			.SetFont(CORE_FONT("Fonts/Roboto-Bold", 12))
+			.SetFont(DEFAULT_FONT("Bold", 12))
 			.SetColorAndOpacity(FLinearColor(230.0f / 255.0f, 230.0f / 255.0f, 230.0f / 255.0f, 0.9f));
 		Set("StateTree.State.Title", StateTitle);
 
 		const FEditableTextBoxStyle StateTitleEditableText = FEditableTextBoxStyle()
 			.SetTextStyle(NormalText)
-			.SetFont(CORE_FONT("Fonts/Roboto-Bold", 12))
+			.SetFont(DEFAULT_FONT("Bold", 12))
 			.SetBackgroundImageNormal(CORE_BOX_BRUSH("Common/TextBox", FMargin(4.0f / 16.0f)))
 			.SetBackgroundImageHovered(CORE_BOX_BRUSH("Common/TextBox_Hovered", FMargin(4.0f / 16.0f)))
 			.SetBackgroundImageFocused(CORE_BOX_BRUSH("Common/TextBox_Hovered", FMargin(4.0f / 16.0f)))
@@ -90,7 +88,7 @@ FStateTreeEditorStyle::FStateTreeEditorStyle()
 	// Details
 	{
 		const FTextBlockStyle Details = FTextBlockStyle(NormalText)
-			.SetFont(CORE_FONT("Fonts/Roboto-Regular", 10))
+			.SetFont(DEFAULT_FONT("Regular", 10))
 			.SetColorAndOpacity(FLinearColor(230.0f / 255.0f, 230.0f / 255.0f, 230.0f / 255.0f, 0.75f));
 		Set("StateTree.Details", Details);
 
@@ -107,7 +105,7 @@ FStateTreeEditorStyle::FStateTreeEditorStyle()
 	// Task
 	{
 		const FTextBlockStyle TaskTitle = FTextBlockStyle(NormalText)
-			.SetFont(CORE_FONT("Fonts/Roboto-Regular", 10))
+			.SetFont(DEFAULT_FONT("Regular", 10))
 			.SetColorAndOpacity(FLinearColor(230.0f / 255.0f, 230.0f / 255.0f, 230.0f / 255.0f, 0.85f));
 		Set("StateTree.Task.Title", TaskTitle);
 
@@ -132,15 +130,15 @@ FStateTreeEditorStyle::FStateTreeEditorStyle()
 	{
 		Set("StateTreeDebugger.Element.Normal",
 			FTextBlockStyle(NormalText)
-			.SetFont(CORE_FONT("Fonts/Roboto-Regular", 10)));
+			.SetFont(DEFAULT_FONT("Regular", 10)));
 
 		Set("StateTreeDebugger.Element.Bold",
 			FTextBlockStyle(NormalText)
-			.SetFont(CORE_FONT("Fonts/Roboto-Bold", 10)));
+			.SetFont(DEFAULT_FONT("Bold", 10)));
 
 		Set("StateTreeDebugger.Element.Subdued",
 			FTextBlockStyle(NormalText)
-			.SetFont(CORE_FONT("Fonts/Roboto-Regular", 10))
+			.SetFont(DEFAULT_FONT("Regular", 10))
 			.SetColorAndOpacity(FSlateColor::UseSubduedForeground()));
 	}
 	
@@ -179,10 +177,13 @@ FStateTreeEditorStyle::FStateTreeEditorStyle()
 		.SetFont(FAppStyle::GetFontStyle(TEXT("PropertyWindow.NormalFont")))
 		.SetFontSize(12));
 
+	// Parameter labels
 	Set("StateTree.Param.Label", FTextBlockStyle(NormalText)
 		.SetFont(FAppStyle::GetFontStyle(TEXT("PropertyWindow.BoldFont")))
 		.SetFontSize(7));
 
+	Set("StateTree.Param.Background", new FSlateRoundedBoxBrush(FStyleColors::Hover, 6.f));
+	
 	// Condition Indent combo button
 	const FButtonStyle IndentButton = FButtonStyle()
 		.SetNormal(FSlateRoundedBoxBrush(FLinearColor::Transparent, 2.0f))
