StateTree: Introduced TStateTreePropertyRefExternalHandle.
- also InstanceDataStorage is stored as a sharedptr now

diff --git a/Engine/Plugins/Runtime/StateTree/Source/StateTreeModule/Private/StateTree.cpp b/Engine/Plugins/Runtime/StateTree/Source/StateTreeModule/Private/StateTree.cpp
--- a/Engine/Plugins/Runtime/StateTree/Source/StateTreeModule/Private/StateTree.cpp
+++ b/Engine/Plugins/Runtime/StateTree/Source/StateTreeModule/Private/StateTree.cpp
@@ -435,12 +435,6 @@ bool UStateTree::Link()
 
 	if (States.Num() > 0 && Nodes.Num() > 0)
 	{
-		if (!DefaultInstanceData.IsValid())
-		{
-			UE_LOG(LogStateTree, Error, TEXT("%s: StartTree does not have instance data. Please recompile the StateTree asset."), *GetFullName());
-			return false;
-		}
-
 		// Check that all nodes are valid.
 		for (FConstStructView Node : Nodes)
 		{
@@ -757,8 +751,7 @@ TArray<FStateTreeMemoryUsage> UStateTree::CalculateEstimatedMemoryUsage() const
 
 	if (!bIsLinked
 		|| States.IsEmpty()
-		|| !Nodes.IsValid()
-		|| !DefaultInstanceData.IsValid())
+		|| !Nodes.IsValid())
 	{
 		return MemoryUsages;
 	}
@@ -949,12 +942,9 @@ TArray<FStateTreeMemoryUsage> UStateTree::CalculateEstimatedMemoryUsage() const
 	TreeMemUsage.EstimatedMemoryUsage += MaxSubtreeUsage;
 	TreeMemUsage.NodeCount += MaxSubtreeNodeCount;
 
-	if (SharedInstanceData.IsValid())
-	{
-		FStateTreeMemoryUsage& SharedMemUsage = MemoryUsages[SharedMemUsageIndex];
-		SharedMemUsage.NodeCount = SharedInstanceData.Num();
-		SharedMemUsage.EstimatedMemoryUsage = SharedInstanceData.GetEstimatedMemoryUsage();
-	}
+	FStateTreeMemoryUsage& SharedMemUsage = MemoryUsages[SharedMemUsageIndex];
+	SharedMemUsage.NodeCount = SharedInstanceData.Num();
+	SharedMemUsage.EstimatedMemoryUsage = SharedInstanceData.GetEstimatedMemoryUsage();
 
 	return MemoryUsages;
 }
diff --git a/Engine/Plugins/Runtime/StateTree/Source/StateTreeModule/Private/StateTreeExecutionContext.cpp b/Engine/Plugins/Runtime/StateTree/Source/StateTreeModule/Private/StateTreeExecutionContext.cpp
--- a/Engine/Plugins/Runtime/StateTree/Source/StateTreeModule/Private/StateTreeExecutionContext.cpp
+++ b/Engine/Plugins/Runtime/StateTree/Source/StateTreeModule/Private/StateTreeExecutionContext.cpp
@@ -184,20 +184,17 @@ EStateTreeRunStatus FStateTreeExecutionContext::Start(const FInstancedPropertyBa
 		return EStateTreeRunStatus::Failed;
 	}
 
-	if (InstanceData.IsValid())
+	FStateTreeExecutionState& Exec = GetExecState();
+	if (!ensureMsgf(Exec.CurrentPhase == EStateTreeUpdatePhase::Unset, TEXT("%hs can't be called while already in %s ('%s' using StateTree '%s')."),
+			__FUNCTION__, *UEnum::GetDisplayValueAsText(Exec.CurrentPhase).ToString(), *GetNameSafe(&Owner), *GetFullNameSafe(&RootStateTree)))
 	{
-		const FStateTreeExecutionState& Exec = GetExecState();
-		if (!ensureMsgf(Exec.CurrentPhase == EStateTreeUpdatePhase::Unset, TEXT("%hs can't be called while already in %s ('%s' using StateTree '%s')."),
-				__FUNCTION__, *UEnum::GetDisplayValueAsText(Exec.CurrentPhase).ToString(), *GetNameSafe(&Owner), *GetFullNameSafe(&RootStateTree)))
-		{
-			return EStateTreeRunStatus::Failed;
-		}
+		return EStateTreeRunStatus::Failed;
+	}
 
-		// Stop if still running previous state.
-		if (Exec.TreeRunStatus == EStateTreeRunStatus::Running)
-		{
-			Stop();
-		}
+	// Stop if still running previous state.
+	if (Exec.TreeRunStatus == EStateTreeRunStatus::Running)
+	{
+		Stop();
 	}
 
 	// Initialize instance data. No active states yet, so we'll initialize the evals and global tasks.
@@ -208,8 +205,6 @@ EStateTreeRunStatus FStateTreeExecutionContext::Start(const FInstancedPropertyBa
 		SetGlobalParameters(RootStateTree.GetDefaultParameters());
 	}
 
-	FStateTreeExecutionState& Exec = GetExecState();
-
 	// Initialize for the init frame.
 	FStateTreeExecutionFrame& InitFrame = Exec.ActiveFrames.AddDefaulted_GetRef();
 	InitFrame.StateTree = &RootStateTree;
@@ -218,12 +213,6 @@ EStateTreeRunStatus FStateTreeExecutionContext::Start(const FInstancedPropertyBa
 	InitFrame.bIsGlobalFrame = true;
 	
 	UpdateInstanceData({}, Exec.ActiveFrames);
-	if (!InstanceData.IsValid())
-	{
-		STATETREE_LOG(Warning, TEXT("%hs: Failed to initialize instance data on '%s' using StateTree '%s'. Try to recompile the StateTree asset."),
-			__FUNCTION__, *GetNameSafe(&Owner), *GetFullNameSafe(&RootStateTree));
-		return EStateTreeRunStatus::Failed;
-	}
 
 	if (!CollectActiveExternalData())
 	{
@@ -334,11 +323,6 @@ EStateTreeRunStatus FStateTreeExecutionContext::Stop(EStateTreeRunStatus Complet
 			__FUNCTION__, *GetNameSafe(&Owner), *GetFullNameSafe(&RootStateTree));
 		return EStateTreeRunStatus::Failed;
 	}
-	
-	if (!InstanceData.IsValid())
-	{
-		return EStateTreeRunStatus::Failed;
-	}
 
 	if (!CollectActiveExternalData())
 	{
@@ -424,13 +408,6 @@ EStateTreeRunStatus FStateTreeExecutionContext::Tick(const float DeltaTime)
 		return EStateTreeRunStatus::Failed;
 	}
 
-	if (!InstanceData.IsValid())
-	{
-		STATETREE_LOG(Error, TEXT("%hs: Tick called on %s using StateTree %s with invalid instance data. Start() must be called before Tick()."),
-			__FUNCTION__, *GetNameSafe(&Owner), *GetFullNameSafe(&RootStateTree));
-		return EStateTreeRunStatus::Failed;
-	}
-
 	if (!CollectActiveExternalData())
 	{
 		STATETREE_LOG(Warning, TEXT("%hs: Failed to collect external data ('%s' using StateTree '%s')"),
@@ -603,13 +580,6 @@ void FStateTreeExecutionContext::SendEvent(const FGameplayTag Tag, const FConstS
 		return;
 	}
 
-	if (!InstanceData.IsValid())
-	{
-		STATETREE_LOG(Error, TEXT("%hs: SendEvent called on %s using StateTree %s with invalid instance data. Start() must be called before sending events."),
-			__FUNCTION__, *GetNameSafe(&Owner), *GetFullNameSafe(&RootStateTree));
-		return;
-	}
-
 	STATETREE_LOG_AND_TRACE(Verbose, TEXT("Send Event '%s'"), *Tag.ToString());
 
 	FStateTreeEventQueue& EventQueue = InstanceData.GetMutableEventQueue();
@@ -627,13 +597,6 @@ void FStateTreeExecutionContext::RequestTransition(const FStateTreeTransitionReq
 		return;
 	}
 
-	if (!InstanceData.IsValid())
-	{
-		STATETREE_LOG(Error, TEXT("%hs: RequestTransition called on %s using StateTree %s with invalid instance data. Start() must be called before requesting transition."),
-			__FUNCTION__, *GetNameSafe(&Owner), *GetFullNameSafe(&RootStateTree));
-		return;
-	}
-
 	FStateTreeExecutionState& Exec = GetExecState();
 
 	if (bAllowDirectTransitions)
@@ -911,7 +874,7 @@ void FStateTreeExecutionContext::UpdateInstanceData(TConstArrayView<FStateTreeEx
 	InstanceData.ResetTemporaryInstances();
 }
 
-FStateTreeDataView FStateTreeExecutionContext::GetDataView(const FStateTreeExecutionFrame* ParentFrame, const FStateTreeExecutionFrame& CurrentFrame, const FStateTreeDataHandle Handle) const
+FStateTreeDataView FStateTreeExecutionContext::GetDataView(FStateTreeInstanceStorage& InstanceDataStorage, FStateTreeInstanceStorage* CurrentlyProcessedSharedInstanceStorage, const FStateTreeExecutionFrame* ParentFrame, const FStateTreeExecutionFrame& CurrentFrame, TConstArrayView<FStateTreeDataView> ContextAndExternalDataViews, const FStateTreeDataHandle Handle)
 {
 	switch (Handle.GetSource())
 	{
@@ -919,14 +882,14 @@ FStateTreeDataView FStateTreeExecutionContext::GetDataView(const FStateTreeExecu
 		return {};
 
 	case EStateTreeDataSourceType::GlobalInstanceData:
-		return InstanceDataStorage->GetMutableStruct(CurrentFrame.GlobalInstanceIndexBase.Get() + Handle.GetIndex());
+		return InstanceDataStorage.GetMutableStruct(CurrentFrame.GlobalInstanceIndexBase.Get() + Handle.GetIndex());
 	case EStateTreeDataSourceType::GlobalInstanceDataObject:
-		return InstanceDataStorage->GetMutableObject(CurrentFrame.GlobalInstanceIndexBase.Get() + Handle.GetIndex());
+		return InstanceDataStorage.GetMutableObject(CurrentFrame.GlobalInstanceIndexBase.Get() + Handle.GetIndex());
 		
 	case EStateTreeDataSourceType::ActiveInstanceData:
-		return InstanceDataStorage->GetMutableStruct(CurrentFrame.ActiveInstanceIndexBase.Get() + Handle.GetIndex());
+		return InstanceDataStorage.GetMutableStruct(CurrentFrame.ActiveInstanceIndexBase.Get() + Handle.GetIndex());
 	case EStateTreeDataSourceType::ActiveInstanceDataObject:
-		return InstanceDataStorage->GetMutableObject(CurrentFrame.ActiveInstanceIndexBase.Get() + Handle.GetIndex());
+		return InstanceDataStorage.GetMutableObject(CurrentFrame.ActiveInstanceIndexBase.Get() + Handle.GetIndex());
 
 	case EStateTreeDataSourceType::SharedInstanceData:
 		check(CurrentlyProcessedSharedInstanceStorage);
@@ -936,9 +899,11 @@ FStateTreeDataView FStateTreeExecutionContext::GetDataView(const FStateTreeExecu
 		return CurrentlyProcessedSharedInstanceStorage->GetMutableObject(Handle.GetIndex());
 
 	case EStateTreeDataSourceType::ContextData:
+		check(!ContextAndExternalDataViews.IsEmpty())
 		return ContextAndExternalDataViews[Handle.GetIndex()];
 
 	case EStateTreeDataSourceType::ExternalData:
+		check(!ContextAndExternalDataViews.IsEmpty())
 		return ContextAndExternalDataViews[CurrentFrame.ExternalDataBaseIndex.Get() + Handle.GetIndex()];
 
 	case EStateTreeDataSourceType::GlobalParameterData:
@@ -946,22 +911,22 @@ FStateTreeDataView FStateTreeExecutionContext::GetDataView(const FStateTreeExecu
 			// Defined in parent frame or is root state tree parameters
 			if (ParentFrame)
 			{
-				return GetDataView(nullptr, *ParentFrame, CurrentFrame.GlobalParameterDataHandle);
+				return GetDataView(InstanceDataStorage, CurrentlyProcessedSharedInstanceStorage, nullptr, *ParentFrame, ContextAndExternalDataViews, CurrentFrame.GlobalParameterDataHandle);
 			}
 
-			return InstanceDataStorage->GetMutableGlobalParameters();
+			return InstanceDataStorage.GetMutableGlobalParameters();
 		}
 
 	case EStateTreeDataSourceType::SubtreeParameterData:
 		{
 			// Defined in parent frame.
 			check(ParentFrame);
-			return GetDataView(nullptr, *ParentFrame, CurrentFrame.StateParameterDataHandle);
+			return GetDataView(InstanceDataStorage, CurrentlyProcessedSharedInstanceStorage, nullptr, *ParentFrame, ContextAndExternalDataViews, CurrentFrame.StateParameterDataHandle);
 		}
 
 	case EStateTreeDataSourceType::StateParameterData:
 		{
-			FCompactStateTreeParameters& Params = InstanceDataStorage->GetMutableStruct(CurrentFrame.ActiveInstanceIndexBase.Get() + Handle.GetIndex()).Get<FCompactStateTreeParameters>();
+			FCompactStateTreeParameters& Params = InstanceDataStorage.GetMutableStruct(CurrentFrame.ActiveInstanceIndexBase.Get() + Handle.GetIndex()).Get<FCompactStateTreeParameters>();
 			return Params.Parameters.GetMutableValue();
 		}
 	default:
@@ -3577,10 +3542,6 @@ void FStateTreeExecutionContext::DebugPrintInternalLayout()
 
 int32 FStateTreeExecutionContext::GetStateChangeCount() const
 {
-	if (!InstanceData.IsValid())
-	{
-		return 0;
-	}
 	const FStateTreeExecutionState& Exec = GetExecState();
 	return Exec.StateChangeCount;
 }
@@ -3588,12 +3549,7 @@ int32 FStateTreeExecutionContext::GetStateChangeCount() const
 #endif // WITH_STATETREE_DEBUG
 
 FString FStateTreeExecutionContext::GetActiveStateName() const
-{
-	if (!InstanceData.IsValid())
-	{
-		return FString(TEXT("<None>"));
-	}
-	
+{	
 	const FStateTreeExecutionState& Exec = GetExecState();
 
 	FString FullStateName;
@@ -3655,13 +3611,7 @@ FString FStateTreeExecutionContext::GetActiveStateName() const
 
 TArray<FName> FStateTreeExecutionContext::GetActiveStateNames() const
 {
-	TArray<FName> Result;
-
-	if (!InstanceData.IsValid())
-	{
-		return Result;
-	}
-	
+	TArray<FName> Result;	
 	const FStateTreeExecutionState& Exec = GetExecState();
 
 	// Active States
diff --git a/Engine/Plugins/Runtime/StateTree/Source/StateTreeModule/Private/StateTreeInstanceData.cpp b/Engine/Plugins/Runtime/StateTree/Source/StateTreeModule/Private/StateTreeInstanceData.cpp
--- a/Engine/Plugins/Runtime/StateTree/Source/StateTreeModule/Private/StateTreeInstanceData.cpp
+++ b/Engine/Plugins/Runtime/StateTree/Source/StateTreeModule/Private/StateTreeInstanceData.cpp
@@ -7,9 +7,14 @@
 #include "VisualLogger/VisualLogger.h"
 #include "StateTree.h"
 #include "Serialization/ObjectAndNameAsStringProxyArchive.h"
+#include "Serialization/CustomVersion.h"
+#include "Serialization/PropertyLocalizationDataGathering.h"
 
 #include UE_INLINE_GENERATED_CPP_BY_NAME(StateTreeInstanceData)
 
+const FGuid FStateTreeInstanceStorageCustomVersion::GUID(0x60C4F0DE, 0x8B264C34, 0xAA937201, 0x5DFF09CC);
+FCustomVersionRegistration GRegisterStateTreeInstanceStorageCustomVersion(FStateTreeInstanceStorageCustomVersion::GUID, FStateTreeInstanceStorageCustomVersion::LatestVersion, TEXT("StateTreeInstanceStorage"));
+
 namespace UE::StateTree
 {
 
@@ -47,6 +52,30 @@ namespace UE::StateTree
 		return DuplicateObject(&Instance, &InOwner);
 	}
 
+#if WITH_EDITORONLY_DATA
+	void GatherForLocalization(const FString& PathToParent, const UScriptStruct* Struct, const void* StructData, const void* DefaultStructData, FPropertyLocalizationDataGatherer& PropertyLocalizationDataGatherer, const EPropertyLocalizationGathererTextFlags GatherTextFlags)
+	{
+		const FStateTreeInstanceData* ThisInstance = static_cast<const FStateTreeInstanceData*>(StructData);
+		const FStateTreeInstanceData* DefaultInstance = static_cast<const FStateTreeInstanceData*>(DefaultStructData);
+
+		PropertyLocalizationDataGatherer.GatherLocalizationDataFromStruct(PathToParent, Struct, StructData, DefaultStructData, GatherTextFlags);
+
+		const uint8* DefaultInstanceMemory = nullptr;
+		if (DefaultInstance)
+		{
+			DefaultInstanceMemory = reinterpret_cast<const uint8*>(&DefaultInstance->GetStorage());
+		}
+		
+		const UScriptStruct* StructTypePtr = FStateTreeInstanceStorage::StaticStruct();
+		PropertyLocalizationDataGatherer.GatherLocalizationDataFromStructWithCallbacks(PathToParent + TEXT(".InstanceStorage"), StructTypePtr, &ThisInstance->GetStorage(), DefaultInstanceMemory, GatherTextFlags);
+	}
+
+	void RegisterInstanceDataForLocalization()
+	{
+		{ static const FAutoRegisterLocalizationDataGatheringCallback AutomaticRegistrationOfLocalizationGatherer(TBaseStructure<FStateTreeInstanceData>::Get(), &GatherForLocalization); }
+	}
+#endif // WITH_EDITORONLY_DATA
+
 } // UE::StateTree
 
 
@@ -172,9 +201,30 @@ void FStateTreeInstanceStorage::SetGlobalParameters(const FInstancedPropertyBag&
 // FStateTreeInstanceData
 //----------------------------------------------------------------//
 
-FStateTreeInstanceData::FStateTreeInstanceData()
+FStateTreeInstanceData::FStateTreeInstanceData() = default;
+
+FStateTreeInstanceData::FStateTreeInstanceData(const FStateTreeInstanceData& Other)
+{
+	InstanceStorage = MakeShared<FStateTreeInstanceStorage>(*Other.InstanceStorage);
+}
+
+FStateTreeInstanceData::FStateTreeInstanceData(FStateTreeInstanceData&& Other)
+{
+	InstanceStorage = Other.InstanceStorage;
+	Other.InstanceStorage = MakeShared<FStateTreeInstanceStorage>();
+}
+
+FStateTreeInstanceData& FStateTreeInstanceData::operator=(const FStateTreeInstanceData& Other)
 {
-	InstanceStorage.InitializeAs<FStateTreeInstanceStorage>();
+	InstanceStorage = MakeShared<FStateTreeInstanceStorage>(*Other.InstanceStorage);
+	return *this;
+}
+
+FStateTreeInstanceData& FStateTreeInstanceData::operator=(FStateTreeInstanceData&& Other)
+{
+	InstanceStorage = Other.InstanceStorage;
+	Other.InstanceStorage = MakeShared<FStateTreeInstanceStorage>();
+	return *this;
 }
 
 FStateTreeInstanceData::~FStateTreeInstanceData()
@@ -184,14 +234,22 @@ FStateTreeInstanceData::~FStateTreeInstanceData()
 
 const FStateTreeInstanceStorage& FStateTreeInstanceData::GetStorage() const
 {
-	check(InstanceStorage.GetMemory() != nullptr);
-	return *reinterpret_cast<const FStateTreeInstanceStorage*>(InstanceStorage.GetMemory());
+	return *InstanceStorage;
+}
+
+TWeakPtr<FStateTreeInstanceStorage> FStateTreeInstanceData::GetWeakMutableStorage()
+{
+	return InstanceStorage;
+}
+
+TWeakPtr<const FStateTreeInstanceStorage> FStateTreeInstanceData::GetWeakStorage() const
+{
+	return InstanceStorage;
 }
 
 FStateTreeInstanceStorage& FStateTreeInstanceData::GetMutableStorage()
 {
-	check(InstanceStorage.GetMemory() != nullptr);
-	return *reinterpret_cast<FStateTreeInstanceStorage*>(InstanceStorage.GetMutableMemory());
+	return *InstanceStorage;
 }
 
 FStateTreeEventQueue& FStateTreeInstanceData::GetMutableEventQueue()
@@ -252,18 +310,6 @@ bool FStateTreeInstanceData::Identical(const FStateTreeInstanceData* Other, uint
 		return false;
 	}
 
-	// Identical if both are uninitialized.
-	if (!IsValid() && !Other->IsValid())
-	{
-		return true;
-	}
-
-	// Not identical if one is valid and other is not.
-	if (IsValid() != Other->IsValid())
-	{
-		return false;
-	}
-
 	const FStateTreeInstanceStorage& Storage = GetStorage();
 	const FStateTreeInstanceStorage& OtherStorage = Other->GetStorage();
 
@@ -330,6 +376,65 @@ bool FStateTreeInstanceData::Identical(const FStateTreeInstanceData* Other, uint
 	return bResult;
 }
 
+void FStateTreeInstanceData::AddStructReferencedObjects(FReferenceCollector& Collector)
+{
+	Collector.AddPropertyReferencesWithStructARO(FStateTreeInstanceStorage::StaticStruct(), &GetMutableStorage());
+}
+
+bool FStateTreeInstanceData::Serialize(FArchive& Ar)
+{
+	Ar.UsingCustomVersion(FStateTreeInstanceStorageCustomVersion::GUID);
+
+	if (Ar.IsLoading())
+	{
+		if (Ar.CustomVer(FStateTreeInstanceStorageCustomVersion::GUID) < FStateTreeInstanceStorageCustomVersion::AddedCustomSerialization)
+		{
+#if WITH_EDITORONLY_DATA
+PRAGMA_DISABLE_DEPRECATION_WARNINGS
+			StaticStruct()->SerializeTaggedProperties(Ar, (uint8*)this, StaticStruct(), nullptr);
+
+			if (InstanceStorage_DEPRECATED.IsValid())
+			{
+				InstanceStorage = MakeShared<FStateTreeInstanceStorage>(MoveTemp(InstanceStorage_DEPRECATED.GetMutable()));
+				InstanceStorage_DEPRECATED.Reset();
+				return true;
+			}
+PRAGMA_ENABLE_DEPRECATION_WARNINGS
+#endif // WITH_EDITORONLY_DATA
+
+			InstanceStorage = MakeShared<FStateTreeInstanceStorage>();
+			return true;
+		}
+
+		InstanceStorage = MakeShared<FStateTreeInstanceStorage>();
+	}
+
+	FStateTreeInstanceStorage::StaticStruct()->SerializeItem(Ar, &InstanceStorage.Get(), nullptr);
+
+	return true;
+}
+
+void FStateTreeInstanceData::GetPreloadDependencies(TArray<UObject*>& OutDeps)
+{
+	UScriptStruct* ScriptStruct = FStateTreeInstanceStorage::StaticStruct();
+	OutDeps.Add(ScriptStruct);
+
+	if (UScriptStruct::ICppStructOps* CppStructOps = ScriptStruct->GetCppStructOps())
+	{
+		CppStructOps->GetPreloadDependencies(&GetMutableStorage(), OutDeps);
+	}
+
+	for (TPropertyValueIterator<FStructProperty> It(ScriptStruct, &GetMutableStorage()); It; ++It)
+	{
+		const UScriptStruct* StructType = It.Key()->Struct;
+		if (UScriptStruct::ICppStructOps* CppStructOps = StructType->GetCppStructOps())
+		{
+			void* StructDataPtr = const_cast<void*>(It.Value());
+			CppStructOps->GetPreloadDependencies(StructDataPtr, OutDeps);
+		}
+	}
+}
+
 void FStateTreeInstanceData::CopyFrom(UObject& InOwner, const FStateTreeInstanceData& InOther)
 {
 	if (&InOther == this)
@@ -446,11 +551,6 @@ void FStateTreeInstanceData::ShrinkTo(const int32 NumStructs)
 	Storage.InstanceStructs.SetNum(NumStructs);
 }
 
-bool FStateTreeInstanceData::IsValid() const
-{
-	return InstanceStorage.IsValid();
-}
-
 void FStateTreeInstanceData::Reset()
 {
 	FStateTreeInstanceStorage& Storage = GetMutableStorage();
diff --git a/Engine/Plugins/Runtime/StateTree/Source/StateTreeModule/Private/StateTreeModule.cpp b/Engine/Plugins/Runtime/StateTree/Source/StateTreeModule/Private/StateTreeModule.cpp
--- a/Engine/Plugins/Runtime/StateTree/Source/StateTreeModule/Private/StateTreeModule.cpp
+++ b/Engine/Plugins/Runtime/StateTree/Source/StateTreeModule/Private/StateTreeModule.cpp
@@ -20,6 +20,10 @@
 
 #endif // WITH_STATETREE_DEBUGGER
 
+#if WITH_EDITORONLY_DATA
+#include "StateTreeInstanceData.h"
+#endif // WITH_EDITORONLY_DATA
+
 #define LOCTEXT_NAMESPACE "StateTree"
 
 class FStateTreeModule : public IStateTreeModule
@@ -102,6 +106,10 @@ void FStateTreeModule::StartupModule()
 #endif // !WITH_EDITOR
 
 #endif // WITH_STATETREE_DEBUGGER
+
+#if WITH_EDITORONLY_DATA
+	UE::StateTree::RegisterInstanceDataForLocalization();
+#endif // WITH_EDITORONLY_DATA
 }
 
 void FStateTreeModule::ShutdownModule()
diff --git a/Engine/Plugins/Runtime/StateTree/Source/StateTreeModule/Public/StateTree.h b/Engine/Plugins/Runtime/StateTree/Source/StateTreeModule/Public/StateTree.h
--- a/Engine/Plugins/Runtime/StateTree/Source/StateTreeModule/Public/StateTree.h
+++ b/Engine/Plugins/Runtime/StateTree/Source/StateTreeModule/Public/StateTree.h
@@ -156,6 +156,9 @@ public:
 	/** @return Id of the transition matching a given runtime transition index; invalid Id if transition not found. */
 	FGuid GetTransitionIdFromIndex(const FStateTreeIndex16 Index) const;	
 
+	/** @return Property bindings */
+	const FStateTreePropertyBindings& GetPropertyBindings() const { return PropertyBindings; }
+
 	UE_DEPRECATED(5.4, "Replaced with GetNumContextDataViews() which contains context data and external data only.")
 	int32 GetNumDataViews() const { return 0; }
 
diff --git a/Engine/Plugins/Runtime/StateTree/Source/StateTreeModule/Public/StateTreeExecutionContext.h b/Engine/Plugins/Runtime/StateTree/Source/StateTreeModule/Public/StateTreeExecutionContext.h
--- a/Engine/Plugins/Runtime/StateTree/Source/StateTreeModule/Public/StateTreeExecutionContext.h
+++ b/Engine/Plugins/Runtime/StateTree/Source/StateTreeModule/Public/StateTreeExecutionContext.h
@@ -128,23 +128,6 @@ public:
 
 	/** @return True of the the execution context is valid and initialized. */ 
 	bool IsValid() const { return RootStateTree.IsReadyToRun(); }
-
-	/**
-	 * @param PropertyRef Property's reference to get pointer to.
-	 * @return Pointer to referenced property if succeeded.
-	 */
-	template<class T>
-	T* GetMutablePropertyPtr(const FStateTreePropertyRef& PropertyRef)
-	{
-		const FStateTreePropertyBindings& PropertyBindings = CurrentlyProcessedFrame->StateTree->PropertyBindings;
-		if(const FStateTreePropertyAccess* PropertyAccess = PropertyBindings.GetPropertyAccess(PropertyRef))
-		{
-			FStateTreeDataView SourceView = GetDataView(CurrentlyProcessedParentFrame, *CurrentlyProcessedFrame, PropertyAccess->SourceDataHandle);
-			return PropertyBindings.GetMutablePropertyPtr<T>(SourceView, *PropertyAccess);
-		}
-
-		return nullptr;
-	}
 	
 	/**
 	 * Start executing.
@@ -240,6 +223,12 @@ public:
 
 	/** @return the currently processed state if applicable. */
 	FStateTreeStateHandle GetCurrentlyProcessedState() const { return CurrentlyProcessedState; }
+
+	/** @return the currently processed execution frame if applicable. */
+	const FStateTreeExecutionFrame* GetCurrentlyProcessedFrame() const { return CurrentlyProcessedFrame; }
+
+	/** @return the currently processed execution parent frame if applicable. */
+	const FStateTreeExecutionFrame* GetCurrentlyProcessedParentFrame() const { return CurrentlyProcessedParentFrame; }
 	
 	/** @return Pointer to a State or null if state not found */ 
 	const FCompactStateTreeState* GetStateFromHandle(const FStateTreeStateHandle StateHandle) const
@@ -382,6 +371,9 @@ public:
 	 */
 	void RequestTransition(const FStateTreeTransitionRequest& Request);
 
+	/** @return data view of the specified handle relative to given frame. */
+	static FStateTreeDataView GetDataView(FStateTreeInstanceStorage& InstanceDataStorage, FStateTreeInstanceStorage* CurrentlyProcessedSharedInstanceStorage, const FStateTreeExecutionFrame* ParentFrame, const FStateTreeExecutionFrame& CurrentFrame, TConstArrayView<FStateTreeDataView> ContextAndExternalDataViews, const FStateTreeDataHandle Handle);
+
 protected:
 
 	/** Max number of execution frames handled during state selection. */
@@ -528,9 +520,12 @@ protected:
 	FString DebugGetEventsAsString() const;
 
 	/** @return data view of the specified handle relative to given frame. */
-	FStateTreeDataView GetDataView(const FStateTreeExecutionFrame* ParentFrame, const FStateTreeExecutionFrame& CurrentFrame, const FStateTreeDataHandle Handle) const;
+	FStateTreeDataView GetDataView(const FStateTreeExecutionFrame* ParentFrame, const FStateTreeExecutionFrame& CurrentFrame, const FStateTreeDataHandle Handle) const
+	{
+		return GetDataView(*InstanceDataStorage, CurrentlyProcessedSharedInstanceStorage, ParentFrame, CurrentFrame, ContextAndExternalDataViews, Handle);
+	}
 
-	/** @return data view of the specified handle relative to given frame. */
+	/** @return true if handle source is valid cified handle relative to given frame. */
 	bool IsHandleSourceValid(const FStateTreeExecutionFrame* ParentFrame, const FStateTreeExecutionFrame& CurrentFrame, const FStateTreeDataHandle Handle) const;
 
 	/** @return data view of the specified handle relative to the given frame, or tries to find a matching temporary instance. */
diff --git a/Engine/Plugins/Runtime/StateTree/Source/StateTreeModule/Public/StateTreeInstanceData.h b/Engine/Plugins/Runtime/StateTree/Source/StateTreeModule/Public/StateTreeInstanceData.h
--- a/Engine/Plugins/Runtime/StateTree/Source/StateTreeModule/Public/StateTreeInstanceData.h
+++ b/Engine/Plugins/Runtime/StateTree/Source/StateTreeModule/Public/StateTreeInstanceData.h
@@ -6,6 +6,7 @@
 #include "StateTreeEvents.h"
 #include "StateTreeTypes.h"
 #include "StateTreeExecutionTypes.h"
+#include "Templates/SharedPointer.h"
 #include "StateTreeInstanceData.generated.h"
 
 struct FStateTreeTransitionRequest;
@@ -49,6 +50,27 @@ struct STATETREEMODULE_API FStateTreeTemporaryInstanceData
 	FInstancedStruct Instance;
 };
 
+struct STATETREEMODULE_API FStateTreeInstanceStorageCustomVersion
+{
+	enum Type
+	{
+		// Before any version changes were made in the plugin
+		BeforeCustomVersionWasAdded = 0,
+		// Added custom serialization
+		AddedCustomSerialization,
+
+		// -----<new versions can be added above this line>-------------------------------------------------
+		VersionPlusOne,
+		LatestVersion = VersionPlusOne - 1
+	};
+
+	/** The GUID for this custom version number */
+	const static FGuid GUID;
+
+private:
+	FStateTreeInstanceStorageCustomVersion() = default;
+};
+
 /**
  * State Tree instance data is used to store the runtime state of a State Tree. It is used together with FStateTreeExecution context to tick the state tree.
  * You are supposed to use FStateTreeInstanceData as a property to store the instance data. That ensures that any UObject references will get GC'd correctly.
@@ -243,6 +265,12 @@ struct STATETREEMODULE_API FStateTreeInstanceData
 	GENERATED_BODY()
 
 	FStateTreeInstanceData();
+	FStateTreeInstanceData(const FStateTreeInstanceData& Other);
+	FStateTreeInstanceData(FStateTreeInstanceData&& Other);
+
+	FStateTreeInstanceData& operator=(const FStateTreeInstanceData& Other);
+	FStateTreeInstanceData& operator=(FStateTreeInstanceData&& Other);
+
 	~FStateTreeInstanceData();
 	
 	/** Initializes the array with specified items. */
@@ -265,9 +293,6 @@ struct STATETREEMODULE_API FStateTreeInstanceData
 	/** Resets the data to empty. */
 	void Reset();
 
-	/** @return true if the instance is correctly initialized. */
-	bool IsValid() const;
-
 	/** @return Number of items in the instance data. */
 	int32 Num() const
 	{
@@ -345,10 +370,16 @@ struct STATETREEMODULE_API FStateTreeInstanceData
 	FStateTreeInstanceStorage& GetMutableStorage();
 	const FStateTreeInstanceStorage& GetStorage() const;
 
+	TWeakPtr<FStateTreeInstanceStorage> GetWeakMutableStorage();
+	TWeakPtr<const FStateTreeInstanceStorage> GetWeakStorage() const;
+
 	int32 GetEstimatedMemoryUsage() const;
 	
 	/** Type traits */
 	bool Identical(const FStateTreeInstanceData* Other, uint32 PortFlags) const;
+	void AddStructReferencedObjects(FReferenceCollector& Collector);
+	bool Serialize(FArchive& Ar);
+	void GetPreloadDependencies(TArray<UObject*>& OutDeps);
 
 	/**
 	 * Adds temporary instance data associated with specified frame and data handle.
@@ -404,18 +435,39 @@ struct STATETREEMODULE_API FStateTreeInstanceData
 	UE_DEPRECATED(5.4, "Use Num() instead.")
 	int32 GetNumItems() const { return 0; }
 
+	UE_DEPRECATED(5.4, "InstanceData is always valid.")
+	bool IsValid() const { return true; }
+
 protected:
 	/** Storage for the actual instance data, always stores FStateTreeInstanceStorage. */
+	TSharedRef<FStateTreeInstanceStorage> InstanceStorage = MakeShared<FStateTreeInstanceStorage>();
+
+#if WITH_EDITORONLY_DATA
+PRAGMA_DISABLE_DEPRECATION_WARNINGS
+
 	UPROPERTY()
-	TInstancedStruct<FStateTreeInstanceStorage> InstanceStorage;
+	TInstancedStruct<FStateTreeInstanceStorage> InstanceStorage_DEPRECATED;
+	
+PRAGMA_ENABLE_DEPRECATION_WARNINGS
+#endif // WITH_EDITORONLY_DATA
 };
 
+#if WITH_EDITORONLY_DATA
+namespace UE::StateTree
+{
+	void RegisterInstanceDataForLocalization();
+}
+#endif // WITH_EDITORONLY_DATA
+
 template<>
 struct TStructOpsTypeTraits<FStateTreeInstanceData> : public TStructOpsTypeTraitsBase2<FStateTreeInstanceData>
 {
 	enum
 	{
 		WithIdentical = true,
+		WithAddStructReferencedObjects = true,
+		WithSerializer = true,
+		WithGetPreloadDependencies = true,
 	};
 };
 
@@ -452,7 +504,7 @@ template <typename T>
 struct TStateTreeInstanceDataStructRef
 {
 	TStateTreeInstanceDataStructRef(FStateTreeInstanceData& InInstanceData, const FStateTreeExecutionFrame& CurrentFrame, const FStateTreeDataHandle InDataHandle)
-		: Storage(InInstanceData.GetMutableStorage())
+		: WeakStorage(InInstanceData.GetWeakMutableStorage())
 		, WeakStateTree(CurrentFrame.StateTree)
 		, RootState(CurrentFrame.RootState)
 		, DataHandle(InDataHandle)
@@ -466,6 +518,13 @@ struct TStateTreeInstanceDataStructRef
 
 	T* GetPtr()
 	{
+		if (!WeakStorage.IsValid())
+		{
+			return nullptr;
+		}
+
+		FStateTreeInstanceStorage& Storage = *WeakStorage.Pin();
+
 		const FStateTreeExecutionState& Exec = Storage.GetExecutionState();
 		const UStateTree* StateTree = WeakStateTree.Get();
 		
@@ -477,9 +536,9 @@ struct TStateTreeInstanceDataStructRef
 		FStructView Struct;
 		if (CurrentFrame)
 		{
-			if (IsHandleSourceValid(*CurrentFrame, DataHandle))
+			if (IsHandleSourceValid(Storage, *CurrentFrame, DataHandle))
 			{
-				Struct = GetDataView(*CurrentFrame, DataHandle);
+				Struct = GetDataView(Storage, *CurrentFrame, DataHandle);
 			}
 			else
 			{
@@ -501,7 +560,7 @@ struct TStateTreeInstanceDataStructRef
 
 protected:
 
-	FStructView GetDataView(const FStateTreeExecutionFrame& CurrentFrame, const FStateTreeDataHandle Handle)
+	FStructView GetDataView(FStateTreeInstanceStorage& Storage, const FStateTreeExecutionFrame& CurrentFrame, const FStateTreeDataHandle Handle)
 	{
 		switch (DataHandle.GetSource())
 		{
@@ -515,7 +574,7 @@ protected:
 		return {};
 	}
 	
-	bool IsHandleSourceValid(const FStateTreeExecutionFrame& CurrentFrame, const FStateTreeDataHandle Handle) const
+	bool IsHandleSourceValid(FStateTreeInstanceStorage& Storage, const FStateTreeExecutionFrame& CurrentFrame, const FStateTreeDataHandle Handle) const
 	{
 		switch (Handle.GetSource())
 		{
@@ -533,7 +592,7 @@ protected:
 		return false;
 	}
 	
-	FStateTreeInstanceStorage& Storage;
+	TWeakPtr<FStateTreeInstanceStorage> WeakStorage = nullptr;
 	TWeakObjectPtr<const UStateTree> WeakStateTree = nullptr;
 	FStateTreeStateHandle RootState = FStateTreeStateHandle::Invalid;
 	FStateTreeDataHandle DataHandle = FStateTreeDataHandle::Invalid;
diff --git a/Engine/Plugins/Runtime/StateTree/Source/StateTreeModule/Public/StateTreePropertyRef.h b/Engine/Plugins/Runtime/StateTree/Source/StateTreeModule/Public/StateTreePropertyRef.h
--- a/Engine/Plugins/Runtime/StateTree/Source/StateTreeModule/Public/StateTreePropertyRef.h
+++ b/Engine/Plugins/Runtime/StateTree/Source/StateTreeModule/Public/StateTreePropertyRef.h
@@ -6,6 +6,32 @@
 #include "StateTreeExecutionContext.h"
 #include "StateTreePropertyRef.generated.h"
 
+struct FStateTreePropertyRef;
+
+namespace UE::StateTree::PropertyRefHelpers
+{
+    /**
+	 * @param PropertyRef Property's reference to get pointer to.
+	 * @param InstanceDataStorage Instance Data Storage
+	 * @param ExecutionFrame Execution frame owning referenced property
+	 * @param ParentExecutionFrame Parent of execution frame owning referenced property
+	 * @return Pointer to referenced property if succeeded.
+	 */
+	template<class T>
+	static T* GetMutablePtrToProperty(const FStateTreePropertyRef& PropertyRef, FStateTreeInstanceStorage& InstanceDataStorage, const FStateTreeExecutionFrame& ExecutionFrame, const FStateTreeExecutionFrame* ParentExecutionFrame)
+	{
+		const FStateTreePropertyBindings& PropertyBindings = ExecutionFrame.StateTree->GetPropertyBindings();
+		if (const FStateTreePropertyAccess* PropertyAccess = PropertyBindings.GetPropertyAccess(PropertyRef))
+		{
+			// Passing empty ContextAndExternalDataViews, as PropertyRef is not allowed to point to context or external data.
+			FStateTreeDataView SourceView = FStateTreeExecutionContext::GetDataView(InstanceDataStorage, nullptr, ParentExecutionFrame, ExecutionFrame, {}, PropertyAccess->SourceDataHandle);
+			return PropertyBindings.GetMutablePropertyPtr<T>(SourceView, *PropertyAccess);
+		}
+
+		return nullptr;
+	}
+}
+
 /**
  * Property ref allows to get a pointer to selected property in StateTree.
  * The expected type of the reference should be set in "RefType" meta specifier.
@@ -44,7 +70,10 @@ struct STATETREEMODULE_API FStateTreePropertyRef
 	template<class T>
 	T* GetMutablePtr(FStateTreeExecutionContext& Context) const
 	{
-		return Context.GetMutablePropertyPtr<T>(*this);
+		const FStateTreeExecutionFrame* CurrentlyProcessedFrame = Context.GetCurrentlyProcessedFrame();
+		check(CurrentlyProcessedFrame);
+
+		return UE::StateTree::PropertyRefHelpers::GetMutablePtrToProperty<T>(*this, Context.GetMutableInstanceData()->GetMutableStorage(), *CurrentlyProcessedFrame, Context.GetCurrentlyProcessedParentFrame());
 	}
 
 	/**
@@ -90,6 +119,66 @@ struct TStateTreePropertyRef
 		return PropertyRef.GetMutablePtr<TRef>(Context);
 	}
 
+	/**
+	 * Used internally.
+	 * @return internal property ref
+	 */
+	FStateTreePropertyRef GetInternalPropertyRef() const
+	{
+		return PropertyRef;
+	}
+
 private:
 	FStateTreePropertyRef PropertyRef;
-};
\ No newline at end of file
+};
+
+/**
+ * External Handle allows to wrap-up property reference to make it accessible without having an access to StateTreeExecutionContext. Useful for capturing property reference in callbacks.
+ */
+template<class TRef>
+struct TStateTreePropertyRefExternalHandle
+{
+	TStateTreePropertyRefExternalHandle(FStateTreePropertyRef InPropertyRef, FStateTreeExecutionContext& InContext)
+		: WeakInstanceStorage(InContext.GetMutableInstanceData()->GetWeakMutableStorage())
+		, WeakStateTree(InContext.GetCurrentlyProcessedFrame()->StateTree)
+		, RootState(InContext.GetCurrentlyProcessedFrame()->RootState)
+		, PropertyRef(InPropertyRef)
+	{}
+
+	TStateTreePropertyRefExternalHandle(TStateTreePropertyRef<TRef> InPropertyRef, FStateTreeExecutionContext& InContext)
+		: TStateTreePropertyRefExternalHandle(InPropertyRef.GetInternalPropertyRef(), InContext)
+	{}
+
+	/** @return pointer to the property if possible, nullptr otherwise. */
+	TRef* GetMutablePtr() const
+	{
+		if (!WeakInstanceStorage.IsValid())
+		{
+			return nullptr;
+		}
+
+		FStateTreeInstanceStorage& InstanceStorage = *WeakInstanceStorage.Pin();
+
+		const TArray<FStateTreeExecutionFrame>& ActiveFrames = InstanceStorage.GetExecutionState().ActiveFrames;
+		const int32 FrameIndex = ActiveFrames.IndexOfByPredicate([this](const FStateTreeExecutionFrame& Frame)
+		{
+			return Frame.RootState == RootState && Frame.StateTree == WeakStateTree;
+		});
+
+		if (FrameIndex == INDEX_NONE)
+		{
+			return nullptr;
+		}
+
+		const FStateTreeExecutionFrame& Frame = ActiveFrames[FrameIndex];
+		const FStateTreeExecutionFrame* ParentFrame = FrameIndex > 0 ? &ActiveFrames[FrameIndex - 1] : nullptr;
+
+		return UE::StateTree::PropertyRefHelpers::GetMutablePtrToProperty<TRef>(PropertyRef, InstanceStorage, Frame, ParentFrame);
+	}
+
+private:
+	TWeakPtr<FStateTreeInstanceStorage> WeakInstanceStorage;
+	TWeakObjectPtr<const UStateTree> WeakStateTree = nullptr;
+	FStateTreeStateHandle RootState = FStateTreeStateHandle::Invalid;
+	FStateTreePropertyRef PropertyRef;
+};
