StateTree: Fixed crash when resolving bindings without value into Weak or Soft Objects.

diff --git a/Engine/Plugins/Runtime/StateTree/Source/StateTreeModule/Private/StateTreePropertyBindings.cpp b/Engine/Plugins/Runtime/StateTree/Source/StateTreeModule/Private/StateTreePropertyBindings.cpp
--- a/Engine/Plugins/Runtime/StateTree/Source/StateTreeModule/Private/StateTreePropertyBindings.cpp
+++ b/Engine/Plugins/Runtime/StateTree/Source/StateTreeModule/Private/StateTreePropertyBindings.cpp
@@ -1727,19 +1727,45 @@ bool FStateTreePropertyPath::ResolveIndirectionsWithValue(const FStateTreeDataVi
 			// Check to see if this is a simple weak object property (eg. not an array of weak objects).
 			else if (const FWeakObjectProperty* WeakObjectProperty = CastField<FWeakObjectProperty>(Property))
 			{
-				const TWeakObjectPtr<UObject>& WeakObjectPtr = *reinterpret_cast<const TWeakObjectPtr<UObject>*>(CurrentAddress + Offset);
-				const UObject* Object = WeakObjectPtr.Get();
-				CurrentAddress = reinterpret_cast<const uint8*>(Object);
-				CurrentStruct = WeakObjectProperty->PropertyClass;
+				if (bWithValue)
+				{
+					const TWeakObjectPtr<UObject>& WeakObjectPtr = *reinterpret_cast<const TWeakObjectPtr<UObject>*>(CurrentAddress + Offset);
+					const UObject* Object = WeakObjectPtr.Get();
+					CurrentAddress = reinterpret_cast<const uint8*>(Object);
+
+					if (Object)
+					{
+						CurrentStruct = Object->GetClass();
+						Indirection.InstanceStruct = CurrentStruct;
+					}
+				}
+				else
+				{
+					CurrentStruct = WeakObjectProperty->PropertyClass;
+				}
+				
 				Indirection.AccessType = EStateTreePropertyAccessType::WeakObject;
 			}
 			// Check to see if this is a simple soft object property (eg. not an array of soft objects).
 			else if (const FSoftObjectProperty* SoftObjectProperty = CastField<FSoftObjectProperty>(Property))
 			{
-				const FSoftObjectPtr& SoftObjectPtr = *reinterpret_cast<const FSoftObjectPtr*>(CurrentAddress + Offset);
-				const UObject* Object = SoftObjectPtr.Get();
-				CurrentAddress = reinterpret_cast<const uint8*>(Object);
-				CurrentStruct = SoftObjectProperty->PropertyClass;
+				if (bWithValue)
+				{
+					const FSoftObjectPtr& SoftObjectPtr = *reinterpret_cast<const FSoftObjectPtr*>(CurrentAddress + Offset);
+					const UObject* Object = SoftObjectPtr.Get();
+					CurrentAddress = reinterpret_cast<const uint8*>(Object);
+
+					if (Object)
+					{
+						CurrentStruct = Object->GetClass();
+						Indirection.InstanceStruct = CurrentStruct;
+					}			
+				}
+				else
+				{			
+					CurrentStruct = SoftObjectProperty->PropertyClass;
+				}
+
 				Indirection.AccessType = EStateTreePropertyAccessType::SoftObject;
 			}
 			else
