StateTree: Blueprint Property Ref.
- State Tree property ref which can be used in BP tasks

diff --git a/Engine/Plugins/Runtime/StateTree/Source/StateTreeEditorModule/Private/Customizations/StateTreeBindingExtension.cpp b/Engine/Plugins/Runtime/StateTree/Source/StateTreeEditorModule/Private/Customizations/StateTreeBindingExtension.cpp
--- a/Engine/Plugins/Runtime/StateTree/Source/StateTreeEditorModule/Private/Customizations/StateTreeBindingExtension.cpp
+++ b/Engine/Plugins/Runtime/StateTree/Source/StateTreeEditorModule/Private/Customizations/StateTreeBindingExtension.cpp
@@ -38,7 +38,7 @@ UObject* FindEditorBindingsOwner(UObject* InObject)
 	return Result;
 }
 
-UStruct* ResolveLeafValueStructType(FStateTreeDataView ValueView, const TArray<FBindingChainElement>& InBindingChain)
+UStruct* ResolveLeafValueStructType(FStateTreeDataView ValueView, TConstArrayView<FBindingChainElement> InBindingChain)
 {
 	if (ValueView.GetMemory() == nullptr)
 	{
@@ -96,7 +96,7 @@ UStruct* ResolveLeafValueStructType(FStateTreeDataView ValueView, const TArray<F
 	return Result;
 }
 
-void MakeStructPropertyPathFromBindingChain(const FGuid StructID, const TArray<FBindingChainElement>& InBindingChain, FStateTreeDataView DataView, FStateTreePropertyPath& OutPath)
+void MakeStructPropertyPathFromBindingChain(const FGuid StructID, TConstArrayView<FBindingChainElement> InBindingChain, FStateTreeDataView DataView, FStateTreePropertyPath& OutPath)
 {
 	OutPath.Reset();
 	OutPath.SetStructID(StructID);
@@ -268,7 +268,7 @@ struct FCachedBindingData : public TSharedFromThis<FCachedBindingData>
 	{
 	}
 
-	void AddBinding(const TArray<FBindingChainElement>& InBindingChain)
+	void AddBinding(TConstArrayView<FBindingChainElement> InBindingChain)
 	{
 		if (InBindingChain.IsEmpty())
 		{
@@ -302,8 +302,7 @@ struct FCachedBindingData : public TSharedFromThis<FCachedBindingData>
 		const int32 SourceStructIndex = InBindingChain[0].ArrayIndex;
 		check(SourceStructIndex >= 0 && SourceStructIndex < AccessibleStructs.Num());
 				
-		TArray<FBindingChainElement> SourceBindingChain = InBindingChain;
-		SourceBindingChain.RemoveAt(0); // remove struct index.
+		TConstArrayView<FBindingChainElement> SourceBindingChain(InBindingChain.begin() + 1, InBindingChain.Num() - 1); // remove struct index.
 
 		FStateTreeDataView DataView;
 		BindingOwner->GetDataViewByID(AccessibleStructs[SourceStructIndex].ID, DataView);
@@ -391,12 +390,20 @@ struct FCachedBindingData : public TSharedFromThis<FCachedBindingData>
 		const UEdGraphSchema_K2* Schema = GetDefault<UEdGraphSchema_K2>();
 		check(Schema);
 
+		FStateTreeDataView TargetDataView;
+		BindingOwner->GetDataViewByID(TargetPath.GetStructID(), TargetDataView);
+
 		FEdGraphPinType PinType;
 
-		if (UE::StateTree::PropertyRefHelpers::IsPropertyRef(*Property))
+		if (UE::StateTree::PropertyRefHelpers::IsPropertyRef(*Property) && TargetDataView.IsValid())
 		{
 			// Use internal type to construct PinType if it's property of PropertyRef type.
-			PinType = UE::StateTree::PropertyRefHelpers::GetPropertyRefInternalTypeAsPin(*Property);
+			TArray<FStateTreePropertyPathIndirection> TargetIndirections;
+			if (ensure(TargetPath.ResolveIndirectionsWithValue(TargetDataView, TargetIndirections)))
+			{
+				const uint8* PropertyRef = TargetIndirections.Last().GetPropertyAddress();
+				PinType = UE::StateTree::PropertyRefHelpers::GetPropertyRefInternalTypeAsPin(*Property, PropertyRef);
+			}
 		}
 		else
 		{
@@ -412,11 +419,10 @@ struct FCachedBindingData : public TSharedFromThis<FCachedBindingData>
 
 				// Check that the binding is valid.
 				FStateTreeDataView SourceDataView;
-				FStateTreeDataView TargetDataView;
 				const FProperty* SourceLeafProperty = nullptr;
 				const UStruct* SourceStruct = nullptr;
 				if (BindingOwner->GetDataViewByID(SourcePath->GetStructID(), SourceDataView)
-					&& BindingOwner->GetDataViewByID(TargetPath.GetStructID(), TargetDataView))
+					&& TargetDataView.IsValid())
 				{
 					TArray<FStateTreePropertyPathIndirection> SourceIndirections;
 					TArray<FStateTreePropertyPathIndirection> TargetIndirections;
@@ -428,18 +434,19 @@ struct FCachedBindingData : public TSharedFromThis<FCachedBindingData>
 						&& TargetPath.ResolveIndirectionsWithValue(TargetDataView, TargetIndirections)
 						&& !TargetIndirections.IsEmpty())
 					{
-						const FStateTreePropertyPathIndirection LastTargetIndirection = TargetIndirections.Last();
+						const FStateTreePropertyPathIndirection TargetLeafIndirection = TargetIndirections.Last();
 						if (SourceIndirections.Num() > 0)
 						{
 							// Binding to a source property.
-							SourceLeafProperty = SourceIndirections.Last().GetProperty();
-							bIsValidBinding = ArePropertiesCompatible(SourceLeafProperty, LastTargetIndirection.GetProperty(), LastTargetIndirection.GetPropertyAddress());
+							const FStateTreePropertyPathIndirection SourceLeafIndirection = SourceIndirections.Last();
+							SourceLeafProperty = SourceLeafIndirection.GetProperty();
+							bIsValidBinding = ArePropertiesCompatible(SourceLeafProperty, TargetLeafIndirection.GetProperty(), SourceLeafIndirection.GetPropertyAddress(), TargetLeafIndirection.GetPropertyAddress());
 						}
 						else
 						{
 							// Binding to a source context struct.
 							SourceStruct = SourceDataView.GetStruct();
-							bIsValidBinding = ArePropertyAndContextStructCompatible(SourceStruct, LastTargetIndirection.GetProperty());
+							bIsValidBinding = ArePropertyAndContextStructCompatible(SourceStruct, TargetLeafIndirection.GetProperty());
 						}
 					}
 				}
@@ -545,7 +552,7 @@ struct FCachedBindingData : public TSharedFromThis<FCachedBindingData>
 		return ArePropertyAndContextStructCompatible(InStruct, PropertyHandle->GetProperty());
 	}
 			
-	bool CanBindToProperty(const FProperty* SourceProperty)
+	bool CanBindToProperty(const FProperty* SourceProperty, TConstArrayView<FBindingChainElement> InBindingChain)
 	{
 		ConditionallyUpdateData();
 
@@ -555,10 +562,32 @@ struct FCachedBindingData : public TSharedFromThis<FCachedBindingData>
 			return true;
 		}
 
+		UObject* OwnerObject = WeakOwnerObject.Get();
+		if (!OwnerObject)
+		{
+			return false;
+		}
+
+		IStateTreeEditorPropertyBindingsOwner* BindingOwner = Cast<IStateTreeEditorPropertyBindingsOwner>(OwnerObject);
+		if (!BindingOwner)
+		{
+			return false;
+		}
+
+		const int32 SourceStructIndex = InBindingChain[0].ArrayIndex;
+		check(AccessibleStructs.IsValidIndex(SourceStructIndex));
+
+		FStateTreeDataView SourceDataView;
+		BindingOwner->GetDataViewByID(AccessibleStructs[SourceStructIndex].ID, SourceDataView);
+
+		FStateTreePropertyPath SourcePath;
+		UE::StateTree::PropertyBinding::MakeStructPropertyPathFromBindingChain(AccessibleStructs[SourceStructIndex].ID, InBindingChain, SourceDataView, SourcePath);
+
+		TArray<FStateTreePropertyPathIndirection> SourceIndirections;
 		void* TargetValueAddress = nullptr;
-		if (PropertyHandle->GetValueData(TargetValueAddress) == FPropertyAccess::Success)
+		if (PropertyHandle->GetValueData(TargetValueAddress) == FPropertyAccess::Success && SourcePath.ResolveIndirectionsWithValue(SourceDataView, SourceIndirections))
 		{
-			return ArePropertiesCompatible(SourceProperty, PropertyHandle->GetProperty(), TargetValueAddress);
+			return ArePropertiesCompatible(SourceProperty, PropertyHandle->GetProperty(), SourceIndirections.Last().GetPropertyAddress(), TargetValueAddress);
 		}
 		
 		return false;
@@ -579,7 +608,8 @@ struct FCachedBindingData : public TSharedFromThis<FCachedBindingData>
 			}
 		}
 
-		return SourceProperty->HasAnyPropertyFlags(CPF_Edit);
+		return SourceProperty->HasAnyPropertyFlags(CPF_Edit) 
+			&& (!SourceProperty->HasAnyPropertyFlags(CPF_NativeAccessSpecifierPrivate | CPF_NativeAccessSpecifierProtected) || SourceProperty->GetBoolMetaData(FBlueprintMetadata::MD_AllowPrivateAccess));
 	}
 
 	static bool ArePropertyAndContextStructCompatible(const UStruct* SourceStruct, const FProperty* TargetProperty)
@@ -596,7 +626,7 @@ struct FCachedBindingData : public TSharedFromThis<FCachedBindingData>
 		return false;
 	}
 
-	static bool ArePropertiesCompatible(const FProperty* SourceProperty, const FProperty* TargetProperty, const void* TargetPropertyValue)
+	static bool ArePropertiesCompatible(const FProperty* SourceProperty, const FProperty* TargetProperty, const void* SourcePropertyValue, const void* TargetPropertyValue)
 	{
 		// @TODO: Refactor FStateTreePropertyBindings::ResolveCopyType() so that we can use it directly here.
 		
@@ -648,10 +678,10 @@ struct FCachedBindingData : public TSharedFromThis<FCachedBindingData>
 				}
 			}
 		}
-		else if (TargetStructProperty && TargetStructProperty->Struct == FStateTreePropertyRef::StaticStruct())
+		else if (TargetStructProperty && UE::StateTree::PropertyRefHelpers::IsPropertyRef(*TargetStructProperty))
 		{
 			check(TargetPropertyValue);
-			bCanBind = UE::StateTree::PropertyRefHelpers::IsPropertyRefCompatibleWithProperty(*TargetStructProperty, *SourceProperty);
+			bCanBind = UE::StateTree::PropertyRefHelpers::IsPropertyRefCompatibleWithProperty(*TargetStructProperty, *SourceProperty, TargetPropertyValue, SourcePropertyValue);
 		}
 		else
 		{
@@ -662,7 +692,7 @@ struct FCachedBindingData : public TSharedFromThis<FCachedBindingData>
 		return bCanBind;
 	}
 
-	UStruct* ResolveIndirection(TArray<FBindingChainElement> InBindingChain)
+	UStruct* ResolveIndirection(TConstArrayView<FBindingChainElement> InBindingChain)
 	{
 		UObject* OwnerObject = WeakOwnerObject.Get();
 		if (!OwnerObject)
@@ -678,9 +708,6 @@ struct FCachedBindingData : public TSharedFromThis<FCachedBindingData>
 
 		const int32 SourceStructIndex = InBindingChain[0].ArrayIndex;
 		check(SourceStructIndex >= 0 && SourceStructIndex < AccessibleStructs.Num());
-		
-		TArray<FBindingChainElement> SourceBindingChain = InBindingChain;
-		SourceBindingChain.RemoveAt(0);
 
 		FStateTreeDataView DataView;
 		if (BindingOwner->GetDataViewByID(AccessibleStructs[SourceStructIndex].ID, DataView))
@@ -872,7 +899,7 @@ void FStateTreeBindingExtension::ExtendWidgetRow(FDetailWidgetRow& InWidgetRow,
 
 	Args.OnCanBindPropertyWithBindingChain = FOnCanBindPropertyWithBindingChain::CreateLambda([CachedBindingData](FProperty* InProperty, TConstArrayView<FBindingChainElement> InBindingChain)
 		{
-			return CachedBindingData->CanBindToProperty(InProperty);
+			return CachedBindingData->CanBindToProperty(InProperty, InBindingChain);
 		});
 
 	Args.OnCanBindToContextStruct = FOnCanBindToContextStruct::CreateLambda([CachedBindingData](const UStruct* InStruct)
@@ -890,7 +917,7 @@ void FStateTreeBindingExtension::ExtendWidgetRow(FDetailWidgetRow& InWidgetRow,
 			return true;
 		});
 
-	Args.OnAddBinding = FOnAddBinding::CreateLambda([CachedBindingData](FName InPropertyName, const TArray<FBindingChainElement>& InBindingChain)
+	Args.OnAddBinding = FOnAddBinding::CreateLambda([CachedBindingData](FName InPropertyName, TConstArrayView<FBindingChainElement> InBindingChain)
 		{
 			CachedBindingData->AddBinding(InBindingChain);
 		});
@@ -927,7 +954,7 @@ void FStateTreeBindingExtension::ExtendWidgetRow(FDetailWidgetRow& InWidgetRow,
 
 	if (BindingOwner)
 	{
-		Args.OnResolveIndirection = FOnResolveIndirection::CreateLambda([CachedBindingData](TArray<FBindingChainElement> InBindingChain)
+		Args.OnResolveIndirection = FOnResolveIndirection::CreateLambda([CachedBindingData](TConstArrayView<FBindingChainElement> InBindingChain)
 		{
 			return CachedBindingData->ResolveIndirection(InBindingChain);
 		});
diff --git a/Engine/Plugins/Runtime/StateTree/Source/StateTreeEditorModule/Private/Customizations/StateTreeBlueprintPropertyRefDetails.cpp b/Engine/Plugins/Runtime/StateTree/Source/StateTreeEditorModule/Private/Customizations/StateTreeBlueprintPropertyRefDetails.cpp
new file mode 100644
--- /dev/null
+++ b/Engine/Plugins/Runtime/StateTree/Source/StateTreeEditorModule/Private/Customizations/StateTreeBlueprintPropertyRefDetails.cpp
@@ -0,0 +1,153 @@
+// Copyright Epic Games, Inc. All Rights Reserved.
+
+#include "StateTreeBlueprintPropertyRefDetails.h"
+#include "DetailWidgetRow.h"
+#include "SPinTypeSelector.h"
+#include "StateTreePropertyRef.h"
+#include "StateTreeEditorNode.h"
+#include "IDetailChildrenBuilder.h"
+#include "DetailLayoutBuilder.h"
+
+#define LOCTEXT_NAMESPACE "StateTreeEditor"
+
+namespace UE::StateTree::BlueprintPropertyRef
+{
+	static bool IsInStateTreeNode(TSharedRef<IPropertyHandle> PropertyHandle)
+	{
+		TSharedPtr<IPropertyHandle> TestPropertyHandle = PropertyHandle->GetParentHandle();
+		while (TestPropertyHandle)
+		{
+			if (const FStructProperty* StructProperty = CastField<FStructProperty>(TestPropertyHandle->GetProperty()))
+			{
+				if (StructProperty->Struct == FStateTreeEditorNode::StaticStruct())
+				{
+					return true;
+				}
+			}
+
+			TestPropertyHandle = TestPropertyHandle->GetParentHandle();
+		}
+
+		return false;
+	}
+}
+
+TSharedRef<IPropertyTypeCustomization> FStateTreeBlueprintPropertyRefDetails::MakeInstance()
+{
+	return MakeShared<FStateTreeBlueprintPropertyRefDetails>();
+}
+
+void FStateTreeBlueprintPropertyRefDetails::CustomizeHeader(TSharedRef<IPropertyHandle> InPropertyHandle, class FDetailWidgetRow& HeaderRow, IPropertyTypeCustomizationUtils& InCustomizationUtils)
+{
+	// Don't allow to modify property ref internal type if it's a part of State Tree node.
+	if (UE::StateTree::BlueprintPropertyRef::IsInStateTreeNode(InPropertyHandle))
+	{
+		HeaderRow.NameContent()
+		[
+			InPropertyHandle->CreatePropertyNameWidget()
+		];
+
+		return;
+	}
+
+	auto GetPinInfo = [WeakPropertyHandle = TWeakPtr<IPropertyHandle>(InPropertyHandle)]()
+	{
+		FEdGraphPinType Result;
+		if (TSharedPtr<IPropertyHandle> PropertyHandle = WeakPropertyHandle.Pin())
+		{
+			TSharedPtr<IPropertyHandle> RefTypePropertyHandle = PropertyHandle->GetChildHandle(GET_MEMBER_NAME_STRING_CHECKED(FStateTreeBlueprintPropertyRef, RefType));
+			TSharedPtr<IPropertyHandle> IsRefToArrayPropertyHandle = PropertyHandle->GetChildHandle(GET_MEMBER_NAME_STRING_CHECKED(FStateTreeBlueprintPropertyRef, bIsRefToArray));
+			TSharedPtr<IPropertyHandle> TypeObjectPropertyHandle = PropertyHandle->GetChildHandle(GET_MEMBER_NAME_STRING_CHECKED(FStateTreeBlueprintPropertyRef, TypeObject));
+
+			if (RefTypePropertyHandle && IsRefToArrayPropertyHandle && TypeObjectPropertyHandle)
+			{
+				uint8 RefType = 0u;
+				bool bIsRefToArray = false;
+				UObject* TypeObject = nullptr;
+
+				RefTypePropertyHandle->GetValue(RefType);
+				IsRefToArrayPropertyHandle->GetValue(bIsRefToArray);
+				TypeObjectPropertyHandle->GetValue(TypeObject);
+
+				FStateTreeBlueprintPropertyRef PropertyRef;
+				PropertyRef.RefType = static_cast<EStateTreePropertyRefType>(RefType);
+				PropertyRef.bIsRefToArray = bIsRefToArray;
+				PropertyRef.TypeObject = TypeObject;
+				Result = UE::StateTree::PropertyRefHelpers::GetBlueprintPropertyRefInternalTypeAsPin(PropertyRef);
+			}	
+		}
+
+		return Result;
+	};
+
+	auto PinInfoChanged = [WeakPropertyHandle = TWeakPtr<IPropertyHandle>(InPropertyHandle)](const FEdGraphPinType& PinType)
+	{
+		if (TSharedPtr<IPropertyHandle> PropertyHandle = WeakPropertyHandle.Pin())
+		{
+			EStateTreePropertyRefType RefType;
+			bool bIsRefToArray = false;
+			UObject* TypeObject = nullptr;
+
+			UE::StateTree::PropertyRefHelpers::GetBlueprintPropertyRefInternalTypeFromPin(PinType, RefType, bIsRefToArray, TypeObject);
+
+			TSharedPtr<IPropertyHandle> RefTypePropertyHandle = PropertyHandle->GetChildHandle(GET_MEMBER_NAME_STRING_CHECKED(FStateTreeBlueprintPropertyRef, RefType));
+			TSharedPtr<IPropertyHandle> IsRefToArrayPropertyHandle = PropertyHandle->GetChildHandle(GET_MEMBER_NAME_STRING_CHECKED(FStateTreeBlueprintPropertyRef, bIsRefToArray));
+			TSharedPtr<IPropertyHandle> TypeObjectPropertyHandle = PropertyHandle->GetChildHandle(GET_MEMBER_NAME_STRING_CHECKED(FStateTreeBlueprintPropertyRef, TypeObject));
+
+			if (RefTypePropertyHandle && IsRefToArrayPropertyHandle && TypeObjectPropertyHandle)
+			{
+				RefTypePropertyHandle->SetValue(static_cast<uint8>(RefType));
+				IsRefToArrayPropertyHandle->SetValue(bIsRefToArray);
+				TypeObjectPropertyHandle->SetValue(TypeObject);
+			}
+		}
+	};
+
+	HeaderRow.NameContent()
+	[
+		InPropertyHandle->CreatePropertyNameWidget()
+	]
+	.ValueContent()
+	[
+		SNew(SHorizontalBox)
+		
+		+ SHorizontalBox::Slot()
+		.VAlign(VAlign_Center)
+		.AutoWidth()
+		.Padding(FMargin(0, 0, 6, 0))
+		[
+			SNew(STextBlock)
+			.Text(LOCTEXT("ReferenceTo", "Reference to"))
+			.Font(IPropertyTypeCustomizationUtils::GetRegularFont())
+			.ColorAndOpacity(FSlateColor::UseSubduedForeground())
+			.ToolTipText(LOCTEXT("ReferenceTo_Tooltip", "Specifies the type of the referenced property. The referenced property is bound using property binding in the State Tree."))
+		]
+		
+		+ SHorizontalBox::Slot()
+		.AutoWidth()
+		[
+			SNew(SPinTypeSelector, FGetPinTypeTree::CreateUObject(GetDefault<UStateTreePropertyRefSchema>(), &UEdGraphSchema_K2::GetVariableTypeTree))
+				.OnPinTypeChanged_Lambda(PinInfoChanged)
+				.TargetPinType_Lambda(GetPinInfo)
+				.Schema(GetDefault<UStateTreePropertyRefSchema>())
+				.Font(IPropertyTypeCustomizationUtils::GetRegularFont())
+				.bAllowArrays(true)
+		]
+	];
+}
+
+void FStateTreeBlueprintPropertyRefDetails::CustomizeChildren(TSharedRef<IPropertyHandle> InPropertyHandle, IDetailChildrenBuilder& ChildBuilder, IPropertyTypeCustomizationUtils& InCustomizationUtils)
+{
+	if(!UE::StateTree::BlueprintPropertyRef::IsInStateTreeNode(InPropertyHandle))
+	{
+		TSharedPtr<IPropertyHandle> IsOptionalPropertyHandle = InPropertyHandle->GetChildHandle(GET_MEMBER_NAME_STRING_CHECKED(FStateTreeBlueprintPropertyRef, bIsOptional));
+		ChildBuilder.AddProperty(IsOptionalPropertyHandle.ToSharedRef());
+	}
+}
+
+bool UStateTreePropertyRefSchema::SupportsPinTypeContainer(TWeakPtr<const FEdGraphSchemaAction> SchemaAction, const FEdGraphPinType& PinType, const EPinContainerType& ContainerType) const
+{
+	return ContainerType == EPinContainerType::None || ContainerType == EPinContainerType::Array;
+}
+
+#undef LOCTEXT_NAMESPACE
\ No newline at end of file
diff --git a/Engine/Plugins/Runtime/StateTree/Source/StateTreeEditorModule/Private/Customizations/StateTreeBlueprintPropertyRefDetails.h b/Engine/Plugins/Runtime/StateTree/Source/StateTreeEditorModule/Private/Customizations/StateTreeBlueprintPropertyRefDetails.h
new file mode 100644
--- /dev/null
+++ b/Engine/Plugins/Runtime/StateTree/Source/StateTreeEditorModule/Private/Customizations/StateTreeBlueprintPropertyRefDetails.h
@@ -0,0 +1,37 @@
+// Copyright Epic Games, Inc. All Rights Reserved.
+
+#pragma once
+
+#include "IPropertyTypeCustomization.h"
+#include "EdGraphSchema_K2.h"
+#include "StateTreeBlueprintPropertyRefDetails.generated.h"
+
+class SWidget;
+class IPropertyHandle;
+
+/**
+ * Type customization for FStateTreeBlueprintPropertyRef.
+ */
+class FStateTreeBlueprintPropertyRefDetails : public IPropertyTypeCustomization
+{
+public:
+	/** Makes a new instance of this detail layout class for a specific detail view requesting it */
+	static TSharedRef<IPropertyTypeCustomization> MakeInstance();
+
+private:
+
+	/** IPropertyTypeCustomization interface */
+	virtual void CustomizeHeader(TSharedRef<IPropertyHandle> InPropertyHandle, class FDetailWidgetRow& HeaderRow, IPropertyTypeCustomizationUtils& InCustomizationUtils) override;
+	virtual void CustomizeChildren(TSharedRef<IPropertyHandle> InPropertyHandle, class IDetailChildrenBuilder& ChildBuilder, IPropertyTypeCustomizationUtils& InCustomizationUtils) override;
+};
+
+/**
+ * Specific property ref schema to allow customizing the requirements (e.g. supported containers).
+ */
+UCLASS()
+class STATETREEEDITORMODULE_API UStateTreePropertyRefSchema : public UEdGraphSchema_K2
+{
+	GENERATED_BODY()
+public:
+	virtual bool SupportsPinTypeContainer(TWeakPtr<const FEdGraphSchemaAction> SchemaAction, const FEdGraphPinType& PinType, const EPinContainerType& ContainerType) const override;
+};
diff --git a/Engine/Plugins/Runtime/StateTree/Source/StateTreeEditorModule/Private/Customizations/StateTreeEditorNodeDetails.cpp b/Engine/Plugins/Runtime/StateTree/Source/StateTreeEditorModule/Private/Customizations/StateTreeEditorNodeDetails.cpp
--- a/Engine/Plugins/Runtime/StateTree/Source/StateTreeEditorModule/Private/Customizations/StateTreeEditorNodeDetails.cpp
+++ b/Engine/Plugins/Runtime/StateTree/Source/StateTreeEditorModule/Private/Customizations/StateTreeEditorNodeDetails.cpp
@@ -35,6 +35,7 @@
 #include "Framework/Notifications/NotificationManager.h"
 #include "Widgets/Notifications/SNotificationList.h"
 #include "Debugger/StateTreeDebuggerUIExtensions.h"
+#include "StateTreePropertyBindings.h"
 
 #define LOCTEXT_NAMESPACE "StateTreeEditor"
 
@@ -180,7 +181,16 @@ namespace UE::StateTreeEditor::Internal
 				if (UE::StateTree::PropertyRefHelpers::IsPropertyRef(*Property))
 				{
 					// Use internal type to construct PinType if it's property of PropertyRef type.
-					PinType = UE::StateTree::PropertyRefHelpers::GetPropertyRefInternalTypeAsPin(*Property);
+					FStateTreeDataView TargetDataView;
+					if (ensure(EditorData->GetDataViewByID(ID, TargetDataView)))
+					{
+						TArray<FStateTreePropertyPathIndirection> TargetIndirections;
+						if (ensure(Path.ResolveIndirectionsWithValue(TargetDataView, TargetIndirections)))
+						{
+							const void* PropertyRef = TargetIndirections.Last().GetPropertyAddress();
+							PinType = UE::StateTree::PropertyRefHelpers::GetPropertyRefInternalTypeAsPin(*Property, PropertyRef);
+						}
+					}
 				}
 				else
 				{
diff --git a/Engine/Plugins/Runtime/StateTree/Source/StateTreeEditorModule/Private/K2Node_StateTreeBlueprintPropertyRef.cpp b/Engine/Plugins/Runtime/StateTree/Source/StateTreeEditorModule/Private/K2Node_StateTreeBlueprintPropertyRef.cpp
new file mode 100644
--- /dev/null
+++ b/Engine/Plugins/Runtime/StateTree/Source/StateTreeEditorModule/Private/K2Node_StateTreeBlueprintPropertyRef.cpp
@@ -0,0 +1,181 @@
+// Copyright Epic Games, Inc. All Rights Reserved.
+
+#include "K2Node_StateTreeBlueprintPropertyRef.h"
+#include "StateTreePropertyRef.h"
+#include "Blueprint/StateTreeNodeBlueprintBase.h"
+#include "EdGraphSchema_K2.h"
+#include "BlueprintActionDatabaseRegistrar.h"
+#include "BlueprintNodeSpawner.h"
+#include "Kismet2/BlueprintEditorUtils.h"
+#include "K2Node_VariableGet.h"
+#include "StateTreePropertyRefHelpers.h"
+#include "KismetCompiledFunctionContext.h"
+#include "EdGraphUtilities.h"
+
+#define LOCTEXT_NAMESPACE "K2Node_StateTreeBlueprintPropertyRef"
+
+namespace UE::StateTree::BlueprintPropertyRef
+{
+	const FName PropertyRefPinName(TEXT("PropertyRef"));
+
+	const FName SucceededExecPinName(TEXT("Succeeded"));
+	const FName FailedExecPinName(TEXT("Failed"));
+
+	const FName OutReferencePinName(TEXT("Reference"));
+}
+
+void UK2Node_StateTreeBlueprintPropertyRef::AllocateDefaultPins()
+{
+	CreatePin(EGPD_Input, UEdGraphSchema_K2::PC_Exec, UEdGraphSchema_K2::PN_Execute);
+	CreatePin(EGPD_Input, UEdGraphSchema_K2::PC_Struct, FStateTreeBlueprintPropertyRef::StaticStruct(), UE::StateTree::BlueprintPropertyRef::PropertyRefPinName);
+
+	CreatePin(EGPD_Output, UEdGraphSchema_K2::PC_Exec, UE::StateTree::BlueprintPropertyRef::SucceededExecPinName);
+    CreatePin(EGPD_Output, UEdGraphSchema_K2::PC_Exec, UE::StateTree::BlueprintPropertyRef::FailedExecPinName);
+
+	UEdGraphNode::FCreatePinParams OutPinParams;
+	OutPinParams.bIsReference = true;
+	CreatePin(EGPD_Output, UEdGraphSchema_K2::PC_Wildcard, UE::StateTree::BlueprintPropertyRef::OutReferencePinName, OutPinParams);
+}
+
+FText UK2Node_StateTreeBlueprintPropertyRef::GetNodeTitle(ENodeTitleType::Type TitleType) const
+{
+	return LOCTEXT("NodeTitle", "Get Property Reference");
+}
+
+FText UK2Node_StateTreeBlueprintPropertyRef::GetTooltipText() const
+{
+	return LOCTEXT("NodeTooltip", "Returns bound property reference.");
+}
+
+FText UK2Node_StateTreeBlueprintPropertyRef::GetMenuCategory() const
+{
+	return LOCTEXT("NodeCategory", "StateTree|PropertyRef");
+}
+
+void UK2Node_StateTreeBlueprintPropertyRef::GetMenuActions(FBlueprintActionDatabaseRegistrar& ActionRegistrar) const
+{
+	Super::GetMenuActions(ActionRegistrar);
+	UClass* Action = GetClass();
+	if (ActionRegistrar.IsOpenForRegistration(Action))
+    {
+        UBlueprintNodeSpawner* Spawner = UBlueprintNodeSpawner::Create(Action);
+        ActionRegistrar.AddBlueprintAction(Action, Spawner);
+    }
+}
+
+void UK2Node_StateTreeBlueprintPropertyRef::ReconstructNode()
+{
+	Super::ReconstructNode();
+	UpdateOutputPin();
+}
+
+void UK2Node_StateTreeBlueprintPropertyRef::ClearCachedBlueprintData(UBlueprint* Blueprint)
+{
+	Super::ClearCachedBlueprintData(Blueprint);
+	UpdateOutputPin();
+}
+
+void UK2Node_StateTreeBlueprintPropertyRef::ValidateNodeDuringCompilation(FCompilerResultsLog& MessageLog) const
+{
+	Super::ValidateNodeDuringCompilation(MessageLog);
+
+	const UClass* BlueprintClass = GetBlueprintClassFromNode();
+	if (!BlueprintClass->IsChildOf<UStateTreeNodeBlueprintBase>())
+	{
+		const FText ErrorText = LOCTEXT("InvalidSelfType", "This blueprint (self) is not a 'State Tree Blueprint Node'.");
+		MessageLog.Error(*ErrorText.ToString(), this);
+	}
+}
+
+void UK2Node_StateTreeBlueprintPropertyRef::UpdateOutputPin() const
+{
+	UEdGraphPin& OutValuePin = *FindPinChecked(UE::StateTree::BlueprintPropertyRef::OutReferencePinName);
+    const UEdGraphPin& PropertyRefPin = *FindPinChecked(UE::StateTree::BlueprintPropertyRef::PropertyRefPinName);
+
+    const UBlueprint* Blueprint = FBlueprintEditorUtils::FindBlueprintForNodeChecked(this);
+    if (Blueprint && Blueprint->GeneratedClass && PropertyRefPin.LinkedTo.Num() == 1)
+    {
+        UK2Node_VariableGet* VariableGet = Cast<UK2Node_VariableGet>(PropertyRefPin.LinkedTo[0]->GetOwningNode());
+        if (VariableGet)
+        {
+            UObject* CDO = Blueprint->GeneratedClass->GetDefaultObject<UObject>();
+            FProperty* Property = Blueprint->GeneratedClass->FindPropertyByName(VariableGet->VariableReference.GetMemberName());
+            FStructProperty* StructProperty = CastField<FStructProperty>(Property);
+            if (StructProperty && StructProperty->Struct == FStateTreeBlueprintPropertyRef::StaticStruct())
+            {
+                FStateTreeBlueprintPropertyRef PropertyRef;
+                Property->GetValue_InContainer(CDO, &PropertyRef);
+				FEdGraphPinType PinType = UE::StateTree::PropertyRefHelpers::GetBlueprintPropertyRefInternalTypeAsPin(PropertyRef);
+
+				PinType.bIsReference = true;
+				OutValuePin.PinType = PinType;
+                return;
+            }
+        }
+    }
+
+    FCreatePinParams PinParams{};
+    OutValuePin.PinType = FEdGraphPinType(UEdGraphSchema_K2::PC_Wildcard, NAME_None, nullptr, PinParams.ContainerType, PinParams.bIsReference, PinParams.ValueTerminalType);
+}
+
+void FKCHandler_StateTreeBlueprintPropertyRefGet::RegisterNets(FKismetFunctionContext& Context, UEdGraphNode* Node)
+{
+	const UK2Node_StateTreeBlueprintPropertyRef* GetRefNode = CastChecked<UK2Node_StateTreeBlueprintPropertyRef>(Node);
+
+	UEdGraphPin* OutPin = GetRefNode->FindPinChecked(UE::StateTree::BlueprintPropertyRef::OutReferencePinName);
+	FBPTerminal* OutValueTerminal = new FBPTerminal();
+	Context.InlineGeneratedValues.Add(OutValueTerminal);
+	OutValueTerminal->CopyFromPin(OutPin, Context.NetNameMap->MakeValidName(OutPin));
+	Context.NetMap.Add(OutPin, OutValueTerminal);
+
+	FBPTerminal*& IsValidTerminal = TemporaryBoolTerminals.FindOrAdd(Node);
+	if (!IsValidTerminal)
+	{
+		IsValidTerminal = Context.CreateLocalTerminal();
+		IsValidTerminal->bPassedByReference = true;
+		IsValidTerminal->Source = Node;
+		IsValidTerminal->Type.PinCategory = UEdGraphSchema_K2::PC_Boolean;
+	}	
+}
+
+void FKCHandler_StateTreeBlueprintPropertyRefGet::Compile(FKismetFunctionContext& Context, UEdGraphNode* Node)
+{
+	const UK2Node_StateTreeBlueprintPropertyRef* GetRefNode = CastChecked<UK2Node_StateTreeBlueprintPropertyRef>(Node);
+   
+    UEdGraphPin* PropertyRefPin = GetRefNode->FindPinChecked(UE::StateTree::BlueprintPropertyRef::PropertyRefPinName);
+    FBPTerminal* PropertyRefTerminal = Context.NetMap.FindRef(FEdGraphUtilities::GetNetFromPin(PropertyRefPin));
+	check(PropertyRefTerminal);
+    
+	const UEdGraphPin* OutPin = GetRefNode->FindPinChecked(UE::StateTree::BlueprintPropertyRef::OutReferencePinName);
+    FBPTerminal* OutTerminal = Context.NetMap.FindRef(OutPin);
+	check(OutTerminal);
+
+	FBlueprintCompiledStatement* GetRef = new FBlueprintCompiledStatement();
+	GetRef->Type = KCST_CallFunction;
+	GetRef->FunctionToCall = UStateTreeNodeBlueprintBase::StaticClass()->FindFunctionByName(GET_FUNCTION_NAME_CHECKED(UStateTreeNodeBlueprintBase, GetPropertyReference));
+	GetRef->FunctionContext = nullptr;
+	GetRef->RHS.Add(PropertyRefTerminal);
+
+    OutTerminal->InlineGeneratedParameter = GetRef;
+	OutTerminal->bPassedByReference = true;
+
+	FBPTerminal* IsValidTerminal = TemporaryBoolTerminals.FindChecked( Node);
+	FBlueprintCompiledStatement& HasPropertyFunction = Context.AppendStatementForNode(Node);
+	HasPropertyFunction.Type = KCST_CallFunction;
+	HasPropertyFunction.FunctionToCall = UStateTreeNodeBlueprintBase::StaticClass()->FindFunctionByName(GET_FUNCTION_NAME_CHECKED(UStateTreeNodeBlueprintBase, IsPropertyRefValid));
+	HasPropertyFunction.FunctionContext = nullptr;
+	HasPropertyFunction.LHS = IsValidTerminal;
+	HasPropertyFunction.RHS.Add(PropertyRefTerminal);
+
+	FBlueprintCompiledStatement& OnInvalidRef = Context.AppendStatementForNode(Node);
+	OnInvalidRef.Type = KCST_GotoIfNot;
+	OnInvalidRef.LHS = IsValidTerminal;
+	Context.GotoFixupRequestMap.Add(&OnInvalidRef, GetRefNode->FindPinChecked(UE::StateTree::BlueprintPropertyRef::FailedExecPinName));
+
+	FBlueprintCompiledStatement& OnValidRef = Context.AppendStatementForNode(Node);
+	OnValidRef.Type = KCST_UnconditionalGoto;
+	OnValidRef.LHS = IsValidTerminal;
+	Context.GotoFixupRequestMap.Add(&OnValidRef, GetRefNode->FindPinChecked(UE::StateTree::BlueprintPropertyRef::SucceededExecPinName));
+}
+
+#undef LOCTEXT_NAMESPACE
\ No newline at end of file
diff --git a/Engine/Plugins/Runtime/StateTree/Source/StateTreeEditorModule/Private/K2Node_StateTreeBlueprintPropertyRef.h b/Engine/Plugins/Runtime/StateTree/Source/StateTreeEditorModule/Private/K2Node_StateTreeBlueprintPropertyRef.h
new file mode 100644
--- /dev/null
+++ b/Engine/Plugins/Runtime/StateTree/Source/StateTreeEditorModule/Private/K2Node_StateTreeBlueprintPropertyRef.h
@@ -0,0 +1,43 @@
+// Copyright Epic Games, Inc. All Rights Reserved.
+
+#pragma once
+#include "K2Node.h"
+#include "KismetCompilerMisc.h"
+#include "K2Node_StateTreeBlueprintPropertyRef.generated.h"
+
+class FKCHandler_StateTreeBlueprintPropertyRefGet : public FNodeHandlingFunctor
+{
+public:
+	using FNodeHandlingFunctor::FNodeHandlingFunctor;
+	virtual void RegisterNets(FKismetFunctionContext& Context, UEdGraphNode* Node) override;
+	virtual void Compile(FKismetFunctionContext& Context, UEdGraphNode* Node) override;
+
+private:
+	TMap<UEdGraphNode*, FBPTerminal*> TemporaryBoolTerminals;
+};
+
+UCLASS()
+class UK2Node_StateTreeBlueprintPropertyRef : public UK2Node
+{
+	GENERATED_BODY()
+
+public:
+    FText GetNodeTitle(ENodeTitleType::Type TitleType) const override;
+    FText GetTooltipText() const override;
+    FText GetMenuCategory() const override;
+
+	virtual void GetMenuActions(FBlueprintActionDatabaseRegistrar& ActionRegistrar) const override;
+	virtual void ReconstructNode() override;
+	virtual void ClearCachedBlueprintData(UBlueprint* Blueprint) override;
+
+	virtual void ValidateNodeDuringCompilation(FCompilerResultsLog& MessageLog) const override;
+
+	virtual FNodeHandlingFunctor* CreateNodeHandler(class FKismetCompilerContext& CompilerContext) const override
+	{
+		return new FKCHandler_StateTreeBlueprintPropertyRefGet(CompilerContext);
+	}
+
+protected:
+	void UpdateOutputPin() const;
+	void AllocateDefaultPins() override;
+};
\ No newline at end of file
diff --git a/Engine/Plugins/Runtime/StateTree/Source/StateTreeEditorModule/Private/StateTreeCompiler.cpp b/Engine/Plugins/Runtime/StateTree/Source/StateTreeEditorModule/Private/StateTreeCompiler.cpp
--- a/Engine/Plugins/Runtime/StateTree/Source/StateTreeEditorModule/Private/StateTreeCompiler.cpp
+++ b/Engine/Plugins/Runtime/StateTree/Source/StateTreeEditorModule/Private/StateTreeCompiler.cpp
@@ -614,23 +614,28 @@ bool FStateTreeCompiler::CreateStateTasksAndParameters()
 		// Add as binding source.
 		BindingsCompiler.AddSourceStruct(LinkedParamsDesc);
 
-		// Check that the bindings for this struct are still all valid.
-		TArray<FStateTreePropertyPathBinding> CopyBindings;
-		TArray<FStateTreePropertyPathBinding> ReferenceBindings;
-		if (!GetAndValidateBindings(LinkedParamsDesc, FStateTreeDataView(CompactStateTreeParameters.Parameters.GetMutableValue()), CopyBindings, ReferenceBindings))
-		{
-			return false;
-		}
-
 		int32 BatchIndex = INDEX_NONE;
-		if (!BindingsCompiler.CompileBatch(LinkedParamsDesc, CopyBindings, BatchIndex))
-		{
-			return false;
-		}
 
-		if (!BindingsCompiler.CompileReferences(LinkedParamsDesc, ReferenceBindings, FStateTreeDataView(CompactStateTreeParameters.Parameters.GetMutableValue())))
+		// Subtrees parameters cannot have bindings
+		if(State->Type != EStateTreeStateType::Subtree)
 		{
-			return false;
+			// Check that the bindings for this struct are still all valid.
+			TArray<FStateTreePropertyPathBinding> CopyBindings;
+			TArray<FStateTreePropertyPathBinding> ReferenceBindings;
+			if (!GetAndValidateBindings(LinkedParamsDesc, FStateTreeDataView(CompactStateTreeParameters.Parameters.GetMutableValue()), CopyBindings, ReferenceBindings))
+			{
+				return false;
+			}
+
+			if (!BindingsCompiler.CompileBatch(LinkedParamsDesc, CopyBindings, BatchIndex))
+			{
+				return false;
+			}
+
+			if (!BindingsCompiler.CompileReferences(LinkedParamsDesc, ReferenceBindings, FStateTreeDataView(CompactStateTreeParameters.Parameters.GetMutableValue()), IDToStructValue))
+			{
+				return false;
+			}
 		}
 			
 		if (const auto Validation = UE::StateTree::Compiler::IsValidIndex16(BatchIndex); Validation.DidFail())
@@ -1071,7 +1076,7 @@ bool FStateTreeCompiler::CreateCondition(UStateTreeState& State, const FStateTre
 		return false;
 	}
 
-	if (!BindingsCompiler.CompileReferences(StructDesc, ReferenceBindings, InstanceDataView))
+	if (!BindingsCompiler.CompileReferences(StructDesc, ReferenceBindings, InstanceDataView, IDToStructValue))
 	{
 		return false;
 	}
@@ -1246,7 +1251,7 @@ bool FStateTreeCompiler::CreateTask(UStateTreeState* State, const FStateTreeEdit
 		return false;
 	}
 
-	if (!BindingsCompiler.CompileReferences(StructDesc, ReferenceBindings, InstanceDataView))
+	if (!BindingsCompiler.CompileReferences(StructDesc, ReferenceBindings, InstanceDataView, IDToStructValue))
 	{
 		return false;
 	}
@@ -1359,7 +1364,7 @@ bool FStateTreeCompiler::CreateEvaluator(const FStateTreeEditorNode& EvalNode, c
 		return false;
 	}
 
-	if (!BindingsCompiler.CompileReferences(StructDesc, ReferenceBindings, InstanceDataView))
+	if (!BindingsCompiler.CompileReferences(StructDesc, ReferenceBindings, InstanceDataView, IDToStructValue))
 	{
 		return false;
 	}
@@ -1374,7 +1379,7 @@ bool FStateTreeCompiler::CreateEvaluator(const FStateTreeEditorNode& EvalNode, c
 	return true;
 }
 
-bool FStateTreeCompiler::IsPropertyOfType(UScriptStruct& Type, const FStateTreeBindableStructDesc& Struct, FStateTreePropertyPath Path) const
+bool FStateTreeCompiler::IsPropertyOfTypeOrChild(UScriptStruct& Type, const FStateTreeBindableStructDesc& Struct, FStateTreePropertyPath Path) const
 {
 	TArray<FStateTreePropertyPathIndirection> Indirection;
 	const bool bResolved = Path.ResolveIndirections(Struct.Struct, Indirection);
@@ -1386,7 +1391,7 @@ bool FStateTreeCompiler::IsPropertyOfType(UScriptStruct& Type, const FStateTreeB
 		{
 			if (const FStructProperty* OwnerStructProperty = CastField<FStructProperty>(OwnerProperty))
 			{
-				return OwnerStructProperty->Struct == &Type;
+				return OwnerStructProperty->Struct->IsChildOf(&Type);
 			}
 		}
 	}
@@ -1575,8 +1580,8 @@ bool FStateTreeCompiler::GetAndValidateBindings(const FStateTreeBindableStructDe
 		// Special case fo AnyEnum. StateTreeBindingExtension allows AnyEnums to bind to other enum types.
 		// The actual copy will be done via potential type promotion copy, into the value property inside the AnyEnum.
 		// We amend the paths here to point to the 'Value' property.
-		const bool bSourceIsAnyEnum = IsPropertyOfType(*TBaseStructure<FStateTreeAnyEnum>::Get(), *SourceStruct, Binding.GetSourcePath());
-		const bool bTargetIsAnyEnum = IsPropertyOfType(*TBaseStructure<FStateTreeAnyEnum>::Get(), TargetStruct, Binding.GetTargetPath());
+		const bool bSourceIsAnyEnum = IsPropertyOfTypeOrChild(*TBaseStructure<FStateTreeAnyEnum>::Get(), *SourceStruct, Binding.GetSourcePath());
+		const bool bTargetIsAnyEnum = IsPropertyOfTypeOrChild(*TBaseStructure<FStateTreeAnyEnum>::Get(), TargetStruct, Binding.GetTargetPath());
 		if (bSourceIsAnyEnum || bTargetIsAnyEnum)
 		{
 			if (bSourceIsAnyEnum)
@@ -1589,7 +1594,7 @@ bool FStateTreeCompiler::GetAndValidateBindings(const FStateTreeBindableStructDe
 			}
 		}
 
-		if (IsPropertyOfType(*FStateTreePropertyRef::StaticStruct(), TargetStruct, Binding.GetTargetPath()))
+		if (IsPropertyOfTypeOrChild(*FStateTreePropertyRef::StaticStruct(), TargetStruct, Binding.GetTargetPath()))
 		{
 			OutReferenceBindings.Add(BindingCopy);
 		}
@@ -1624,20 +1629,35 @@ bool FStateTreeCompiler::GetAndValidateBindings(const FStateTreeBindableStructDe
 		const FProperty* Property = *It;
 		check(Property);
 		const FName PropertyName = Property->GetFName();
-		const bool bIsOptional = UE::StateTree::PropertyHelpers::HasOptionalMetadata(*Property);
 
 		if (UE::StateTree::PropertyRefHelpers::IsPropertyRef(*Property))
 		{
-			if (bIsOptional == false && !IsPropertyBound(PropertyName, OutReferenceBindings))
+			TArray<FStateTreePropertyPathIndirection> TargetIndirections;
+			FStateTreePropertyPath TargetPath(TargetStruct.ID, PropertyName);
+			if (!TargetPath.ResolveIndirectionsWithValue(TargetValue, TargetIndirections))
 			{
 				Log.Reportf(EMessageSeverity::Error, TargetStruct,
+						TEXT("Couldn't resolve path to '%s' for target %s."),
+						*PropertyName.ToString(), *TargetStruct.ToString());
+				bResult = false;
+			}
+			else
+			{
+				const void* PropertyRef = TargetIndirections.Last().GetPropertyAddress();
+				const bool bIsOptional = UE::StateTree::PropertyRefHelpers::IsPropertyRefMarkedAsOptional(*Property, PropertyRef);
+
+				if (bIsOptional == false && !IsPropertyBound(PropertyName, OutReferenceBindings))
+				{
+					Log.Reportf(EMessageSeverity::Error, TargetStruct,
 						TEXT("Property reference '%s' on % s is expected to have a binding."),
 						*PropertyName.ToString(), *TargetStruct.ToString());
 					bResult = false;
+				}
 			}
 		}
 		else
 		{
+			const bool bIsOptional = UE::StateTree::PropertyHelpers::HasOptionalMetadata(*Property);
 			const EStateTreePropertyUsage Usage = UE::StateTree::GetUsageFromMetaData(Property);
 			if (Usage == EStateTreePropertyUsage::Input)
 			{
diff --git a/Engine/Plugins/Runtime/StateTree/Source/StateTreeEditorModule/Private/StateTreeEditorModule.cpp b/Engine/Plugins/Runtime/StateTree/Source/StateTreeEditorModule/Private/StateTreeEditorModule.cpp
--- a/Engine/Plugins/Runtime/StateTree/Source/StateTreeEditorModule/Private/StateTreeEditorModule.cpp
+++ b/Engine/Plugins/Runtime/StateTree/Source/StateTreeEditorModule/Private/StateTreeEditorModule.cpp
@@ -24,6 +24,7 @@
 #include "StateTreeEditorStyle.h"
 #include "StateTreeNodeClassCache.h"
 #include "Modules/ModuleManager.h"
+#include "Customizations/StateTreeBlueprintPropertyRefDetails.h"
 
 #define LOCTEXT_NAMESPACE "StateTreeEditor"
 
@@ -91,6 +92,7 @@ void FStateTreeEditorModule::StartupModule()
 	PropertyModule.RegisterCustomPropertyTypeLayout("StateTreeReference", FOnGetPropertyTypeCustomizationInstance::CreateStatic(&FStateTreeReferenceDetails::MakeInstance));
 	PropertyModule.RegisterCustomPropertyTypeLayout("StateTreeEditorColorRef", FOnGetPropertyTypeCustomizationInstance::CreateStatic(&FStateTreeEditorColorRefDetails::MakeInstance));
 	PropertyModule.RegisterCustomPropertyTypeLayout("StateTreeEditorColor", FOnGetPropertyTypeCustomizationInstance::CreateStatic(&FStateTreeEditorColorDetails::MakeInstance));
+	PropertyModule.RegisterCustomPropertyTypeLayout("StateTreeBlueprintPropertyRef", FOnGetPropertyTypeCustomizationInstance::CreateStatic(&FStateTreeBlueprintPropertyRefDetails::MakeInstance));
 	PropertyModule.RegisterCustomClassLayout("StateTreeState", FOnGetDetailCustomizationInstance::CreateStatic(&FStateTreeStateDetails::MakeInstance));
 	PropertyModule.RegisterCustomClassLayout("StateTreeEditorData", FOnGetDetailCustomizationInstance::CreateStatic(&FStateTreeEditorDataDetails::MakeInstance));
 
@@ -120,6 +122,7 @@ void FStateTreeEditorModule::ShutdownModule()
 		PropertyModule.UnregisterCustomPropertyTypeLayout("StateTreeEditorNode");
 		PropertyModule.UnregisterCustomPropertyTypeLayout("StateTreeStateParameters");
 		PropertyModule.UnregisterCustomPropertyTypeLayout("StateTreeAnyEnum");
+		PropertyModule.UnregisterCustomPropertyTypeLayout("StateTreeBlueprintPropertyRef");
 		PropertyModule.NotifyCustomizationModuleChanged();
 	}
 }
diff --git a/Engine/Plugins/Runtime/StateTree/Source/StateTreeEditorModule/Private/StateTreePropertyBindingCompiler.cpp b/Engine/Plugins/Runtime/StateTree/Source/StateTreeEditorModule/Private/StateTreePropertyBindingCompiler.cpp
--- a/Engine/Plugins/Runtime/StateTree/Source/StateTreeEditorModule/Private/StateTreePropertyBindingCompiler.cpp
+++ b/Engine/Plugins/Runtime/StateTree/Source/StateTreeEditorModule/Private/StateTreePropertyBindingCompiler.cpp
@@ -131,7 +131,7 @@ bool FStateTreePropertyBindingCompiler::CompileBatch(const FStateTreeBindableStr
 	return true;
 }
 
-bool FStateTreePropertyBindingCompiler::CompileReferences(const FStateTreeBindableStructDesc& TargetStruct, TConstArrayView<FStateTreePropertyPathBinding> PropertyReferenceBindings, FStateTreeDataView InstanceDataView)
+bool FStateTreePropertyBindingCompiler::CompileReferences(const FStateTreeBindableStructDesc& TargetStruct, TConstArrayView<FStateTreePropertyPathBinding> PropertyReferenceBindings, FStateTreeDataView InstanceDataView, const TMap<FGuid, const FStateTreeDataView>& IDToStructValue)
 {
 	for (const FStateTreePropertyPathBinding& Binding : PropertyReferenceBindings)
 	{
@@ -149,10 +149,17 @@ bool FStateTreePropertyBindingCompiler::CompileReferences(const FStateTreeBindab
 			return false;
 		}
 
+		const FStateTreeDataView* SourceDataView = IDToStructValue.Find(Binding.GetSourcePath().GetStructID());
+		if (!SourceDataView)
+		{
+			Log->Reportf(EMessageSeverity::Error, TargetStruct,
+				TEXT("Could not find a binding source data view."));
+			return false;
+		}
+
 		FString Error;
 		TArray<FStateTreePropertyPathIndirection> SourceIndirections;
-		
-		if (!Binding.GetSourcePath().ResolveIndirections(SourceStruct->Struct, SourceIndirections, &Error))
+		if (!Binding.GetSourcePath().ResolveIndirectionsWithValue(*SourceDataView, SourceIndirections, &Error))
 		{
 			Log->Reportf(EMessageSeverity::Error, TargetStruct, TEXT("Resolving path in %s: %s"), *SourceStruct->ToString(), *Error);
 			return false;
@@ -167,16 +174,19 @@ bool FStateTreePropertyBindingCompiler::CompileReferences(const FStateTreeBindab
 			return false;
 		}
 
-		TArray<FStateTreePropertyIndirection> TargetIndirections;
-		FStateTreePropertyIndirection TargetFirstIndirection;
-		FStateTreePropertyPathIndirection TargetLeafIndirection;
-		if (!FStateTreePropertyBindings::ResolvePath(InstanceDataView.GetStruct(), Binding.GetTargetPath(), TargetIndirections, TargetFirstIndirection, TargetLeafIndirection))
+		TArray<FStateTreePropertyPathIndirection> TargetIndirections;
+		if (!Binding.GetTargetPath().ResolveIndirectionsWithValue(InstanceDataView, TargetIndirections, &Error))
 		{
 			Log->Reportf(EMessageSeverity::Error, TargetStruct, TEXT("Resolving path in %s: %s"), *TargetStruct.ToString(), *Error);
 			return false;
 		}
 
-		if (!UE::StateTree::PropertyRefHelpers::IsPropertyRefCompatibleWithProperty(*TargetLeafIndirection.GetProperty(), *SourceIndirections.Last().GetProperty()))
+		FStateTreePropertyPathIndirection& TargetLeafIndirection = TargetIndirections.Last();
+		FStateTreePropertyRef* PropertyRef = reinterpret_cast<FStateTreePropertyRef*>(const_cast<uint8*>(TargetLeafIndirection.GetPropertyAddress()));
+		check(PropertyRef);
+
+		FStateTreePropertyPathIndirection& SourceLeafIndirection = SourceIndirections.Last();
+		if (!UE::StateTree::PropertyRefHelpers::IsPropertyRefCompatibleWithProperty(*TargetLeafIndirection.GetProperty(), *SourceLeafIndirection.GetProperty(), PropertyRef, SourceLeafIndirection.GetPropertyAddress()))
 		{
 			Log->Reportf(EMessageSeverity::Error, TargetStruct,
 				TEXT("%s cannot reference %s, types are incompatible."),		
@@ -202,8 +212,10 @@ bool FStateTreePropertyBindingCompiler::CompileReferences(const FStateTreeBindab
 
 		if (!ReferenceIndex.IsValid())
 		{
-			// If referencing another PropertyRef, reuse it's index.
-			if (UE::StateTree::PropertyRefHelpers::IsPropertyRef(*SourceIndirections.Last().GetProperty()))
+			// If referencing another non global or subtree parameter PropertyRef, reuse it's index.
+			if (UE::StateTree::PropertyRefHelpers::IsPropertyRef(*SourceIndirections.Last().GetProperty()) 
+				&& SourceStruct->DataHandle.GetSource() != EStateTreeDataSourceType::GlobalParameterData 
+				&& SourceStruct->DataHandle.GetSource() != EStateTreeDataSourceType::SubtreeParameterData)
 			{
 				const FCompiledReference* ReferencedReference = CompiledReferences.FindByPredicate([&Binding](const FCompiledReference& CompiledReference)
 				{
@@ -234,9 +246,7 @@ bool FStateTreePropertyBindingCompiler::CompileReferences(const FStateTreeBindab
 		}
 
 		// Store index in instance data.
-		uint8* RawData = FStateTreePropertyBindings::GetAddress(InstanceDataView, TargetIndirections, TargetFirstIndirection, TargetLeafIndirection.GetProperty());
-		check(RawData);
-		reinterpret_cast<FStateTreePropertyRef*>(RawData)->RefAccessIndex = ReferenceIndex;
+		PropertyRef->RefAccessIndex = ReferenceIndex;
 
 		FCompiledReference& CompiledReference = CompiledReferences.AddDefaulted_GetRef();
 		CompiledReference.Path = Binding.GetTargetPath();
diff --git a/Engine/Plugins/Runtime/StateTree/Source/StateTreeEditorModule/Public/StateTreeCompiler.h b/Engine/Plugins/Runtime/StateTree/Source/StateTreeEditorModule/Public/StateTreeCompiler.h
--- a/Engine/Plugins/Runtime/StateTree/Source/StateTreeEditorModule/Public/StateTreeCompiler.h
+++ b/Engine/Plugins/Runtime/StateTree/Source/StateTreeEditorModule/Public/StateTreeCompiler.h
@@ -53,7 +53,7 @@ private:
 	bool CreateTask(UStateTreeState* State, const FStateTreeEditorNode& TaskNode, const FStateTreeDataHandle TaskDataHandle);
 	bool CreateEvaluator(const FStateTreeEditorNode& EvalNode, const FStateTreeDataHandle EvalDataHandle);
 	bool GetAndValidateBindings(const FStateTreeBindableStructDesc& TargetStruct, FStateTreeDataView TargetValue, TArray<FStateTreePropertyPathBinding>& OutCopyBindings, TArray<FStateTreePropertyPathBinding>& OutReferenceBindings) const;
-	bool IsPropertyOfType(UScriptStruct& Type, const FStateTreeBindableStructDesc& Struct, FStateTreePropertyPath Path) const;
+	bool IsPropertyOfTypeOrChild(UScriptStruct& Type, const FStateTreeBindableStructDesc& Struct, FStateTreePropertyPath Path) const;
 	bool ValidateStructRef(const FStateTreeBindableStructDesc& SourceStruct, FStateTreePropertyPath SourcePath,
 							const FStateTreeBindableStructDesc& TargetStruct, FStateTreePropertyPath TargetPath) const;
 	bool CompileAndValidateNode(const UStateTreeState* SourceState, const FStateTreeBindableStructDesc& NodeDesc, FStructView NodeView, const FStateTreeDataView InstanceData);
diff --git a/Engine/Plugins/Runtime/StateTree/Source/StateTreeEditorModule/Public/StateTreePropertyBindingCompiler.h b/Engine/Plugins/Runtime/StateTree/Source/StateTreeEditorModule/Public/StateTreePropertyBindingCompiler.h
--- a/Engine/Plugins/Runtime/StateTree/Source/StateTreeEditorModule/Public/StateTreePropertyBindingCompiler.h
+++ b/Engine/Plugins/Runtime/StateTree/Source/StateTreeEditorModule/Public/StateTreePropertyBindingCompiler.h
@@ -41,7 +41,7 @@ struct STATETREEEDITORMODULE_API FStateTreePropertyBindingCompiler
 	  * @param InstanceDataView - view to the instance data
 	  * @return True on success, false on failure.
 	 */
-	[[nodiscard]] bool CompileReferences(const FStateTreeBindableStructDesc& TargetStruct, TConstArrayView<FStateTreePropertyPathBinding> PropertyReferenceBindings, FStateTreeDataView InstanceDataView);
+	[[nodiscard]] bool CompileReferences(const FStateTreeBindableStructDesc& TargetStruct, TConstArrayView<FStateTreePropertyPathBinding> PropertyReferenceBindings, FStateTreeDataView InstanceDataView, const TMap<FGuid, const FStateTreeDataView>& IDToStructValue);
 
 	/** Finalizes compilation, should be called once all batches are compiled. */
 	void Finalize();
@@ -65,6 +65,9 @@ struct STATETREEEDITORMODULE_API FStateTreePropertyBindingCompiler
 		return SourceStructs[Index];
 	}
 
+	UE_DEPRECATED(5.4, "Use CompileReferences with additional IDToStructValue parameter instead.")
+	[[nodiscard]] bool CompileReferences(const FStateTreeBindableStructDesc& TargetStruct, TConstArrayView<FStateTreePropertyPathBinding> PropertyReferenceBindings, FStateTreeDataView InstanceDataView) { return false; }
+
 	const FStateTreeBindableStructDesc* GetSourceStructDescByID(const FGuid& ID) const
 	{
 		return SourceStructs.FindByPredicate([ID](const FStateTreeBindableStructDesc& Structs) { return (Structs.ID == ID); });
diff --git a/Engine/Plugins/Runtime/StateTree/Source/StateTreeEditorModule/StateTreeEditorModule.Build.cs b/Engine/Plugins/Runtime/StateTree/Source/StateTreeEditorModule/StateTreeEditorModule.Build.cs
--- a/Engine/Plugins/Runtime/StateTree/Source/StateTreeEditorModule/StateTreeEditorModule.Build.cs
+++ b/Engine/Plugins/Runtime/StateTree/Source/StateTreeEditorModule/StateTreeEditorModule.Build.cs
@@ -49,7 +49,8 @@ namespace UnrealBuildTool.Rules
 				"DeveloperSettings",
 				"RewindDebuggerInterface",
 				"DetailCustomizations",
-				"AppFramework"
+				"AppFramework",
+				"KismetCompiler",
 			}
 			);
 
diff --git a/Engine/Plugins/Runtime/StateTree/Source/StateTreeModule/Private/Blueprint/StateTreeNodeBlueprintBase.cpp b/Engine/Plugins/Runtime/StateTree/Source/StateTreeModule/Private/Blueprint/StateTreeNodeBlueprintBase.cpp
--- a/Engine/Plugins/Runtime/StateTree/Source/StateTreeModule/Private/Blueprint/StateTreeNodeBlueprintBase.cpp
+++ b/Engine/Plugins/Runtime/StateTree/Source/StateTreeModule/Private/Blueprint/StateTreeNodeBlueprintBase.cpp
@@ -4,6 +4,8 @@
 #include "AIController.h"
 #include "StateTreeExecutionContext.h"
 #include "VisualLogger/VisualLogger.h"
+#include "StateTreePropertyRef.h"
+#include "StateTreePropertyRefHelpers.h"
 
 #include UE_INLINE_GENERATED_CPP_BY_NAME(StateTreeNodeBlueprintBase)
 
@@ -40,21 +42,70 @@ void UStateTreeNodeBlueprintBase::SetCachedInstanceDataFromContext(const FStateT
 {
 	if (FStateTreeInstanceData* InstanceData = Context.GetMutableInstanceData())
 	{
-		InstanceStorage = &InstanceData->GetMutableStorage();
+		WeakInstanceStorage = InstanceData->GetWeakMutableStorage();
 	}
 	CachedState = Context.GetCurrentlyProcessedState();
 	CachedOwner = Context.GetOwner();
+
+	const FStateTreeExecutionFrame* CurrentlyProcessedFrame = Context.GetCurrentlyProcessedFrame();
+	check(CurrentlyProcessedFrame);
+
+	CachedFrameStateTree = CurrentlyProcessedFrame->StateTree;
+	CachedFrameRootState = CurrentlyProcessedFrame->RootState;
 }
 
 void UStateTreeNodeBlueprintBase::ClearCachedInstanceData() const
 {
-	InstanceStorage = nullptr;
+	WeakInstanceStorage = nullptr;
 	CachedState = FStateTreeStateHandle::Invalid;
 	CachedOwner = nullptr;
+	CachedFrameStateTree = nullptr;
+	CachedFrameRootState = FStateTreeStateHandle::Invalid;
+}
+
+void* UStateTreeNodeBlueprintBase::GetMutablePtrToProperty(const FStateTreeBlueprintPropertyRef& PropertyRef, FProperty*& OutSourceProperty) const
+{
+	TSharedPtr<FStateTreeInstanceStorage> InstanceStorage = WeakInstanceStorage.Pin();
+	if (InstanceStorage == nullptr)
+	{
+		UE_VLOG_UELOG(this, LogStateTree, Error, TEXT("Trying to GetMutablePtrToProperty while node is not active."));
+		return nullptr;
+	}
+
+	check(CachedFrameStateTree && CachedFrameRootState.IsValid());
+
+	TConstArrayView<FStateTreeExecutionFrame> ActiveFrames = InstanceStorage->GetExecutionState().ActiveFrames;
+
+	const int32 CurrentFrameIndex = ActiveFrames.IndexOfByPredicate([this](const FStateTreeExecutionFrame& Frame)
+	{
+		return Frame.StateTree == CachedFrameStateTree && Frame.RootState == CachedFrameRootState;
+	});
+
+	check(CurrentFrameIndex != INDEX_NONE);
+	const FStateTreeExecutionFrame& CurrentFrame = ActiveFrames[CurrentFrameIndex];
+
+	const FStateTreeExecutionFrame* CurrentParentFrame = nullptr;
+	if (CurrentFrameIndex > 0)
+	{
+		CurrentParentFrame = &ActiveFrames[CurrentFrameIndex - 1];
+	}
+
+	const FProperty* SourceProperty = nullptr;
+	void* PropertyAddress = UE::StateTree::PropertyRefHelpers::GetMutablePtrToProperty<void>(PropertyRef, *InstanceStorage, CurrentFrame, CurrentParentFrame, &SourceProperty);
+	if (PropertyAddress && UE::StateTree::PropertyRefHelpers::IsBlueprintPropertyRefCompatibleWithProperty(*SourceProperty, &PropertyRef))
+	{
+		OutSourceProperty = const_cast<FProperty*>(SourceProperty);
+		return PropertyAddress;
+	}
+
+	OutSourceProperty = nullptr;
+	return nullptr;
 }
 
 void UStateTreeNodeBlueprintBase::SendEvent(const FStateTreeEvent& Event)
 {
+	TSharedPtr<FStateTreeInstanceStorage> InstanceStorage = WeakInstanceStorage.Pin();
+
 	if (InstanceStorage == nullptr || CachedOwner == nullptr)
 	{
 		UE_VLOG_UELOG(this, LogStateTree, Error, TEXT("Trying to call SendEvent() while node is not active. Use SendEvent() on UStateTreeComponent instead for sending signals externally."));
@@ -65,6 +116,8 @@ void UStateTreeNodeBlueprintBase::SendEvent(const FStateTreeEvent& Event)
 
 void UStateTreeNodeBlueprintBase::RequestTransition(const FStateTreeStateLink& TargetState, const EStateTreeTransitionPriority Priority)
 {
+	TSharedPtr<FStateTreeInstanceStorage> InstanceStorage = WeakInstanceStorage.Pin();
+
 	if (InstanceStorage == nullptr || CachedOwner == nullptr)
 	{
 		UE_VLOG_UELOG(this, LogStateTree, Error, TEXT("Trying to call SendEvent() while node is not active. Use SendEvent() on UStateTreeComponent instead for sending signals externally."));
@@ -78,3 +131,32 @@ void UStateTreeNodeBlueprintBase::RequestTransition(const FStateTreeStateLink& T
 	
 	InstanceStorage->AddTransitionRequest(CachedOwner, Request);
 }
+
+bool UStateTreeNodeBlueprintBase::IsPropertyRefValid(const FStateTreeBlueprintPropertyRef& PropertyRef) const
+{
+	FProperty* SourceProperty = nullptr;
+	return GetMutablePtrToProperty(PropertyRef, SourceProperty) != nullptr;
+}
+
+DEFINE_FUNCTION(UStateTreeNodeBlueprintBase::execGetPropertyReference)
+{
+	P_GET_STRUCT_REF(FStateTreeBlueprintPropertyRef, PropertyRef);
+	Stack.StepCompiledIn<FProperty>(nullptr);
+	P_FINISH;
+
+	FProperty* SourceProperty = nullptr;
+	if (void* PropertyAddress = P_THIS->GetMutablePtrToProperty(PropertyRef, SourceProperty))
+	{
+		Stack.MostRecentPropertyAddress = reinterpret_cast<uint8*>(PropertyAddress);
+		Stack.MostRecentProperty = const_cast<FProperty*>(SourceProperty);
+		if (RESULT_PARAM)
+		{
+			Stack.MostRecentProperty->CopyCompleteValueToScriptVM(RESULT_PARAM, Stack.MostRecentPropertyAddress);
+		}
+	}
+	else
+	{
+		Stack.MostRecentPropertyAddress = nullptr;
+		Stack.MostRecentProperty = nullptr;
+	}
+}
diff --git a/Engine/Plugins/Runtime/StateTree/Source/StateTreeModule/Private/StateTreePropertyBindings.cpp b/Engine/Plugins/Runtime/StateTree/Source/StateTreeModule/Private/StateTreePropertyBindings.cpp
--- a/Engine/Plugins/Runtime/StateTree/Source/StateTreeModule/Private/StateTreePropertyBindings.cpp
+++ b/Engine/Plugins/Runtime/StateTree/Source/StateTreeModule/Private/StateTreePropertyBindings.cpp
@@ -287,7 +287,7 @@ bool FStateTreePropertyBindings::ResolvePaths()
 	return bResult;
 }
 
-bool FStateTreePropertyBindings::ResolvePath(const UStruct* Struct, const FStateTreePropertyPath& Path, TArray<FStateTreePropertyIndirection>& OutIndirections, FStateTreePropertyIndirection& OutFirstIndirection, FStateTreePropertyPathIndirection& OutLeafIndirection)
+bool FStateTreePropertyBindings::ResolvePath(const UStruct* Struct, const FStateTreePropertyPath& Path, FStateTreePropertyIndirection& OutFirstIndirection, FStateTreePropertyPathIndirection& OutLeafIndirection)
 {
 	if (!Struct)
  	{
@@ -411,9 +411,9 @@ bool FStateTreePropertyBindings::ResolvePath(const UStruct* Struct, const FState
 		FStateTreePropertyIndirection* PrevIndirection = &OutFirstIndirection;
 		for (int32 Index = 1; Index < TempIndirections.Num(); Index++)
 		{
-			const int32 IndirectionIndex = OutIndirections.Num();
+			const int32 IndirectionIndex = PropertyIndirections.Num();
 			PrevIndirection->NextIndex = FStateTreeIndex16(IndirectionIndex); // Set PrevIndirection before array add, as it can invalidate the pointer.
-			FStateTreePropertyIndirection& NewIndirection = OutIndirections.Add_GetRef(TempIndirections[Index]);
+			FStateTreePropertyIndirection& NewIndirection = PropertyIndirections.Add_GetRef(TempIndirections[Index]);
 			PrevIndirection = &NewIndirection;
 		}
 	}
@@ -811,7 +811,7 @@ EStateTreePropertyAccessCompatibility FStateTreePropertyBindings::GetPropertyCom
 	return EStateTreePropertyAccessCompatibility::Incompatible;
 }
 
-uint8* FStateTreePropertyBindings::GetAddress(FStateTreeDataView InStructView, TConstArrayView<FStateTreePropertyIndirection> Indirections, const FStateTreePropertyIndirection& FirstIndirection, const FProperty* LeafProperty)
+uint8* FStateTreePropertyBindings::GetAddress(FStateTreeDataView InStructView, const FStateTreePropertyIndirection& FirstIndirection, const FProperty* LeafProperty) const
 {
 	uint8* Address = InStructView.GetMutableMemory();
 	if (Address == nullptr)
@@ -904,7 +904,7 @@ uint8* FStateTreePropertyBindings::GetAddress(FStateTreeDataView InStructView, T
 				*StaticEnum<EStateTreePropertyAccessType>()->GetValueAsString(Indirection->Type), *LeafProperty->GetNameCPP());
 		}
 
-		Indirection = Indirection->NextIndex.IsValid() ? &Indirections[Indirection->NextIndex.Get()] : nullptr;
+		Indirection = Indirection->NextIndex.IsValid() ? &PropertyIndirections[Indirection->NextIndex.Get()] : nullptr;
 	}
 
 	return Address;
diff --git a/Engine/Plugins/Runtime/StateTree/Source/StateTreeModule/Private/StateTreePropertyRefHelpers.cpp b/Engine/Plugins/Runtime/StateTree/Source/StateTreeModule/Private/StateTreePropertyRefHelpers.cpp
--- a/Engine/Plugins/Runtime/StateTree/Source/StateTreeModule/Private/StateTreePropertyRefHelpers.cpp
+++ b/Engine/Plugins/Runtime/StateTree/Source/StateTreeModule/Private/StateTreePropertyRefHelpers.cpp
@@ -1,18 +1,21 @@
 // Copyright Epic Games, Inc. All Rights Reserved.
 
-#if WITH_EDITOR
 #include "StateTreePropertyRefHelpers.h"
 #include "StateTreePropertyRef.h"
 #include "UObject/TextProperty.h"
 #include "UObject/EnumProperty.h"
 #include "UObject/Class.h"
-#include "IPropertyAccessEditor.h"
 #include "StateTreePropertyBindings.h"
-#include "EdGraph/EdGraphPin.h"
+
+#if WITH_EDITOR
 #include "EdGraphSchema_K2.h"
+#include "IPropertyAccessEditor.h"
+#include "EdGraph/EdGraphPin.h"
+#endif
 
 namespace UE::StateTree::PropertyRefHelpers
 {
+#if WITH_EDITOR
 	static const FName BoolName = TEXT("bool");
 	static const FName ByteName = TEXT("byte");
 	static const FName Int32Name = TEXT("int32");
@@ -25,10 +28,23 @@ namespace UE::StateTree::PropertyRefHelpers
 
 	static const FName IsRefToArrayName = TEXT("IsRefToArray");
 	static const FName RefTypeName = TEXT("RefType");
+	static const FName IsOptionalName = TEXT("Optional");
 
-	bool IsPropertyRefCompatibleWithProperty(const FProperty& RefProperty, const FProperty& SourceProperty)
+	bool ArePropertyRefsCompatible(const FProperty& TargetRefProperty, const FProperty& SourceRefProperty, const void* TargetRefAddress, const void* SourceRefAddress)
 	{
-		ensure(IsPropertyRef(RefProperty));
+		check(IsPropertyRef(SourceRefProperty) && IsPropertyRef(TargetRefProperty));
+		check(SourceRefAddress && TargetRefAddress);
+
+		FEdGraphPinType SourceRefPin = GetPropertyRefInternalTypeAsPin(SourceRefProperty, SourceRefAddress);
+		FEdGraphPinType TargetRefPin = GetPropertyRefInternalTypeAsPin(TargetRefProperty, TargetRefAddress);
+
+		return SourceRefPin.PinCategory == TargetRefPin.PinCategory && SourceRefPin.ContainerType == TargetRefPin.ContainerType 
+			&& SourceRefPin.PinSubCategoryObject == TargetRefPin.PinSubCategoryObject;
+	}
+
+	bool IsNativePropertyRefCompatibleWithProperty(const FProperty& RefProperty, const FProperty& SourceProperty)
+	{
+		check(IsPropertyRef(RefProperty));
 
 		const FProperty* TestProperty = &SourceProperty;
 		const bool bIsTargetRefArray = RefProperty.HasMetaData(IsRefToArrayName);
@@ -53,17 +69,6 @@ namespace UE::StateTree::PropertyRefHelpers
 
 		const FName TargetTypeName = FName(*TargetTypeNameStr);
 
-		const FStructProperty* SourceStructProperty = CastField<FStructProperty>(TestProperty);
-
-		// Compare properties metadata directly if SourceProperty is PropertyRef as well
-		if (SourceStructProperty && SourceStructProperty->Struct == FStateTreePropertyRef::StaticStruct())
-		{
-			const FName SourceTypeName(SourceStructProperty->GetMetaData(RefTypeName));
-			const bool bIsSourceRefArray = SourceStructProperty->GetBoolMetaData(IsRefToArrayName);
-
-			return SourceTypeName == TargetTypeName && bIsSourceRefArray == bIsTargetRefArray;
-		}
-
 		if(TargetTypeName == BoolName)
 		{
 			return CastField<FBoolProperty>(TestProperty) != nullptr;
@@ -108,7 +113,7 @@ namespace UE::StateTree::PropertyRefHelpers
 				TargetRefField = LoadObject<UField>(nullptr, *TargetTypeNameStr);
 			}
 
-			if (SourceStructProperty)
+			if (const FStructProperty* SourceStructProperty = CastField<FStructProperty>(TestProperty))
 			{
 				return SourceStructProperty->Struct->IsChildOf(Cast<UStruct>(TargetRefField));
 			}
@@ -128,6 +133,33 @@ namespace UE::StateTree::PropertyRefHelpers
 		return false;
 	}
 
+	bool IsPropertyRefCompatibleWithProperty(const FProperty& RefProperty, const FProperty& SourceProperty, const void* PropertyRefAddress, const void* SourceAddress)
+	{
+		check(PropertyRefAddress);
+		check(SourceAddress);
+		check(IsPropertyRef(RefProperty));
+
+		if (IsPropertyRef(SourceProperty))
+		{
+			return ArePropertyRefsCompatible(RefProperty, SourceProperty, PropertyRefAddress, SourceAddress);
+		}
+
+		if (const FStructProperty* StructProperty = CastField<FStructProperty>(&RefProperty))
+		{
+			if (StructProperty->Struct == FStateTreePropertyRef::StaticStruct())
+			{
+				return IsNativePropertyRefCompatibleWithProperty(RefProperty, SourceProperty);
+			}
+			else if (StructProperty->Struct == FStateTreeBlueprintPropertyRef::StaticStruct())
+			{
+				return IsBlueprintPropertyRefCompatibleWithProperty(SourceProperty, PropertyRefAddress);
+			}
+		}
+
+		checkNoEntry();
+		return false;
+	}
+
 	bool IsPropertyAccessibleForPropertyRef(const FProperty& SourceProperty, FStateTreeBindableStructDesc SourceStruct, bool bIsOutput)
 	{
 		switch (SourceStruct.DataSource)
@@ -184,7 +216,82 @@ namespace UE::StateTree::PropertyRefHelpers
 		return IsPropertyAccessibleForPropertyRef(SourceProperty, SourceStruct, bIsOutput);
 	}
 
-	FEdGraphPinType GetPropertyRefInternalTypeAsPin(const FProperty& RefProperty)
+	FEdGraphPinType GetBlueprintPropertyRefInternalTypeAsPin(const FStateTreeBlueprintPropertyRef& PropertyRef)
+	{
+		FEdGraphPinType PinType;
+		PinType.PinSubCategory = NAME_None;
+
+		if (PropertyRef.IsRefToArray())
+		{
+			PinType.ContainerType = EPinContainerType::Array;
+		}
+
+		switch (PropertyRef.GetRefType())
+		{
+		case EStateTreePropertyRefType::None:
+			break;
+
+		case EStateTreePropertyRefType::Bool:
+			PinType.PinCategory = UEdGraphSchema_K2::PC_Boolean;
+			break;
+
+		case EStateTreePropertyRefType::Byte:
+			PinType.PinCategory = UEdGraphSchema_K2::PC_Byte;
+			break;
+
+		case EStateTreePropertyRefType::Int32:
+			PinType.PinCategory = UEdGraphSchema_K2::PC_Int;
+			break;
+
+		case EStateTreePropertyRefType::Int64:
+			PinType.PinCategory = UEdGraphSchema_K2::PC_Int64;
+			break;
+
+		case EStateTreePropertyRefType::Float:
+			PinType.PinCategory = UEdGraphSchema_K2::PC_Real;
+			PinType.PinSubCategory = UEdGraphSchema_K2::PC_Float;
+			break;
+
+		case EStateTreePropertyRefType::Double:
+			PinType.PinCategory = UEdGraphSchema_K2::PC_Real;
+			PinType.PinSubCategory = UEdGraphSchema_K2::PC_Double;
+			break;
+
+		case EStateTreePropertyRefType::Name:
+			PinType.PinCategory = UEdGraphSchema_K2::PC_Name;
+			break;
+
+		case EStateTreePropertyRefType::String:
+			PinType.PinCategory = UEdGraphSchema_K2::PC_String;
+			break;
+
+		case EStateTreePropertyRefType::Text:
+			PinType.PinCategory = UEdGraphSchema_K2::PC_Text;
+			break;
+
+		case EStateTreePropertyRefType::Enum:
+			PinType.PinCategory = UEdGraphSchema_K2::PC_Enum;
+			PinType.PinSubCategoryObject = PropertyRef.GetTypeObject();
+			break;
+
+		case EStateTreePropertyRefType::Struct:
+			PinType.PinCategory = UEdGraphSchema_K2::PC_Struct;
+			PinType.PinSubCategoryObject = PropertyRef.GetTypeObject();
+			break;
+
+		case EStateTreePropertyRefType::Object:
+			PinType.PinCategory = UEdGraphSchema_K2::PC_Object;
+			PinType.PinSubCategoryObject = PropertyRef.GetTypeObject();
+			break;
+
+		default:
+			ensureMsgf(false, TEXT("Unhandled type %s"), *UEnum::GetValueAsString(PropertyRef.GetRefType()));
+			break;
+		}
+		return PinType;
+	}
+
+	FEdGraphPinType GetNativePropertyRefInternalTypeAsPin(const FProperty& RefProperty)
 	{
 		ensure(IsPropertyRef(RefProperty));
 
@@ -265,14 +372,261 @@ namespace UE::StateTree::PropertyRefHelpers
 		return PinType;
 	}
 
+	FEdGraphPinType GetPropertyRefInternalTypeAsPin(const FProperty& RefProperty, const void* PropertyRefAddress)
+	{
+		if (const FStructProperty* StructProperty = CastField<FStructProperty>(&RefProperty))
+		{
+			if (StructProperty->Struct == FStateTreePropertyRef::StaticStruct())
+			{
+				return GetNativePropertyRefInternalTypeAsPin(RefProperty);
+			}
+			else if (StructProperty->Struct == FStateTreeBlueprintPropertyRef::StaticStruct())
+			{
+				check(PropertyRefAddress);
+				return GetBlueprintPropertyRefInternalTypeAsPin(*reinterpret_cast<const FStateTreeBlueprintPropertyRef*>(PropertyRefAddress));
+			}
+		}
+
+		checkNoEntry();
+		return FEdGraphPinType();
+	}
+
+	void STATETREEMODULE_API GetBlueprintPropertyRefInternalTypeFromPin(const FEdGraphPinType& PinType, EStateTreePropertyRefType& OutRefType, bool& bOutIsArray, UObject*& OutObjectType)
+	{
+		OutRefType = EStateTreePropertyRefType::None;
+		bOutIsArray = false;
+		OutObjectType = nullptr;
+
+		// Set container type
+		switch (PinType.ContainerType)
+		{
+		case EPinContainerType::Array:
+			bOutIsArray = true;
+			break;
+		case EPinContainerType::Set:
+			ensureMsgf(false, TEXT("Unsuported container type [Set] "));
+			break;
+		case EPinContainerType::Map:
+			ensureMsgf(false, TEXT("Unsuported container type [Map] "));
+			break;
+		default:
+			break;
+		}
+	
+		// Value type
+		if (PinType.PinCategory == UEdGraphSchema_K2::PC_Boolean)
+		{
+			OutRefType = EStateTreePropertyRefType::Bool;
+		}
+		else if (PinType.PinCategory == UEdGraphSchema_K2::PC_Byte)
+		{
+			if (UEnum* Enum = Cast<UEnum>(PinType.PinSubCategoryObject))
+			{
+				OutRefType = EStateTreePropertyRefType::Enum;
+				OutObjectType = PinType.PinSubCategoryObject.Get();
+			}
+			else
+			{
+				OutRefType = EStateTreePropertyRefType::Byte;
+			}
+		}
+		else if (PinType.PinCategory == UEdGraphSchema_K2::PC_Int)
+		{
+			OutRefType = EStateTreePropertyRefType::Int32;
+		}
+		else if (PinType.PinCategory == UEdGraphSchema_K2::PC_Int64)
+		{
+			OutRefType = EStateTreePropertyRefType::Int64;
+		}
+		else if (PinType.PinCategory == UEdGraphSchema_K2::PC_Real)
+		{
+			if (PinType.PinSubCategory == UEdGraphSchema_K2::PC_Float)
+			{
+				OutRefType = EStateTreePropertyRefType::Float;
+			}
+			else if (PinType.PinSubCategory == UEdGraphSchema_K2::PC_Double)
+			{
+				OutRefType = EStateTreePropertyRefType::Double;
+			}		
+		}
+		else if (PinType.PinCategory == UEdGraphSchema_K2::PC_Name)
+		{
+			OutRefType = EStateTreePropertyRefType::Name;
+		}
+		else if (PinType.PinCategory == UEdGraphSchema_K2::PC_String)
+		{
+			OutRefType = EStateTreePropertyRefType::String;
+		}
+		else if (PinType.PinCategory == UEdGraphSchema_K2::PC_Text)
+		{
+			OutRefType = EStateTreePropertyRefType::Text;
+		}
+		else if (PinType.PinCategory == UEdGraphSchema_K2::PC_Enum)
+		{
+			OutRefType = EStateTreePropertyRefType::Enum;
+			OutObjectType = PinType.PinSubCategoryObject.Get();
+		}
+		else if (PinType.PinCategory == UEdGraphSchema_K2::PC_Struct)
+		{
+			OutRefType = EStateTreePropertyRefType::Struct;
+			OutObjectType = PinType.PinSubCategoryObject.Get();
+		}
+		else if (PinType.PinCategory == UEdGraphSchema_K2::PC_Object)
+		{
+			OutRefType = EStateTreePropertyRefType::Object;
+			OutObjectType = PinType.PinSubCategoryObject.Get();
+		}
+		else if (PinType.PinCategory == UEdGraphSchema_K2::PC_SoftObject)
+		{
+			OutRefType = EStateTreePropertyRefType::SoftObject;
+			OutObjectType = PinType.PinSubCategoryObject.Get();
+		}
+		else if (PinType.PinCategory == UEdGraphSchema_K2::PC_Class)
+		{
+			OutRefType = EStateTreePropertyRefType::Class;
+			OutObjectType = PinType.PinSubCategoryObject.Get();
+		}
+		else if (PinType.PinCategory == UEdGraphSchema_K2::PC_SoftClass)
+		{
+			OutRefType = EStateTreePropertyRefType::SoftClass;
+			OutObjectType = PinType.PinSubCategoryObject.Get();
+		}
+		else
+		{
+			ensureMsgf(false, TEXT("Unhandled pin category %s"), *PinType.PinCategory.ToString());
+		}
+	}
+
+	bool STATETREEMODULE_API IsPropertyRefMarkedAsOptional(const FProperty& RefProperty, const void* PropertyRefAddress)
+	{
+		if (const FStructProperty* StructProperty = CastField<FStructProperty>(&RefProperty))
+		{
+			if (StructProperty->Struct == FStateTreePropertyRef::StaticStruct())
+			{
+				return RefProperty.HasMetaData(IsOptionalName);
+			}
+			else if (StructProperty->Struct == FStateTreeBlueprintPropertyRef::StaticStruct())
+			{
+				check(PropertyRefAddress);
+				return reinterpret_cast<const FStateTreeBlueprintPropertyRef*>(PropertyRefAddress)->IsOptional();
+			}
+		}
+
+		checkNoEntry();
+		return false;
+	}
+#endif
+
+	bool IsBlueprintPropertyRefCompatibleWithProperty(const FProperty& SourceProperty, const void* PropertyRefAddress)
+	{
+		const FStateTreeBlueprintPropertyRef& PropertyRef = *reinterpret_cast<const FStateTreeBlueprintPropertyRef*>(PropertyRefAddress);
+		const FProperty* TestProperty = &SourceProperty;
+		if (PropertyRef.IsRefToArray())
+		{
+			if (const FArrayProperty* ArrayProperty = CastField<FArrayProperty>(TestProperty))
+			{
+				TestProperty = ArrayProperty->Inner;
+			}
+			else
+			{
+				return false;
+			}
+		}
+
+		switch (PropertyRef.GetRefType())
+		{
+		case EStateTreePropertyRefType::None:
+			return false;
+
+		case EStateTreePropertyRefType::Bool:
+			return Validator<bool>::IsValid(*TestProperty);
+
+		case EStateTreePropertyRefType::Byte:
+			return Validator<uint8>::IsValid(*TestProperty);
+
+		case EStateTreePropertyRefType::Int32:
+			return Validator<int32>::IsValid(*TestProperty);
+
+		case EStateTreePropertyRefType::Int64:
+			return Validator<int64>::IsValid(*TestProperty);
+
+		case EStateTreePropertyRefType::Float:
+			return Validator<float>::IsValid(*TestProperty);
+
+		case EStateTreePropertyRefType::Double:
+			return Validator<double>::IsValid(*TestProperty);
+
+		case EStateTreePropertyRefType::Name:
+			return Validator<FName>::IsValid(*TestProperty);
+
+		case EStateTreePropertyRefType::String:
+			return Validator<FString>::IsValid(*TestProperty);
+
+		case EStateTreePropertyRefType::Text:
+			return Validator<FText>::IsValid(*TestProperty);
+
+		case EStateTreePropertyRefType::Enum:
+			if (const UEnum* Enum = Cast<UEnum>(PropertyRef.GetTypeObject()))
+			{
+				return IsPropertyCompatibleWithEnum(*TestProperty, *Enum);
+			}
+			return false;
+
+		case EStateTreePropertyRefType::Struct:
+			if (const UScriptStruct* Struct = Cast<UScriptStruct>(PropertyRef.GetTypeObject()))
+			{
+				return IsPropertyCompatibleWithStruct(*TestProperty, *Struct);
+			}
+			return false;
+
+		case EStateTreePropertyRefType::Object:
+			if (const UClass* Class = Cast<UClass>(PropertyRef.GetTypeObject()))
+			{
+				return IsPropertyCompatibleWithClass(*TestProperty, *Class);
+			}
+			return false;
+
+		default:
+			checkNoEntry();
+		}
+
+		return false;
+	}
+
 	bool IsPropertyRef(const FProperty& Property)
 	{
 		if (const FStructProperty* StructProperty = CastField<FStructProperty>(&Property))
 		{
-			return StructProperty->Struct == FStateTreePropertyRef::StaticStruct();
+			return StructProperty->Struct->IsChildOf(FStateTreePropertyRef::StaticStruct());
 		}
 
 		return false;
 	}
-}
-#endif
\ No newline at end of file
+
+	bool IsPropertyCompatibleWithEnum(const FProperty& Property, const UEnum& Enum)
+	{
+		if (const FEnumProperty* EnumProperty = CastField<FEnumProperty>(&Property))
+		{
+			return EnumProperty->GetEnum() == &Enum;
+		}
+		return false;
+	}
+
+	bool IsPropertyCompatibleWithClass(const FProperty& Property, const UClass& Class)
+	{
+		if (const FObjectProperty* ObjectProperty = CastField<FObjectProperty>(&Property))
+		{
+			return ObjectProperty->PropertyClass == &Class;
+		}
+		return false;
+	}
+
+	bool IsPropertyCompatibleWithStruct(const FProperty& Property, const UScriptStruct& Struct)
+	{
+		if (const FStructProperty* StructProperty = CastField<FStructProperty>(&Property))
+		{
+			return StructProperty->Struct == &Struct;
+		}
+		return false;
+	}
+}
\ No newline at end of file
diff --git a/Engine/Plugins/Runtime/StateTree/Source/StateTreeModule/Public/Blueprint/StateTreeNodeBlueprintBase.h b/Engine/Plugins/Runtime/StateTree/Source/StateTreeModule/Public/Blueprint/StateTreeNodeBlueprintBase.h
--- a/Engine/Plugins/Runtime/StateTree/Source/StateTreeModule/Public/Blueprint/StateTreeNodeBlueprintBase.h
+++ b/Engine/Plugins/Runtime/StateTree/Source/StateTreeModule/Public/Blueprint/StateTreeNodeBlueprintBase.h
@@ -3,6 +3,7 @@
 #pragma once
 
 #include "StateTreeExecutionTypes.h"
+#include "UObject/ObjectMacros.h"
 #include "StateTreeNodeBlueprintBase.generated.h"
 
 struct FStateTreeEvent;
@@ -10,6 +11,8 @@ struct FStateTreeEventQueue;
 struct FStateTreeInstanceStorage;
 struct FStateTreeLinker;
 struct FStateTreeExecutionContext;
+struct FStateTreeBlueprintPropertyRef;
+class UStateTree;
 
 UENUM()
 enum class EStateTreeBlueprintPropertyCategory : uint8
@@ -44,6 +47,14 @@ public:
 	UFUNCTION(BlueprintCallable, Category = "StateTree", meta = (HideSelfPin = "true", DisplayName = "StateTree Request Transition"))
 	void RequestTransition(const FStateTreeStateLink& TargetState, const EStateTreeTransitionPriority Priority = EStateTreeTransitionPriority::Normal);
 
+	/** Returns a reference to selected property in State Tree. */
+	UFUNCTION(CustomThunk)
+	void GetPropertyReference(const FStateTreeBlueprintPropertyRef& PropertyRef) const;
+
+	/** Returns true if reference to selected property in State Tree is accessible. */
+	UFUNCTION()
+	bool IsPropertyRefValid(const FStateTreeBlueprintPropertyRef& PropertyRef) const;
+
 protected:
 	virtual UWorld* GetWorld() const override;
 	AActor* GetOwnerActor(const FStateTreeExecutionContext& Context) const;
@@ -58,16 +69,26 @@ protected:
 	void ClearCachedEventQueue() const { ClearCachedInstanceData(); }
 	
 private:
+	DECLARE_FUNCTION(execGetPropertyReference);
 
-	/** Cached State where the node is processed on. */
-	mutable FStateTreeStateHandle CachedState;
+	void* GetMutablePtrToProperty(const FStateTreeBlueprintPropertyRef& PropertyRef, FProperty*& OutSourceProperty) const;
 	
 	/** Cached instance data while the node is active. */
-	mutable FStateTreeInstanceStorage* InstanceStorage = nullptr;
+	mutable TWeakPtr<FStateTreeInstanceStorage> WeakInstanceStorage;
 
 	/** Cached owner while the node is active. */
 	UPROPERTY()
 	mutable TObjectPtr<UObject> CachedOwner = nullptr;
+
+	/** Cached State Tree of owning execution frame. */
+	UPROPERTY()
+	mutable TObjectPtr<const UStateTree> CachedFrameStateTree = nullptr;
+
+	/** Cached root state of owning execution frame. */
+	mutable FStateTreeStateHandle CachedFrameRootState;
+
+	/** Cached State where the node is processed on. */
+	mutable FStateTreeStateHandle CachedState;
 };
 
 #if UE_ENABLE_INCLUDE_ORDER_DEPRECATED_IN_5_2
diff --git a/Engine/Plugins/Runtime/StateTree/Source/StateTreeModule/Public/StateTreePropertyBindings.h b/Engine/Plugins/Runtime/StateTree/Source/StateTreeModule/Public/StateTreePropertyBindings.h
--- a/Engine/Plugins/Runtime/StateTree/Source/StateTreeModule/Public/StateTreePropertyBindings.h
+++ b/Engine/Plugins/Runtime/StateTree/Source/StateTreeModule/Public/StateTreePropertyBindings.h
@@ -899,7 +899,7 @@ struct STATETREEMODULE_API FStateTreePropertyBindings
 	 * @param PropertyAccess Access to the property for which we want to obtain a pointer.
 	 * @return Pointer to referenced property if it's type match, nullptr otherwise.
 	 */
-	template< class T >
+	template<class T>
 	T* GetMutablePropertyPtr(FStateTreeDataView SourceView, const FStateTreePropertyAccess& PropertyAccess) const
 	{
 		check(SourceView.GetStruct() == PropertyAccess.SourceStructType);
@@ -939,7 +939,6 @@ struct STATETREEMODULE_API FStateTreePropertyBindings
 	 */
 	[[nodiscard]] static bool ResolveCopyType(const FStateTreePropertyPathIndirection& SourceIndirection, const FStateTreePropertyPathIndirection& TargetIndirection, FStateTreePropertyCopy& OutCopy);
 
-	
 	UE_DEPRECATED(5.3, "Should not be used, will be removed in a future version.")
 	TArrayView<FStateTreeBindableStructDesc> GetSourceStructs() { return  SourceStructs; };
 
@@ -953,23 +952,12 @@ struct STATETREEMODULE_API FStateTreePropertyBindings
 	}
 
 private:
-	[[nodiscard]] static bool ResolvePath(const UStruct* Struct, const FStateTreePropertyPath& Path, TArray<FStateTreePropertyIndirection>& OutIndirections, FStateTreePropertyIndirection& OutFirstIndirection, FStateTreePropertyPathIndirection& OutLeafIndirection);
-	static uint8* GetAddress(FStateTreeDataView InStructView, TConstArrayView<FStateTreePropertyIndirection> Indirections, const FStateTreePropertyIndirection& FirstIndirection, const FProperty* LeafProperty);
-	
-	uint8* GetAddress(FStateTreeDataView InStructView, const FStateTreePropertyIndirection& FirstIndirection, const FProperty* LeafProperty) const
-	{
-		return GetAddress(InStructView, PropertyIndirections, FirstIndirection, LeafProperty);
-	}
-
-	[[nodiscard]] bool ResolvePath(const UStruct* Struct, const FStateTreePropertyPath& Path, FStateTreePropertyIndirection& OutFirstIndirection, FStateTreePropertyPathIndirection& OutLeafIndirection)
-	{
-		return ResolvePath(Struct, Path, PropertyIndirections, OutFirstIndirection, OutLeafIndirection);
-	}
-
+	[[nodiscard]] bool ResolvePath(const UStruct* Struct, const FStateTreePropertyPath& Path, FStateTreePropertyIndirection& OutFirstIndirection, FStateTreePropertyPathIndirection& OutLeafIndirection);
 	const FStateTreeBindableStructDesc* GetSourceDescByHandle(const FStateTreeDataHandle SourceDataHandle);
 
 	void PerformCopy(const FStateTreePropertyCopy& Copy, uint8* SourceAddress, uint8* TargetAddress) const;
 	void PerformResetObjects(const FStateTreePropertyCopy& Copy, uint8* TargetAddress) const;
+	uint8* GetAddress(FStateTreeDataView InStructView, const FStateTreePropertyIndirection& FirstIndirection, const FProperty* LeafProperty) const;
 
 	/** Array of expected source structs. */
 	UPROPERTY()
diff --git a/Engine/Plugins/Runtime/StateTree/Source/StateTreeModule/Public/StateTreePropertyRef.h b/Engine/Plugins/Runtime/StateTree/Source/StateTreeModule/Public/StateTreePropertyRef.h
--- a/Engine/Plugins/Runtime/StateTree/Source/StateTreeModule/Public/StateTreePropertyRef.h
+++ b/Engine/Plugins/Runtime/StateTree/Source/StateTreeModule/Public/StateTreePropertyRef.h
@@ -4,6 +4,7 @@
 
 #include "StateTreeIndexTypes.h"
 #include "StateTreeExecutionContext.h"
+#include "StateTreePropertyRefHelpers.h"
 #include "StateTreePropertyRef.generated.h"
 
 struct FStateTreePropertyRef;
@@ -12,20 +13,63 @@ namespace UE::StateTree::PropertyRefHelpers
 {
     /**
 	 * @param PropertyRef Property's reference to get pointer to.
-	 * @param InstanceDataStorage Instance Data Storage
-	 * @param ExecutionFrame Execution frame owning referenced property
-	 * @param ParentExecutionFrame Parent of execution frame owning referenced property
-	 * @return Pointer to referenced property if succeeded.
+	 * @param InstanceDataStorage Instance Data Storage.
+	 * @param ExecutionFrame Execution frame owning referenced property.
+	 * @param ParentExecutionFrame Parent of execution frame owning referenced property.
+	 * @param OutSourceProperty On success, returns referenced property.
+	 * @return Pointer to referenced property value if succeeded.
 	 */
 	template<class T>
-	static T* GetMutablePtrToProperty(const FStateTreePropertyRef& PropertyRef, FStateTreeInstanceStorage& InstanceDataStorage, const FStateTreeExecutionFrame& ExecutionFrame, const FStateTreeExecutionFrame* ParentExecutionFrame)
+	static T* GetMutablePtrToProperty(const FStateTreePropertyRef& PropertyRef, FStateTreeInstanceStorage& InstanceDataStorage, const FStateTreeExecutionFrame& ExecutionFrame, const FStateTreeExecutionFrame* ParentExecutionFrame, const FProperty** OutSourceProperty = nullptr)
 	{
 		const FStateTreePropertyBindings& PropertyBindings = ExecutionFrame.StateTree->GetPropertyBindings();
 		if (const FStateTreePropertyAccess* PropertyAccess = PropertyBindings.GetPropertyAccess(PropertyRef))
 		{
 			// Passing empty ContextAndExternalDataViews, as PropertyRef is not allowed to point to context or external data.
 			FStateTreeDataView SourceView = FStateTreeExecutionContext::GetDataView(InstanceDataStorage, nullptr, ParentExecutionFrame, ExecutionFrame, {}, PropertyAccess->SourceDataHandle);
-			return PropertyBindings.GetMutablePropertyPtr<T>(SourceView, *PropertyAccess);
+			
+			// The only possibility when PropertyRef references another PropertyRef is when source one is a global or subtree parameter, i.e lives in parent execution frame.
+			// If that's the case, referenced PropertyRef is obtained and we recursively take the address where it points to.
+			if (UE::StateTree::PropertyRefHelpers::IsPropertyRef(*PropertyAccess->SourceLeafProperty))
+			{
+				check(PropertyAccess->SourceDataHandle.GetSource() == EStateTreeDataSourceType::GlobalParameterData || PropertyAccess->SourceDataHandle.GetSource() == EStateTreeDataSourceType::SubtreeParameterData);
+
+				if (ParentExecutionFrame == nullptr)
+				{
+					return nullptr;
+				}
+
+				const FStateTreePropertyRef* ReferencedPropertyRef = PropertyBindings.GetMutablePropertyPtr<FStateTreePropertyRef>(SourceView, *PropertyAccess);
+				if (ReferencedPropertyRef == nullptr)
+				{
+					return nullptr;
+				}
+
+				const TArray<FStateTreeExecutionFrame>& ActiveFrames = InstanceDataStorage.GetExecutionState().ActiveFrames;
+				const int32 FrameIndex = ActiveFrames.IndexOfByPredicate([ParentExecutionFrame](const FStateTreeExecutionFrame& Frame)
+				{
+					return Frame.RootState == ParentExecutionFrame->RootState && Frame.StateTree == ParentExecutionFrame->StateTree;
+				});
+
+				if (FrameIndex == INDEX_NONE)
+				{
+					return nullptr;
+				}
+
+				const FStateTreeExecutionFrame& Frame = ActiveFrames[FrameIndex];
+				const FStateTreeExecutionFrame* ParentFrame = FrameIndex > 0 ? &ActiveFrames[FrameIndex - 1] : nullptr;
+
+				return GetMutablePtrToProperty<T>(*ReferencedPropertyRef, InstanceDataStorage, Frame, ParentFrame, OutSourceProperty);
+			}
+			else
+			{
+				if (OutSourceProperty)
+				{
+					*OutSourceProperty = PropertyAccess->SourceLeafProperty;
+				}
+
+				return PropertyBindings.GetMutablePropertyPtr<T>(SourceView, *PropertyAccess);
+			}
 		}
 
 		return nullptr;
@@ -182,3 +226,66 @@ private:
 	FStateTreeStateHandle RootState = FStateTreeStateHandle::Invalid;
 	FStateTreePropertyRef PropertyRef;
 };
+
+UENUM()
+enum class EStateTreePropertyRefType : uint8
+{
+	None,
+	Bool,
+	Byte,
+	Int32,
+	Int64,
+	Float,
+	Double,
+	Name,
+	String,
+	Text,
+	Enum,
+	Struct,
+	Object,
+	SoftObject,
+	Class,
+	SoftClass,
+};
+
+/**
+ * FStateTreeBlueprintPropertyRef is a PropertyRef intended to be used in State Tree Blueprint nodes like tasks, conditions or evaluators, but also as a StateTree parameter.
+ */
+USTRUCT(BlueprintType, DisplayName = "State Tree Property Ref")
+struct STATETREEMODULE_API FStateTreeBlueprintPropertyRef : public FStateTreePropertyRef
+{
+	GENERATED_BODY()
+
+	FStateTreeBlueprintPropertyRef() = default;
+
+	/** Returns PropertyRef's type */
+	EStateTreePropertyRefType GetRefType() const { return RefType; }
+
+	/** Returns true if referenced property is an array. */
+	bool IsRefToArray() const { return bIsRefToArray; }
+
+	/** Returns selected ScriptStruct, Class or Enum. */
+	UObject* GetTypeObject() const { return TypeObject; }
+
+	/** Returns true if PropertyRef was marked as optional. */
+	bool IsOptional() const { return bIsOptional; }
+
+private:
+	/** Specifies the type of property to reference */
+	UPROPERTY(EditAnywhere, Category = "InternalType")
+	EStateTreePropertyRefType RefType = EStateTreePropertyRefType::None;
+
+	/** If specified, the reference is to an TArray<RefType> */
+	UPROPERTY(EditAnywhere, Category = "InternalType")
+	uint8 bIsRefToArray : 1;
+
+	/** If specified, the reference can be left unbound, otherwise the State Tree compiler report error if the reference is not bound. */
+	UPROPERTY(EditAnywhere, Category = "Parameter")
+	uint8 bIsOptional : 1;
+
+	/** Specifies the type of property to reference together with RefType, used for Enums, Structs, Objects and Classes. */
+	UPROPERTY(EditAnywhere, Category= "InternalType")
+	TObjectPtr<UObject> TypeObject = nullptr;
+
+	friend class FStateTreeBlueprintPropertyRefDetails;
+};
\ No newline at end of file
diff --git a/Engine/Plugins/Runtime/StateTree/Source/StateTreeModule/Public/StateTreePropertyRefHelpers.h b/Engine/Plugins/Runtime/StateTree/Source/StateTreeModule/Public/StateTreePropertyRefHelpers.h
--- a/Engine/Plugins/Runtime/StateTree/Source/StateTreeModule/Public/StateTreePropertyRefHelpers.h
+++ b/Engine/Plugins/Runtime/StateTree/Source/StateTreeModule/Public/StateTreePropertyRefHelpers.h
@@ -15,49 +15,88 @@ struct FBindingChainElement;
 struct FStateTreeBindableStructDesc;
 struct FStateTreePropertyPathIndirection;
 struct FEdGraphPinType;
+struct FStateTreeBlueprintPropertyRef;
+enum class EStateTreePropertyRefType : uint8;
 
 namespace UE::StateTree::PropertyRefHelpers
 {
 #if WITH_EDITOR
 	/**
-	 * @param RefProperty Property of PropertyRef type
-	 * @param SourceProperty Property to check it's type compatibility
-	 * @return true if SourceProperty type is compatible with PropertyRef
+	 * @param RefProperty Property of PropertyRef type.
+	 * @param SourceProperty Property to check it's type compatibility.
+	 * @param PropertyRefAddress Address of PropertyRef
+	 * @param SourceAddress Address of checked property value.
+	 * @return true if SourceProperty type is compatible with PropertyRef.
 	 */
-	bool STATETREEMODULE_API IsPropertyRefCompatibleWithProperty(const FProperty& RefProperty, const FProperty& SourceProperty);
+	bool STATETREEMODULE_API IsPropertyRefCompatibleWithProperty(const FProperty& RefProperty, const FProperty& SourceProperty, const void* PropertyRefAddress, const void* SourceAddress);
 
 	/**
-	 * @param SourcePropertyPathIndirections Path indirections of the referenced property
-	 * @param SourceStruct Bindable owner of referenced property
-	 * @return true if property can be referenced by PropertyRef
+	 * @param SourcePropertyPathIndirections Path indirections of the referenced property.
+	 * @param SourceStruct Bindable owner of referenced property.
+	 * @return true if property can be referenced by PropertyRef.
 	 */
 	bool STATETREEMODULE_API IsPropertyAccessibleForPropertyRef(TConstArrayView<FStateTreePropertyPathIndirection> SourcePropertyPathIndirections, FStateTreeBindableStructDesc SourceStruct);
 
 	/**
-	 * @param SourceProperty Referenced property
-	 * @param BindingChain Binding chain to referenced property
-	 * @param SourceStruct Bindable owner of referenced property
-	 * @return true if property can be referenced by PropertyRef
+	 * @param SourceProperty Referenced property.
+	 * @param BindingChain Binding chain to referenced property.
+	 * @param SourceStruct Bindable owner of referenced property.
+	 * @return true if property can be referenced by PropertyRef.
 	 */
 	bool STATETREEMODULE_API IsPropertyAccessibleForPropertyRef(const FProperty& SourceProperty, TConstArrayView<FBindingChainElement> BindingChain, FStateTreeBindableStructDesc SourceStruct);
 
 	/**
-	 * @param Property Property to check
-	 * @return true if Property is a PropertyRef
+	 * @param RefProperty Property of PropertyRef type.
+	 * @param PropertyRefAddress Address of PropertyRef.
+	 * @return true if PropertyRef is marked as optional.
 	 */
-	bool STATETREEMODULE_API IsPropertyRef(const FProperty& Property);
+	bool STATETREEMODULE_API IsPropertyRefMarkedAsOptional(const FProperty& RefProperty, const void* PropertyRefAddress);
 
 	/**
 	 * @param RefProperty Property of PropertyRef type
+	 * @param PropertyRefAddress Address of PropertyRef.
 	 * @return PinType for PropertyRef's internal type
 	 */
-	FEdGraphPinType STATETREEMODULE_API GetPropertyRefInternalTypeAsPin(const FProperty& RefProperty);
+	FEdGraphPinType STATETREEMODULE_API GetPropertyRefInternalTypeAsPin(const FProperty& RefProperty, const void* PropertyRefAddress);
+
+	/**
+	 * @param InPropertyRef PropertyRef to get internal type from.
+	 * @return PinType for Blueprint PropertyRef's internal type
+	 */
+	FEdGraphPinType STATETREEMODULE_API GetBlueprintPropertyRefInternalTypeAsPin(const FStateTreeBlueprintPropertyRef& InPropertyRef);
+
+	/**
+	 * @param PinType Pin to get type from.
+	 * @param OutRefType PropertyRef's referenced type.
+	 * @param bOutIsArray True if PropertyRef references an array property.
+	 * @param OutObjectType Referenced type's specific object.
+	 */
+	void STATETREEMODULE_API GetBlueprintPropertyRefInternalTypeFromPin(const FEdGraphPinType& PinType, EStateTreePropertyRefType& OutRefType, bool& bOutIsArray, UObject*& OutObjectType);
 #endif
 
+	/**
+	 * @param Property Property to check
+	 * @return true if Property is a PropertyRef
+	 */
+	bool STATETREEMODULE_API IsPropertyRef(const FProperty& Property);
+
+	/**
+	 * @param SourceProperty Property to check it's type compatibility.
+	 * @param PropertyRefAddress Address of PropertyRef
+	 * @return true if SourceProperty type is compatible with Blueprint PropertyRef.
+	 */
+	bool STATETREEMODULE_API IsBlueprintPropertyRefCompatibleWithProperty(const FProperty& SourceProperty, const void* PropertyRefAddress);
+
 	template<class T, class = void>
 	struct Validator
 	{};
 
+	template<>
+	struct Validator<void>
+	{
+		static bool IsValid(const FProperty& Property) { return true; }
+	};
+
 	template<>
 	struct Validator<FBoolProperty::TCppType>
 	{
@@ -126,6 +165,9 @@ namespace UE::StateTree::PropertyRefHelpers
 		}
 	};
 
+	/* Checks if provided property is compatible with selected ScriptStruct. */
+	bool STATETREEMODULE_API IsPropertyCompatibleWithStruct(const FProperty& Property, const class UScriptStruct& Struct);
+
 	template<class T>
 	struct Validator<T, decltype(TBaseStructure<T>::Get, void())>
 	{
@@ -140,6 +182,9 @@ namespace UE::StateTree::PropertyRefHelpers
 		}
 	};
 
+	/* Checks if provided property is compatible with selected Class. */
+	bool STATETREEMODULE_API IsPropertyCompatibleWithClass(const FProperty& Property, const class UClass& Class);
+
 	template<class T>
 	struct Validator<T, typename TEnableIf<TIsDerivedFrom<typename TRemovePointer<T>::Type, UObject>::IsDerived>::Type>
 	{
@@ -154,6 +199,9 @@ namespace UE::StateTree::PropertyRefHelpers
 		}
 	};
 
+	/* Checks if provided property is compatible with selected Enum. */
+	bool STATETREEMODULE_API IsPropertyCompatibleWithEnum(const FProperty& Property, const class UEnum& Enum);
+
 	template<class T>
 	struct Validator<T, typename TEnableIf<TIsEnum<T>::Value>::Type>
 	{
diff --git a/Engine/Plugins/Runtime/StateTree/Source/StateTreeTestSuite/Private/StateTreeTest.cpp b/Engine/Plugins/Runtime/StateTree/Source/StateTreeTestSuite/Private/StateTreeTest.cpp
--- a/Engine/Plugins/Runtime/StateTree/Source/StateTreeTestSuite/Private/StateTreeTest.cpp
+++ b/Engine/Plugins/Runtime/StateTree/Source/StateTreeTestSuite/Private/StateTreeTest.cpp
@@ -1759,9 +1759,17 @@ struct FStateTreeTest_References : FAITestBase
 		PropertyBindings.Add(UE::StateTree::Tests::MakeBinding(SourceDesc.ID, TEXT("Item.A"), TargetDesc.ID, TEXT("RefToInt")));
 		PropertyBindings.Add(UE::StateTree::Tests::MakeBinding(SourceDesc.ID, TEXT("Array"), TargetDesc.ID, TEXT("RefToStructArray")));
 
+		FStateTreeTest_PropertyRefSourceStruct Source;
+		FStateTreeDataView SourceView = FStateTreeDataView(FStructView::Make(Source));
+
 		FStateTreeTest_PropertyRefTargetStruct Target;
 		FStateTreeDataView TargetView(FStructView::Make(Target));
-		const bool bCompileReferencesResult = BindingCompiler.CompileReferences(TargetDesc, PropertyBindings, TargetView);
+
+		TMap<FGuid, const FStateTreeDataView> IDToStructValue;
+		IDToStructValue.Emplace(SourceDesc.ID, SourceView);
+		IDToStructValue.Emplace(TargetDesc.ID, TargetView);
+
+		const bool bCompileReferencesResult = BindingCompiler.CompileReferences(TargetDesc, PropertyBindings, TargetView, IDToStructValue);
 		AITEST_TRUE("CompileReferences should succeed", bCompileReferencesResult);	
 
 		BindingCompiler.Finalize();
@@ -1769,10 +1777,6 @@ struct FStateTreeTest_References : FAITestBase
 		const bool bResolveResult = Bindings.ResolvePaths();
 		AITEST_TRUE("ResolvePaths should succeed", bResolveResult);
 
-		FStateTreeTest_PropertyRefSourceStruct Source;
-
-		FStateTreeDataView SourceView = FStateTreeDataView(FStructView::Make(Source));
-
 		{
 			const FStateTreePropertyAccess* PropertyAccess = Bindings.GetPropertyAccess(Target.RefToStruct);
 			AITEST_NOT_NULL("GetPropertyAccess should succeed", PropertyAccess);
@@ -1841,26 +1845,37 @@ struct FStateTreeTest_ReferencesConstness : FAITestBase
 		FStateTreePropertyPathBinding ContextPropertyBinding = UE::StateTree::Tests::MakeBinding(SourceAsTaskDesc.ID, TEXT("Item"), TargetDesc.ID, TEXT("RefToStruct"));
 		FStateTreePropertyPathBinding ContextOutputPropertyBinding = UE::StateTree::Tests::MakeBinding(SourceAsTaskDesc.ID, TEXT("Item"), TargetDesc.ID, TEXT("RefToStruct"));
 
+		FStateTreeTest_PropertyRefSourceStruct SourceAsTask;
+		FStateTreeDataView SourceAsTaskView(FStructView::Make(SourceAsTask));
+
+		FStateTreeTest_PropertyRefSourceStruct SourceAsContext;
+		FStateTreeDataView SourceAsContextView(FStructView::Make(SourceAsContext));
+
 		FStateTreeTest_PropertyRefTargetStruct Target;
 		FStateTreeDataView TargetView(FStructView::Make(Target));
 
+		TMap<FGuid, const FStateTreeDataView> IDToStructValue;
+		IDToStructValue.Emplace(SourceAsTaskDesc.ID, SourceAsTaskView);
+		IDToStructValue.Emplace(SourceAsContextDesc.ID, SourceAsContextView);
+		IDToStructValue.Emplace(TargetDesc.ID, TargetView);
+
 		{
-			const bool bCompileReferenceResult = BindingCompiler.CompileReferences(TargetDesc, {TaskPropertyBinding}, TargetView);
+			const bool bCompileReferenceResult = BindingCompiler.CompileReferences(TargetDesc, {TaskPropertyBinding}, TargetView, IDToStructValue);
 			AITEST_FALSE("CompileReferences should fail", bCompileReferenceResult);
 		}
 
 		{
-			const bool bCompileReferenceResult = BindingCompiler.CompileReferences(TargetDesc, {TaskOutputPropertyBinding}, TargetView);
+			const bool bCompileReferenceResult = BindingCompiler.CompileReferences(TargetDesc, {TaskOutputPropertyBinding}, TargetView, IDToStructValue);
 			AITEST_TRUE("CompileReferences should succeed", bCompileReferenceResult);
 		}
 
 		{
-			const bool bCompileReferenceResult = BindingCompiler.CompileReferences(TargetDesc, {ContextPropertyBinding}, TargetView);
+			const bool bCompileReferenceResult = BindingCompiler.CompileReferences(TargetDesc, {ContextPropertyBinding}, TargetView, IDToStructValue);
 			AITEST_FALSE("CompileReferences should fail", bCompileReferenceResult);
 		}
 
 		{	
-			const bool bCompileReferenceResult = BindingCompiler.CompileReferences(TargetDesc, {ContextOutputPropertyBinding}, TargetView);
+			const bool bCompileReferenceResult = BindingCompiler.CompileReferences(TargetDesc, {ContextOutputPropertyBinding}, TargetView, IDToStructValue);
 			AITEST_FALSE("CompileReferences should fail", bCompileReferenceResult);
 		}
 
