Fixed properties in FStateTreeBlueprintPropertyRef that were not initialized properly

diff --git a/Engine/Plugins/Runtime/StateTree/Source/StateTreeModule/Public/StateTreePropertyRef.h b/Engine/Plugins/Runtime/StateTree/Source/StateTreeModule/Public/StateTreePropertyRef.h
--- a/Engine/Plugins/Runtime/StateTree/Source/StateTreeModule/Public/StateTreePropertyRef.h
+++ b/Engine/Plugins/Runtime/StateTree/Source/StateTreeModule/Public/StateTreePropertyRef.h
@@ -277,11 +277,11 @@ private:
 
 	/** If specified, the reference is to an TArray<RefType> */
 	UPROPERTY(EditAnywhere, Category = "InternalType")
-	uint8 bIsRefToArray : 1;
+	uint8 bIsRefToArray : 1 = false;
 
 	/** If specified, the reference can be left unbound, otherwise the State Tree compiler report error if the reference is not bound. */
 	UPROPERTY(EditAnywhere, Category = "Parameter")
-	uint8 bIsOptional : 1;
+	uint8 bIsOptional : 1 = false;
 
 	/** Specifies the type of property to reference together with RefType, used for Enums, Structs, Objects and Classes. */
 	UPROPERTY(EditAnywhere, Category= "InternalType")
