[StateTreeDebugger] handle buffered data differently if it can be reused in multiple traces during the same game session or if it need to be reset when stopping the trace

diff --git a/Engine/Plugins/Runtime/StateTree/Source/StateTreeModule/Private/Debugger/StateTreeTrace.cpp b/Engine/Plugins/Runtime/StateTree/Source/StateTreeModule/Private/Debugger/StateTreeTrace.cpp
--- a/Engine/Plugins/Runtime/StateTree/Source/StateTreeModule/Private/Debugger/StateTreeTrace.cpp
+++ b/Engine/Plugins/Runtime/StateTree/Source/StateTreeModule/Private/Debugger/StateTreeTrace.cpp
@@ -8,6 +8,7 @@
 #include "ObjectTrace.h"
 #include "Serialization/BufferArchive.h"
 #include "StateTree.h"
+#include "StateTreeDelegates.h"
 #include "StateTreeExecutionTypes.h"
 #include "UObject/Package.h"
 #include "Trace/Trace.inl"
@@ -121,12 +122,7 @@ namespace UE::StateTreeTrace
 {
 
 FDelegateHandle GOnWorldTickStartDelegateHandle;
-
-#if WITH_EDITOR
-FDelegateHandle GOnPIEStartDelegateHandle;
-#endif // WITH_EDITOR
-
-double GRecordingWorldTime = -1;
+FDelegateHandle GTracingStateChangedDelegateHandle;
 
 /** Struct to keep track if a given phase was traced or not. */
 struct FPhaseTraceStatusPair
@@ -263,6 +259,7 @@ struct FInstanceEventBufferedData
 /** Struct to keep track of the buffered event data and flush them. */
 struct FBufferedDataList
 {
+	double RecordingWorldTime = -1;
 	double TracedRecordingWorldTime = -1;
 	
 	/**
@@ -280,6 +277,10 @@ struct FBufferedDataList
 	/** Flag use to prevent reentrant calls */
 	bool bFlushing = false;
 
+	uint16 NextAssetDebugId = 1;
+	int32 CurrentVersion = 0;
+	int32 FlushedVersion = -1;
+
 	void Flush(const FStateTreeInstanceDebugId InstanceId)
 	{
 		if (bFlushing)
@@ -289,20 +290,26 @@ struct FBufferedDataList
 
 		TGuardValue<bool> GuardReentry(bFlushing, true);
 
-		// Trace asset events first since they are required for instance lifetime event types.
-		for (const FAssetDebugIdEventBufferedData& AssetDebugIdEventData : AssetDebugIdEvents)
+		if (FlushedVersion != CurrentVersion)
 		{
-			AssetDebugIdEventData.Trace();
-		}
-		AssetDebugIdEvents.Empty();
+			FlushedVersion = CurrentVersion;
 
-		// Then trace instance lifetime events since they are required for other event types.
-		// It is also associated to an older world time.
-		for (TPair<FStateTreeInstanceDebugId, FInstanceEventBufferedData>& Pair : InstanceLifetimeEvents)
-		{
-			Pair.Value.Trace();
+			// Trace asset events first since they are required for instance lifetime event types.
+			// Events are preserved in case the trace session is stopped and then a new one gets started
+			// in the same game session. In which case we need to output the ids to that new trace.
+			for (const FAssetDebugIdEventBufferedData& AssetDebugIdEventData : AssetDebugIdEvents)
+			{
+				AssetDebugIdEventData.Trace();
+			}
+
+			// Then trace instance lifetime events since they are required for other event types.
+			// It is also associated to an older world time.
+			// Events are also preserved for the same reason as AssetDebugIdEvents.
+			for (TPair<FStateTreeInstanceDebugId, FInstanceEventBufferedData>& Pair : InstanceLifetimeEvents)
+			{
+				Pair.Value.Trace();
+			}
 		}
-		InstanceLifetimeEvents.Empty();
 
 		TraceWorldTime();
 
@@ -317,11 +324,11 @@ struct FBufferedDataList
 	 */
 	void TraceWorldTime()
 	{
-		if (TracedRecordingWorldTime != GRecordingWorldTime)
+		if (TracedRecordingWorldTime != RecordingWorldTime)
 		{
-			TracedRecordingWorldTime = GRecordingWorldTime;
+			TracedRecordingWorldTime = RecordingWorldTime;
 			UE_TRACE_LOG(StateTreeDebugger, WorldTimestampEvent, StateTreeDebugChannel)
-				<< WorldTimestampEvent.WorldTime(GRecordingWorldTime);
+				<< WorldTimestampEvent.WorldTime(RecordingWorldTime);
 		}
 	}
 
@@ -350,7 +357,7 @@ struct FBufferedDataList
 							<< PhaseEvent.Phase(static_cast<std::underlying_type_t<EStateTreeUpdatePhase>>(StackEntry.Phase))
 							<< PhaseEvent.StateIndex(StackEntry.StateHandle.Index)
 							<< PhaseEvent.EventType(static_cast<std::underlying_type_t<EStateTreeTraceEventType>>(EStateTreeTraceEventType::Push));
-			
+
 						StackEntry.bTraced = true;
 					}
 				}
@@ -358,6 +365,15 @@ struct FBufferedDataList
 			}
 		}
 	}
+
+	void BumpTraceVersion()
+	{
+		// Bump version so shareable data will be flush in the next trace (e.g. Asset ids, instance lifetime events, etc.)
+		CurrentVersion++;
+
+		// Force world time to trace on the next event
+		RecordingWorldTime = -1;
+	}
 };
 
 /**
@@ -367,14 +383,6 @@ struct FBufferedDataList
  */
 thread_local FBufferedDataList GBufferedEvents;
 
-struct FStateTreeTraceDebugIdPair
-{
-	TWeakObjectPtr<const UStateTree> WeakStateTree;
-	FStateTreeIndex16 Id;
-};
-thread_local uint16 GNextAssetDebugId(1);
-thread_local TArray<FStateTreeTraceDebugIdPair> GStateTreeTraceDebugIds;
-
 /**
  * Pushed or pops an entry on the Phase stack for a given Instance.
  * Will send the Pop events for phases popped if their associated Push events were sent.
@@ -473,54 +481,52 @@ void SerializeDataViewToArchive(FBufferArchive& Ar, const FStateTreeDataView Dat
 
 void RegisterGlobalDelegates()
 {
-#if WITH_EDITOR
-	GOnPIEStartDelegateHandle = FEditorDelegates::BeginPIE.AddLambda([&LastRecordingWorldTime=GBufferedEvents.TracedRecordingWorldTime](const bool bIsSimulating)
-		{
-			LastRecordingWorldTime = -1;
-			GStateTreeTraceDebugIds.Reset();
-			GNextAssetDebugId = 1;
-		});
-#endif // WITH_EDITOR
-	
-	GOnWorldTickStartDelegateHandle = FWorldDelegates::OnWorldTickStart.AddLambda([&WorldTime=GRecordingWorldTime](const UWorld* TickedWorld, ELevelTick TickType, float DeltaTime)
+	GOnWorldTickStartDelegateHandle = FWorldDelegates::OnWorldTickStart.AddLambda([&WorldTime=GBufferedEvents.RecordingWorldTime](const UWorld* TickedWorld, ELevelTick TickType, float DeltaTime)
 		{
 #if OBJECT_TRACE_ENABLED
 			WorldTime = FObjectTrace::GetWorldElapsedTime(TickedWorld);
 #endif// OBJECT_TRACE_ENABLED
 		});
+
+	GTracingStateChangedDelegateHandle = UE::StateTree::Delegates::OnTracingStateChanged.AddLambda([](const bool bTracesEnabled)
+	{
+		// Bump trace version when disabling traces so next trace will flush buffered events that are still relevant.
+		if (!bTracesEnabled)
+		{
+			GBufferedEvents.BumpTraceVersion();
+		}
+	});
 }
 
 void UnregisterGlobalDelegates()
 {
-#if WITH_EDITOR
-	FEditorDelegates::BeginPIE.Remove(GOnPIEStartDelegateHandle);
-#endif // WITH_EDITOR
-	
 	FWorldDelegates::OnWorldTickStart.Remove(GOnWorldTickStartDelegateHandle);
 	GOnWorldTickStartDelegateHandle.Reset();
+
+	UE::StateTree::Delegates::OnTracingStateChanged.Remove(GTracingStateChangedDelegateHandle);
+	GTracingStateChangedDelegateHandle.Reset();
 }
 
 FStateTreeIndex16 FindOrAddDebugIdForAsset(const UStateTree* StateTree)
 {
 	FStateTreeIndex16 AssetDebugId;
-	const FStateTreeTraceDebugIdPair* ExistingPair = GStateTreeTraceDebugIds.FindByPredicate([StateTree](const FStateTreeTraceDebugIdPair& Pair)
+	const FAssetDebugIdEventBufferedData* ExistingPair = GBufferedEvents.AssetDebugIdEvents.FindByPredicate([StateTree](const FAssetDebugIdEventBufferedData& BufferedData)
 	{
-		return Pair.WeakStateTree == StateTree;
+		return BufferedData.WeakStateTree == StateTree;
 	});
 
 	if (ExistingPair == nullptr)
 	{
 		if (ensure(StateTree != nullptr))
 		{
-			AssetDebugId = FStateTreeIndex16(GNextAssetDebugId++);
-			GStateTreeTraceDebugIds.Emplace(FStateTreeTraceDebugIdPair(StateTree, AssetDebugId));
-
+			AssetDebugId = FStateTreeIndex16(GBufferedEvents.NextAssetDebugId++);
+			GBufferedEvents.AssetDebugIdEvents.Emplace(StateTree, AssetDebugId);
 			OutputAssetDebugIdEvent(StateTree, AssetDebugId);
 		}
 	}
 	else
 	{
-		AssetDebugId = ExistingPair->Id;
+		AssetDebugId = ExistingPair->AssetDebugId;
 	}
 
 	return AssetDebugId;
@@ -546,10 +552,6 @@ void OutputAssetDebugIdEvent(
 			<< AssetDebugIdEvent.CompiledDataHash(StateTree->LastCompiledEditorDataHash)
 			<< AssetDebugIdEvent.AssetDebugId(AssetDebugId.Get());
 	}
-	else
-	{
-		GBufferedEvents.AssetDebugIdEvents.Emplace(StateTree, AssetDebugId);
-	}
 }
 
 void OutputInstanceLifetimeEvent(
@@ -573,15 +575,18 @@ void OutputInstanceLifetimeEvent(
 			<< InstanceEvent.EventType(static_cast<std::underlying_type_t<EStateTreeTraceEventType>>(EventType))
 			<< InstanceEvent.AssetDebugId(AssetDebugId.Get());
 	}
-	else
+
+	// Buffer these events regardless of the status of the channel since they will be used
+	// when flushing buffered event when a late recording is started or more than one trace are started
+	// during the same game session (i.e. Start Traces -> Stop Traces -> Start Traces).
+	if (!GBufferedEvents.bFlushing)
 	{
 		if (EventType == EStateTreeTraceEventType::Push)
 		{
-			GBufferedEvents.InstanceLifetimeEvents.Emplace(InstanceId, FInstanceEventBufferedData(GRecordingWorldTime, StateTree, InstanceId, InstanceName));
+			GBufferedEvents.InstanceLifetimeEvents.Emplace(InstanceId, FInstanceEventBufferedData(GBufferedEvents.RecordingWorldTime, StateTree, InstanceId, InstanceName));
 		}
 		else if (EventType == EStateTreeTraceEventType::Pop)
 		{
-			// Remove matching instance events since if it was not sent then no other events were sent between, hence not needed in the trace.
 			GBufferedEvents.InstanceLifetimeEvents.Remove(InstanceId);
 		}
 		else
@@ -753,7 +758,7 @@ void OutputActiveStatesEventTrace(
 		if (FInstanceEventBufferedData* ExisingBufferedData = GBufferedEvents.InstanceLifetimeEvents.Find(InstanceId))
 		{
 			ExisingBufferedData->ActiveStates= FInstanceEventBufferedData::FActiveStates(ActiveFrames);
-			ExisingBufferedData->ActiveStatesRecordingWorldTime = GRecordingWorldTime;
+			ExisingBufferedData->ActiveStatesRecordingWorldTime = GBufferedEvents.RecordingWorldTime;
 		}
 	}
 }
