StateTree: Consolidated searching for frame and it's parent into a single helper function.

diff --git a/Engine/Plugins/Runtime/StateTree/Source/StateTreeModule/Private/Blueprint/StateTreeNodeBlueprintBase.cpp b/Engine/Plugins/Runtime/StateTree/Source/StateTreeModule/Private/Blueprint/StateTreeNodeBlueprintBase.cpp
--- a/Engine/Plugins/Runtime/StateTree/Source/StateTreeModule/Private/Blueprint/StateTreeNodeBlueprintBase.cpp
+++ b/Engine/Plugins/Runtime/StateTree/Source/StateTreeModule/Private/Blueprint/StateTreeNodeBlueprintBase.cpp
@@ -75,23 +75,12 @@ void* UStateTreeNodeBlueprintBase::GetMutablePtrToProperty(const FStateTreeBluep
 	check(CachedFrameStateTree && CachedFrameRootState.IsValid());
 
 	TConstArrayView<FStateTreeExecutionFrame> ActiveFrames = InstanceStorage->GetExecutionState().ActiveFrames;
-
-	const int32 CurrentFrameIndex = ActiveFrames.IndexOfByPredicate([this](const FStateTreeExecutionFrame& Frame)
-	{
-		return Frame.StateTree == CachedFrameStateTree && Frame.RootState == CachedFrameRootState;
-	});
-
-	check(CurrentFrameIndex != INDEX_NONE);
-	const FStateTreeExecutionFrame& CurrentFrame = ActiveFrames[CurrentFrameIndex];
-
 	const FStateTreeExecutionFrame* CurrentParentFrame = nullptr;
-	if (CurrentFrameIndex > 0)
-	{
-		CurrentParentFrame = &ActiveFrames[CurrentFrameIndex - 1];
-	}
+	const FStateTreeExecutionFrame* CurrentFrame = FStateTreeExecutionContext::FindFrame(CachedFrameStateTree, CachedFrameRootState, ActiveFrames, CurrentParentFrame);
+	check(CurrentFrame);
 
 	const FProperty* SourceProperty = nullptr;
-	void* PropertyAddress = UE::StateTree::PropertyRefHelpers::GetMutablePtrToProperty<void>(PropertyRef, *InstanceStorage, CurrentFrame, CurrentParentFrame, &SourceProperty);
+	void* PropertyAddress = UE::StateTree::PropertyRefHelpers::GetMutablePtrToProperty<void>(PropertyRef, *InstanceStorage, *CurrentFrame, CurrentParentFrame, &SourceProperty);
 	if (PropertyAddress && UE::StateTree::PropertyRefHelpers::IsBlueprintPropertyRefCompatibleWithProperty(*SourceProperty, &PropertyRef))
 	{
 		OutSourceProperty = const_cast<FProperty*>(SourceProperty);
diff --git a/Engine/Plugins/Runtime/StateTree/Source/StateTreeModule/Private/StateTreeExecutionContext.cpp b/Engine/Plugins/Runtime/StateTree/Source/StateTreeModule/Private/StateTreeExecutionContext.cpp
--- a/Engine/Plugins/Runtime/StateTree/Source/StateTreeModule/Private/StateTreeExecutionContext.cpp
+++ b/Engine/Plugins/Runtime/StateTree/Source/StateTreeModule/Private/StateTreeExecutionContext.cpp
@@ -936,6 +936,27 @@ FStateTreeDataView FStateTreeExecutionContext::GetDataView(FStateTreeInstanceSto
 	return {};
 }
 
+const FStateTreeExecutionFrame* FStateTreeExecutionContext::FindFrame(const UStateTree* StateTree, FStateTreeStateHandle RootState, TConstArrayView<FStateTreeExecutionFrame> Frames, const FStateTreeExecutionFrame*& OutParentFrame)
+{
+	const int32 FrameIndex = Frames.IndexOfByPredicate([&StateTree, RootState](const FStateTreeExecutionFrame& Frame)
+	{
+		return Frame.StateTree == StateTree && Frame.RootState == RootState;
+	});
+
+	if(FrameIndex == INDEX_NONE)
+	{
+		OutParentFrame = nullptr;
+		return nullptr;
+	}
+
+	if (FrameIndex > 0)
+	{
+		OutParentFrame = &Frames[FrameIndex - 1];
+	}
+
+	return &Frames[FrameIndex];
+}
+
 bool FStateTreeExecutionContext::IsHandleSourceValid(const FStateTreeExecutionFrame* ParentFrame, const FStateTreeExecutionFrame& CurrentFrame, const FStateTreeDataHandle Handle) const
 {
 	// Checks that the instance data is valid for specific handle types.
@@ -2409,15 +2430,12 @@ bool FStateTreeExecutionContext::TriggerTransitions()
 	for (const FStateTreeTransitionRequest& Request : InstanceData.GetTransitionRequests())
 	{
 		// Find frame associated with the request.
-		const int32 FrameIndex = Exec.ActiveFrames.IndexOfByPredicate([&Request](const FStateTreeExecutionFrame& Frame)
-		{
-			return Frame.StateTree == Request.SourceStateTree && Frame.RootState == Request.SourceRootState;
-		});
+		const FStateTreeExecutionFrame* ParentFrame = nullptr;
+		const FStateTreeExecutionFrame* CurrentFrame = FindFrame(Request.SourceStateTree, Request.SourceRootState, Exec.ActiveFrames, ParentFrame);
 
-		if (FrameIndex != INDEX_NONE)
+		if (CurrentFrame)
 		{
-			const FStateTreeExecutionFrame& CurrentFrame = Exec.ActiveFrames[FrameIndex];
-			if (RequestTransition(CurrentFrame, Request.TargetState, Request.Priority))
+			if (RequestTransition(*CurrentFrame, Request.TargetState, Request.Priority))
 			{
 				NextTransitionSource = FStateTreeTransitionSource(EStateTreeTransitionSourceType::ExternalRequest, Request.TargetState, Request.Priority);
 			}
diff --git a/Engine/Plugins/Runtime/StateTree/Source/StateTreeModule/Public/StateTreeExecutionContext.h b/Engine/Plugins/Runtime/StateTree/Source/StateTreeModule/Public/StateTreeExecutionContext.h
--- a/Engine/Plugins/Runtime/StateTree/Source/StateTreeModule/Public/StateTreeExecutionContext.h
+++ b/Engine/Plugins/Runtime/StateTree/Source/StateTreeModule/Public/StateTreeExecutionContext.h
@@ -374,6 +374,9 @@ public:
 	/** @return data view of the specified handle relative to given frame. */
 	static FStateTreeDataView GetDataView(FStateTreeInstanceStorage& InstanceDataStorage, FStateTreeInstanceStorage* CurrentlyProcessedSharedInstanceStorage, const FStateTreeExecutionFrame* ParentFrame, const FStateTreeExecutionFrame& CurrentFrame, TConstArrayView<FStateTreeDataView> ContextAndExternalDataViews, const FStateTreeDataHandle Handle);
 
+	/** Looks for a frame in provided list of frames. */
+	static const FStateTreeExecutionFrame* FindFrame(const UStateTree* StateTree, FStateTreeStateHandle RootState, TConstArrayView<FStateTreeExecutionFrame> Frames, const FStateTreeExecutionFrame*& OutParentFrame);
+
 protected:
 
 	/** Max number of execution frames handled during state selection. */
diff --git a/Engine/Plugins/Runtime/StateTree/Source/StateTreeModule/Public/StateTreePropertyRef.h b/Engine/Plugins/Runtime/StateTree/Source/StateTreeModule/Public/StateTreePropertyRef.h
--- a/Engine/Plugins/Runtime/StateTree/Source/StateTreeModule/Public/StateTreePropertyRef.h
+++ b/Engine/Plugins/Runtime/StateTree/Source/StateTreeModule/Public/StateTreePropertyRef.h
@@ -45,21 +45,16 @@ namespace UE::StateTree::PropertyRefHelpers
 					return nullptr;
 				}
 
-				const TArray<FStateTreeExecutionFrame>& ActiveFrames = InstanceDataStorage.GetExecutionState().ActiveFrames;
-				const int32 FrameIndex = ActiveFrames.IndexOfByPredicate([ParentExecutionFrame](const FStateTreeExecutionFrame& Frame)
-				{
-					return Frame.RootState == ParentExecutionFrame->RootState && Frame.StateTree == ParentExecutionFrame->StateTree;
-				});
-
-				if (FrameIndex == INDEX_NONE)
+				const FStateTreeExecutionFrame* ParentFrame = nullptr;		
+				TConstArrayView<FStateTreeExecutionFrame> ActiveFrames = InstanceDataStorage.GetExecutionState().ActiveFrames;
+				const FStateTreeExecutionFrame* Frame = FStateTreeExecutionContext::FindFrame(ParentExecutionFrame->StateTree, ParentExecutionFrame->RootState, ActiveFrames, ParentFrame);
+				
+				if (Frame == nullptr)
 				{
 					return nullptr;
 				}
 
-				const FStateTreeExecutionFrame& Frame = ActiveFrames[FrameIndex];
-				const FStateTreeExecutionFrame* ParentFrame = FrameIndex > 0 ? &ActiveFrames[FrameIndex - 1] : nullptr;
-
-				return GetMutablePtrToProperty<T>(*ReferencedPropertyRef, InstanceDataStorage, Frame, ParentFrame, OutSourceProperty);
+				return GetMutablePtrToProperty<T>(*ReferencedPropertyRef, InstanceDataStorage, *Frame, ParentFrame, OutSourceProperty);
 			}
 			else
 			{
@@ -202,22 +197,17 @@ struct TStateTreePropertyRefExternalHandle
 		}
 
 		FStateTreeInstanceStorage& InstanceStorage = *WeakInstanceStorage.Pin();
+		TConstArrayView<FStateTreeExecutionFrame> ActiveFrames = InstanceStorage.GetExecutionState().ActiveFrames;
 
-		const TArray<FStateTreeExecutionFrame>& ActiveFrames = InstanceStorage.GetExecutionState().ActiveFrames;
-		const int32 FrameIndex = ActiveFrames.IndexOfByPredicate([this](const FStateTreeExecutionFrame& Frame)
-		{
-			return Frame.RootState == RootState && Frame.StateTree == WeakStateTree;
-		});
+		const FStateTreeExecutionFrame* ParentFrame = nullptr;
+		const FStateTreeExecutionFrame* Frame = FStateTreeExecutionContext::FindFrame(WeakStateTree.Get(), RootState, ActiveFrames, ParentFrame);
 
-		if (FrameIndex == INDEX_NONE)
+		if (Frame == nullptr)
 		{
 			return nullptr;
 		}
 
-		const FStateTreeExecutionFrame& Frame = ActiveFrames[FrameIndex];
-		const FStateTreeExecutionFrame* ParentFrame = FrameIndex > 0 ? &ActiveFrames[FrameIndex - 1] : nullptr;
-
-		return UE::StateTree::PropertyRefHelpers::GetMutablePtrToProperty<TRef>(PropertyRef, InstanceStorage, Frame, ParentFrame);
+		return UE::StateTree::PropertyRefHelpers::GetMutablePtrToProperty<TRef>(PropertyRef, InstanceStorage, *Frame, ParentFrame);
 	}
 
 private:
