StateTree: Fixed incorrect InstanceDataNum for CompactStates

diff --git a/Engine/Plugins/Runtime/StateTree/Source/StateTreeEditorModule/Private/StateTreeCompiler.cpp b/Engine/Plugins/Runtime/StateTree/Source/StateTreeEditorModule/Private/StateTreeCompiler.cpp
--- a/Engine/Plugins/Runtime/StateTree/Source/StateTreeEditorModule/Private/StateTreeCompiler.cpp
+++ b/Engine/Plugins/Runtime/StateTree/Source/StateTreeEditorModule/Private/StateTreeCompiler.cpp
@@ -694,14 +694,15 @@ bool FStateTreeCompiler::CreateStateTasksAndParameters()
 			return false;
 		}
 
-		if (const auto Validation = UE::StateTree::Compiler::IsValidCount8(InstanceDataIndex); Validation.DidFail())
+		const int32 InstanceDataNum = InstanceDataIndex - FirstInstanceDataIndex[i];
+		if (const auto Validation = UE::StateTree::Compiler::IsValidCount8(InstanceDataNum); Validation.DidFail())
 		{
-			Validation.Log(Log, TEXT("InstanceDataIndex"));
+			Validation.Log(Log, TEXT("InstanceDataNum"));
 			return false;
 		}
 
 		CompactState.TasksNum = uint8(TasksNum);
-		CompactState.InstanceDataNum = uint8(InstanceDataIndex);
+		CompactState.InstanceDataNum = uint8(InstanceDataNum);
 	}
 	
 	return true;
