Extend FStateTreePropertyRef to be able to bind to multiple property types

diff --git a/Engine/Plugins/Runtime/StateTree/Source/StateTreeEditorModule/Private/Customizations/StateTreeBindingExtension.cpp b/Engine/Plugins/Runtime/StateTree/Source/StateTreeEditorModule/Private/Customizations/StateTreeBindingExtension.cpp
--- a/Engine/Plugins/Runtime/StateTree/Source/StateTreeEditorModule/Private/Customizations/StateTreeBindingExtension.cpp
+++ b/Engine/Plugins/Runtime/StateTree/Source/StateTreeEditorModule/Private/Customizations/StateTreeBindingExtension.cpp
@@ -394,8 +394,8 @@ struct FCachedBindingData : public TSharedFromThis<FCachedBindingData>
 		BindingOwner->GetDataViewByID(TargetPath.GetStructID(), TargetDataView);
 
 		FEdGraphPinType PinType;
-
-		if (UE::StateTree::PropertyRefHelpers::IsPropertyRef(*Property) && TargetDataView.IsValid())
+		const bool bIsPropertyRef = UE::StateTree::PropertyRefHelpers::IsPropertyRef(*Property);
+		if (bIsPropertyRef && TargetDataView.IsValid())
 		{
 			// Use internal type to construct PinType if it's property of PropertyRef type.
 			TArray<FStateTreePropertyPathIndirection> TargetIndirections;
@@ -410,6 +410,8 @@ struct FCachedBindingData : public TSharedFromThis<FCachedBindingData>
 			Schema->ConvertPropertyToPinType(Property, PinType);
 		}
 
+
+		FTextBuilder TooltipBuilder;
 		if (const FStateTreePropertyPath* SourcePath = EditorBindings->GetPropertyBindingSource(TargetPath))
 		{
 			const FStateTreeBindableStructDesc* SourceDesc = UE::StateTree::PropertyBinding::FindStruct(AccessibleStructs, SourcePath->GetStructID());
@@ -464,13 +466,18 @@ struct FCachedBindingData : public TSharedFromThis<FCachedBindingData>
 
 					if (SourcePath->IsPathEmpty())
 					{
-						TooltipText = FText::Format(LOCTEXT("ExistingBindingTooltip", "Property is bound to {0}."), FText::FromString(SourceDesc->ToString()));
+						TooltipBuilder.AppendLineFormat(LOCTEXT("ExistingBindingTooltip", "Property is bound to {0}."), FText::FromString(SourceDesc->ToString()));
 					}
 					else
 					{
-						TooltipText = FText::Format(LOCTEXT("ExistingBindingWithPropertyTooltip", "Property is bound to {0} property {1}."), FText::FromString(SourceDesc->ToString()), FText::FromString(SourcePath->ToString()));
+						TooltipBuilder.AppendLineFormat(LOCTEXT("ExistingBindingWithPropertyTooltip", "Property is bound to {0} property {1}."), FText::FromString(SourceDesc->ToString()), FText::FromString(SourcePath->ToString()));
 					}
-					
+
+					if (bIsPropertyRef) // Update the pin type with source property so that property ref that can binds to multiple types display the binded one.
+					{
+						Schema->ConvertPropertyToPinType(SourceLeafProperty, PinType);
+					}
+
 					Image = FAppStyle::GetBrush(PropertyIcon);
 					Color = Schema->GetPinTypeColor(PinType);
 				}
@@ -491,12 +498,12 @@ struct FCachedBindingData : public TSharedFromThis<FCachedBindingData>
 
 					if (SourcePath->IsPathEmpty())
 					{
-						TooltipText = FText::Format(LOCTEXT("MismatchingBindingTooltip", "Property is bound to {0}, but binding source type '{1}' does not match property type '{2}'."),
+						TooltipBuilder.AppendLineFormat(LOCTEXT("MismatchingBindingTooltip", "Property is bound to {0}, but binding source type '{1}' does not match property type '{2}'."),
 							FText::FromString(SourceDesc->ToString()), SourceType, TargetType);
 					}
 					else
 					{
-						TooltipText = FText::Format(LOCTEXT("MismatchingBindingTooltipWithProperty", "Property is bound to {0} property {1}, but binding source type '{2}' does not match property type '{3}'."),
+						TooltipBuilder.AppendLineFormat(LOCTEXT("MismatchingBindingTooltipWithProperty", "Property is bound to {0} property {1}, but binding source type '{2}' does not match property type '{3}'."),
 							FText::FromString(SourceDesc->ToString()), FText::FromString(SourcePath->ToString()), SourceType, TargetType);
 					}
 
@@ -508,7 +515,7 @@ struct FCachedBindingData : public TSharedFromThis<FCachedBindingData>
 			{
 				// Missing source
 				Text = FText::Format(LOCTEXT("MissingSource", "???.{0}"), FText::FromString(SourcePath->ToString()));
-				TooltipText = FText::Format(LOCTEXT("MissingBindingTooltip", "Missing binding source for property path '{0}'."), FText::FromString(SourcePath->ToString()));
+				TooltipBuilder.AppendLineFormat(LOCTEXT("MissingBindingTooltip", "Missing binding source for property path '{0}'."), FText::FromString(SourcePath->ToString()));
 				Image = FCoreStyle::Get().GetBrush("Icons.ErrorWithColor");
 				Color = FLinearColor::White;
 			}
@@ -519,13 +526,32 @@ struct FCachedBindingData : public TSharedFromThis<FCachedBindingData>
 		{
 			// No bindings
 			Text = FText::GetEmpty();
-			TooltipText = FText::Format(LOCTEXT("BindTooltip", "Bind {0} to value from another property."), UE::StateTree::PropertyBinding::GetPropertyTypeText(Property));
+			TooltipBuilder.AppendLineFormat(LOCTEXT("BindTooltip", "Bind {0} to value from another property."), UE::StateTree::PropertyBinding::GetPropertyTypeText(Property));
+
 			Image = FAppStyle::GetBrush(PropertyIcon);
 			Color = Schema->GetPinTypeColor(PinType);
 
 			CachedSourcePath.Reset();
 		}
 
+		if (bIsPropertyRef)
+		{
+			if (Property->HasMetaData(UE::StateTree::PropertyRefHelpers::IsRefToArrayName))
+			{
+				TooltipBuilder.AppendLineFormat(LOCTEXT("PropertyRefBindingTooltip", "Supported types are Array of {0}"), FText::FromString(Property->GetMetaData(UE::StateTree::PropertyRefHelpers::RefTypeName)));
+			}
+			else
+			{
+				TooltipBuilder.AppendLineFormat(LOCTEXT("PropertyRefBindingTooltip", "Supported types are {0}"), FText::FromString(Property->GetMetaData(UE::StateTree::PropertyRefHelpers::RefTypeName)));
+				if (Property->HasMetaData(UE::StateTree::PropertyRefHelpers::CanRefToArrayName))
+				{
+					TooltipBuilder.AppendLine(LOCTEXT("PropertyRefBindingTooltip", "Supports Arrays"));
+				}
+			}
+		}
+
+		TooltipText = TooltipBuilder.ToText();
+
 		bIsDataCached = true;
 	}
 
diff --git a/Engine/Plugins/Runtime/StateTree/Source/StateTreeModule/Private/StateTreePropertyRefHelpers.cpp b/Engine/Plugins/Runtime/StateTree/Source/StateTreeModule/Private/StateTreePropertyRefHelpers.cpp
--- a/Engine/Plugins/Runtime/StateTree/Source/StateTreeModule/Private/StateTreePropertyRefHelpers.cpp
+++ b/Engine/Plugins/Runtime/StateTree/Source/StateTreeModule/Private/StateTreePropertyRefHelpers.cpp
@@ -25,9 +25,9 @@ namespace UE::StateTree::PropertyRefHelpers
 	static const FName NameName = TEXT("Name");
 	static const FName StringName = TEXT("String");
 	static const FName TextName = TEXT("Text");
-
-	static const FName IsRefToArrayName = TEXT("IsRefToArray");
-	static const FName RefTypeName = TEXT("RefType");
+	const FName IsRefToArrayName = TEXT("IsRefToArray");
+	const FName CanRefToArrayName = TEXT("CanRefToArray");
+	const FName RefTypeName = TEXT("RefType");
 	static const FName IsOptionalName = TEXT("Optional");
 
 	bool ArePropertyRefsCompatible(const FProperty& TargetRefProperty, const FProperty& SourceRefProperty, const void* TargetRefAddress, const void* SourceRefAddress)
@@ -47,86 +47,147 @@ namespace UE::StateTree::PropertyRefHelpers
 		check(IsPropertyRef(RefProperty));
 
 		const FProperty* TestProperty = &SourceProperty;
+		const bool bCanTargetRefArray = RefProperty.HasMetaData(CanRefToArrayName);
 		const bool bIsTargetRefArray = RefProperty.HasMetaData(IsRefToArrayName);
 
-		if (bIsTargetRefArray)
+		if (bIsTargetRefArray || bCanTargetRefArray)
 		{
 			if (const FArrayProperty* ArrayProperty = CastField<FArrayProperty>(TestProperty))
 			{
 				TestProperty = ArrayProperty->Inner;
 			}
-			else
+			else if(!bCanTargetRefArray)
 			{
 				return false;
 			}
 		}
 
-		const FString& TargetTypeNameStr = RefProperty.GetMetaData(RefTypeName);
-		if (TargetTypeNameStr.IsEmpty())
+		FString TargetTypeNameFullStr = RefProperty.GetMetaData(RefTypeName);
+		if (TargetTypeNameFullStr.IsEmpty())
 		{
 			return false;
 		}
 
-		const FName TargetTypeName = FName(*TargetTypeNameStr);
+		TargetTypeNameFullStr.RemoveSpacesInline();
 
-		if(TargetTypeName == BoolName)
-		{
-			return CastField<FBoolProperty>(TestProperty) != nullptr;
-		}
-		else if(TargetTypeName == ByteName)
-		{
-			return CastField<FByteProperty>(TestProperty) != nullptr;
-		}
-		else if(TargetTypeName == Int32Name)
-		{
-			return CastField<FIntProperty>(TestProperty) != nullptr;
-		}
-		else if(TargetTypeName == Int64Name)
-		{
-			return CastField<FInt64Property>(TestProperty) != nullptr;
-		}
-		else if(TargetTypeName == FloatName)
-		{
-			return CastField<FFloatProperty>(TestProperty) != nullptr;
-		}
-		else if(TargetTypeName == DoubleName)
-		{
-			return CastField<FDoubleProperty>(TestProperty) != nullptr;
-		}
-		else if(TargetTypeName == NameName)
-		{
-			return CastField<FNameProperty>(TestProperty) != nullptr;
-		}
-		else if(TargetTypeName == StringName)
-		{
-			return CastField<FStrProperty>(TestProperty) != nullptr;
-		}
-		else if(TargetTypeName == TextName)
-		{
-			return CastField<FTextProperty>(TestProperty) != nullptr;
-		}
-		else
+		TArray<FString> TargetTypes;
+		TargetTypeNameFullStr.ParseIntoArray(TargetTypes, TEXT(","), true);
+
+		const FStructProperty* SourceStructProperty = CastField<FStructProperty>(TestProperty);
+		// Check inside loop are only allowed to return true to avoid shortcircuiting the loop.
+		for (const FString& TargetTypeNameStr : TargetTypes)
 		{
-			UField* TargetRefField = UClass::TryFindTypeSlow<UField>(TargetTypeNameStr);
-			if (!TargetRefField)
+			const FName TargetTypeName = FName(*TargetTypeNameStr);
+			// Compare properties metadata directly if SourceProperty is PropertyRef as well
+			if (SourceStructProperty && SourceStructProperty->Struct == FStateTreePropertyRef::StaticStruct())
+
 			{
-				TargetRefField = LoadObject<UField>(nullptr, *TargetTypeNameStr);
+				const FName SourceTypeName(SourceStructProperty->GetMetaData(RefTypeName));
+				const bool bIsSourceRefArray = SourceStructProperty->GetBoolMetaData(IsRefToArrayName);
+				if (SourceTypeName == TargetTypeName && bIsSourceRefArray == bIsTargetRefArray)
+				{
+					return true;
+				}
+
 			}
 
-			if (const FStructProperty* SourceStructProperty = CastField<FStructProperty>(TestProperty))
+			if (TargetTypeName == BoolName)
 			{
-				return SourceStructProperty->Struct->IsChildOf(Cast<UStruct>(TargetRefField));
+				if (TestProperty->IsA<FBoolProperty>())
+				{
+					return true;
+				}
 			}
-
-			if (const FObjectProperty* ObjectProperty = CastField<FObjectProperty>(TestProperty))
+			else if (TargetTypeName == ByteName)
 			{
-				// Only referencing object of the same exact class should be allowed. Otherwise one could e.g assign UObject to AActor property through reference to UObject.
-				return ObjectProperty->PropertyClass == Cast<UStruct>(TargetRefField);
+				if (TestProperty->IsA<FByteProperty>())
+				{
+					return true;
+				}
+			}
+			else if (TargetTypeName == Int32Name)
+			{
+				if (TestProperty->IsA<FIntProperty>())
+				{
+					return true;
+				}
+			}
+			else if (TargetTypeName == Int64Name)
+			{
+				if (TestProperty->IsA<FInt64Property>())
+				{
+					return true;
+				}
+			}
+			else if (TargetTypeName == FloatName)
+			{
+				if (TestProperty->IsA<FFloatProperty>())
+				{
+					return true;
+				}
+			}
+			else if (TargetTypeName == DoubleName)
+			{
+				if (TestProperty->IsA<FDoubleProperty>())
+				{
+					return true;
+				}
+			}
+			else if (TargetTypeName == NameName)
+			{
+				if (TestProperty->IsA<FNameProperty>())
+				{
+					return true;
+				}
 			}
 
-			if (const FEnumProperty* EnumProperty = CastField<FEnumProperty>(TestProperty))
+			else if (TargetTypeName == StringName)
+
+			{
+				if (TestProperty->IsA<FStrProperty>())
+				{
+					return true;
+				}
+			}
+			else if (TargetTypeName == TextName)
 			{
-				return EnumProperty->GetEnum() == TargetRefField;
+				if (TestProperty->IsA<FTextProperty>())
+				{
+					return true;
+				}
+			}
+			else
+			{
+				UField* TargetRefField = UClass::TryFindTypeSlow<UField>(TargetTypeNameStr);
+				if (!TargetRefField)
+				{
+					TargetRefField = LoadObject<UField>(nullptr, *TargetTypeNameStr);
+				}
+
+				if (SourceStructProperty)
+				{
+					if (SourceStructProperty->Struct->IsChildOf(Cast<UStruct>(TargetRefField)))
+					{
+						return true;
+					}
+				}
+
+				if (const FObjectProperty* ObjectProperty = CastField<FObjectProperty>(TestProperty))
+				{
+					// Only referencing object of the same exact class should be allowed. Otherwise one could e.g assign UObject to AActor property through reference to UObject.
+					if(ObjectProperty->PropertyClass == Cast<UStruct>(TargetRefField))
+					{
+						return true;
+
+					}
+				}
+				else if (const FEnumProperty* EnumProperty = CastField<FEnumProperty>(TestProperty))
+				{
+					if (EnumProperty->GetEnum() == TargetRefField)
+					{
+						return true;
+					}
+				}
 			}
 		}
 
@@ -365,7 +426,11 @@ namespace UE::StateTree::PropertyRefHelpers
 			}
 			else
 			{
-				checkNoEntry();
+				checkCode(
+					TArray<FString> AllowedTypes;
+					TargetTypeNameStr.ParseIntoArray(AllowedTypes, TEXT(","), true);
+					check(AllowedTypes.Num() > 1); // We are trying to bind to multiple type so its normal we can't find what the pin type should be.
+				)
 			}
 		}
 
diff --git a/Engine/Plugins/Runtime/StateTree/Source/StateTreeModule/Public/StateTreePropertyRef.h b/Engine/Plugins/Runtime/StateTree/Source/StateTreeModule/Public/StateTreePropertyRef.h
--- a/Engine/Plugins/Runtime/StateTree/Source/StateTreeModule/Public/StateTreePropertyRef.h
+++ b/Engine/Plugins/Runtime/StateTree/Source/StateTreeModule/Public/StateTreePropertyRef.h
@@ -11,7 +11,7 @@ struct FStateTreePropertyRef;
 
 namespace UE::StateTree::PropertyRefHelpers
 {
-    /**
+	/**
 	 * @param PropertyRef Property's reference to get pointer to.
 	 * @param InstanceDataStorage Instance Data Storage.
 	 * @param ExecutionFrame Execution frame owning referenced property.
@@ -19,7 +19,7 @@ namespace UE::StateTree::PropertyRefHelpers
 	 * @param OutSourceProperty On success, returns referenced property.
 	 * @return Pointer to referenced property value if succeeded.
 	 */
-	template<class T>
+	template <class T>
 	static T* GetMutablePtrToProperty(const FStateTreePropertyRef& PropertyRef, FStateTreeInstanceStorage& InstanceDataStorage, const FStateTreeExecutionFrame& ExecutionFrame, const FStateTreeExecutionFrame* ParentExecutionFrame, const FProperty** OutSourceProperty = nullptr)
 	{
 		const FStateTreePropertyBindings& PropertyBindings = ExecutionFrame.StateTree->GetPropertyBindings();
@@ -69,7 +69,28 @@ namespace UE::StateTree::PropertyRefHelpers
 
 		return nullptr;
 	}
-}
+
+	/**
+	 * @param PropertyRef Property's reference to get pointer to.
+	 * @param InstanceDataStorage Instance Data Storage
+	 * @param ExecutionFrame Execution frame owning referenced property
+	 * @param ParentExecutionFrame Parent of execution frame owning referenced property
+	 * @return A tuple of pointer to referenced property if succeeded.
+	 */
+	template <class... T>
+	static TTuple<T*...> GetMutablePtrTupleToProperty(const FStateTreePropertyRef& PropertyRef, FStateTreeInstanceStorage& InstanceDataStorage, const FStateTreeExecutionFrame& ExecutionFrame, const FStateTreeExecutionFrame* ParentExecutionFrame)
+	{
+		const FStateTreePropertyBindings& PropertyBindings = ExecutionFrame.StateTree->GetPropertyBindings();
+		if (const FStateTreePropertyAccess* PropertyAccess = PropertyBindings.GetPropertyAccess(PropertyRef))
+		{
+			// Passing empty ContextAndExternalDataViews, as PropertyRef is not allowed to point to context or external data.
+			FStateTreeDataView SourceView = FStateTreeExecutionContext::GetDataView(InstanceDataStorage, nullptr, ParentExecutionFrame, ExecutionFrame, {}, PropertyAccess->SourceDataHandle);
+			return TTuple<T*...>(PropertyBindings.GetMutablePropertyPtr<T>(SourceView, *PropertyAccess)...);
+		}
+
+		return TTuple<T*...>{};
+	}
+} // namespace UE::StateTree::PropertyRefHelpers
 
 /**
  * Property ref allows to get a pointer to selected property in StateTree.
@@ -77,10 +98,13 @@ namespace UE::StateTree::PropertyRefHelpers
  *
  * Meta specifiers for the type:
  *  - RefType = "<type>"
- *		- Specifies the type of property to reference.
+ *		- Specifies a comma separated list of type of property to reference.
  *		- Supported types are: bool, byte, int32, int64, float, double, Name, String, Text, UObject pointers, and structs.
+ *		- If multiple types are specified, GetMutablePtrTuple can be used to access the correct type.
  *  - IsRefToArray
  *		- If specified, the reference is to an TArray<RefType>
+ *	- CanRefToArray
+ *		- If specified, the reference can bind to a Reftype or TArray<RefType>
  *  - Optional
  *		- If specified, the reference can be left unbound, otherwise the compiler report error if the reference is not bound.
  *
@@ -89,7 +113,7 @@ namespace UE::StateTree::PropertyRefHelpers
  *  // Reference to float
  *	UPROPERTY(EditAnywhere, meta = (RefType = "float"))
  *	FStateTreePropertyRef RefToFloat;
- * 
+ *
  *  // Reference to FTestStructBase
  *	UPROPERTY(EditAnywhere, meta = (RefType = "/Script/ModuleName.TestStructBase"))
  *	FStateTreePropertyRef RefToTest;
@@ -97,6 +121,10 @@ namespace UE::StateTree::PropertyRefHelpers
  *  // Reference to TArray<FTestStructBase>
  *	UPROPERTY(EditAnywhere, meta = (RefType = "/Script/ModuleName.TestStructBase", IsRefToArray))
  *	FStateTreePropertyRef RefToArrayOfTests;
+ *
+ *  // Reference to Vector, TArray<FVector>, AActor*, TArray<AActor*>
+ *	UPROPERTY(EditAnywhere, meta = (RefType = "Vector, AActor", CanRefToArray))
+ *	FStateTreePropertyRef RefToLocationLikeTypes;
  */
 USTRUCT()
 struct STATETREEMODULE_API FStateTreePropertyRef
@@ -106,7 +134,7 @@ struct STATETREEMODULE_API FStateTreePropertyRef
 	FStateTreePropertyRef() = default;
 
 	/** @return pointer to the property if possible, nullptr otherwise. */
-	template<class T>
+	template <class T>
 	T* GetMutablePtr(FStateTreeExecutionContext& Context) const
 	{
 		const FStateTreeExecutionFrame* CurrentlyProcessedFrame = Context.GetCurrentlyProcessedFrame();
@@ -115,6 +143,16 @@ struct STATETREEMODULE_API FStateTreePropertyRef
 		return UE::StateTree::PropertyRefHelpers::GetMutablePtrToProperty<T>(*this, Context.GetMutableInstanceData()->GetMutableStorage(), *CurrentlyProcessedFrame, Context.GetCurrentlyProcessedParentFrame());
 	}
 
+	/** @return a tuple of pointers of the given types to the property if possible, nullptr otherwise. */
+	template <class... T>
+	TTuple<T*...> GetMutablePtrTuple(FStateTreeExecutionContext& Context) const
+	{
+		const FStateTreeExecutionFrame* CurrentlyProcessedFrame = Context.GetCurrentlyProcessedFrame();
+		check(CurrentlyProcessedFrame);
+
+		return UE::StateTree::PropertyRefHelpers::GetMutablePtrTupleToProperty<T...>(*this, Context.GetMutableInstanceData()->GetMutableStorage(), *CurrentlyProcessedFrame, Context.GetCurrentlyProcessedParentFrame());
+	}
+
 	/**
 	 * Used internally.
 	 * @return index to referenced property access
@@ -132,15 +170,15 @@ private:
 };
 
 /**
- * TStateTreePropertyRef is a type-safe FStateTreePropertyRef wrapper against the given type.
+ * TStateTreePropertyRef is a type-safe FStateTreePropertyRef wrapper against a single given type.
  * @note When used as a property, this automatically defines PropertyRef property meta-data.
- * 
+ *
  * Example:
  *
  *  // Reference to float
  *	UPROPERTY(EditAnywhere)
  *	TStateTreePropertyRef<float> RefToFloat;
- * 
+ *
  *  // Reference to FTestStructBase
  *	UPROPERTY(EditAnywhere)
  *	TStateTreePropertyRef<FTestStructBase> RefToTest;
@@ -148,8 +186,12 @@ private:
  *  // Reference to TArray<FTestStructBase>
  *	UPROPERTY(EditAnywhere)
  *	TStateTreePropertyRef<TArray<FTestStructBase>> RefToArrayOfTests;
+ *
+ *  // Reference to FTestStructBase or TArray<FTestStructBase>
+ *	UPROPERTY(EditAnywhere, meta = (CanRefToArray))
+ *	TStateTreePropertyRef<FTestStructBase> RefToSingleOrArrayOfTests;
  */
-template<class TRef>
+template <class TRef>
 struct TStateTreePropertyRef
 {
 	/** @return pointer to the property if possible, nullptr otherwise. */
@@ -158,6 +200,12 @@ struct TStateTreePropertyRef
 		return PropertyRef.GetMutablePtr<TRef>(Context);
 	}
 
+	/** @return a tuple of pointer to the property of the type or array of type, nullptr otherwise. */
+	TTuple<TRef*, TArray<TRef>*> GetMutablePtrTuple(FStateTreeExecutionContext& Context) const
+	{
+		return PropertyRef.GetMutablePtrTuple<TRef, TArray<TRef>>(Context);
+	}
+
 	/**
 	 * Used internally.
 	 * @return internal property ref
@@ -174,49 +222,81 @@ private:
 /**
  * External Handle allows to wrap-up property reference to make it accessible without having an access to StateTreeExecutionContext. Useful for capturing property reference in callbacks.
  */
-template<class TRef>
-struct TStateTreePropertyRefExternalHandle
+struct FStateTreePropertyRefExternalHandle
 {
-	TStateTreePropertyRefExternalHandle(FStateTreePropertyRef InPropertyRef, FStateTreeExecutionContext& InContext)
+	FStateTreePropertyRefExternalHandle(FStateTreePropertyRef InPropertyRef, FStateTreeExecutionContext& InContext)
 		: WeakInstanceStorage(InContext.GetMutableInstanceData()->GetWeakMutableStorage())
 		, WeakStateTree(InContext.GetCurrentlyProcessedFrame()->StateTree)
 		, RootState(InContext.GetCurrentlyProcessedFrame()->RootState)
 		, PropertyRef(InPropertyRef)
-	{}
-
-	TStateTreePropertyRefExternalHandle(TStateTreePropertyRef<TRef> InPropertyRef, FStateTreeExecutionContext& InContext)
-		: TStateTreePropertyRefExternalHandle(InPropertyRef.GetInternalPropertyRef(), InContext)
-	{}
+	{
+	}
 
 	/** @return pointer to the property if possible, nullptr otherwise. */
+	template <class TRef>
 	TRef* GetMutablePtr() const
+	{
+		return GetMutablePtrTuple<TRef>().Key;
+	}
+
+	/** @return a tuple of pointers of the given types to the property if possible, nullptr otherwise. */
+	template <class... TRef>
+	TTuple<TRef*...> GetMutablePtrTuple() const
 	{
 		if (!WeakInstanceStorage.IsValid())
 		{
-			return nullptr;
+			return TTuple<TRef*...>{};
 		}
 
 		FStateTreeInstanceStorage& InstanceStorage = *WeakInstanceStorage.Pin();
 		TConstArrayView<FStateTreeExecutionFrame> ActiveFrames = InstanceStorage.GetExecutionState().ActiveFrames;
-
 		const FStateTreeExecutionFrame* ParentFrame = nullptr;
 		const FStateTreeExecutionFrame* Frame = FStateTreeExecutionContext::FindFrame(WeakStateTree.Get(), RootState, ActiveFrames, ParentFrame);
 
 		if (Frame == nullptr)
 		{
-			return nullptr;
+			return TTuple<TRef*...>{};
 		}
 
-		return UE::StateTree::PropertyRefHelpers::GetMutablePtrToProperty<TRef>(PropertyRef, InstanceStorage, *Frame, ParentFrame);
+		return UE::StateTree::PropertyRefHelpers::GetMutablePtrTupleToProperty<TRef...>(PropertyRef, InstanceStorage, Frame, ParentFrame);
 	}
 
-private:
+protected:
 	TWeakPtr<FStateTreeInstanceStorage> WeakInstanceStorage;
 	TWeakObjectPtr<const UStateTree> WeakStateTree = nullptr;
 	FStateTreeStateHandle RootState = FStateTreeStateHandle::Invalid;
 	FStateTreePropertyRef PropertyRef;
 };
 
+/**
+ * Single type safe external handle allows to wrap-up property reference to make it accessible without having an access to StateTreeExecutionContext. Useful for capturing property reference in callbacks.
+ */
+template <class TRef>
+struct TStateTreePropertyRefExternalHandle : public FStateTreePropertyRefExternalHandle
+{
+	using FStateTreePropertyRefExternalHandle::FStateTreePropertyRefExternalHandle;
+	TStateTreePropertyRefExternalHandle(TStateTreePropertyRef<TRef> InPropertyRef, FStateTreeExecutionContext& InContext)
+		: FStateTreePropertyRefExternalHandle(InPropertyRef.GetInternalPropertyRef(), InContext)
+	{
+	}
+
+	/** @return pointer to the property if possible, nullptr otherwise. */
+	TRef* GetMutablePtr() const
+	{
+		return FStateTreePropertyRefExternalHandle::GetMutablePtr<TRef>();
+	}
+
+	/** @return a tuple of pointer to the property of the type or array of type, nullptr otherwise. */
+	TTuple<TRef*, TArray<TRef>*> GetMutablePtrTuple() const
+	{
+		return FStateTreePropertyRefExternalHandle::GetMutablePtrTuple<TRef, TArray<TRef>>();
+	}
+
+private:
+	using FStateTreePropertyRefExternalHandle::GetMutablePtr;
+	using FStateTreePropertyRefExternalHandle::GetMutablePtrTuple;
+};
+
 UENUM()
 enum class EStateTreePropertyRefType : uint8
 {
diff --git a/Engine/Plugins/Runtime/StateTree/Source/StateTreeModule/Public/StateTreePropertyRefHelpers.h b/Engine/Plugins/Runtime/StateTree/Source/StateTreeModule/Public/StateTreePropertyRefHelpers.h
--- a/Engine/Plugins/Runtime/StateTree/Source/StateTreeModule/Public/StateTreePropertyRefHelpers.h
+++ b/Engine/Plugins/Runtime/StateTree/Source/StateTreeModule/Public/StateTreePropertyRefHelpers.h
@@ -21,6 +21,9 @@ enum class EStateTreePropertyRefType : uint8;
 namespace UE::StateTree::PropertyRefHelpers
 {
 #if WITH_EDITOR
+	STATETREEMODULE_API extern const FName IsRefToArrayName;
+	STATETREEMODULE_API extern const FName CanRefToArrayName;
+	STATETREEMODULE_API extern const FName RefTypeName;
 	/**
 	 * @param RefProperty Property of PropertyRef type.
 	 * @param SourceProperty Property to check it's type compatibility.
