Fux loctext macro warning
#rnx

diff --git a/Engine/Plugins/Runtime/StateTree/Source/StateTreeEditorModule/Private/Customizations/StateTreeBindingExtension.cpp b/Engine/Plugins/Runtime/StateTree/Source/StateTreeEditorModule/Private/Customizations/StateTreeBindingExtension.cpp
--- a/Engine/Plugins/Runtime/StateTree/Source/StateTreeEditorModule/Private/Customizations/StateTreeBindingExtension.cpp
+++ b/Engine/Plugins/Runtime/StateTree/Source/StateTreeEditorModule/Private/Customizations/StateTreeBindingExtension.cpp
@@ -538,14 +538,14 @@ struct FCachedBindingData : public TSharedFromThis<FCachedBindingData>
 		{
 			if (Property->HasMetaData(UE::StateTree::PropertyRefHelpers::IsRefToArrayName))
 			{
-				TooltipBuilder.AppendLineFormat(LOCTEXT("PropertyRefBindingTooltip", "Supported types are Array of {0}"), FText::FromString(Property->GetMetaData(UE::StateTree::PropertyRefHelpers::RefTypeName)));
+				TooltipBuilder.AppendLineFormat(LOCTEXT("PropertyRefBindingTooltipArray", "Supported types are Array of {0}"), FText::FromString(Property->GetMetaData(UE::StateTree::PropertyRefHelpers::RefTypeName)));
 			}
 			else
 			{
 				TooltipBuilder.AppendLineFormat(LOCTEXT("PropertyRefBindingTooltip", "Supported types are {0}"), FText::FromString(Property->GetMetaData(UE::StateTree::PropertyRefHelpers::RefTypeName)));
 				if (Property->HasMetaData(UE::StateTree::PropertyRefHelpers::CanRefToArrayName))
 				{
-					TooltipBuilder.AppendLine(LOCTEXT("PropertyRefBindingTooltip", "Supports Arrays"));
+					TooltipBuilder.AppendLine(LOCTEXT("PropertyRefBindingTooltipCanSupportArray", "Supports Arrays"));
 				}
 			}
 		}
