StateTree: Added support for event payload in transitions and event capture for states. (contributed)
- Transitions can now be set up to require a specific payload type and/or a gameplay tag
- Required transition event payload is visible to property bindind for transition conditions
- States can require an event type as enter condition, the event is captured when state is selected
- Required state event payload is visible to binding for state enter conditions, tasks and transistions

diff --git a/Engine/Plugins/Runtime/StateTree/Source/StateTreeEditorModule/Private/Customizations/StateTreeBindingExtension.cpp b/Engine/Plugins/Runtime/StateTree/Source/StateTreeEditorModule/Private/Customizations/StateTreeBindingExtension.cpp
--- a/Engine/Plugins/Runtime/StateTree/Source/StateTreeEditorModule/Private/Customizations/StateTreeBindingExtension.cpp
+++ b/Engine/Plugins/Runtime/StateTree/Source/StateTreeEditorModule/Private/Customizations/StateTreeBindingExtension.cpp
@@ -226,10 +226,14 @@ FText GetSectionNameFromDataSource(const EStateTreeBindableStructSource Source)
 		return LOCTEXT("Evaluators", "Evaluators");
 	case EStateTreeBindableStructSource::GlobalTask:
 		return LOCTEXT("StateGlobalTasks", "Global Tasks");
-	case EStateTreeBindableStructSource::State:
-		return LOCTEXT("StateParameters", "State");
+	case EStateTreeBindableStructSource::StateParameter:
+		return LOCTEXT("StateParameters", "State Parameters");
 	case EStateTreeBindableStructSource::Task:
 		return LOCTEXT("Tasks", "Tasks");
+	case EStateTreeBindableStructSource::TransitionEvent:
+		return LOCTEXT("TransitionEvents", "Transition Events");
+	case EStateTreeBindableStructSource::StateEvent:
+		return LOCTEXT("StateSelectionEvents", "State Selection Events");
 	default:
 		return FText::GetEmpty();
 	}
@@ -565,7 +569,10 @@ struct FCachedBindingData : public TSharedFromThis<FCachedBindingData>
 		{
 			const bool bIsStateTreeNode = AccessibleStructs.ContainsByPredicate([InStruct](const FStateTreeBindableStructDesc& AccessibleStruct)
 			{
-				return (AccessibleStruct.DataSource != EStateTreeBindableStructSource::Context && AccessibleStruct.DataSource != EStateTreeBindableStructSource::Parameter)
+				return AccessibleStruct.DataSource != EStateTreeBindableStructSource::Context
+					&& AccessibleStruct.DataSource != EStateTreeBindableStructSource::Parameter
+					&& AccessibleStruct.DataSource != EStateTreeBindableStructSource::TransitionEvent
+					&& AccessibleStruct.DataSource != EStateTreeBindableStructSource::StateEvent
 					&& AccessibleStruct.Struct == InStruct;
 			});
 
@@ -634,8 +641,7 @@ struct FCachedBindingData : public TSharedFromThis<FCachedBindingData>
 			}
 		}
 
-		return SourceProperty->HasAnyPropertyFlags(CPF_Edit) 
-			&& (!SourceProperty->HasAnyPropertyFlags(CPF_NativeAccessSpecifierPrivate | CPF_NativeAccessSpecifierProtected) || SourceProperty->GetBoolMetaData(FBlueprintMetadata::MD_AllowPrivateAccess));
+		return SourceProperty && IsPropertyBindable(*SourceProperty);
 	}
 
 	static bool ArePropertyAndContextStructCompatible(const UStruct* SourceStruct, const FProperty* TargetProperty)
@@ -821,6 +827,12 @@ private:
 	bool bIsDataCached = false;
 };
 
+bool IsPropertyBindable(const FProperty& Property)
+{
+	return Property.HasAnyPropertyFlags(CPF_Edit) 
+			&& (!Property.HasAnyPropertyFlags(CPF_NativeAccessSpecifierPrivate | CPF_NativeAccessSpecifierProtected) || Property.GetBoolMetaData(FBlueprintMetadata::MD_AllowPrivateAccess));
+}
+
 } // UE::StateTree::PropertyBinding
 
 bool FStateTreeBindingExtension::IsPropertyExtendable(const UClass* InObjectClass, const IPropertyHandle& PropertyHandle) const
diff --git a/Engine/Plugins/Runtime/StateTree/Source/StateTreeEditorModule/Private/Customizations/StateTreeBindingExtension.h b/Engine/Plugins/Runtime/StateTree/Source/StateTreeEditorModule/Private/Customizations/StateTreeBindingExtension.h
--- a/Engine/Plugins/Runtime/StateTree/Source/StateTreeEditorModule/Private/Customizations/StateTreeBindingExtension.h
+++ b/Engine/Plugins/Runtime/StateTree/Source/StateTreeEditorModule/Private/Customizations/StateTreeBindingExtension.h
@@ -11,6 +11,7 @@ class IDetailLayoutBuilder;
 class IPropertyAccessEditor;
 struct FStateTreeEditorPropertyPath;
 struct FStateTreePropertyPath;
+class FProperty;
 
 namespace UE::StateTree::PropertyBinding
 {
@@ -34,6 +35,11 @@ namespace UE::StateTree::PropertyBinding
 
 	DECLARE_MULTICAST_DELEGATE_TwoParams(FOnStateTreePropertyBindingChanged, const FStateTreePropertyPath& /*SourcePath*/, const FStateTreePropertyPath& /*TargetPath*/);
 	extern STATETREEEDITORMODULE_API FOnStateTreePropertyBindingChanged OnStateTreePropertyBindingChanged;
+
+	/**
+	 * Returns true if provided Property is bindable.
+	 */
+	bool IsPropertyBindable(const FProperty& Property);
 } // UE::StateTree::PropertyBinding
 
 class FStateTreeBindingExtension : public IDetailPropertyExtensionHandler
diff --git a/Engine/Plugins/Runtime/StateTree/Source/StateTreeEditorModule/Private/Customizations/StateTreeEventDescDetails.cpp b/Engine/Plugins/Runtime/StateTree/Source/StateTreeEditorModule/Private/Customizations/StateTreeEventDescDetails.cpp
new file mode 100644
--- /dev/null
+++ b/Engine/Plugins/Runtime/StateTree/Source/StateTreeEditorModule/Private/Customizations/StateTreeEventDescDetails.cpp
@@ -0,0 +1,234 @@
+// Copyright Epic Games, Inc. All Rights Reserved.
+
+#include "StateTreeEventDescDetails.h"
+#include "DetailWidgetRow.h"
+#include "StateTreeState.h"
+#include "Widgets/Text/STextBlock.h"
+#include "DetailLayoutBuilder.h"
+#include "IDetailChildrenBuilder.h"
+#include "IPropertyUtilities.h"
+#include "StateTreeBindingExtension.h"
+#include "StateTreeEditorStyle.h"
+#include "Styling/StyleColors.h"
+#include "Kismet2/BlueprintEditorUtils.h"
+#include "Widgets/Images/SImage.h"
+
+#define LOCTEXT_NAMESPACE "StateTreeEditor"
+
+/** 
+ * Node builder for StateTreeEvent Payload.
+ * Draws recursively all children of the selected Payload type without possibility of modifying them.
+ */
+class FStateTreeEventPayloadDetails : public IDetailCustomNodeBuilder
+{
+public:
+	using FVisitedStructs = TArray<UStruct*, TInlineAllocator<16>>;
+
+	FStateTreeEventPayloadDetails(const FProperty& Property, const FVisitedStructs& InVisitedStructs)
+		: Name(Property.GetName())
+		, VisitedStructs(InVisitedStructs)
+	{
+		NameWidget = GenerateNameWidget(Property);
+		ValueWidget = GenerateValueWidget(Property);
+
+		if (const FStructProperty* StructProperty = CastField<FStructProperty>(&Property))
+		{
+			Struct = StructProperty->Struct;
+		}
+		else if (const FObjectProperty* ObjectProperty = CastField<FObjectProperty>(&Property))
+		{
+			Struct = ObjectProperty->PropertyClass;
+		}
+
+		if (Struct)
+		{
+			if (VisitedStructs.Contains(Struct))
+			{
+				Struct = nullptr;
+			}
+			else
+			{
+				VisitedStructs.Add(Struct);
+			}
+		}
+	}
+
+	FStateTreeEventPayloadDetails(FName InName, TSharedRef<SWidget> InNameWidget, TSharedRef<SWidget> InValueWidget, UStruct* InStruct, const FVisitedStructs& InVisitedStructs)
+		: Name(InName)
+		, NameWidget(InNameWidget)
+		, ValueWidget(InValueWidget)	
+		, VisitedStructs(InVisitedStructs)
+		, Struct(InStruct)
+	{
+		if (VisitedStructs.Contains(Struct))
+		{
+			Struct = nullptr;
+		}
+		else
+		{
+			VisitedStructs.Add(Struct);
+		}
+	}
+
+	virtual bool InitiallyCollapsed() const override
+	{
+		return true;
+	}
+
+	virtual void GenerateHeaderRowContent( FDetailWidgetRow& NodeRow ) override 
+	{
+		NodeRow.ShouldAutoExpand(false);
+		NodeRow.NameContent()
+		[
+			NameWidget.ToSharedRef()
+		];
+
+		NodeRow.ValueContent()
+		[
+			ValueWidget.ToSharedRef()
+		];
+	}
+
+	virtual void GenerateChildContent(IDetailChildrenBuilder& ChildrenBuilder) override
+	{
+		for (TFieldIterator<const FProperty> PropIt(Struct); PropIt; ++PropIt)
+		{
+			const FProperty& Property = **PropIt;
+			if (UE::StateTree::PropertyBinding::IsPropertyBindable(Property))
+			{
+				ChildrenBuilder.AddCustomBuilder(MakeShared<FStateTreeEventPayloadDetails>(Property, VisitedStructs));
+			}
+		}
+	}
+
+	FName GetName() const override
+	{
+		return Name;
+	}
+
+private:
+	TSharedRef<SWidget> GenerateNameWidget(const FProperty& Property) const
+	{
+		return SNew(SHorizontalBox)
+					+SHorizontalBox::Slot()
+					.AutoWidth()
+					.VAlign(VAlign_Center)
+					[
+						SNew(STextBlock)
+						.Text(Property.GetDisplayNameText())
+						.Font(IDetailLayoutBuilder::GetDetailFont())
+						.ToolTipText(Property.GetToolTipText())
+					]
+
+					+SHorizontalBox::Slot()
+					.AutoWidth()
+					.VAlign(VAlign_Center)
+					.Padding(4.0f, 0.0f)
+					[
+						SNew(SBorder)
+						.Padding(FMargin(6, 1))
+						.BorderImage(FStateTreeEditorStyle::Get().GetBrush("StateTree.Param.Background"))
+						[
+							SNew(STextBlock)
+							.TextStyle(FStateTreeEditorStyle::Get(), "StateTree.Param.Label")
+							.ColorAndOpacity(FStyleColors::Foreground)
+							.Text(LOCTEXT("LabelOutput", "OUT"))
+						]
+					];
+	}
+
+	TSharedRef<SWidget> GenerateValueWidget(const FProperty& Property) const
+	{
+		const UEdGraphSchema_K2* Schema = GetDefault<UEdGraphSchema_K2>();
+		FEdGraphPinType PinType;
+		Schema->ConvertPropertyToPinType(&Property, PinType);
+
+		const FSlateBrush* Icon = FBlueprintEditorUtils::GetIconFromPin(PinType, true);
+		const FText Text = GetPinTypeText(PinType);
+		const FLinearColor IconColor = Schema->GetPinTypeColor(PinType);
+
+		return SNew(SHorizontalBox)
+					+SHorizontalBox::Slot()
+					.AutoWidth()
+					.VAlign(VAlign_Center)
+					.Padding(4.0f, 0.0f)
+					[
+						SNew(SImage)
+							.Image(Icon)
+							.ColorAndOpacity(IconColor)
+					]
+
+					+ SHorizontalBox::Slot()
+					.AutoWidth()
+					.VAlign(VAlign_Center)
+					[
+						SNew(STextBlock)
+							.Font(IDetailLayoutBuilder::GetDetailFont())
+							.Text(Text)
+					];
+	}
+
+	static FText GetPinTypeText(const FEdGraphPinType& PinType)
+	{
+		const FName PinSubCategory = PinType.PinSubCategory;
+		const UObject* PinSubCategoryObject = PinType.PinSubCategoryObject.Get();
+		if (PinSubCategory != UEdGraphSchema_K2::PSC_Bitmask && PinSubCategoryObject)
+		{
+			if (const UField* Field = Cast<const UField>(PinSubCategoryObject))
+			{
+				return Field->GetDisplayNameText();
+			}
+			return FText::FromString(PinSubCategoryObject->GetName());
+		}
+
+		return UEdGraphSchema_K2::GetCategoryText(PinType.PinCategory, NAME_None, true);
+	}
+
+	FName Name;
+	TSharedPtr<SWidget> NameWidget;
+	TSharedPtr<SWidget> ValueWidget;
+	FVisitedStructs VisitedStructs;
+	UStruct* Struct = nullptr;
+};
+
+TSharedRef<IPropertyTypeCustomization> FStateTreeEventDescDetails::MakeInstance()
+{
+	return MakeShared<FStateTreeEventDescDetails>();
+}
+
+void FStateTreeEventDescDetails::CustomizeHeader(TSharedRef<IPropertyHandle> PropertyHandle, FDetailWidgetRow& HeaderRow, IPropertyTypeCustomizationUtils& CustomizationUtils)
+{
+	HeaderRow
+	.NameContent()
+	[
+		PropertyHandle->CreatePropertyNameWidget()
+	]
+	.ValueContent()
+	[
+		PropertyHandle->CreatePropertyValueWidget()
+	];
+}
+
+void FStateTreeEventDescDetails::CustomizeChildren(TSharedRef<IPropertyHandle> PropertyHandle, IDetailChildrenBuilder& ChildBuilder, IPropertyTypeCustomizationUtils& CustomizationUtils)
+{
+	TSharedPtr<IPropertyHandle> TagProperty = PropertyHandle->GetChildHandle(GET_MEMBER_NAME_CHECKED(FStateTreeEventDesc, Tag));
+	check(TagProperty);
+	ChildBuilder.AddProperty(TagProperty.ToSharedRef());
+
+	TSharedPtr<IPropertyHandle> PayloadProperty = PropertyHandle->GetChildHandle(GET_MEMBER_NAME_CHECKED(FStateTreeEventDesc, PayloadStruct));
+	check(PayloadProperty);
+
+	UObject* PayloadType = nullptr;
+	if (PayloadProperty->GetValue(PayloadType))
+	{
+		UStruct* PayloadStruct = Cast<UStruct>(PayloadType);
+		ChildBuilder.AddCustomBuilder(MakeShared<FStateTreeEventPayloadDetails>(TEXT("EventPayload"), PayloadProperty->CreatePropertyNameWidget(), PayloadProperty->CreatePropertyValueWidget(), PayloadStruct, FStateTreeEventPayloadDetails::FVisitedStructs()));
+	}
+
+	PayloadProperty->SetOnPropertyValueChanged(FSimpleDelegate::CreateLambda([PropUtils = CustomizationUtils.GetPropertyUtilities()]()
+	{
+		PropUtils->RequestForceRefresh(); 
+	}));
+}
+
+#undef LOCTEXT_NAMESPACE
diff --git a/Engine/Plugins/Runtime/StateTree/Source/StateTreeEditorModule/Private/Customizations/StateTreeEventDescDetails.h b/Engine/Plugins/Runtime/StateTree/Source/StateTreeEditorModule/Private/Customizations/StateTreeEventDescDetails.h
new file mode 100644
--- /dev/null
+++ b/Engine/Plugins/Runtime/StateTree/Source/StateTreeEditorModule/Private/Customizations/StateTreeEventDescDetails.h
@@ -0,0 +1,22 @@
+// Copyright Epic Games, Inc. All Rights Reserved.
+
+#pragma once
+
+#include "IPropertyTypeCustomization.h"
+
+class IPropertyHandle;
+class FDetailWidgetRow;
+class IPropertyTypeCustomizationUtils;
+class IDetailChildrenBuilder;
+
+/**
+ * Type customization for FStateTreeEventDesc.
+ */
+class FStateTreeEventDescDetails : public IPropertyTypeCustomization
+{
+public:
+	static TSharedRef<IPropertyTypeCustomization> MakeInstance();
+
+	virtual void CustomizeHeader(TSharedRef<IPropertyHandle> PropertyHandle, FDetailWidgetRow& HeaderRow, IPropertyTypeCustomizationUtils& CustomizationUtils) override;
+	virtual void CustomizeChildren(TSharedRef<IPropertyHandle> PropertyHandle, IDetailChildrenBuilder& ChildBuilder, IPropertyTypeCustomizationUtils& CustomizationUtils) override;
+};
\ No newline at end of file
diff --git a/Engine/Plugins/Runtime/StateTree/Source/StateTreeEditorModule/Private/Customizations/StateTreeTransitionDetails.cpp b/Engine/Plugins/Runtime/StateTree/Source/StateTreeEditorModule/Private/Customizations/StateTreeTransitionDetails.cpp
--- a/Engine/Plugins/Runtime/StateTree/Source/StateTreeEditorModule/Private/Customizations/StateTreeTransitionDetails.cpp
+++ b/Engine/Plugins/Runtime/StateTree/Source/StateTreeEditorModule/Private/Customizations/StateTreeTransitionDetails.cpp
@@ -45,7 +45,7 @@ void FStateTreeTransitionDetails::CustomizeHeader(TSharedRef<IPropertyHandle> St
 
 	TriggerProperty = StructProperty->GetChildHandle(GET_MEMBER_NAME_CHECKED(FStateTreeTransition, Trigger));
 	PriorityProperty = StructProperty->GetChildHandle(GET_MEMBER_NAME_CHECKED(FStateTreeTransition, Priority));
-	EventTagProperty = StructProperty->GetChildHandle(GET_MEMBER_NAME_CHECKED(FStateTreeTransition, EventTag));
+	RequiredEventProperty = StructProperty->GetChildHandle(GET_MEMBER_NAME_CHECKED(FStateTreeTransition, RequiredEvent));
 	StateProperty = StructProperty->GetChildHandle(GET_MEMBER_NAME_CHECKED(FStateTreeTransition, State));
 	DelayTransitionProperty = StructProperty->GetChildHandle(GET_MEMBER_NAME_CHECKED(FStateTreeTransition, bDelayTransition));
 	DelayDurationProperty = StructProperty->GetChildHandle(GET_MEMBER_NAME_CHECKED(FStateTreeTransition, DelayDuration));
@@ -87,7 +87,7 @@ void FStateTreeTransitionDetails::CustomizeHeader(TSharedRef<IPropertyHandle> St
 void FStateTreeTransitionDetails::CustomizeChildren(TSharedRef<class IPropertyHandle> StructPropertyHandle, class IDetailChildrenBuilder& StructBuilder, IPropertyTypeCustomizationUtils& StructCustomizationUtils)
 {
 	check(TriggerProperty);
-	check(EventTagProperty);
+	check(RequiredEventProperty);
 	check(DelayTransitionProperty);
 	check(DelayDurationProperty);
 	check(DelayRandomVarianceProperty);
@@ -114,7 +114,7 @@ void FStateTreeTransitionDetails::CustomizeChildren(TSharedRef<class IPropertyHa
 	StructBuilder.AddProperty(TriggerProperty.ToSharedRef());
 
 	// Show event only when the trigger is set to Event. 
-	StructBuilder.AddProperty(EventTagProperty.ToSharedRef())
+	StructBuilder.AddProperty(RequiredEventProperty.ToSharedRef())
 		.Visibility(TAttribute<EVisibility>::Create(TAttribute<EVisibility>::FGetter::CreateLambda([WeakSelf]()
 		{
 			if (const TSharedPtr<FStateTreeTransitionDetails> Self = WeakSelf.Pin())
@@ -199,13 +199,34 @@ FText FStateTreeTransitionDetails::GetDescription() const
 	}
 
 	EStateTreeTransitionTrigger Trigger = GetTrigger();
+
 	FText TriggerText = UEnum::GetDisplayValueAsText(Trigger);
 
 	if (Trigger == EStateTreeTransitionTrigger::OnEvent)
 	{
-		FGameplayTag EventTag;
-		UE::StateTree::PropertyHelpers::GetStructValue<FGameplayTag>(EventTagProperty, EventTag);
-		TriggerText = FText::Format(LOCTEXT("TransitionOnEvent", "On Event {0}"), FText::FromName(EventTag.GetTagName()));
+		FStateTreeEventDesc RequiredEvent;
+		UE::StateTree::PropertyHelpers::GetStructValue<FStateTreeEventDesc>(RequiredEventProperty, RequiredEvent);
+
+		TArray<FText> PayloadItems;
+		
+		if (RequiredEvent.IsValid())
+		{
+			if (RequiredEvent.Tag.IsValid())
+			{
+				PayloadItems.Add(FText::Format(LOCTEXT("TransitionEventTag", "Tag: '{0}'"), FText::FromName(RequiredEvent.Tag.GetTagName())));
+			}
+			
+			if (RequiredEvent.PayloadStruct)
+			{
+				PayloadItems.Add(FText::Format(LOCTEXT("TransitionEventPayload", "Payload: '{0}'"), RequiredEvent.PayloadStruct->GetDisplayNameText()));
+			}
+		}
+		else
+		{
+			PayloadItems.Add(LOCTEXT("TransitionInvalidEvent", "Invalid"));
+		}
+		
+		TriggerText = FText::Format(LOCTEXT("TransitionOnEvent", "On Event ({0})"), FText::Join(INVTEXT(", "), PayloadItems));
 	}
 	
 	FText TargetText;
diff --git a/Engine/Plugins/Runtime/StateTree/Source/StateTreeEditorModule/Private/Customizations/StateTreeTransitionDetails.h b/Engine/Plugins/Runtime/StateTree/Source/StateTreeEditorModule/Private/Customizations/StateTreeTransitionDetails.h
--- a/Engine/Plugins/Runtime/StateTree/Source/StateTreeEditorModule/Private/Customizations/StateTreeTransitionDetails.h
+++ b/Engine/Plugins/Runtime/StateTree/Source/StateTreeEditorModule/Private/Customizations/StateTreeTransitionDetails.h
@@ -39,7 +39,7 @@ private:
 
 	TSharedPtr<IPropertyHandle> TriggerProperty;
 	TSharedPtr<IPropertyHandle> PriorityProperty;
-	TSharedPtr<IPropertyHandle> EventTagProperty;
+	TSharedPtr<IPropertyHandle> RequiredEventProperty;
 	TSharedPtr<IPropertyHandle> StateProperty;
 	TSharedPtr<IPropertyHandle> DelayTransitionProperty;
 	TSharedPtr<IPropertyHandle> DelayDurationProperty;
diff --git a/Engine/Plugins/Runtime/StateTree/Source/StateTreeEditorModule/Private/StateTreeCompiler.cpp b/Engine/Plugins/Runtime/StateTree/Source/StateTreeEditorModule/Private/StateTreeCompiler.cpp
--- a/Engine/Plugins/Runtime/StateTree/Source/StateTreeEditorModule/Private/StateTreeCompiler.cpp
+++ b/Engine/Plugins/Runtime/StateTree/Source/StateTreeEditorModule/Private/StateTreeCompiler.cpp
@@ -565,6 +565,8 @@ bool FStateTreeCompiler::CreateStateTasksAndParameters()
 			const FCompactStateTreeState& ParentCompactState = StateTree->States[CompactState.Parent.Index];
 			const int32 InstanceDataBegin = FirstInstanceDataIndex[CompactState.Parent.Index] + (int32)ParentCompactState.InstanceDataNum;
 			FirstInstanceDataIndex[i] = InstanceDataBegin;
+
+			CompactState.Depth = ParentCompactState.Depth + 1;
 		}
 
 		int32 InstanceDataIndex = FirstInstanceDataIndex[i];
@@ -602,7 +604,7 @@ bool FStateTreeCompiler::CreateStateTasksAndParameters()
 			State->Name,
 			State->Parameters.Parameters.GetPropertyBagStruct(),
 			CompactState.ParameterDataHandle,
-			EStateTreeBindableStructSource::State,
+			EStateTreeBindableStructSource::StateParameter,
 			State->Parameters.ID
 		};
 
@@ -614,6 +616,36 @@ bool FStateTreeCompiler::CreateStateTasksAndParameters()
 		// Add as binding source.
 		BindingsCompiler.AddSourceStruct(LinkedParamsDesc);
 
+		if (State->bHasRequiredEventToEnter)
+		{
+			CompactState.EventDataIndex = FStateTreeIndex16(InstanceDataIndex++);
+			CompactState.RequiredEventToEnter.Tag = State->RequiredEventToEnter.Tag;
+			CompactState.RequiredEventToEnter.PayloadStruct = State->RequiredEventToEnter.PayloadStruct;
+
+			FStateTreeBindableStructDesc Desc;
+			Desc.Struct = FStateTreeEvent::StaticStruct();
+			Desc.Name = CompactState.Name;
+			Desc.ID = State->GetEventID();
+			Desc.DataSource = EStateTreeBindableStructSource::StateEvent;
+			Desc.DataHandle = FStateTreeDataHandle(EStateTreeDataSourceType::StateEvent, CompactState.EventDataIndex.Get(), CompactStateHandle);
+
+			BindingsCompiler.AddSourceStruct(Desc);
+
+			if (!CompactState.RequiredEventToEnter.IsValid())
+			{
+				Log.Reportf(EMessageSeverity::Error, Desc,
+					TEXT("Event is marked as required, but isn't set up."));
+				return false;
+			}
+		}
+
+		if (CompactState.Depth >= FStateTreeActiveStates::MaxStates)
+		{
+			Log.Reportf(EMessageSeverity::Error, LinkedParamsDesc,
+				TEXT("Exceeds the maximum depth of execution (%u)"), FStateTreeActiveStates::MaxStates);
+			return false;
+		}
+
 		int32 BatchIndex = INDEX_NONE;
 
 		// Subtrees parameters cannot have bindings
@@ -818,12 +850,13 @@ bool FStateTreeCompiler::CreateStateTransitions()
 		
 		for (FStateTreeTransition& Transition : SourceState->Transitions)
 		{
-			IDToTransition.Add(Transition.ID, StateTree->Transitions.Num());
+			const int32 TransitionIndex = StateTree->Transitions.Num();
+			IDToTransition.Add(Transition.ID, TransitionIndex);
 
 			FCompactStateTransition& CompactTransition = StateTree->Transitions.AddDefaulted_GetRef();
 			CompactTransition.Trigger = Transition.Trigger;
 			CompactTransition.Priority = Transition.Priority;
-			CompactTransition.EventTag = Transition.EventTag;
+
 			CompactTransition.bTransitionEnabled = Transition.bTransitionEnabled;
 
 			if (Transition.State.LinkType == EStateTreeTransitionType::NextSelectableState)
@@ -866,7 +899,42 @@ bool FStateTreeCompiler::CreateStateTransitions()
 			{
 				return false;
 			}
-			
+
+			if (Transition.Trigger == EStateTreeTransitionTrigger::OnEvent)
+			{
+				CompactTransition.RequiredEvent.Tag = Transition.RequiredEvent.Tag;
+				CompactTransition.RequiredEvent.PayloadStruct = Transition.RequiredEvent.PayloadStruct;
+
+				FStateTreeBindableStructDesc Desc;
+				Desc.Struct = FStateTreeEvent::StaticStruct();
+				Desc.Name = FName(TEXT("Transition"));
+				Desc.ID = Transition.GetEventID();
+				Desc.DataSource = EStateTreeBindableStructSource::TransitionEvent;
+				Desc.DataHandle = FStateTreeDataHandle(EStateTreeDataSourceType::TransitionEvent, TransitionIndex);
+
+				if (!Transition.RequiredEvent.IsValid())
+				{
+					Log.Reportf(EMessageSeverity::Error, Desc,
+						TEXT("On Event Transition requires at least tag or payload to be set up."),
+						*Transition.State.Name.ToString());
+					return false;
+				}
+
+				if (!CompactTransition.State.IsCompletionState())
+				{
+					FCompactStateTreeState& TransitionTargetState = StateTree->States[CompactTransition.State.Index];
+					if (TransitionTargetState.RequiredEventToEnter.IsValid() && !TransitionTargetState.RequiredEventToEnter.IsSubsetOfAnotherDesc(CompactTransition.RequiredEvent))
+					{
+						Log.Reportf(EMessageSeverity::Error, Desc,
+							TEXT("On Event transition to %s will never succeed as transition and state required events are incompatible."),
+							*TransitionTargetState.Name.ToString());
+						return false;
+					}
+				}
+
+				BindingsCompiler.AddSourceStruct(Desc);
+			}
+
 			const int32 ConditionsBegin = Nodes.Num();
 			if (const auto Validation = UE::StateTree::Compiler::IsValidCount16(ConditionsBegin); Validation.DidFail())
 			{
diff --git a/Engine/Plugins/Runtime/StateTree/Source/StateTreeEditorModule/Private/StateTreeEditorData.cpp b/Engine/Plugins/Runtime/StateTree/Source/StateTreeEditorModule/Private/StateTreeEditorData.cpp
--- a/Engine/Plugins/Runtime/StateTree/Source/StateTreeEditorModule/Private/StateTreeEditorData.cpp
+++ b/Engine/Plugins/Runtime/StateTree/Source/StateTreeEditorModule/Private/StateTreeEditorData.cpp
@@ -277,7 +277,7 @@ void UStateTreeEditorData::GetAccessibleStructs(const TConstArrayView<const USta
 
 	if (BaseProgress == EStateTreeVisitor::Continue)
 	{
-		TArray<FStateTreeBindableStructDesc> TaskDescs;
+		TArray<FStateTreeBindableStructDesc, TInlineAllocator<32>> BindableDescs;
 
 		for (const UStateTreeState* State : Path)
 		{
@@ -286,21 +286,61 @@ void UStateTreeEditorData::GetAccessibleStructs(const TConstArrayView<const USta
 				continue;
 			}
 			
-			const EStateTreeVisitor StateProgress = VisitStateNodes(*State, [&OutStructDescs, &TaskDescs, TargetStructID]
+			const EStateTreeVisitor StateProgress = VisitStateNodes(*State, [&OutStructDescs, &BindableDescs, &Path, TargetStructID]
 				(const UStateTreeState* State, const FStateTreeBindableStructDesc& Desc, const FStateTreeDataView Value)
 				{
 					// Stop iterating as soon as we find the target node.
 					if (Desc.ID == TargetStructID)
 					{
-						OutStructDescs.Append(TaskDescs);
+						OutStructDescs.Append(BindableDescs);
 						return EStateTreeVisitor::Break;
 					}
 
 					// Not at target yet, collect all bindable source accessible so far.
 					if (Desc.DataSource == EStateTreeBindableStructSource::Task
-						|| Desc.DataSource == EStateTreeBindableStructSource::State)
+						|| Desc.DataSource == EStateTreeBindableStructSource::StateParameter)
 					{
-						TaskDescs.Add(Desc);
+						BindableDescs.Add(Desc);
+					}
+
+					switch (Desc.DataSource)
+					{
+						case EStateTreeBindableStructSource::StateParameter:
+						case EStateTreeBindableStructSource::Task:
+						case EStateTreeBindableStructSource::StateEvent:
+							BindableDescs.Add(Desc);
+							break;
+
+						case EStateTreeBindableStructSource::TransitionEvent:
+						{
+							// Checking if BindableStruct's owning Transition contains the Target.
+							if (State == Path.Last())
+							{
+								for (const FStateTreeTransition& Transition : State->Transitions)
+								{
+									bool bFoundOwningTransition = false;
+									for (const FStateTreeEditorNode& ConditionNode : Transition.Conditions)
+									{
+										if (ConditionNode.ID == TargetStructID)
+										{
+											if(Transition.GetEventID() == Desc.ID)
+											{
+												BindableDescs.Add(Desc);
+											}
+
+											bFoundOwningTransition = true;
+											break;
+										}
+									}
+
+									if (bFoundOwningTransition)
+									{
+										break;
+									}
+								}
+							}
+							break;
+						}
 					}
 							
 					return EStateTreeVisitor::Continue;
@@ -817,7 +857,7 @@ EStateTreeVisitor UStateTreeEditorData::VisitStateNodes(const UStateTreeState& S
 			Desc.Struct = State.Parameters.Parameters.GetPropertyBagStruct();
 			Desc.Name = State.Name;
 			Desc.ID = State.Parameters.ID;
-			Desc.DataSource = EStateTreeBindableStructSource::State;
+			Desc.DataSource = EStateTreeBindableStructSource::StateParameter;
 
 			if (InFunc(&State, Desc, FStateTreeDataView(const_cast<FInstancedPropertyBag&>(State.Parameters.Parameters).GetMutableValue())) == EStateTreeVisitor::Break)
 			{
@@ -825,6 +865,23 @@ EStateTreeVisitor UStateTreeEditorData::VisitStateNodes(const UStateTreeState& S
 			}
 		}
 	}
+
+	if (bContinue)
+	{
+		if (State.bHasRequiredEventToEnter )
+		{
+			FStateTreeBindableStructDesc Desc;
+			Desc.Struct = FStateTreeEvent::StaticStruct();
+			Desc.Name = State.Name;
+			Desc.ID = State.GetEventID();
+			Desc.DataSource = EStateTreeBindableStructSource::StateEvent;
+
+			if (InFunc(&State, Desc, FStateTreeDataView(FStructView::Make(const_cast<UStateTreeState&>(State).RequiredEventToEnter.GetTemporaryEvent()))) == EStateTreeVisitor::Break)
+			{
+				bContinue = false;
+			}
+		}
+	}
 	
 	if (bContinue)
 	{
@@ -890,6 +947,21 @@ EStateTreeVisitor UStateTreeEditorData::VisitStateNodes(const UStateTreeState& S
 		// Transitions
 		for (const FStateTreeTransition& Transition : State.Transitions)
 		{
+			if (Transition.Trigger == EStateTreeTransitionTrigger::OnEvent)
+			{
+				FStateTreeBindableStructDesc Desc;
+				Desc.Struct = FStateTreeEvent::StaticStruct();
+				Desc.Name = FName(TEXT("Transition"));
+				Desc.ID = Transition.GetEventID();
+				Desc.DataSource = EStateTreeBindableStructSource::TransitionEvent;
+
+				if (InFunc(&State, Desc, FStateTreeDataView(FStructView::Make(const_cast<FStateTreeTransition&>(Transition).RequiredEvent.GetTemporaryEvent()))) == EStateTreeVisitor::Break)
+				{
+					bContinue = false;
+					break;
+				}
+			}
+
 			for (const FStateTreeEditorNode& Node : Transition.Conditions)
 			{
 				if (const FStateTreeConditionBase* Cond = Node.Node.GetPtr<FStateTreeConditionBase>())
diff --git a/Engine/Plugins/Runtime/StateTree/Source/StateTreeEditorModule/Private/StateTreeEditorModule.cpp b/Engine/Plugins/Runtime/StateTree/Source/StateTreeEditorModule/Private/StateTreeEditorModule.cpp
--- a/Engine/Plugins/Runtime/StateTree/Source/StateTreeEditorModule/Private/StateTreeEditorModule.cpp
+++ b/Engine/Plugins/Runtime/StateTree/Source/StateTreeEditorModule/Private/StateTreeEditorModule.cpp
@@ -13,6 +13,7 @@
 #include "Customizations/StateTreeStateLinkDetails.h"
 #include "Customizations/StateTreeStateParametersDetails.h"
 #include "Customizations/StateTreeTransitionDetails.h"
+#include "Customizations/StateTreeEventDescDetails.h"
 #include "PropertyEditorModule.h"
 #include "StateTree.h"
 #include "StateTreeCompiler.h"
@@ -85,6 +86,7 @@ void FStateTreeEditorModule::StartupModule()
 	// Register the details customizer
 	FPropertyEditorModule& PropertyModule = FModuleManager::LoadModuleChecked<FPropertyEditorModule>("PropertyEditor");
 	PropertyModule.RegisterCustomPropertyTypeLayout("StateTreeTransition", FOnGetPropertyTypeCustomizationInstance::CreateStatic(&FStateTreeTransitionDetails::MakeInstance));
+	PropertyModule.RegisterCustomPropertyTypeLayout("StateTreeEventDesc", FOnGetPropertyTypeCustomizationInstance::CreateStatic(&FStateTreeEventDescDetails::MakeInstance));
 	PropertyModule.RegisterCustomPropertyTypeLayout("StateTreeStateLink", FOnGetPropertyTypeCustomizationInstance::CreateStatic(&FStateTreeStateLinkDetails::MakeInstance));
 	PropertyModule.RegisterCustomPropertyTypeLayout("StateTreeEditorNode", FOnGetPropertyTypeCustomizationInstance::CreateStatic(&FStateTreeEditorNodeDetails::MakeInstance));
 	PropertyModule.RegisterCustomPropertyTypeLayout("StateTreeStateParameters", FOnGetPropertyTypeCustomizationInstance::CreateStatic(&FStateTreeStateParametersDetails::MakeInstance));
@@ -118,6 +120,7 @@ void FStateTreeEditorModule::ShutdownModule()
 	{
 		FPropertyEditorModule& PropertyModule = FModuleManager::LoadModuleChecked<FPropertyEditorModule>("PropertyEditor");
 		PropertyModule.UnregisterCustomPropertyTypeLayout("StateTreeTransition");
+		PropertyModule.UnregisterCustomPropertyTypeLayout("StateTreeEventDesc");
 		PropertyModule.UnregisterCustomPropertyTypeLayout("StateTreeStateLink");
 		PropertyModule.UnregisterCustomPropertyTypeLayout("StateTreeEditorNode");
 		PropertyModule.UnregisterCustomPropertyTypeLayout("StateTreeStateParameters");
diff --git a/Engine/Plugins/Runtime/StateTree/Source/StateTreeEditorModule/Private/StateTreeState.cpp b/Engine/Plugins/Runtime/StateTree/Source/StateTreeEditorModule/Private/StateTreeState.cpp
--- a/Engine/Plugins/Runtime/StateTree/Source/StateTreeEditorModule/Private/StateTreeState.cpp
+++ b/Engine/Plugins/Runtime/StateTree/Source/StateTreeEditorModule/Private/StateTreeState.cpp
@@ -43,11 +43,24 @@ FStateTreeTransition::FStateTreeTransition(const EStateTreeTransitionTrigger InT
 
 FStateTreeTransition::FStateTreeTransition(const EStateTreeTransitionTrigger InTrigger, const FGameplayTag InEventTag, const EStateTreeTransitionType InType, const UStateTreeState* InState)
 	: Trigger(InTrigger)
-	, EventTag(InEventTag)
+	, RequiredEvent{InEventTag}
 {
 	State = InState ? InState->GetLinkToState() : FStateTreeStateLink(InType);
 }
 
+void FStateTreeTransition::PostSerialize(const FArchive& Ar)
+{
+#if WITH_EDITORONLY_DATA
+PRAGMA_DISABLE_DEPRECATION_WARNINGS
+	if (EventTag_DEPRECATED.IsValid())
+	{
+		RequiredEvent.Tag = EventTag_DEPRECATED;
+		EventTag_DEPRECATED = FGameplayTag();
+	}
+PRAGMA_ENABLE_DEPRECATION_WARNINGS
+#endif // WITH_EDITORONLY_DATA
+}
+
 //////////////////////////////////////////////////////////////////////////
 // UStateTreeState
 
diff --git a/Engine/Plugins/Runtime/StateTree/Source/StateTreeEditorModule/Public/StateTreeState.h b/Engine/Plugins/Runtime/StateTree/Source/StateTreeEditorModule/Public/StateTreeState.h
--- a/Engine/Plugins/Runtime/StateTree/Source/StateTreeEditorModule/Public/StateTreeState.h
+++ b/Engine/Plugins/Runtime/StateTree/Source/StateTreeEditorModule/Public/StateTreeState.h
@@ -4,10 +4,57 @@
 
 #include "StateTreeEditorNode.h"
 #include "StateTreeEditorTypes.h"
+#include "StateTreeEvents.h"
 #include "StateTreeState.generated.h"
 
 class UStateTreeState;
 
+/**
+ * Editor representation of an event description.
+ */
+USTRUCT()
+struct FStateTreeEventDesc
+{
+	GENERATED_BODY()
+
+	FStateTreeEventDesc() = default;
+
+	FStateTreeEventDesc(FGameplayTag InTag)
+		: Tag(InTag)
+	{}
+
+	/** Event Tag. */
+	UPROPERTY(EditDefaultsOnly, Category = "Event")
+	FGameplayTag Tag;
+
+	/** Event Payload Struct. */
+	UPROPERTY(EditDefaultsOnly, Category = "Event")
+	TObjectPtr<const UScriptStruct> PayloadStruct;
+
+	bool IsValid() const
+	{
+		return Tag.IsValid() || PayloadStruct;
+	}
+
+	FStateTreeEvent& GetTemporaryEvent()
+	{
+		TemporaryEvent.Tag = Tag;
+		TemporaryEvent.Payload = FInstancedStruct(PayloadStruct);
+
+		return TemporaryEvent;
+	}
+
+	bool operator==(const FStateTreeEventDesc& Other) const
+	{
+		return Tag == Other.Tag && PayloadStruct == Other.PayloadStruct;
+	}
+
+private:
+	/** Temporary event used as a source value in bindings. */
+	UPROPERTY(Transient)
+	FStateTreeEvent TemporaryEvent;
+};
+
 /**
  * Editor representation of a transition in StateTree
  */
@@ -16,9 +63,16 @@ struct STATETREEEDITORMODULE_API FStateTreeTransition
 {
 	GENERATED_BODY()
 
+	// Macro needed to avoid deprecation errors with members being copied or created.
+PRAGMA_DISABLE_DEPRECATION_WARNINGS
 	FStateTreeTransition() = default;
 	FStateTreeTransition(const EStateTreeTransitionTrigger InTrigger, const EStateTreeTransitionType InType, const UStateTreeState* InState = nullptr);
 	FStateTreeTransition(const EStateTreeTransitionTrigger InTrigger, const FGameplayTag InEventTag, const EStateTreeTransitionType InType, const UStateTreeState* InState = nullptr);
+	FStateTreeTransition(const FStateTreeTransition&) = default;
+	FStateTreeTransition(FStateTreeTransition&&) = default;
+	FStateTreeTransition& operator=(const FStateTreeTransition&) = default;
+	FStateTreeTransition& operator=(FStateTreeTransition&&) = default;
+PRAGMA_ENABLE_DEPRECATION_WARNINGS
 
 	template<typename T, typename... TArgs>
 	TStateTreeEditorNode<T>& AddCondition(TArgs&&... InArgs)
@@ -34,13 +88,19 @@ struct STATETREEEDITORMODULE_API FStateTreeTransition
 		return static_cast<TStateTreeEditorNode<T>&>(CondNode);
 	}
 
+	FGuid GetEventID() const
+	{
+		return FGuid::Combine(ID, FGuid::NewDeterministicGuid(TEXT("Event")));
+	}
+
+	void PostSerialize(const FArchive& Ar);
+
 	/** When to try trigger the transition. */
 	UPROPERTY(EditDefaultsOnly, Category = "Transition")
 	EStateTreeTransitionTrigger Trigger = EStateTreeTransitionTrigger::OnStateCompleted;
 
-	/** Tag of the State Tree event that triggers the transition. */
-	UPROPERTY(EditDefaultsOnly, Category = "Transition")
-	FGameplayTag EventTag;
+	UPROPERTY(EditDefaultsOnly, Category = "Transition", DisplayName = "Required Event")
+	FStateTreeEventDesc RequiredEvent; 
 
 	/** Transition target state. */
 	UPROPERTY(EditDefaultsOnly, Category = "Transition", meta=(DisplayName="Transition To"))
@@ -76,6 +136,21 @@ struct STATETREEEDITORMODULE_API FStateTreeTransition
 	/** True if the Transition is Enabled (i.e. not explicitly disabled in the asset). */
 	UPROPERTY(EditDefaultsOnly, Category = "Debug")
 	bool bTransitionEnabled = true;
+
+#if WITH_EDITORONLY_DATA
+	UE_DEPRECATED(5.5, "Use RequiredEvent.Tag instead.")
+	UPROPERTY()
+	FGameplayTag EventTag_DEPRECATED;
+#endif // WITH_EDITORONLY_DATA
+};
+
+template<>
+struct TStructOpsTypeTraits<FStateTreeTransition> : public TStructOpsTypeTraitsBase2<FStateTreeTransition>
+{
+	enum 
+	{
+		WithPostSerialize = true,
+	};
 };
 
 
@@ -213,6 +288,10 @@ public:
 		return Transition;
 	}
 
+	FGuid GetEventID() const
+	{
+		return FGuid::Combine(ID, FGuid::NewDeterministicGuid(TEXT("Event")));
+	}
 
 	// ~StateTree Builder API
 
@@ -244,6 +323,12 @@ public:
 	UPROPERTY(EditDefaultsOnly, Category = "State")
 	FStateTreeStateParameters Parameters;
 
+	UPROPERTY(EditDefaultsOnly, Category = "State", meta=(InlineEditConditionToggle))
+	bool bHasRequiredEventToEnter = false;
+
+	UPROPERTY(EditDefaultsOnly, Category = "State", meta = (EditCondition = "bHasRequiredEventToEnter"))
+	FStateTreeEventDesc RequiredEventToEnter;
+
 	UPROPERTY(EditDefaultsOnly, Category = "State", meta = (IgnoreForMemberInitializationTest))
 	FGuid ID;
 
diff --git a/Engine/Plugins/Runtime/StateTree/Source/StateTreeModule/Private/Debugger/StateTreeTraceTypes.cpp b/Engine/Plugins/Runtime/StateTree/Source/StateTreeModule/Private/Debugger/StateTreeTraceTypes.cpp
--- a/Engine/Plugins/Runtime/StateTree/Source/StateTreeModule/Private/Debugger/StateTreeTraceTypes.cpp
+++ b/Engine/Plugins/Runtime/StateTree/Source/StateTreeModule/Private/Debugger/StateTreeTraceTypes.cpp
@@ -143,9 +143,14 @@ FString FStateTreeTraceTransitionEvent::GetValueString(const UStateTree& StateTr
 		{
 			ensureAlways(Transition->Priority == TransitionSource.Priority);
 			ensureAlways(Transition->State == TransitionSource.TargetState);
-			if (Transition->EventTag.IsValid())
+			if (Transition->RequiredEvent.Tag.IsValid())
 			{
-				StrBuilder.Appendf(TEXT("\n\t%s"), *Transition->EventTag.ToString()); 
+				StrBuilder.Appendf(TEXT("\n\t%s"), *Transition->RequiredEvent.Tag.ToString()); 
+			}
+
+			if (Transition->RequiredEvent.PayloadStruct)
+			{
+				StrBuilder.Appendf(TEXT("\n\t%s"), *Transition->RequiredEvent.PayloadStruct->GetName()); 
 			}
 		}
 		else
diff --git a/Engine/Plugins/Runtime/StateTree/Source/StateTreeModule/Private/StateTree.cpp b/Engine/Plugins/Runtime/StateTree/Source/StateTreeModule/Private/StateTree.cpp
--- a/Engine/Plugins/Runtime/StateTree/Source/StateTreeModule/Private/StateTree.cpp
+++ b/Engine/Plugins/Runtime/StateTree/Source/StateTreeModule/Private/StateTree.cpp
@@ -619,7 +619,6 @@ bool UStateTree::PatchBindings()
 		}
 	}
 
-
 	TMap<FStateTreeDataHandle, FStateTreeDataView> DataViews;
 	TMap<FStateTreeIndex16, FStateTreeDataView> BindingBatchDataView;
 
@@ -705,12 +704,19 @@ bool UStateTree::PatchBindings()
 		for (int32 Index = Batch.BindingsBegin; Index != Batch.BindingsEnd; Index++)
 		{
 			FStateTreePropertyPathBinding& Binding = PropertyPathBindings[Index];
-			FStateTreeDataView SourceView = GetDataSourceView(Binding.GetSourceDataHandle());
-			
-			if (!Binding.GetMutableSourcePath().UpdateSegmentsFromValue(SourceView, &ErrorMsg))
+
+			const EStateTreeDataSourceType Source = Binding.GetSourceDataHandle().GetSource();
+			const bool bIsSourceEvent = Source == EStateTreeDataSourceType::TransitionEvent || Source == EStateTreeDataSourceType::StateEvent;
+
+			if(!bIsSourceEvent)
 			{
-				UE_LOG(LogStateTree, Error, TEXT("%hs: Failed to update source instance structs for property binding '%s'. Reason: %s"), __FUNCTION__, *Binding.GetTargetPath().ToString(), *ErrorMsg);
-				return false;
+				FStateTreeDataView SourceView = GetDataSourceView(Binding.GetSourceDataHandle());
+
+				if (!Binding.GetMutableSourcePath().UpdateSegmentsFromValue(SourceView, &ErrorMsg))
+				{
+					UE_LOG(LogStateTree, Error, TEXT("%hs: Failed to update source instance structs for property binding '%s'. Reason: %s"), __FUNCTION__, *Binding.GetTargetPath().ToString(), *ErrorMsg);
+					return false;
+				}
 			}
 
 			if (!Binding.GetMutableTargetPath().UpdateSegmentsFromValue(TargetView, &ErrorMsg))
diff --git a/Engine/Plugins/Runtime/StateTree/Source/StateTreeModule/Private/StateTreeEvents.cpp b/Engine/Plugins/Runtime/StateTree/Source/StateTreeModule/Private/StateTreeEvents.cpp
--- a/Engine/Plugins/Runtime/StateTree/Source/StateTreeModule/Private/StateTreeEvents.cpp
+++ b/Engine/Plugins/Runtime/StateTree/Source/StateTreeModule/Private/StateTreeEvents.cpp
@@ -6,13 +6,27 @@
 
 #include UE_INLINE_GENERATED_CPP_BY_NAME(StateTreeEvents)
 
+//----------------------------------------------------------------//
+// FStateTreeSharedEvent
+//----------------------------------------------------------------//
+
+void FStateTreeSharedEvent::AddStructReferencedObjects(FReferenceCollector& Collector)
+{
+	Collector.AddPropertyReferencesWithStructARO(FStateTreeEvent::StaticStruct(), Event.Get());
+}
+
+
+//----------------------------------------------------------------//
+// FStateTreeEventQueue
+//----------------------------------------------------------------//
+
 void FStateTreeEventQueue::SendEvent(const UObject* Owner, const FGameplayTag& Tag, const FConstStructView Payload, const FName Origin)
 {
-	if (Events.Num() >= MaxActiveEvents)
+	if (SharedEvents.Num() >= MaxActiveEvents)
 	{
 		UE_VLOG_UELOG(Owner, LogStateTree, Error, TEXT("%s: Too many events send on '%s'. Dropping event %s"), ANSI_TO_TCHAR(__FUNCTION__), *GetNameSafe(Owner), *Tag.ToString());
 		return;
 	}
 
-	Events.Emplace(Tag, Payload, Origin);
+	SharedEvents.Emplace(Tag, Payload, Origin);
 }
diff --git a/Engine/Plugins/Runtime/StateTree/Source/StateTreeModule/Private/StateTreeExecutionContext.cpp b/Engine/Plugins/Runtime/StateTree/Source/StateTreeModule/Private/StateTreeExecutionContext.cpp
--- a/Engine/Plugins/Runtime/StateTree/Source/StateTreeModule/Private/StateTreeExecutionContext.cpp
+++ b/Engine/Plugins/Runtime/StateTree/Source/StateTreeModule/Private/StateTreeExecutionContext.cpp
@@ -246,16 +246,16 @@ EStateTreeRunStatus FStateTreeExecutionContext::Start(const FInstancedPropertyBa
 
 		static const FStateTreeStateHandle RootState = FStateTreeStateHandle(0);
 
-		TArray<FStateTreeExecutionFrame, TFixedAllocator<MaxExecutionFrames>> NextActiveFrames;
-		if (SelectState(InitFrame, RootState, NextActiveFrames))
+		FStateSelectionResult StateSelectionResult;
+		if (SelectState(InitFrame, RootState, StateSelectionResult))
 		{
-			check(!NextActiveFrames.IsEmpty());
-			if (NextActiveFrames.Last().ActiveStates.Last().IsCompletionState())
+			check(StateSelectionResult.ContainsFrames());
+			if (StateSelectionResult.GetSelectedFrames().Last().ActiveStates.Last().IsCompletionState())
 			{
 				// Transition to a terminal state (succeeded/failed).
 				STATETREE_LOG(Warning, TEXT("%hs: Tree %s at StateTree start on '%s' using StateTree '%s'."),
-					__FUNCTION__, NextActiveFrames.Last().ActiveStates.Last() == FStateTreeStateHandle::Succeeded ? TEXT("succeeded") : TEXT("failed"), *GetNameSafe(&Owner), *GetFullNameSafe(&RootStateTree));
-				Exec.TreeRunStatus = NextActiveFrames.Last().ActiveStates.Last().ToCompletionStatus();
+					__FUNCTION__, StateSelectionResult.GetSelectedFrames().Last().ActiveStates.Last() == FStateTreeStateHandle::Succeeded ? TEXT("succeeded") : TEXT("failed"), *GetNameSafe(&Owner), *GetFullNameSafe(&RootStateTree));
+				Exec.TreeRunStatus = StateSelectionResult.GetSelectedFrames().Last().ActiveStates.Last().ToCompletionStatus();
 			}
 			else
 			{
@@ -263,7 +263,8 @@ EStateTreeRunStatus FStateTreeExecutionContext::Start(const FInstancedPropertyBa
 				FStateTreeTransitionResult Transition;
 				Transition.TargetState = RootState;
 				Transition.CurrentRunStatus = Exec.LastTickStatus;
-				Transition.NextActiveFrames = NextActiveFrames; // Enter state will update Exec.ActiveFrames.
+				Transition.NextActiveFrames = StateSelectionResult.GetSelectedFrames(); // Enter state will update Exec.ActiveFrames.
+				Transition.NextActiveFrameEvents = StateSelectionResult.GetFramesStateSelectionEvents();
 				const EStateTreeRunStatus LastTickStatus = EnterState(Transition);
 			
 				Exec.LastTickStatus = LastTickStatus;
@@ -360,7 +361,7 @@ EStateTreeRunStatus FStateTreeExecutionContext::Stop(EStateTreeRunStatus Complet
 	
 	// Capture events added between ticks.
 	FStateTreeEventQueue& EventQueue = InstanceData.GetMutableEventQueue();
-	EventsToProcess = EventQueue.GetEvents();
+	EventsToProcess = EventQueue.GetEventsView();
 	EventQueue.Reset();
 
 	// Exit states if still in some valid state.
@@ -434,7 +435,7 @@ EStateTreeRunStatus FStateTreeExecutionContext::Tick(const float DeltaTime)
 	Exec.CurrentPhase = EStateTreeUpdatePhase::TickStateTree;
 	
 	// Capture events added between ticks.
-	EventsToProcess = EventQueue.GetEvents();
+	EventsToProcess = EventQueue.GetEventsView();
 	EventQueue.Reset();
 	
 	// Update the delayed transitions.
@@ -467,7 +468,7 @@ EStateTreeRunStatus FStateTreeExecutionContext::Tick(const float DeltaTime)
 			// Append events accumulated during the tick, so that transitions can immediately act on them.
 			// We'll consume the events only if they lead to state change below (EnterState is treated the same as Tick),
 			// or let them be processed next frame if no transition.
-			EventsToProcess.Append(EventQueue.GetEvents());
+			EventsToProcess.Append(EventQueue.GetEventsView());
 
 			// Trigger conditional transitions or state succeed/failed transitions. First tick transition is handled here too.
 			if (TriggerTransitions())
@@ -497,7 +498,7 @@ EStateTreeRunStatus FStateTreeExecutionContext::Tick(const float DeltaTime)
 				}
 
 				// Append and consume the events accumulated during the state exit.
-				EventsToProcess.Append(EventQueue.GetEvents());
+				EventsToProcess.Append(EventQueue.GetEventsView());
 				EventQueue.Reset();
 
 				// Enter state tasks can fail/succeed, treat it same as tick.
@@ -822,6 +823,11 @@ void FStateTreeExecutionContext::UpdateInstanceData(TConstArrayView<FStateTreeEx
 				TempInstanceStructs[BaseIndex + State.ParameterDataHandle.GetIndex()] = FindInstanceTempData(NextFrame, State.ParameterDataHandle);
 			}
 
+			if (State.EventDataIndex.IsValid())
+			{
+				InstanceStructs[BaseIndex + State.EventDataIndex.Get()] = FConstStructView(FStateTreeSharedEvent::StaticStruct());
+			}
+
 			for (int32 TaskIndex = State.TasksBegin; TaskIndex < (State.TasksBegin + State.TasksNum); TaskIndex++)
 			{
 				const FStateTreeTaskBase& Task = NextFrame.StateTree->Nodes[TaskIndex].Get<const FStateTreeTaskBase>();
@@ -874,7 +880,54 @@ void FStateTreeExecutionContext::UpdateInstanceData(TConstArrayView<FStateTreeEx
 	InstanceData.ResetTemporaryInstances();
 }
 
-FStateTreeDataView FStateTreeExecutionContext::GetDataView(FStateTreeInstanceStorage& InstanceDataStorage, FStateTreeInstanceStorage* CurrentlyProcessedSharedInstanceStorage, const FStateTreeExecutionFrame* ParentFrame, const FStateTreeExecutionFrame& CurrentFrame, TConstArrayView<FStateTreeDataView> ContextAndExternalDataViews, const FStateTreeDataHandle Handle)
+FStateTreeDataView FStateTreeExecutionContext::GetDataView(const FStateTreeExecutionFrame* ParentFrame, const FStateTreeExecutionFrame& CurrentFrame, const FStateTreeDataHandle Handle)
+{
+	switch (Handle.GetSource())
+	{
+	case EStateTreeDataSourceType::ContextData:
+		check(!ContextAndExternalDataViews.IsEmpty())
+		return ContextAndExternalDataViews[Handle.GetIndex()];
+
+	case EStateTreeDataSourceType::ExternalData:
+		check(!ContextAndExternalDataViews.IsEmpty())
+		return ContextAndExternalDataViews[CurrentFrame.ExternalDataBaseIndex.Get() + Handle.GetIndex()];
+
+	case EStateTreeDataSourceType::TransitionEvent:
+		{
+			if (CurrentlyProcessedTransitionEvent)
+			{
+				// const_cast because events are read only, but we cannot express that in FStateTreeDataView.
+				return FStateTreeDataView(FStructView::Make(*const_cast<FStateTreeEvent*>(CurrentlyProcessedTransitionEvent)));
+			}
+
+			return nullptr;
+		}
+
+	case EStateTreeDataSourceType::StateEvent:
+		{
+			if (CurrentlyProcessedStateSelectionEvents)
+			{
+				if (const FCompactStateTreeState* State = CurrentFrame.StateTree->GetStateFromHandle(Handle.GetState()))
+				{
+					if (const FStateTreeEvent* Event = CurrentlyProcessedStateSelectionEvents->Events[State->Depth].Get())
+					{
+						// const_cast because events are read only, but we cannot express that in FStateTreeDataView.
+						return FStateTreeDataView(FStructView::Make(*const_cast<FStateTreeEvent*>(Event)));
+					}
+				}
+
+				return {};
+			}
+
+			return GetDataViewFromInstanceStorage(*InstanceDataStorage, CurrentlyProcessedSharedInstanceStorage, ParentFrame, CurrentFrame, Handle);
+		}
+
+	default:
+		return GetDataViewFromInstanceStorage(*InstanceDataStorage, CurrentlyProcessedSharedInstanceStorage, ParentFrame, CurrentFrame, Handle);
+	}
+}
+
+FStateTreeDataView FStateTreeExecutionContext::GetDataViewFromInstanceStorage(FStateTreeInstanceStorage& InstanceDataStorage, FStateTreeInstanceStorage* CurrentlyProcessedSharedInstanceStorage, const FStateTreeExecutionFrame* ParentFrame, const FStateTreeExecutionFrame& CurrentFrame, const FStateTreeDataHandle Handle)
 {
 	switch (Handle.GetSource())
 	{
@@ -898,20 +951,12 @@ FStateTreeDataView FStateTreeExecutionContext::GetDataView(FStateTreeInstanceSto
 		check(CurrentlyProcessedSharedInstanceStorage);
 		return CurrentlyProcessedSharedInstanceStorage->GetMutableObject(Handle.GetIndex());
 
-	case EStateTreeDataSourceType::ContextData:
-		check(!ContextAndExternalDataViews.IsEmpty())
-		return ContextAndExternalDataViews[Handle.GetIndex()];
-
-	case EStateTreeDataSourceType::ExternalData:
-		check(!ContextAndExternalDataViews.IsEmpty())
-		return ContextAndExternalDataViews[CurrentFrame.ExternalDataBaseIndex.Get() + Handle.GetIndex()];
-
 	case EStateTreeDataSourceType::GlobalParameterData:
 		{
 			// Defined in parent frame or is root state tree parameters
 			if (ParentFrame)
 			{
-				return GetDataView(InstanceDataStorage, CurrentlyProcessedSharedInstanceStorage, nullptr, *ParentFrame, ContextAndExternalDataViews, CurrentFrame.GlobalParameterDataHandle);
+				return GetDataViewFromInstanceStorage(InstanceDataStorage, CurrentlyProcessedSharedInstanceStorage, nullptr, *ParentFrame, CurrentFrame.GlobalParameterDataHandle);
 			}
 
 			return InstanceDataStorage.GetMutableGlobalParameters();
@@ -921,7 +966,7 @@ FStateTreeDataView FStateTreeExecutionContext::GetDataView(FStateTreeInstanceSto
 		{
 			// Defined in parent frame.
 			check(ParentFrame);
-			return GetDataView(InstanceDataStorage, CurrentlyProcessedSharedInstanceStorage, nullptr, *ParentFrame, ContextAndExternalDataViews, CurrentFrame.StateParameterDataHandle);
+			return GetDataViewFromInstanceStorage(InstanceDataStorage, CurrentlyProcessedSharedInstanceStorage, nullptr, *ParentFrame, CurrentFrame.StateParameterDataHandle);
 		}
 
 	case EStateTreeDataSourceType::StateParameterData:
@@ -929,6 +974,19 @@ FStateTreeDataView FStateTreeExecutionContext::GetDataView(FStateTreeInstanceSto
 			FCompactStateTreeParameters& Params = InstanceDataStorage.GetMutableStruct(CurrentFrame.ActiveInstanceIndexBase.Get() + Handle.GetIndex()).Get<FCompactStateTreeParameters>();
 			return Params.Parameters.GetMutableValue();
 		}
+
+	case EStateTreeDataSourceType::StateEvent:
+		{
+			FStateTreeSharedEvent& SharedEvent = InstanceDataStorage.GetMutableStruct(CurrentFrame.ActiveInstanceIndexBase.Get() + Handle.GetIndex()).Get<FStateTreeSharedEvent>();
+			check(SharedEvent.IsValid());
+			return FStateTreeDataView(FStructView::Make(SharedEvent));
+		}
+
+	case EStateTreeDataSourceType::ContextData:
+	case EStateTreeDataSourceType::ExternalData:
+	case EStateTreeDataSourceType::TransitionEvent:
+		return {};
+	
 	default:
 		checkf(false, TEXT("Unhandle case %s"), *UEnum::GetValueAsString(Handle.GetSource()));
 	}
@@ -943,7 +1001,7 @@ const FStateTreeExecutionFrame* FStateTreeExecutionContext::FindFrame(const USta
 		return Frame.StateTree == StateTree && Frame.RootState == RootState;
 	});
 
-	if(FrameIndex == INDEX_NONE)
+	if (FrameIndex == INDEX_NONE)
 	{
 		OutParentFrame = nullptr;
 		return nullptr;
@@ -1010,6 +1068,15 @@ bool FStateTreeExecutionContext::IsHandleSourceValid(const FStateTreeExecutionFr
 			&& CurrentFrame.ActiveStates.Contains(Handle.GetState())
 			&& InstanceDataStorage->IsValidIndex(CurrentFrame.ActiveInstanceIndexBase.Get() + Handle.GetIndex());
 
+	case EStateTreeDataSourceType::TransitionEvent:
+		return CurrentlyProcessedTransitionEvent != nullptr;
+
+	case EStateTreeDataSourceType::StateEvent:
+		return CurrentlyProcessedStateSelectionEvents != nullptr
+			|| (CurrentFrame.ActiveInstanceIndexBase.IsValid()
+			&& CurrentFrame.ActiveStates.Contains(Handle.GetState())
+			&& InstanceDataStorage->IsValidIndex(CurrentFrame.ActiveInstanceIndexBase.Get() + Handle.GetIndex()));
+
 	default:
 		checkf(false, TEXT("Unhandle case %s"), *UEnum::GetValueAsString(Handle.GetSource()));
 	}
@@ -1017,7 +1084,7 @@ bool FStateTreeExecutionContext::IsHandleSourceValid(const FStateTreeExecutionFr
 	return false;
 }
 
-FStateTreeDataView FStateTreeExecutionContext::GetDataViewOrTemporary(const FStateTreeExecutionFrame* ParentFrame, const FStateTreeExecutionFrame& CurrentFrame, const FStateTreeDataHandle Handle) const
+FStateTreeDataView FStateTreeExecutionContext::GetDataViewOrTemporary(const FStateTreeExecutionFrame* ParentFrame, const FStateTreeExecutionFrame& CurrentFrame, const FStateTreeDataHandle Handle)
 {
 	if (IsHandleSourceValid(ParentFrame, CurrentFrame, Handle))
 	{
@@ -1083,7 +1150,7 @@ FStateTreeDataView FStateTreeExecutionContext::AddTemporaryInstance(const FState
 	return NewInstance;
 }
 
-bool FStateTreeExecutionContext::CopyBatchOnActiveInstances(const FStateTreeExecutionFrame* ParentFrame, const FStateTreeExecutionFrame& CurrentFrame, const FStateTreeDataView TargetView, const FStateTreeIndex16 BindingsBatch) const
+bool FStateTreeExecutionContext::CopyBatchOnActiveInstances(const FStateTreeExecutionFrame* ParentFrame, const FStateTreeExecutionFrame& CurrentFrame, const FStateTreeDataView TargetView, const FStateTreeIndex16 BindingsBatch)
 {
 	const FStateTreePropertyCopyBatch& Batch = CurrentFrame.StateTree->PropertyBindings.GetBatch(BindingsBatch);
 	check(TargetView.GetStruct() == Batch.TargetStruct.Struct);
@@ -1097,7 +1164,7 @@ bool FStateTreeExecutionContext::CopyBatchOnActiveInstances(const FStateTreeExec
 	return bSucceed;
 }
 
-bool FStateTreeExecutionContext::CopyBatchWithValidation(const FStateTreeExecutionFrame* ParentFrame, const FStateTreeExecutionFrame& CurrentFrame, const FStateTreeDataView TargetView, const FStateTreeIndex16 BindingsBatch) const
+bool FStateTreeExecutionContext::CopyBatchWithValidation(const FStateTreeExecutionFrame* ParentFrame, const FStateTreeExecutionFrame& CurrentFrame, const FStateTreeDataView TargetView, const FStateTreeIndex16 BindingsBatch)
 {
 	const FStateTreePropertyCopyBatch& Batch = CurrentFrame.StateTree->PropertyBindings.GetBatch(BindingsBatch);
 	check(TargetView.GetStruct() == Batch.TargetStruct.Struct);
@@ -1242,6 +1309,69 @@ bool FStateTreeExecutionContext::SetGlobalParameters(const FInstancedPropertyBag
 	return false;
 }
 
+void FStateTreeExecutionContext::CaptureNewStateEvents(TConstArrayView<FStateTreeExecutionFrame> PrevFrames, TConstArrayView<FStateTreeExecutionFrame> NewFrames, TArrayView<FStateTreeFrameStateSelectionEvents> FramesStateSelectionEvents)
+{
+	// Mark the events from delayed transitions as in use, so that each State will receive unique copy of the event struct. 
+	TArray<FStateTreeSharedEvent, TInlineAllocator<16>> EventsInUse;
+	for (const FStateTreeTransitionDelayedState& DelayedTransition : GetExecState().DelayedTransitions)
+	{
+		if (DelayedTransition.CapturedEvent.IsValid())
+		{
+			EventsInUse.Add(DelayedTransition.CapturedEvent);
+		}
+	}
+
+	for (int32 FrameIndex = 0; FrameIndex < NewFrames.Num(); ++FrameIndex)
+	{
+		const FStateTreeExecutionFrame& NewFrame = NewFrames[FrameIndex];
+
+		// Find states that are unique to the new frame.
+		TConstArrayView<FStateTreeStateHandle> UniqueStates = NewFrame.ActiveStates.States;
+		if (PrevFrames.IsValidIndex(FrameIndex))
+		{
+			const FStateTreeExecutionFrame& PrevFrame = PrevFrames[FrameIndex];
+
+			if (PrevFrame.RootState == NewFrame.RootState
+				&& PrevFrame.StateTree == NewFrame.StateTree)
+			{
+				for (int32 StateIndex = 0; StateIndex < NewFrame.ActiveStates.Num(); ++StateIndex)
+				{
+					if (!PrevFrame.ActiveStates.IsValidIndex(StateIndex) || PrevFrame.ActiveStates[StateIndex] == NewFrame.ActiveStates[StateIndex])
+					{
+						UniqueStates = TConstArrayView<FStateTreeStateHandle>(&NewFrame.ActiveStates[StateIndex], NewFrame.ActiveStates.Num() - StateIndex);
+						break;
+					}
+				}
+			}
+		}
+
+		// Capture events for the new states.
+		for (const FStateTreeStateHandle StateHandle : UniqueStates)
+		{
+			if (const FCompactStateTreeState* State = NewFrame.StateTree->GetStateFromHandle(StateHandle))
+			{
+				if (State->EventDataIndex.IsValid())
+				{
+					FStateTreeSharedEvent& StateTreeEvent = InstanceDataStorage->GetMutableStruct(NewFrame.ActiveInstanceIndexBase.Get() + State->EventDataIndex.Get()).Get<FStateTreeSharedEvent>();
+					
+					const FStateTreeSharedEvent& EventToCapture = FramesStateSelectionEvents[FrameIndex].Events[State->Depth];
+					if (EventsInUse.Contains(EventToCapture))
+					{
+						// Event is already spoken for, make a copy.
+						StateTreeEvent = FStateTreeSharedEvent(*EventToCapture);
+					}
+					else
+					{
+						// Event not in use, steal it.
+						StateTreeEvent = FramesStateSelectionEvents[FrameIndex].Events[State->Depth];
+						EventsInUse.Add(EventToCapture);
+					}
+				}
+			}
+		}
+	}
+}
+
 EStateTreeRunStatus FStateTreeExecutionContext::EnterState(FStateTreeTransitionResult& Transition)
 {
 	CSV_SCOPED_TIMING_STAT_EXCLUSIVE(StateTree_EnterState);
@@ -1256,6 +1386,8 @@ EStateTreeRunStatus FStateTreeExecutionContext::EnterState(FStateTreeTransitionR
 	// Allocate new tasks.
 	UpdateInstanceData(Exec.ActiveFrames, Transition.NextActiveFrames);
 
+	CaptureNewStateEvents(Exec.ActiveFrames, Transition.NextActiveFrames, Transition.NextActiveFrameEvents);
+
 	Exec.StateChangeCount++;
 	Exec.CompletedFrameIndex = FStateTreeIndex16::Invalid;
 	Exec.CompletedStateHandle = FStateTreeStateHandle::Invalid;
@@ -2334,18 +2466,25 @@ bool FStateTreeExecutionContext::TestAllConditions(const FStateTreeExecutionFram
 FString FStateTreeExecutionContext::DebugGetEventsAsString() const
 {
 	FString Result;
-	for (const FStateTreeEvent& Event : EventsToProcess)
+	for (const FStateTreeSharedEvent& Event : EventsToProcess)
 	{
 		if (!Result.IsEmpty())
 		{
 			Result += TEXT(", ");
 		}
-		Result += Event.Tag.ToString();
+
+		check(Event.IsValid());
+		Result += Event->Tag.ToString();
 	}
 	return Result;
 }
 
-bool FStateTreeExecutionContext::RequestTransition(const FStateTreeExecutionFrame& CurrentFrame, const FStateTreeStateHandle NextState, const EStateTreeTransitionPriority Priority, const EStateTreeSelectionFallback Fallback)
+bool FStateTreeExecutionContext::RequestTransition(
+	const FStateTreeExecutionFrame& CurrentFrame,
+	const FStateTreeStateHandle NextState,
+	const EStateTreeTransitionPriority Priority,
+	const FStateTreeSharedEvent* TransitionEvent,
+	const EStateTreeSelectionFallback Fallback)
 {
 	// Skip lower priority transitions.
 	if (NextTransition.Priority >= Priority)
@@ -2367,11 +2506,12 @@ bool FStateTreeExecutionContext::RequestTransition(const FStateTreeExecutionFram
 		return true;
 	}
 
-	TArray<FStateTreeExecutionFrame, TFixedAllocator<MaxExecutionFrames>> NewNextActiveFrames;
-	if (SelectState(CurrentFrame, NextState, NewNextActiveFrames, Fallback))
+	FStateSelectionResult StateSelectionResult;
+	if (SelectState(CurrentFrame, NextState, StateSelectionResult, TransitionEvent, Fallback))
 	{
 		SetupNextTransition(CurrentFrame, NextState, Priority);
-		NextTransition.NextActiveFrames = NewNextActiveFrames;
+		NextTransition.NextActiveFrames = StateSelectionResult.GetSelectedFrames();
+		NextTransition.NextActiveFrameEvents = StateSelectionResult.GetFramesStateSelectionEvents();
 
 		STATETREE_LOG(Verbose, TEXT("Transition on state '%s' -[%s]-> state '%s'"),
 			*GetSafeStateName(CurrentFrame, CurrentFrame.ActiveStates.Last()),
@@ -2522,76 +2662,112 @@ bool FStateTreeExecutionContext::TriggerTransitions()
 					}
 
 					// If a delayed transition has passed the delay, and remove it from the queue, and try trigger it.
-					FStateTreeTransitionDelayedState* DelayedState = nullptr;
 					if (Transition.HasDelay())
 					{
-						DelayedState = Exec.FindDelayedTransition(CurrentFrame.StateTree, FStateTreeIndex16(TransitionIndex));
-						if (DelayedState != nullptr && DelayedState->TimeLeft <= 0.0f)
+						bool bTriggeredDelayedTransition = false;
+						const TArray<FStateTreeTransitionDelayedState, TInlineAllocator<8>> ExpiredDelayedStates = Exec.FindAndRemoveExpiredDelayedTransitions(CurrentFrame.StateTree, FStateTreeIndex16(TransitionIndex));
+						for (const FStateTreeTransitionDelayedState& DelayedState : ExpiredDelayedStates)
 						{
 							STATETREE_LOG(Verbose, TEXT("Passed delayed transition from '%s' (%s) -> '%s'"),
 								*GetSafeStateName(CurrentFrame, CurrentFrame.ActiveStates.Last()), *State.Name.ToString(), *GetSafeStateName(CurrentFrame, Transition.State));
 
-							Exec.DelayedTransitions.RemoveAllSwap([StateTree = CurrentFrame.StateTree, TransitionIndex](const FStateTreeTransitionDelayedState& DelayedState)
-								{
-									return DelayedState.StateTree == StateTree && DelayedState.TransitionIndex.Get() == TransitionIndex;
-								});
-
 							// Trigger Delayed Transition when the delay has passed.
-							if (RequestTransition(CurrentFrame, Transition.State, Transition.Priority, Transition.Fallback))
+							if (RequestTransition(CurrentFrame, Transition.State, Transition.Priority, &DelayedState.CapturedEvent, Transition.Fallback))
 							{
 								NextTransitionSource = FStateTreeTransitionSource(FStateTreeIndex16(TransitionIndex), Transition.State, Transition.Priority);
+								bTriggeredDelayedTransition = true;
+								break;
 							}
+						}
+
+						if (bTriggeredDelayedTransition)
+						{
 							continue;
 						}
 					}
 
-					const bool bShouldTrigger = Transition.Trigger == EStateTreeTransitionTrigger::OnTick
-												|| (Transition.Trigger == EStateTreeTransitionTrigger::OnEvent
-													&& HasEventToProcess(Transition.EventTag));
+					TArray<FStateTreeSharedEvent*, TInlineAllocator<8>> TransitionEvents;
 
-					bool bPassed = false; 
-					if (bShouldTrigger)
+					if (Transition.Trigger == EStateTreeTransitionTrigger::OnEvent)
 					{
-						STATETREE_TRACE_TRANSITION_EVENT(FStateTreeTransitionSource(FStateTreeIndex16(TransitionIndex), Transition.State, Transition.Priority), EStateTreeTraceEventType::OnEvaluating);
-						STATETREE_TRACE_SCOPED_PHASE(EStateTreeUpdatePhase::TransitionConditions);
-						bPassed = TestAllConditions(CurrentParentFrame, CurrentFrame, Transition.ConditionsBegin, Transition.ConditionsNum);
-					}
+						check(Transition.RequiredEvent.IsValid());
 
-					if (bPassed)
+						TConstArrayView<FStateTreeSharedEvent> EventsQueue = GetEventsToProcessView();
+						for (FStateTreeSharedEvent Event : EventsQueue)
+						{
+							check(Event.IsValid());
+							if (Transition.RequiredEvent.DoesEventMatchDesc(*Event))
+							{
+								TransitionEvents.Emplace(&Event);
+							}
+						}
+					}
+					else if (Transition.Trigger == EStateTreeTransitionTrigger::OnTick)
 					{
-						// If the transitions is delayed, set up the delay. 
-						if (Transition.HasDelay())
+						TransitionEvents.Emplace();
+					}
+					
+					for (FStateTreeSharedEvent* TransitionEvent : TransitionEvents)
+					{
+						bool bPassed = false; 
+						{
+							FCurrentlyProcessedTransitionEventScope TransitionEventScope(*this, TransitionEvent ? TransitionEvent->Get() : nullptr);
+							STATETREE_TRACE_TRANSITION_EVENT(FStateTreeTransitionSource(FStateTreeIndex16(TransitionIndex), Transition.State, Transition.Priority), EStateTreeTraceEventType::OnEvaluating);
+							STATETREE_TRACE_SCOPED_PHASE(EStateTreeUpdatePhase::TransitionConditions);
+							bPassed = TestAllConditions(CurrentParentFrame, CurrentFrame, Transition.ConditionsBegin, Transition.ConditionsNum);
+						}
+
+						if (bPassed)
 						{
-							if (DelayedState == nullptr)
+							// If the transitions is delayed, set up the delay. 
+							if (Transition.HasDelay())
 							{
-								// Initialize new delayed transition.
-								const float DelayDuration = Transition.Delay.GetRandomDuration();
-								if (DelayDuration > 0.0f)
+								uint32 TransitionEventHash = 0u;
+								if (TransitionEvent && TransitionEvent->IsValid())
+								{		
+									TransitionEventHash = GetTypeHash(*TransitionEvent->Get());
+								}
+
+								const bool bIsDelayedTransitionExisting = Exec.DelayedTransitions.ContainsByPredicate([StateTree = CurrentFrame.StateTree, TransitionIndex, TransitionEventHash](const FStateTreeTransitionDelayedState& DelayedState)
 								{
-									DelayedState = &Exec.DelayedTransitions.AddDefaulted_GetRef();
-									DelayedState->StateTree = CurrentFrame.StateTree;
-									DelayedState->TransitionIndex = FStateTreeIndex16(TransitionIndex);
-									DelayedState->TimeLeft = DelayDuration;
-									BeginDelayedTransition(*DelayedState);
-									STATETREE_LOG(Verbose, TEXT("Delayed transition triggered from '%s' (%s) -> '%s' %.1fs"),
-										*GetSafeStateName(CurrentFrame, CurrentFrame.ActiveStates.Last()), *State.Name.ToString(), *GetSafeStateName(CurrentFrame, Transition.State), DelayedState->TimeLeft);
+									return DelayedState.StateTree == StateTree && DelayedState.TransitionIndex.Get() == TransitionIndex && DelayedState.CapturedEventHash == TransitionEventHash;
+								});
+
+								if (!bIsDelayedTransitionExisting)
+								{
+									// Initialize new delayed transition.
+									const float DelayDuration = Transition.Delay.GetRandomDuration();
+									if (DelayDuration > 0.0f)
+									{
+										FStateTreeTransitionDelayedState& DelayedState = Exec.DelayedTransitions.AddDefaulted_GetRef();
+										DelayedState.StateTree = CurrentFrame.StateTree;
+										DelayedState.TransitionIndex = FStateTreeIndex16(TransitionIndex);
+										DelayedState.TimeLeft = DelayDuration;
+										DelayedState.CapturedEvent = *TransitionEvent;
+										DelayedState.CapturedEventHash = TransitionEventHash;
+
+										BeginDelayedTransition(DelayedState);
+										STATETREE_LOG(Verbose, TEXT("Delayed transition triggered from '%s' (%s) -> '%s' %.1fs"),
+											*GetSafeStateName(CurrentFrame, CurrentFrame.ActiveStates.Last()), *State.Name.ToString(), *GetSafeStateName(CurrentFrame, Transition.State), DelayedState.TimeLeft);
 									
-									// Delay state added, skip requesting the transition.
+										// Delay state added, skip requesting the transition.
+										continue;
+									}
+									// Fallthrough to request transition if duration was zero. 
+								}
+								else
+								{
+									// We get here if the transitions re-triggers during the delay, on which case we'll just ignore it.
 									continue;
 								}
-								// Fallthrough to request transition if duration was zero. 
 							}
-							else
+
+							if (RequestTransition(CurrentFrame, Transition.State, Transition.Priority, TransitionEvent, Transition.Fallback))
 							{
-								// We get here if the transitions re-triggers during the delay, on which case we'll just ignore it.
-								continue;
+								NextTransitionSource = FStateTreeTransitionSource(FStateTreeIndex16(TransitionIndex), Transition.State, Transition.Priority);
+								break;
 							}
 						}
-
-						if (RequestTransition(CurrentFrame, Transition.State, Transition.Priority, Transition.Fallback))
-						{
-							NextTransitionSource = FStateTreeTransitionSource(FStateTreeIndex16(TransitionIndex), Transition.State, Transition.Priority);
-						}
 					}
 				}
 			}
@@ -2695,7 +2871,7 @@ bool FStateTreeExecutionContext::TriggerTransitions()
 						{
 							// No delay allowed on completion conditions.
 							// No priority on completion transitions, use the priority to signal that state is selected.
-							if (RequestTransition(CurrentFrame, Transition.State, EStateTreeTransitionPriority::Normal, Transition.Fallback))
+							if (RequestTransition(CurrentFrame, Transition.State, EStateTreeTransitionPriority::Normal, /*TransitionEvent*/nullptr, Transition.Fallback))
 							{
 								NextTransitionSource = FStateTreeTransitionSource(FStateTreeIndex16(TransitionIndex), Transition.State, Transition.Priority);
 								break;
@@ -2778,9 +2954,12 @@ bool FStateTreeExecutionContext::TriggerTransitions()
 
 bool FStateTreeExecutionContext::SelectState(const FStateTreeExecutionFrame& CurrentFrame,
 											const FStateTreeStateHandle NextState,
-											TArray<FStateTreeExecutionFrame, TFixedAllocator<MaxExecutionFrames>>& OutNextActiveFrames,
+											FStateSelectionResult& OutSelectionResult,
+											const FStateTreeSharedEvent* TransitionEvent,
 											const EStateTreeSelectionFallback Fallback)
 {
+	check(TransitionEvent == nullptr || TransitionEvent->IsValid());
+
 	const FStateTreeExecutionState& Exec = GetExecState();
 
 	if (Exec.ActiveFrames.IsEmpty())
@@ -2839,14 +3018,14 @@ bool FStateTreeExecutionContext::SelectState(const FStateTreeExecutionFrame& Cur
 	if (CurrentFrameIndex != INDEX_NONE)
 	{
 		const int32 NumCommonFrames = CurrentFrameIndex + 1;
-		OutNextActiveFrames = MakeArrayView(Exec.ActiveFrames.GetData(), NumCommonFrames);
+		OutSelectionResult = FStateSelectionResult(MakeArrayView(Exec.ActiveFrames.GetData(), NumCommonFrames));
 		CurrentFrameInActiveFrames  = &Exec.ActiveFrames[CurrentFrameIndex];
 	}
 	else if (CurrentStateTreeIndex != INDEX_NONE)
 	{
 		// If we could not find a common frame, we assume that we jumped to different subtree in same asset.
 		const int32 NumCommonFrames = CurrentStateTreeIndex + 1;
-		OutNextActiveFrames = MakeArrayView(Exec.ActiveFrames.GetData(), NumCommonFrames);
+		OutSelectionResult = FStateSelectionResult(MakeArrayView(Exec.ActiveFrames.GetData(), NumCommonFrames));
 		CurrentFrameInActiveFrames  = &Exec.ActiveFrames[CurrentStateTreeIndex];
 	}
 	else
@@ -2858,8 +3037,8 @@ bool FStateTreeExecutionContext::SelectState(const FStateTreeExecutionFrame& Cur
 	
 	// Append in between state in reverse order, they were collected from leaf towards the root.
 	// Note: NextState will be added by SelectStateInternal() if conditions pass.
-	const int32 LastFrameIndex = OutNextActiveFrames.Num() - 1;
-	FStateTreeExecutionFrame& LastFrame = OutNextActiveFrames[LastFrameIndex];
+	const int32 LastFrameIndex = OutSelectionResult.FramesNum() - 1;
+	FStateTreeExecutionFrame& LastFrame = OutSelectionResult.GetSelectedFrames()[LastFrameIndex];
 
 	LastFrame.ActiveStates.Reset();
 	for (int32 Index = NumParentStates - 1; Index > 0; Index--)
@@ -2867,16 +3046,16 @@ bool FStateTreeExecutionContext::SelectState(const FStateTreeExecutionFrame& Cur
 		LastFrame.ActiveStates.Push(ParentStates[Index]);
 	}
 
-	TArray<FStateTreeExecutionFrame, TFixedAllocator<MaxExecutionFrames>> InitialNextActiveFrames;
-	
+	FStateSelectionResult InitialSelection;
+
 	if (Fallback == EStateTreeSelectionFallback::NextSelectableSibling)
 	{
-		InitialNextActiveFrames = OutNextActiveFrames;
+		InitialSelection = OutSelectionResult;
 	}
 	
 	// We take copy of the last frame and assign it later, as SelectStateInternal() might change the array and invalidate the pointer.
-	const FStateTreeExecutionFrame* CurrentParentFrame = LastFrameIndex > 0 ? &OutNextActiveFrames[LastFrameIndex - 1] : nullptr; 
-	if (SelectStateInternal(CurrentParentFrame, OutNextActiveFrames[LastFrameIndex], CurrentFrameInActiveFrames , NextState, OutNextActiveFrames))
+	const FStateTreeExecutionFrame* CurrentParentFrame = LastFrameIndex > 0 ? &OutSelectionResult.GetSelectedFrames()[LastFrameIndex - 1] : nullptr; 
+	if (SelectStateInternal(CurrentParentFrame, OutSelectionResult.GetSelectedFrames()[LastFrameIndex], CurrentFrameInActiveFrames , NextState, OutSelectionResult, TransitionEvent))
 	{
 		return true;
 	}
@@ -2897,11 +3076,11 @@ bool FStateTreeExecutionContext::SelectState(const FStateTreeExecutionFrame& Cur
 				FStateTreeStateHandle ChildStateHandle = FStateTreeStateHandle(ChildState);
 
 				// Start selection from blank slate.
-				OutNextActiveFrames = InitialNextActiveFrames;
+				OutSelectionResult = InitialSelection;
 	
 				// We take copy of the last frame and assign it later, as SelectStateInternal() might change the array and invalidate the pointer.
-				CurrentParentFrame = LastFrameIndex > 0 ? &OutNextActiveFrames[LastFrameIndex - 1] : nullptr; 
-				if (SelectStateInternal(CurrentParentFrame, OutNextActiveFrames[LastFrameIndex], CurrentFrameInActiveFrames , ChildStateHandle, OutNextActiveFrames))
+				CurrentParentFrame = LastFrameIndex > 0 ? &OutSelectionResult.GetSelectedFrames()[LastFrameIndex - 1] : nullptr; 
+				if (SelectStateInternal(CurrentParentFrame, OutSelectionResult.GetSelectedFrames()[LastFrameIndex], CurrentFrameInActiveFrames, ChildStateHandle, OutSelectionResult))
 				{
 					return true;
 				}
@@ -2917,7 +3096,8 @@ bool FStateTreeExecutionContext::SelectStateInternal(
 	FStateTreeExecutionFrame& CurrentFrame,
 	const FStateTreeExecutionFrame* CurrentFrameInActiveFrames,
 	const FStateTreeStateHandle NextStateHandle,
-	TArray<FStateTreeExecutionFrame, TFixedAllocator<MaxExecutionFrames>>& OutNextActiveFrames)
+	FStateSelectionResult& OutSelectionResult,
+	const FStateTreeSharedEvent* TransitionEvent)
 {
 	CSV_SCOPED_TIMING_STAT_EXCLUSIVE(StateTree_SelectState);
 
@@ -2932,6 +3112,7 @@ bool FStateTreeExecutionContext::SelectStateInternal(
 	}
 
 	FCurrentlyProcessedFrameScope FrameScope(*this, CurrentParentFrame, CurrentFrame);
+	FCurrentFrameStateSelectionEventsScope CapturedEventsScope(*this, OutSelectionResult.GetFramesStateSelectionEvents().Last());
 
 	const UStateTree* CurrentStateTree = CurrentFrame.StateTree;
 	const FCompactStateTreeState& NextState = CurrentStateTree->States[NextStateHandle.Index];
@@ -2952,13 +3133,53 @@ bool FStateTreeExecutionContext::SelectStateInternal(
 		return false;
 	}
 
-	// Check that the state can be entered
-	STATETREE_TRACE_PHASE_BEGIN(EStateTreeUpdatePhase::EnterConditions);
-	const bool bEnterConditionsPassed = TestAllConditions(CurrentParentFrame, CurrentFrame, NextState.EnterConditionsBegin, NextState.EnterConditionsNum);
-	STATETREE_TRACE_PHASE_END(EStateTreeUpdatePhase::EnterConditions);
+	TArray<const FStateTreeSharedEvent*, TInlineAllocator<FStateTreeEventQueue::MaxActiveEvents>> StateSelectionEvents;
+	if (NextState.EventDataIndex.IsValid())
+	{
+		check(NextState.RequiredEventToEnter.IsValid());
 
-	if (bEnterConditionsPassed)
+		TArrayView<FStateTreeSharedEvent> EventsQueue = GetMutableEventsToProcessView();
+		if (TransitionEvent && TransitionEvent->IsValid())
+		{
+			if (NextState.RequiredEventToEnter.DoesEventMatchDesc(*TransitionEvent->Get()))
+			{
+				StateSelectionEvents.Emplace(TransitionEvent);
+			}
+		}
+		else
+		{
+			for (FStateTreeSharedEvent& Event : EventsQueue)
+			{
+				check(Event.IsValid());
+				if (NextState.RequiredEventToEnter.DoesEventMatchDesc(*Event))
+				{
+					StateSelectionEvents.Emplace(&Event);
+				}
+			}
+		}
+	}
+	else
+	{
+		StateSelectionEvents.Emplace();
+	}
+
+	for (const FStateTreeSharedEvent* StateSelectionEvent : StateSelectionEvents)
 	{
+		if (StateSelectionEvent)
+		{
+			CurrentlyProcessedStateSelectionEvents->Events[NextState.Depth] = *StateSelectionEvent;
+		}
+		
+		// Check that the state can be entered
+		STATETREE_TRACE_PHASE_BEGIN(EStateTreeUpdatePhase::EnterConditions);
+		const bool bEnterConditionsPassed = TestAllConditions(CurrentParentFrame, CurrentFrame, NextState.EnterConditionsBegin, NextState.EnterConditionsNum);
+		STATETREE_TRACE_PHASE_END(EStateTreeUpdatePhase::EnterConditions);
+
+		if (!bEnterConditionsPassed)
+		{
+			continue;
+		}
+
 		if (!CurrentFrame.ActiveStates.Push(NextStateHandle))
 		{
 			STATETREE_LOG(Error, TEXT("%hs: Reached max execution depth when trying to select state %s from '%s'.  '%s' using StateTree '%s'."),
@@ -3012,7 +3233,7 @@ bool FStateTreeExecutionContext::SelectStateInternal(
 		{
 			if (NextState.LinkedState.IsValid())
 			{
-				if (OutNextActiveFrames.Num() == MaxExecutionFrames)
+				if (OutSelectionResult.IsFull())
 				{
 					STATETREE_LOG(Error, TEXT("%hs: Reached max execution depth when trying to select state %s from '%s'.  '%s' using StateTree '%s'."),
 						__FUNCTION__, *GetSafeStateName(CurrentFrame, NextStateHandle), *GetStateStatusString(Exec), *GetNameSafe(&Owner), *GetFullNameSafe(CurrentFrame.StateTree));
@@ -3025,7 +3246,7 @@ bool FStateTreeExecutionContext::SelectStateInternal(
 				NewFrame.ExternalDataBaseIndex = CurrentFrame.ExternalDataBaseIndex;
 
 				// Check and prevent recursion.
-				const bool bNewFrameAlreadySelected = OutNextActiveFrames.ContainsByPredicate([&NewFrame](const FStateTreeExecutionFrame& Frame) {
+				const bool bNewFrameAlreadySelected = OutSelectionResult.GetSelectedFrames().ContainsByPredicate([&NewFrame](const FStateTreeExecutionFrame& Frame) {
 					return Frame.IsSameFrame(NewFrame);
 				});
 				
@@ -3057,15 +3278,15 @@ bool FStateTreeExecutionContext::SelectStateInternal(
 					NewFrame.StateParameterDataHandle = NextState.ParameterDataHandle; // Temporary allocated earlier if did not exists.
 				}
 
-				OutNextActiveFrames.Push(NewFrame);
+				OutSelectionResult.PushFrame(NewFrame);
 
 				// If State is linked, proceed to the linked state.
-				if (SelectStateInternal(&CurrentFrame, OutNextActiveFrames.Last(), ExistingFrame, NewFrame.RootState, OutNextActiveFrames))
+				if (SelectStateInternal(&CurrentFrame, OutSelectionResult.GetSelectedFrames().Last(), ExistingFrame, NewFrame.RootState, OutSelectionResult))
 				{
 					return true;
 				}
 				
-				OutNextActiveFrames.Pop();
+				OutSelectionResult.PopFrame();
 			}
 			else
 			{
@@ -3077,7 +3298,7 @@ bool FStateTreeExecutionContext::SelectStateInternal(
 		{
 			if (NextState.LinkedAsset)
 			{
-				if (OutNextActiveFrames.Num() == MaxExecutionFrames)
+				if (OutSelectionResult.IsFull())
 				{
 					STATETREE_LOG(Error, TEXT("%hs: Reached max execution depth when trying to select state %s from '%s'.  '%s' using StateTree '%s'."),
 						__FUNCTION__, *GetSafeStateName(CurrentFrame, NextStateHandle), *GetStateStatusString(Exec), *GetNameSafe(&Owner), *GetFullNameSafe(CurrentFrame.StateTree));
@@ -3098,7 +3319,7 @@ bool FStateTreeExecutionContext::SelectStateInternal(
 				NewFrame.bIsGlobalFrame = true;
 
 				// Check and prevent recursion.
-				const bool bNewFrameAlreadySelected = OutNextActiveFrames.ContainsByPredicate([&NewFrame](const FStateTreeExecutionFrame& Frame) {
+				const bool bNewFrameAlreadySelected = OutSelectionResult.GetSelectedFrames().ContainsByPredicate([&NewFrame](const FStateTreeExecutionFrame& Frame) {
 					return Frame.IsSameFrame(NewFrame);
 				});
 				
@@ -3150,15 +3371,15 @@ bool FStateTreeExecutionContext::SelectStateInternal(
 					}
 				}
 				
-				OutNextActiveFrames.Push(NewFrame);
+				OutSelectionResult.PushFrame(NewFrame);
 
 				// If State is linked, proceed to the linked state.
-				if (SelectStateInternal(&CurrentFrame, OutNextActiveFrames.Last(), ExistingFrame, NewFrame.RootState, OutNextActiveFrames))
+				if (SelectStateInternal(&CurrentFrame, OutSelectionResult.GetSelectedFrames().Last(), ExistingFrame, NewFrame.RootState, OutSelectionResult))
 				{
 					return true;
 				}
 				
-				OutNextActiveFrames.Pop();
+				OutSelectionResult.PopFrame();
 			}
 			else
 			{
@@ -3211,33 +3432,62 @@ bool FStateTreeExecutionContext::SelectStateInternal(
 				if (CurrentFrame.ActiveStates.Contains(Transition.State))
 				{
 					STATETREE_LOG(Error, TEXT("%hs: Loop detected when trying to select state %s from '%s'. Prior states: %s.  '%s' using StateTree '%s'."),
-						__FUNCTION__, *GetSafeStateName(CurrentFrame, NextStateHandle), *GetStateStatusString(Exec), *DebugGetStatePath(OutNextActiveFrames, &CurrentFrame), *GetNameSafe(&Owner), *GetFullNameSafe(CurrentFrame.StateTree));
+						__FUNCTION__, *GetSafeStateName(CurrentFrame, NextStateHandle), *GetStateStatusString(Exec), *DebugGetStatePath(OutSelectionResult.GetSelectedFrames(), &CurrentFrame), *GetNameSafe(&Owner), *GetFullNameSafe(CurrentFrame.StateTree));
 					continue;
 				}
 
-				const bool bShouldTrigger = Transition.Trigger == EStateTreeTransitionTrigger::OnTick
-											|| (Transition.Trigger == EStateTreeTransitionTrigger::OnEvent
-												&& HasEventToProcess(Transition.EventTag));
+				TArray<const FStateTreeSharedEvent*, TInlineAllocator<FStateTreeEventQueue::MaxActiveEvents>> SelectedStateTransitionEvents;
+				if (Transition.Trigger == EStateTreeTransitionTrigger::OnEvent)
+				{
+					check(Transition.RequiredEvent.IsValid());
 
-				bool bTransitionConditionsPassed = false;
-				if (bShouldTrigger)
+					if (StateSelectionEvent)
+					{
+						SelectedStateTransitionEvents.Emplace(StateSelectionEvent);
+					}
+					else
+					{
+						TArrayView<FStateTreeSharedEvent> EventsQueue = GetMutableEventsToProcessView();
+						for (FStateTreeSharedEvent& Event : EventsQueue)
+						{
+							check(Event.IsValid());
+							if (Transition.RequiredEvent.DoesEventMatchDesc(*Event))
+							{
+								SelectedStateTransitionEvents.Emplace(&Event);
+							}
+						}
+					}
+				}
+				else if (Transition.Trigger == EStateTreeTransitionTrigger::OnTick)
 				{
-					STATETREE_TRACE_TRANSITION_EVENT(FStateTreeTransitionSource(FStateTreeIndex16(TransitionIndex), Transition.State, Transition.Priority), EStateTreeTraceEventType::OnEvaluating);
-					STATETREE_TRACE_SCOPED_PHASE(EStateTreeUpdatePhase::TransitionConditions);
-					bTransitionConditionsPassed = bShouldTrigger && TestAllConditions(CurrentParentFrame, CurrentFrame, Transition.ConditionsBegin, Transition.ConditionsNum);
+					SelectedStateTransitionEvents.Emplace();
 				}
 
-				if (bTransitionConditionsPassed)
+				for (const FStateTreeSharedEvent* SelectedStateTransitionEvent : SelectedStateTransitionEvents)
 				{
-					// Using SelectState() instead of SelectStateInternal to treat the transitions the same way as regular transitions,
-					// e.g. it may jump to a completely different branch.
-					TArray<FStateTreeExecutionFrame, TFixedAllocator<MaxExecutionFrames>> NewActiveFrames;
-					if (SelectState(CurrentFrame, Transition.State, NewActiveFrames, Transition.Fallback))
+					bool bTransitionConditionsPassed = false;
+					{
+						FCurrentlyProcessedTransitionEventScope TransitionEventScope(*this, SelectedStateTransitionEvent ? SelectedStateTransitionEvent->Get() : nullptr);
+
+						STATETREE_TRACE_TRANSITION_EVENT(FStateTreeTransitionSource(FStateTreeIndex16(TransitionIndex), Transition.State, Transition.Priority), EStateTreeTraceEventType::OnEvaluating);
+						STATETREE_TRACE_SCOPED_PHASE(EStateTreeUpdatePhase::TransitionConditions);
+
+						bTransitionConditionsPassed = TestAllConditions(CurrentParentFrame, CurrentFrame, Transition.ConditionsBegin, Transition.ConditionsNum);
+					}
+
+					if (bTransitionConditionsPassed)
 					{
-						// Selection succeeded.
-						// Cannot break yet because higher priority transitions may override the selection. 
-						OutNextActiveFrames = NewActiveFrames;
-						CurrentPriority = Transition.Priority;
+						// Using SelectState() instead of SelectStateInternal to treat the transitions the same way as regular transitions,
+						// e.g. it may jump to a completely different branch.
+						FStateSelectionResult StateSelectionResult;
+						if (SelectState(CurrentFrame, Transition.State, StateSelectionResult, SelectedStateTransitionEvent, Transition.Fallback))
+						{
+							// Selection succeeded.
+							// Cannot break yet because higher priority transitions may override the selection. 
+							OutSelectionResult = StateSelectionResult;
+							CurrentPriority = Transition.Priority;
+							break;
+						}
 					}
 				}
 			}
@@ -3256,7 +3506,7 @@ bool FStateTreeExecutionContext::SelectStateInternal(
 				// If the state has children, proceed to select children.
 				for (uint16 ChildState = NextState.ChildrenBegin; ChildState < NextState.ChildrenEnd; ChildState = CurrentStateTree->States[ChildState].GetNextSibling())
 				{
-					if (SelectStateInternal(CurrentParentFrame, CurrentFrame, CurrentFrameInActiveFrames, FStateTreeStateHandle(ChildState), OutNextActiveFrames))
+					if (SelectStateInternal(CurrentParentFrame, CurrentFrame, CurrentFrameInActiveFrames, FStateTreeStateHandle(ChildState), OutSelectionResult))
 					{
 						// Selection succeeded
 						return true;
@@ -3497,14 +3747,15 @@ void FStateTreeExecutionContext::DebugPrintInternalLayout()
 	RootStateTree.PropertyBindings.DebugPrintInternalLayout(DebugString);
 
 	// Transitions
-	DebugString += FString::Printf(TEXT("\nTransitions(%d)\n  [ %-3s | %15s | %-20s | %-40s | %-8s ]\n"), RootStateTree.Transitions.Num()
-		, TEXT("Idx"), TEXT("State"), TEXT("Transition Trigger"), TEXT("Transition Event Tag"), TEXT("Num Cond"));
+	DebugString += FString::Printf(TEXT("\nTransitions(%d)\n  [ %-3s | %15s | %-20s | %-40s | %-40s | %-8s ]\n"), RootStateTree.Transitions.Num()
+		, TEXT("Idx"), TEXT("State"), TEXT("Transition Trigger"), TEXT("Transition Event Tag"), TEXT("Transition Event Payload"), TEXT("Num Cond"));
 	for (const FCompactStateTransition& Transition : RootStateTree.Transitions)
 	{
-		DebugString += FString::Printf(TEXT("  | %3d | %15s | %-20s | %-40s | %8d |\n"),
+		DebugString += FString::Printf(TEXT("  | %3d | %15s | %-20s | %-40s | %-40s | %8d |\n"),
 									Transition.ConditionsBegin, *Transition.State.Describe(),
 									*UEnum::GetDisplayValueAsText(Transition.Trigger).ToString(),
-									*Transition.EventTag.ToString(),
+									*Transition.RequiredEvent.Tag.ToString(),
+									Transition.RequiredEvent.PayloadStruct ? *Transition.RequiredEvent.PayloadStruct->GetName() : TEXT("None"),
 									Transition.ConditionsNum);
 	}
 
diff --git a/Engine/Plugins/Runtime/StateTree/Source/StateTreeModule/Private/StateTreePropertyRefHelpers.cpp b/Engine/Plugins/Runtime/StateTree/Source/StateTreeModule/Private/StateTreePropertyRefHelpers.cpp
--- a/Engine/Plugins/Runtime/StateTree/Source/StateTreeModule/Private/StateTreePropertyRefHelpers.cpp
+++ b/Engine/Plugins/Runtime/StateTree/Source/StateTreeModule/Private/StateTreePropertyRefHelpers.cpp
@@ -226,7 +226,9 @@ namespace UE::StateTree::PropertyRefHelpers
 		switch (SourceStruct.DataSource)
 		{
 		case EStateTreeBindableStructSource::Parameter:
-		case EStateTreeBindableStructSource::State:
+		case EStateTreeBindableStructSource::StateParameter:
+		case EStateTreeBindableStructSource::TransitionEvent:
+		case EStateTreeBindableStructSource::StateEvent:
 			return true;
 
 		case EStateTreeBindableStructSource::Context:
@@ -237,6 +239,7 @@ namespace UE::StateTree::PropertyRefHelpers
 		case EStateTreeBindableStructSource::Evaluator:
 		case EStateTreeBindableStructSource::Task:
 			return bIsOutput || IsPropertyRef(SourceProperty);
+
 		default:
 			checkNoEntry();
 		}
diff --git a/Engine/Plugins/Runtime/StateTree/Source/StateTreeModule/Private/StateTreeTypes.cpp b/Engine/Plugins/Runtime/StateTree/Source/StateTreeModule/Private/StateTreeTypes.cpp
--- a/Engine/Plugins/Runtime/StateTree/Source/StateTreeModule/Private/StateTreeTypes.cpp
+++ b/Engine/Plugins/Runtime/StateTree/Source/StateTreeModule/Private/StateTreeTypes.cpp
@@ -1,6 +1,7 @@
 // Copyright Epic Games, Inc. All Rights Reserved.
 
 #include "StateTreeTypes.h"
+#include "StateTreeEvents.h"
 #include "StateTree.h" // FStateTreeCustomVersion
 
 #include UE_INLINE_GENERATED_CPP_BY_NAME(StateTreeTypes)
@@ -78,3 +79,9 @@ PRAGMA_ENABLE_DEPRECATION_WARNINGS
 	}
 #endif // WITH_EDITORONLY_DAT
 }
+
+bool FCompactEventDesc::DoesEventMatchDesc(const FStateTreeEvent& Event) const
+{
+	const UScriptStruct* EventPayloadStruct = Event.Payload.GetScriptStruct();
+	return (!Tag.IsValid() || !Event.Tag.IsValid() || Event.Tag.MatchesTag(Tag)) && (EventPayloadStruct == nullptr || PayloadStruct == nullptr || EventPayloadStruct->IsChildOf(PayloadStruct));
+}
diff --git a/Engine/Plugins/Runtime/StateTree/Source/StateTreeModule/Public/StateTree.h b/Engine/Plugins/Runtime/StateTree/Source/StateTreeModule/Public/StateTree.h
--- a/Engine/Plugins/Runtime/StateTree/Source/StateTreeModule/Public/StateTree.h
+++ b/Engine/Plugins/Runtime/StateTree/Source/StateTreeModule/Public/StateTree.h
@@ -49,6 +49,8 @@ struct STATETREEMODULE_API FStateTreeCustomVersion
 		OverridableStateParameters,
 		// Added storing global parameters in instance storage
 		StoringGlobalParametersInInstanceStorage,
+		// Added binding to events
+		AddedBindingToEvents,
 
 		// -----<new versions can be added above this line>-------------------------------------------------
 		VersionPlusOne,
diff --git a/Engine/Plugins/Runtime/StateTree/Source/StateTreeModule/Public/StateTreeEvents.h b/Engine/Plugins/Runtime/StateTree/Source/StateTreeModule/Public/StateTreeEvents.h
--- a/Engine/Plugins/Runtime/StateTree/Source/StateTreeModule/Public/StateTreeEvents.h
+++ b/Engine/Plugins/Runtime/StateTree/Source/StateTreeModule/Public/StateTreeEvents.h
@@ -4,6 +4,7 @@
 
 #include "GameplayTagContainer.h"
 #include "StructView.h"
+#include "StateTreeIndexTypes.h"
 #include "StateTreeEvents.generated.h"
 
 /** Enum used for flow control during event iteration. */
@@ -37,13 +38,25 @@ struct STATETREEMODULE_API FStateTreeEvent
 		, Origin(InOrigin)
 	{
 	}
+
+	friend FORCEINLINE uint32 GetTypeHash(const FStateTreeEvent& Event)
+	{
+		uint32 Hash = GetTypeHash(Event.Tag);
+		
+		if (Event.Payload.IsValid())
+		{
+			Hash = HashCombineFast(Hash, Event.Payload.GetScriptStruct()->GetStructTypeHash(Event.Payload.GetMemory()));
+		}
+
+		return Hash;
+	}
 	
 	/** Tag describing the event */
 	UPROPERTY(EditAnywhere, BlueprintReadWrite, Category = "Default", meta=(Categories="StateTreeEvent"))
 	FGameplayTag Tag;
 
 	/** Optional payload for the event. */ 
-	UPROPERTY(EditAnywhere, Category = "Default")
+	UPROPERTY(EditAnywhere, BlueprintReadWrite, Category = "Default")
 	FInstancedStruct Payload;
 
 	/** Optional info to describe who sent the event. */
@@ -51,6 +64,81 @@ struct STATETREEMODULE_API FStateTreeEvent
 	FName Origin;
 };
 
+/**
+ * A struct wrapping FStateTreeEvent in shared struct, used to make it easier to refer to the events during State Tree update.
+ */
+USTRUCT()
+struct FStateTreeSharedEvent
+{
+	GENERATED_BODY()
+
+	FStateTreeSharedEvent() = default;
+
+	explicit FStateTreeSharedEvent(const FGameplayTag InTag, const FConstStructView InPayload, const FName InOrigin)
+		: Event(MakeShared<FStateTreeEvent>(InTag, InPayload, InOrigin))
+	{}
+
+	explicit FStateTreeSharedEvent(const FStateTreeEvent& InEvent)
+		: Event(MakeShared<FStateTreeEvent>(InEvent))
+	{}
+
+	void AddStructReferencedObjects(FReferenceCollector& Collector);
+
+	const FStateTreeEvent* Get() const
+	{
+		return Event.Get();
+	}
+
+	FStateTreeEvent* GetMutable()
+	{
+		return Event.Get();
+	}
+
+	const FStateTreeEvent* operator->() const
+	{
+		return Event.Get();
+	}
+
+	FStateTreeEvent* operator->()
+	{
+		return Event.Get();
+	}
+
+	const FStateTreeEvent& operator*()
+	{
+		check(Event.IsValid());
+		return *Event.Get();
+	}
+
+	FStateTreeEvent& operator*() const
+	{
+		check(Event.IsValid());
+		return *Event.Get();
+	}
+
+	bool IsValid() const
+	{
+		return Event.IsValid();
+	}
+
+	bool operator==(const FStateTreeSharedEvent& Other) const
+	{
+		return Event == Other.Event;
+	}
+
+protected:
+	TSharedPtr<FStateTreeEvent> Event;
+};
+
+template<>
+struct TStructOpsTypeTraits<FStateTreeSharedEvent> : public TStructOpsTypeTraitsBase2<FStateTreeSharedEvent>
+{
+	enum
+	{
+		WithAddStructReferencedObjects = true,
+	};
+};
+
 /**
  * Event queue buffering all the events to be processed by a State Tree.
  */
@@ -63,15 +151,18 @@ struct STATETREEMODULE_API FStateTreeEventQueue
 	static constexpr int32 MaxActiveEvents = 64;
 
 	/** @return const view to all the events in the buffer. */
-	TConstArrayView<FStateTreeEvent> GetEvents() const
+	TConstArrayView<FStateTreeSharedEvent> GetEventsView() const
 	{
-		return Events;
+		return SharedEvents;
 	}
 
+	UE_DEPRECATED(5.5, "Use GetEventsView() instead.")
+	TConstArrayView<FStateTreeEvent> GetEvents() const { return {};	}
+
 	/** Resets the events in the event queue */
 	void Reset()
 	{
-		Events.Reset();
+		SharedEvents.Reset();
 	}
 
 	/**
@@ -85,10 +176,10 @@ struct STATETREEMODULE_API FStateTreeEventQueue
 
 protected:
 	// Used by FStateTreeExecutionState to implement deprecated functionality.
-	TArray<FStateTreeEvent>& GetEventsArray() { return Events; };
+	TArray<FStateTreeSharedEvent>& GetEventsArray() { return SharedEvents; };
 
 	UPROPERTY()
-	TArray<FStateTreeEvent> Events;
+	TArray<FStateTreeSharedEvent> SharedEvents;
 
 	friend struct FStateTreeInstanceData;
 };
diff --git a/Engine/Plugins/Runtime/StateTree/Source/StateTreeModule/Public/StateTreeExecutionContext.h b/Engine/Plugins/Runtime/StateTree/Source/StateTreeModule/Public/StateTreeExecutionContext.h
--- a/Engine/Plugins/Runtime/StateTree/Source/StateTreeModule/Public/StateTreeExecutionContext.h
+++ b/Engine/Plugins/Runtime/StateTree/Source/StateTreeModule/Public/StateTreeExecutionContext.h
@@ -6,6 +6,7 @@
 #include "StateTreeExecutionTypes.h"
 #include "StateTreeNodeBase.h"
 #include "Experimental/ConcurrentLinearAllocator.h"
+#include "Templates/IsInvocable.h"
 
 struct FGameplayTag;
 struct FInstancedPropertyBag;
@@ -191,11 +192,11 @@ public:
 	/** Sends event for the StateTree. */
 	void SendEvent(const FGameplayTag Tag, const FConstStructView Payload = FConstStructView(), const FName Origin = FName()) const;
 
-	/** Iterates over all events. Can only be used during StateTree tick. Expects a lambda which takes const FStateTreeEvent& Event, and returns EStateTreeLoopEvents. */
+	/** Iterates over all events. Can only be used during StateTree tick. Expects a lambda which takes const FStateTreeSharedEvent& Event, and returns EStateTreeLoopEvents. */
 	template<typename TFunc>
-	void ForEachEvent(TFunc&& Function) const
+	TEnableIf<TIsInvocable<TFunc, FStateTreeSharedEvent>::Value, void>::Type ForEachEvent(TFunc&& Function) const
 	{
-		for (const FStateTreeEvent& Event : EventsToProcess)
+		for (const FStateTreeSharedEvent& Event : EventsToProcess)
 		{
 			if (Function(Event) == EStateTreeLoopEvents::Break)
 			{
@@ -204,8 +205,27 @@ public:
 		}
 	}
 
+	/** Iterates over all events. Can only be used during StateTree tick. Expects a lambda which takes const FStateTreeEvent& Event, and returns EStateTreeLoopEvents. Less preferable than FStateTreeSharedEvent version. */
+	template<typename TFunc>
+	TEnableIf<TIsInvocable<TFunc, FStateTreeEvent>::Value, void>::Type ForEachEvent(TFunc&& Function) const
+	{
+		for (const FStateTreeSharedEvent& Event : EventsToProcess)
+		{
+			if (Function(*Event) == EStateTreeLoopEvents::Break)
+			{
+				break;
+			}
+		}
+	}
+
 	/** @return events to process this tick. */
-	TConstArrayView<FStateTreeEvent> GetEventsToProcess() const { return EventsToProcess; }
+	TArrayView<FStateTreeSharedEvent> GetMutableEventsToProcessView() { return EventsToProcess; }
+
+	/** @return events to process this tick. */
+	TConstArrayView<FStateTreeSharedEvent> GetEventsToProcessView() const { return EventsToProcess; }
+
+	UE_DEPRECATED(5.5, "Use GetEventsToProcessView() instead.")
+	TConstArrayView<FStateTreeEvent> GetEventsToProcess() const { return {}; }
 
 	/** @return true if there is a pending event with specified tag. */
 	bool HasEventToProcess(const FGameplayTag Tag) const
@@ -215,9 +235,10 @@ public:
 			return false;
 		}
 		
-		return EventsToProcess.ContainsByPredicate([Tag](const FStateTreeEvent& Event)
+		return EventsToProcess.ContainsByPredicate([Tag](const FStateTreeSharedEvent& Event)
 		{
-			return Event.Tag.MatchesTag(Tag);
+			check(Event.IsValid());
+			return Event->Tag.MatchesTag(Tag);
 		});
 	}
 
@@ -372,15 +393,74 @@ public:
 	void RequestTransition(const FStateTreeTransitionRequest& Request);
 
 	/** @return data view of the specified handle relative to given frame. */
-	static FStateTreeDataView GetDataView(FStateTreeInstanceStorage& InstanceDataStorage, FStateTreeInstanceStorage* CurrentlyProcessedSharedInstanceStorage, const FStateTreeExecutionFrame* ParentFrame, const FStateTreeExecutionFrame& CurrentFrame, TConstArrayView<FStateTreeDataView> ContextAndExternalDataViews, const FStateTreeDataHandle Handle);
+	static FStateTreeDataView GetDataViewFromInstanceStorage(FStateTreeInstanceStorage& InstanceDataStorage, FStateTreeInstanceStorage* CurrentlyProcessedSharedInstanceStorage, const FStateTreeExecutionFrame* ParentFrame, const FStateTreeExecutionFrame& CurrentFrame, const FStateTreeDataHandle Handle);
 
 	/** Looks for a frame in provided list of frames. */
 	static const FStateTreeExecutionFrame* FindFrame(const UStateTree* StateTree, FStateTreeStateHandle RootState, TConstArrayView<FStateTreeExecutionFrame> Frames, const FStateTreeExecutionFrame*& OutParentFrame);
 
+	UE_DEPRECATED(5.5, "Use GetDataViewFromInstanceStorage instead.")
+	static FStateTreeDataView GetDataView(FStateTreeInstanceStorage& InstanceDataStorage, FStateTreeInstanceStorage* CurrentlyProcessedSharedInstanceStorage, const FStateTreeExecutionFrame* ParentFrame, const FStateTreeExecutionFrame& CurrentFrame, TConstArrayView<FStateTreeDataView> ContextAndExternalDataViews, const FStateTreeDataHandle Handle)
+	{
+		return GetDataViewFromInstanceStorage(InstanceDataStorage, CurrentlyProcessedSharedInstanceStorage, ParentFrame, CurrentFrame, Handle);
+	}
+
 protected:
+	/**
+	 * Describes a result of States Selection.
+	 */
+	struct FStateSelectionResult
+	{
+		/** Max number of execution frames handled during state selection. */
+		static constexpr int32 MaxExecutionFrames = 8;
+
+		FStateSelectionResult() = default;
+		explicit FStateSelectionResult(TConstArrayView<FStateTreeExecutionFrame> InFrames)
+		{
+			SelectedFrames = InFrames;
+			FramesStateSelectionEvents.SetNum(SelectedFrames.Num());
+		}
+
+		bool IsFull() const
+		{
+			return SelectedFrames.Num() == MaxExecutionFrames;
+		}
+
+		void PushFrame(FStateTreeExecutionFrame Frame)
+		{
+			SelectedFrames.Add(Frame);
+			FramesStateSelectionEvents.AddDefaulted();
+		}
+
+		void PopFrame()
+		{
+			SelectedFrames.Pop();
+			FramesStateSelectionEvents.Pop();
+		}
+
+		bool ContainsFrames() const
+		{
+			return !SelectedFrames.IsEmpty();
+		}
+
+		int32 FramesNum() const
+		{
+			return SelectedFrames.Num();
+		}
+
+		TArrayView<FStateTreeExecutionFrame> GetSelectedFrames()
+		{
+			return SelectedFrames;
+		}
 
-	/** Max number of execution frames handled during state selection. */
-	static constexpr int32 MaxExecutionFrames = 8;
+		TArrayView<FStateTreeFrameStateSelectionEvents> GetFramesStateSelectionEvents()
+		{
+			return FramesStateSelectionEvents;
+		}
+
+	protected:
+		TArray<FStateTreeExecutionFrame, TFixedAllocator<MaxExecutionFrames>> SelectedFrames;
+		TArray<FStateTreeFrameStateSelectionEvents, TFixedAllocator<MaxExecutionFrames>> FramesStateSelectionEvents;
+	};
 	
 #if WITH_STATETREE_DEBUGGER
 	FStateTreeInstanceDebugId GetInstanceDebugId() const;
@@ -455,6 +535,7 @@ protected:
 		const FStateTreeExecutionFrame& CurrentFrame,
 		const FStateTreeStateHandle NextState,
 		const EStateTreeTransitionPriority Priority,
+		const FStateTreeSharedEvent* TransitionEvent = nullptr,
 		const EStateTreeSelectionFallback Fallback = EStateTreeSelectionFallback::None);
 
 	/**
@@ -478,14 +559,16 @@ protected:
 	 * If NextState is a leaf state, the active states leading from root to the leaf are returned.
 	 * @param CurrentFrame The frame where the NextState is valid. 
 	 * @param NextState The state which we try to select next.
-	 * @param OutNextActiveFrames Active frames and states that got selected.
+	 * @param OutSelectionResult The result of selection.
+	 * @param TransitionEvent Event used by transition to trigger the state selection.
 	 * @param Fallback selection behavior to execute if it fails to select the desired state
 	 * @return True if succeeded to select new active states.
 	 */
 	bool SelectState(
 		const FStateTreeExecutionFrame& CurrentFrame,
 		const FStateTreeStateHandle NextState,
-		TArray<FStateTreeExecutionFrame, TFixedAllocator<MaxExecutionFrames>>& OutNextActiveFrames,
+		FStateSelectionResult& OutSelectionResult,
+		const FStateTreeSharedEvent* TransitionEvent = nullptr,
 		const EStateTreeSelectionFallback Fallback = EStateTreeSelectionFallback::None);
 
 	/**
@@ -496,7 +579,8 @@ protected:
 		FStateTreeExecutionFrame& CurrentFrame,
 		const FStateTreeExecutionFrame* CurrentFrameInActiveFrames,
 		const FStateTreeStateHandle NextStateHandle,
-		TArray<FStateTreeExecutionFrame, TFixedAllocator<MaxExecutionFrames>>& OutNextActiveFrames);
+		FStateSelectionResult& OutSelectionResult,
+		const FStateTreeSharedEvent* TransitionEvent = nullptr);
 
 	/** @return StateTree execution state from the instance storage. */
 	FStateTreeExecutionState& GetExecState()
@@ -523,16 +607,13 @@ protected:
 	FString DebugGetEventsAsString() const;
 
 	/** @return data view of the specified handle relative to given frame. */
-	FStateTreeDataView GetDataView(const FStateTreeExecutionFrame* ParentFrame, const FStateTreeExecutionFrame& CurrentFrame, const FStateTreeDataHandle Handle) const
-	{
-		return GetDataView(*InstanceDataStorage, CurrentlyProcessedSharedInstanceStorage, ParentFrame, CurrentFrame, ContextAndExternalDataViews, Handle);
-	}
+	FStateTreeDataView GetDataView(const FStateTreeExecutionFrame* ParentFrame, const FStateTreeExecutionFrame& CurrentFrame, const FStateTreeDataHandle Handle);
 
 	/** @return true if handle source is valid cified handle relative to given frame. */
 	bool IsHandleSourceValid(const FStateTreeExecutionFrame* ParentFrame, const FStateTreeExecutionFrame& CurrentFrame, const FStateTreeDataHandle Handle) const;
 
 	/** @return data view of the specified handle relative to the given frame, or tries to find a matching temporary instance. */
-	FStateTreeDataView GetDataViewOrTemporary(const FStateTreeExecutionFrame* ParentFrame, const FStateTreeExecutionFrame& CurrentFrame, const FStateTreeDataHandle Handle) const;
+	FStateTreeDataView GetDataViewOrTemporary(const FStateTreeExecutionFrame* ParentFrame, const FStateTreeExecutionFrame& CurrentFrame, const FStateTreeDataHandle Handle);
 
 	/**
 	 * Adds a temporary instance that can be located using frame and data handle later.
@@ -541,10 +622,10 @@ protected:
 	FStateTreeDataView AddTemporaryInstance(const FStateTreeExecutionFrame& Frame, const FStateTreeIndex16 OwnerNodeIndex, const FStateTreeDataHandle DataHandle, FConstStructView NewInstanceData);
 
 	/** Copies a batch of properties to the data in TargetView. Should be used only on active instances, assumes valid handles and does not consider temporary instances. */
-	bool CopyBatchOnActiveInstances(const FStateTreeExecutionFrame* ParentFrame, const FStateTreeExecutionFrame& CurrentFrame, const FStateTreeDataView TargetView, const FStateTreeIndex16 BindingsBatch) const;
+	bool CopyBatchOnActiveInstances(const FStateTreeExecutionFrame* ParentFrame, const FStateTreeExecutionFrame& CurrentFrame, const FStateTreeDataView TargetView, const FStateTreeIndex16 BindingsBatch);
 
 	/** Copies a batch of properties to the data in TargetView. This version validates the data handles and looks up temporary instances. */
-	bool CopyBatchWithValidation(const FStateTreeExecutionFrame* ParentFrame, const FStateTreeExecutionFrame& CurrentFrame, const FStateTreeDataView TargetView, const FStateTreeIndex16 BindingsBatch) const;
+	bool CopyBatchWithValidation(const FStateTreeExecutionFrame* ParentFrame, const FStateTreeExecutionFrame& CurrentFrame, const FStateTreeDataView TargetView, const FStateTreeIndex16 BindingsBatch);
 
 	/**
 	 * Collects external data for all StateTrees in active frames.
@@ -563,6 +644,11 @@ protected:
 	 * @returns true if successfully set the parameters
 	 */
 	bool SetGlobalParameters(const FInstancedPropertyBag& Parameters);
+
+	/**
+	 * Captures StateTree events used during the state selection.
+	 */
+	void CaptureNewStateEvents(TConstArrayView<FStateTreeExecutionFrame> PrevFrames, TConstArrayView<FStateTreeExecutionFrame> NewFrames, TArrayView<FStateTreeFrameStateSelectionEvents> FramesStateSelectionEvents);
 	
 	/** Owner of the instance data. */
 	UObject& Owner;
@@ -580,7 +666,7 @@ protected:
 	TArray<FStateTreeDataView, TConcurrentLinearArrayAllocator<FDefaultBlockAllocationTag>> ContextAndExternalDataViews;
 
 	/** Events to process in current tick. */
-	TArray<FStateTreeEvent, TConcurrentLinearArrayAllocator<FDefaultBlockAllocationTag>> EventsToProcess;
+	TArray<FStateTreeSharedEvent, TConcurrentLinearArrayAllocator<FDefaultBlockAllocationTag>> EventsToProcess;
 
 	FOnCollectStateTreeExternalData CollectExternalDataDelegate;
 
@@ -620,7 +706,6 @@ protected:
 		const FStateTreeExecutionFrame* SavedFrame = nullptr;
 		const FStateTreeExecutionFrame* SavedParentFrame = nullptr;
 	};
-
 	
 	/** Current state we're processing, or invalid if not applicable. */
 	FStateTreeStateHandle CurrentlyProcessedState;
@@ -644,7 +729,51 @@ protected:
 		FStateTreeExecutionContext& Context;
 		FStateTreeStateHandle SavedState = FStateTreeStateHandle::Invalid; 
 	};
-	
+
+	/** Current event we're processing in transition, or invalid if not applicable. */
+	const FStateTreeEvent* CurrentlyProcessedTransitionEvent = nullptr;
+
+	/** Helper struct to track currently processed transition event. */
+	struct FCurrentlyProcessedTransitionEventScope
+	{
+		FCurrentlyProcessedTransitionEventScope(FStateTreeExecutionContext& InContext, const FStateTreeEvent* Event)
+			: Context(InContext)
+		{
+			check(Context.CurrentlyProcessedTransitionEvent == nullptr);
+			Context.CurrentlyProcessedTransitionEvent = Event;
+		}
+
+		~FCurrentlyProcessedTransitionEventScope()
+		{
+			Context.CurrentlyProcessedTransitionEvent = nullptr;
+		}
+
+	private:
+		FStateTreeExecutionContext& Context;
+	};
+
+	/** Current frames events we're processing during the state selection, or invalid if not applicable. */
+	FStateTreeFrameStateSelectionEvents* CurrentlyProcessedStateSelectionEvents = nullptr;
+
+	/** Helper struct to track currently processed state selection events. */
+	struct FCurrentFrameStateSelectionEventsScope
+	{
+		FCurrentFrameStateSelectionEventsScope(FStateTreeExecutionContext& InContext, FStateTreeFrameStateSelectionEvents& InCurrentlyProcessedStateSelectionEvents)
+			: Context(InContext)
+		{
+			SavedStateSelectionEvents = Context.CurrentlyProcessedStateSelectionEvents; 
+			Context.CurrentlyProcessedStateSelectionEvents = &InCurrentlyProcessedStateSelectionEvents;
+		}
+
+		~FCurrentFrameStateSelectionEventsScope()
+		{
+			Context.CurrentlyProcessedStateSelectionEvents = SavedStateSelectionEvents;
+		}
+
+	private:
+		FStateTreeExecutionContext& Context;
+		FStateTreeFrameStateSelectionEvents* SavedStateSelectionEvents = nullptr;
+	};
 
 	/** True if transitions are allowed to be requested directly instead of buffering. */
 	bool bAllowDirectTransitions = false;
diff --git a/Engine/Plugins/Runtime/StateTree/Source/StateTreeModule/Public/StateTreeExecutionTypes.h b/Engine/Plugins/Runtime/StateTree/Source/StateTreeModule/Public/StateTreeExecutionTypes.h
--- a/Engine/Plugins/Runtime/StateTree/Source/StateTreeModule/Public/StateTreeExecutionTypes.h
+++ b/Engine/Plugins/Runtime/StateTree/Source/StateTreeModule/Public/StateTreeExecutionTypes.h
@@ -2,6 +2,7 @@
 
 #pragma once
 #include "StateTreeTypes.h"
+#include "StateTreeEvents.h"
 
 #include "StateTreeExecutionTypes.generated.h"
 
@@ -378,7 +379,7 @@ struct STATETREEMODULE_API FStateTreeActiveStates
 	bool IsEmpty() const { return NumStates == 0; } 
 
 	/** Returns a specified state in the array. */
-	FORCEINLINE FStateTreeStateHandle operator[](const int32 Index) const
+	FORCEINLINE const FStateTreeStateHandle& operator[](const int32 Index) const
 	{
 		check(Index >= 0 && Index < (int32)NumStates);
 		return States[Index];
@@ -525,6 +526,12 @@ struct STATETREEMODULE_API FStateTreeTransitionDelayedState
 
 	UPROPERTY()
 	float TimeLeft = 0.0f;
+
+	UPROPERTY()
+	FStateTreeSharedEvent CapturedEvent;
+
+	UPROPERTY()
+	uint32 CapturedEventHash = 0u;
 };
 
 /** Describes an active branch of a State Tree. */
@@ -604,13 +611,23 @@ PRAGMA_ENABLE_DEPRECATION_WARNINGS
 		StateChangeCount = 0;
 	}
 
-	/** @returns Delayed transition state for a specific transition, or nullptr if it does not exists. */
-	FStateTreeTransitionDelayedState* FindDelayedTransition(const UStateTree* OwnerStateTree, const FStateTreeIndex16 TransitionIndex)
+	UE_DEPRECATED(5.5, "Use FindAndRemoveExpiredDelayedTransitions() instead or search through DelayedTransitions directly.")
+	FStateTreeTransitionDelayedState* FindDelayedTransition(const UStateTree* OwnerStateTree, const FStateTreeIndex16 TransitionIndex) { return nullptr; }
+
+	/** Finds all delayed transition states for a specific transition and removes them. Returns their copies. */
+	TArray<FStateTreeTransitionDelayedState, TInlineAllocator<8>> FindAndRemoveExpiredDelayedTransitions(const UStateTree* OwnerStateTree, const FStateTreeIndex16 TransitionIndex)
 	{
-		return DelayedTransitions.FindByPredicate([OwnerStateTree, TransitionIndex](const FStateTreeTransitionDelayedState& TransitionState)
+		TArray<FStateTreeTransitionDelayedState, TInlineAllocator<8>> Result;
+		for (TArray<FStateTreeTransitionDelayedState>::TIterator It = DelayedTransitions.CreateIterator(); It; ++It)
 		{
-			return TransitionState.StateTree == OwnerStateTree && TransitionState.TransitionIndex == TransitionIndex;
-		});
+			if (It->TimeLeft <= 0.0f && It->StateTree == OwnerStateTree && It->TransitionIndex == TransitionIndex)
+			{
+				Result.Emplace(MoveTemp(*It));
+				It.RemoveCurrentSwap();
+			}
+		}
+
+		return Result;
 	}
 
 	/** Currently active frames (and states) */
@@ -676,6 +693,12 @@ PRAGMA_ENABLE_DEPRECATION_WARNINGS
 #endif	
 };
 
+ /** Contains StateTree events used during State Selection for a single execution frame. */
+struct FStateTreeFrameStateSelectionEvents
+{
+	TStaticArray<FStateTreeSharedEvent, FStateTreeActiveStates::MaxStates> Events;
+};
+
 /**
  * Describes a state tree transition. Source is the state where the transition started, Target describes the state where the transition pointed at,
  * and Next describes the selected state. The reason Transition and Next are different is that Transition state can be a selector state,
@@ -709,6 +732,9 @@ PRAGMA_ENABLE_DEPRECATION_WARNINGS
 	UPROPERTY(EditDefaultsOnly, Category = "Default", BlueprintReadOnly)
 	TArray<FStateTreeExecutionFrame> NextActiveFrames;
 
+	/** Events used in state selection. */
+	TArray<FStateTreeFrameStateSelectionEvents> NextActiveFrameEvents;
+
 	/** Current Run status. */
 	UPROPERTY(EditDefaultsOnly, Category = "Default", BlueprintReadOnly)
 	EStateTreeRunStatus CurrentRunStatus = EStateTreeRunStatus::Unset;
diff --git a/Engine/Plugins/Runtime/StateTree/Source/StateTreeModule/Public/StateTreePropertyBindings.h b/Engine/Plugins/Runtime/StateTree/Source/StateTreeModule/Public/StateTreePropertyBindings.h
--- a/Engine/Plugins/Runtime/StateTree/Source/StateTreeModule/Public/StateTreePropertyBindings.h
+++ b/Engine/Plugins/Runtime/StateTree/Source/StateTreeModule/Public/StateTreePropertyBindings.h
@@ -24,11 +24,15 @@ enum class EStateTreeBindableStructSource : uint8
 	/** Source is StateTree global task */
 	GlobalTask,
 	/** Source is State parameter */
-	State,
+	StateParameter,
 	/** Source is State task */
 	Task,
 	/** Source is State condition */
 	Condition,
+	/** Source is StateTree event used by transition */
+	TransitionEvent,
+	/** Source is StateTree event used by state selection */
+	StateEvent,
 };
 
 
diff --git a/Engine/Plugins/Runtime/StateTree/Source/StateTreeModule/Public/StateTreePropertyRef.h b/Engine/Plugins/Runtime/StateTree/Source/StateTreeModule/Public/StateTreePropertyRef.h
--- a/Engine/Plugins/Runtime/StateTree/Source/StateTreeModule/Public/StateTreePropertyRef.h
+++ b/Engine/Plugins/Runtime/StateTree/Source/StateTreeModule/Public/StateTreePropertyRef.h
@@ -25,8 +25,7 @@ namespace UE::StateTree::PropertyRefHelpers
 		const FStateTreePropertyBindings& PropertyBindings = ExecutionFrame.StateTree->GetPropertyBindings();
 		if (const FStateTreePropertyAccess* PropertyAccess = PropertyBindings.GetPropertyAccess(PropertyRef))
 		{
-			// Passing empty ContextAndExternalDataViews, as PropertyRef is not allowed to point to context or external data.
-			FStateTreeDataView SourceView = FStateTreeExecutionContext::GetDataView(InstanceDataStorage, nullptr, ParentExecutionFrame, ExecutionFrame, {}, PropertyAccess->SourceDataHandle);
+			FStateTreeDataView SourceView = FStateTreeExecutionContext::GetDataViewFromInstanceStorage(InstanceDataStorage, nullptr, ParentExecutionFrame, ExecutionFrame, PropertyAccess->SourceDataHandle);
 			
 			// The only possibility when PropertyRef references another PropertyRef is when source one is a global or subtree parameter, i.e lives in parent execution frame.
 			// If that's the case, referenced PropertyRef is obtained and we recursively take the address where it points to.
@@ -84,7 +83,7 @@ namespace UE::StateTree::PropertyRefHelpers
 		if (const FStateTreePropertyAccess* PropertyAccess = PropertyBindings.GetPropertyAccess(PropertyRef))
 		{
 			// Passing empty ContextAndExternalDataViews, as PropertyRef is not allowed to point to context or external data.
-			FStateTreeDataView SourceView = FStateTreeExecutionContext::GetDataView(InstanceDataStorage, nullptr, ParentExecutionFrame, ExecutionFrame, {}, PropertyAccess->SourceDataHandle);
+			FStateTreeDataView SourceView = FStateTreeExecutionContext::GetDataViewFromInstanceStorage(InstanceDataStorage, nullptr, ParentExecutionFrame, ExecutionFrame, PropertyAccess->SourceDataHandle);
 			return TTuple<T*...>(PropertyBindings.GetMutablePropertyPtr<T>(SourceView, *PropertyAccess)...);
 		}
 
diff --git a/Engine/Plugins/Runtime/StateTree/Source/StateTreeModule/Public/StateTreeTypes.h b/Engine/Plugins/Runtime/StateTree/Source/StateTreeModule/Public/StateTreeTypes.h
--- a/Engine/Plugins/Runtime/StateTree/Source/StateTreeModule/Public/StateTreeTypes.h
+++ b/Engine/Plugins/Runtime/StateTree/Source/StateTreeModule/Public/StateTreeTypes.h
@@ -14,6 +14,7 @@ STATETREEMODULE_API DECLARE_LOG_CATEGORY_EXTERN(LogStateTree, Warning, All);
 #endif // WITH_STATETREE_DEBUG
 
 class UStateTree;
+struct FStateTreeEvent;
 
 namespace UE::StateTree
 {
@@ -251,6 +252,12 @@ enum class EStateTreeDataSourceType : uint8
 
 	/** Parameters for regular and linked states */
 	StateParameterData,
+
+	/** Event used in transition. */
+	TransitionEvent,
+
+	/** Event used in state selection. */
+	StateEvent,
 };
 
 /** Handle to a StateTree data */
@@ -286,7 +293,7 @@ struct STATETREEMODULE_API FStateTreeDataHandle
 		// Require valid state for active instance data
 		check(Source != EStateTreeDataSourceType::ActiveInstanceData || (Source == EStateTreeDataSourceType::ActiveInstanceData && StateHandle.IsValid()));
 		check(Source != EStateTreeDataSourceType::ActiveInstanceDataObject || (Source == EStateTreeDataSourceType::ActiveInstanceDataObject && StateHandle.IsValid()));
-		check(Source == EStateTreeDataSourceType::GlobalParameterData || InIndex != InvalidIndex);
+		check(Source == EStateTreeDataSourceType::GlobalParameterData || IsValidIndex(InIndex));
 	}
 
 	explicit FStateTreeDataHandle(const EStateTreeDataSourceType InSource, const int32 InIndex, const FStateTreeStateHandle InStateHandle = FStateTreeStateHandle::Invalid)
@@ -480,6 +487,51 @@ enum class EStateTreeSelectionFallback : uint8
 	NextSelectableSibling,
 };
 
+/**
+ *  Runtime representation of an event description.
+ */
+USTRUCT()
+struct STATETREEMODULE_API FCompactEventDesc
+{
+	GENERATED_BODY()
+
+	/** Event Payload Struct. */
+	UPROPERTY()
+	TObjectPtr<const UScriptStruct> PayloadStruct = nullptr;
+
+	/** Event Tag. */
+	UPROPERTY()
+	FGameplayTag Tag;
+
+	/** Returns true if describes an event correctly. */
+	bool IsValid() const
+	{
+		return Tag.IsValid() || PayloadStruct;
+	}
+
+	/** Returns true if described events is a subset of events described by another EventDesc. */
+	bool IsSubsetOfAnotherDesc(const FCompactEventDesc& Desc) const
+	{
+		if (Tag.IsValid() && Desc.Tag.IsValid())
+		{
+			if (!Desc.Tag.MatchesTag(Tag) || !Tag.MatchesTag(Desc.Tag))
+			{
+				return false;
+			}
+		}
+
+		if (PayloadStruct && Desc.PayloadStruct)
+		{
+			return PayloadStruct->IsChildOf(Desc.PayloadStruct);
+		}
+
+		return true;
+	}
+
+	/** Returns true provided event matches description. */
+	bool DoesEventMatchDesc(const FStateTreeEvent& Event) const;
+};
+
 /**
  *  Runtime representation of a StateTree transition.
  */
@@ -498,10 +550,10 @@ struct STATETREEMODULE_API FCompactStateTransition
 	{
 		return !Delay.IsEmpty();
 	}
-	
-	/** Transition event tag, used when trigger type is event. */
+
+	/** Event Description */
 	UPROPERTY()
-	FGameplayTag EventTag;
+	FCompactEventDesc RequiredEvent;
 
 	/** Index to first condition to test */
 	UPROPERTY()
@@ -556,6 +608,10 @@ struct STATETREEMODULE_API FCompactStateTreeState
 	/** @return True if the state has any child states */
 	bool HasChildren() const { return ChildrenEnd > ChildrenBegin; }
 
+	/** Description of an event required to enter the state. */
+	UPROPERTY()
+	FCompactEventDesc RequiredEventToEnter;
+
 	/** Name of the State */
 	UPROPERTY()
 	FName Name;
@@ -601,6 +657,9 @@ struct STATETREEMODULE_API FCompactStateTreeState
 	UPROPERTY()
 	FStateTreeIndex16 ParameterBindingsBatch = FStateTreeIndex16::Invalid;
 
+	UPROPERTY()
+	FStateTreeIndex16 EventDataIndex = FStateTreeIndex16::Invalid;
+
 	/** Number of enter conditions */
 	UPROPERTY()
 	uint8 EnterConditionsNum = 0;
@@ -617,6 +676,10 @@ struct STATETREEMODULE_API FCompactStateTreeState
 	UPROPERTY()
 	uint8 InstanceDataNum = 0;
 
+	/** Distance to root state. */
+	UPROPERTY()
+	uint8 Depth = 0;
+
 	/** Type of the state */
 	UPROPERTY()
 	EStateTreeStateType Type = EStateTreeStateType::State;
@@ -1047,4 +1110,4 @@ struct TStructOpsTypeTraits<FStateTreeStateLink> : public TStructOpsTypeTraitsBa
 		WithStructuredSerializer = true,
 		WithPostSerialize = true,
 	};
-};
+};
\ No newline at end of file
diff --git a/Engine/Plugins/Runtime/StateTree/Source/StateTreeTestSuite/Private/StateTreeTest.cpp b/Engine/Plugins/Runtime/StateTree/Source/StateTreeTestSuite/Private/StateTreeTest.cpp
--- a/Engine/Plugins/Runtime/StateTree/Source/StateTreeTestSuite/Private/StateTreeTest.cpp
+++ b/Engine/Plugins/Runtime/StateTree/Source/StateTreeTestSuite/Private/StateTreeTest.cpp
@@ -51,11 +51,13 @@ namespace UE::StateTree::Tests
 		virtual ~FNativeGameplayTags() {}
 		
 		FGameplayTag TestTag;
+		FGameplayTag TestTag2;
 
 		virtual void AddTags() override
 		{
 			UGameplayTagsManager& Manager = UGameplayTagsManager::Get();
 			TestTag = Manager.AddNativeGameplayTag(TEXT("Test.StateTree.Tag"));
+			TestTag2 = Manager.AddNativeGameplayTag(TEXT("Test.StateTree.Tag2"));
 		}
 
 		FORCEINLINE static const FNativeGameplayTags& Get()
@@ -1070,7 +1072,7 @@ struct FStateTreeTest_TransitionDelay : FAITestBase
 		Transition.bDelayTransition = true;
 		Transition.DelayDuration = 0.15f;
 		Transition.DelayRandomVariance = 0.0f;
-		Transition.EventTag = Tag;
+		Transition.RequiredEvent.Tag = Tag;
 
 		// State B
 		auto& Task1 = StateB.AddTask<FTestTask_Stand>(FName(TEXT("Task1")));
@@ -1151,7 +1153,7 @@ struct FStateTreeTest_TransitionDelayZero : FAITestBase
 		Transition.bDelayTransition = true;
 		Transition.DelayDuration = 0.0f;
 		Transition.DelayRandomVariance = 0.0f;
-		Transition.EventTag = Tag;
+		Transition.RequiredEvent.Tag = Tag;
 
 		// State B
 		auto& Task1 = StateB.AddTask<FTestTask_Stand>(FName(TEXT("Task1")));
@@ -1192,6 +1194,181 @@ struct FStateTreeTest_TransitionDelayZero : FAITestBase
 };
 IMPLEMENT_AI_INSTANT_TEST(FStateTreeTest_TransitionDelayZero, "System.StateTree.TransitionDelayZero");
 
+struct FStateTreeTest_StateRequiringEvent : FAITestBase
+{
+	virtual bool InstantTest() override
+	{
+		UStateTree& StateTree = UE::StateTree::Tests::NewStateTree(&GetWorld());
+		UStateTreeEditorData& EditorData = *Cast<UStateTreeEditorData>(StateTree.EditorData);
+
+		UStateTreeState& Root = EditorData.AddSubTree(FName(TEXT("Root")));
+		
+		FGameplayTag ValidTag = UE::StateTree::Tests::FNativeGameplayTags::Get().TestTag;
+		FGameplayTag InvalidTag = UE::StateTree::Tests::FNativeGameplayTags::Get().TestTag2;
+
+		using FValidPayload = FStateTreeTest_PropertyStructA;
+		using FInvalidPayload = FStateTreeTest_PropertyStructB;
+
+		// This state shouldn't be selected as it requires different tag.
+		UStateTreeState& StateA = Root.AddChildState(FName(TEXT("A")));
+		StateA.bHasRequiredEventToEnter  = true;
+		StateA.RequiredEventToEnter.Tag = InvalidTag;
+		auto& TaskA = StateA.AddTask<FTestTask_Stand>(FName(TEXT("TaskA")));
+
+		// This state shouldn't be selected as it requires different payload.
+		UStateTreeState& StateB = Root.AddChildState(FName(TEXT("B")));
+		StateB.bHasRequiredEventToEnter  = true;
+		StateB.RequiredEventToEnter.PayloadStruct = FInvalidPayload::StaticStruct();
+		auto& TaskB = StateB.AddTask<FTestTask_Stand>(FName(TEXT("TaskB")));
+
+		// This state shouldn't be selected as it requires the same tag, but different payload.
+		UStateTreeState& StateC = Root.AddChildState(FName(TEXT("C")));
+		StateC.bHasRequiredEventToEnter  = true;
+		StateC.RequiredEventToEnter.Tag = ValidTag;
+		StateC.RequiredEventToEnter.PayloadStruct = FInvalidPayload::StaticStruct();
+		auto& TaskC = StateC.AddTask<FTestTask_Stand>(FName(TEXT("TaskC")));
+
+		// This state shouldn't be selected as it requires the same payload, but different tag.
+		UStateTreeState& StateD = Root.AddChildState(FName(TEXT("D")));
+		StateD.bHasRequiredEventToEnter  = true;
+		StateD.RequiredEventToEnter.Tag = InvalidTag;
+		StateD.RequiredEventToEnter.PayloadStruct = FValidPayload::StaticStruct();
+		auto& TaskD = StateD.AddTask<FTestTask_Stand>(FName(TEXT("TaskD")));
+
+		// This state should be selected as the arrived event matches the requirement.
+		UStateTreeState& StateE = Root.AddChildState(FName(TEXT("E")));
+		StateE.bHasRequiredEventToEnter  = true;
+		StateE.RequiredEventToEnter.Tag = ValidTag;
+		StateE.RequiredEventToEnter.PayloadStruct = FValidPayload::StaticStruct();
+		auto& TaskE = StateE.AddTask<FTestTask_Stand>(FName(TEXT("TaskE")));
+
+		// This state should be selected only initially when there's not event in the queue.
+		UStateTreeState& StateInitial = Root.AddChildState(FName(TEXT("Initial")));
+		auto& TaskInitial = StateInitial.AddTask<FTestTask_Stand>(FName(TEXT("TaskInitial")));
+		StateInitial.AddTransition(EStateTreeTransitionTrigger::OnEvent, ValidTag, EStateTreeTransitionType::GotoState, &Root);
+
+		FStateTreeCompilerLog Log;
+		FStateTreeCompiler Compiler(Log);
+		const bool bResult = Compiler.Compile(StateTree);
+
+		AITEST_TRUE("StateTree should get compiled", bResult);
+
+		EStateTreeRunStatus Status = EStateTreeRunStatus::Unset;
+		FStateTreeInstanceData InstanceData;
+		FTestStateTreeExecutionContext Exec(StateTree, StateTree, InstanceData);
+		const bool bInitSucceeded = Exec.IsValid();
+		AITEST_TRUE("StateTree should init", bInitSucceeded);
+
+		const FString EnterStateStr(TEXT("EnterState"));
+
+		Status = Exec.Start();
+		AITEST_TRUE("StateTree TaskInitial should not enter state", Exec.Expect(TaskInitial.GetName(), EnterStateStr));
+		Exec.SendEvent(ValidTag, FConstStructView::Make(FValidPayload()));
+		Status = Exec.Tick(0.1f);
+
+		AITEST_FALSE("StateTree TaskA should not enter state", Exec.Expect(TaskA.GetName(), EnterStateStr));
+		AITEST_FALSE("StateTree TaskB should not enter state", Exec.Expect(TaskB.GetName(), EnterStateStr));
+		AITEST_FALSE("StateTree TaskC should not enter state", Exec.Expect(TaskC.GetName(), EnterStateStr));
+		AITEST_FALSE("StateTree TaskD should not enter state", Exec.Expect(TaskD.GetName(), EnterStateStr));
+		AITEST_TRUE("StateTree TaskE should enter state", Exec.Expect(TaskE.GetName(), EnterStateStr));
+		Exec.LogClear();
+
+		return true;
+	}
+};
+IMPLEMENT_AI_INSTANT_TEST(FStateTreeTest_StateRequiringEvent, "System.StateTree.StateRequiringEvent");
+
+struct FStateTreeTest_PassingTransitionEventToStateSelection : FAITestBase
+{
+	virtual bool InstantTest() override
+	{
+		UStateTree& StateTree = UE::StateTree::Tests::NewStateTree(&GetWorld());
+		UStateTreeEditorData& EditorData = *Cast<UStateTreeEditorData>(StateTree.EditorData);
+
+		UStateTreeState& Root = EditorData.AddSubTree(FName(TEXT("Root")));
+
+		FStateTreePropertyPath PathToPayloadMember;
+		{
+			const bool bParseResult = PathToPayloadMember.FromString(TEXT("Payload.A"));
+
+			AITEST_TRUE("Parsing path should succeeed", bParseResult);
+
+			FStateTreeEvent EventWithPayload;
+			EventWithPayload.Payload = FInstancedStruct::Make<FStateTreeTest_PropertyStructA>();
+			const bool bUpdateSegments = PathToPayloadMember.UpdateSegmentsFromValue(FStateTreeDataView(FStructView::Make(EventWithPayload)));
+			AITEST_TRUE("Updating segments should succeeed", bUpdateSegments);
+		}
+
+		// This state shouldn't be selected, because transition's condition and state's enter condition exlude each other.
+		UStateTreeState& StateA = Root.AddChildState(FName(TEXT("A")));
+		StateA.bHasRequiredEventToEnter  = true;
+		StateA.RequiredEventToEnter.PayloadStruct = FStateTreeTest_PropertyStructA::StaticStruct();
+		auto& TaskA = StateA.AddTask<FTestTask_Stand>(FName(TEXT("TaskA")));
+		TStateTreeEditorNode<FStateTreeCompareIntCondition>& AIntCond = StateA.AddEnterCondition<FStateTreeCompareIntCondition>(EGenericAICheck::Equal);
+		AIntCond.GetInstanceData().Right = 0;
+		EditorData.AddPropertyBinding(
+			FStateTreePropertyPath(StateA.GetEventID(), PathToPayloadMember.GetSegments()),
+			FStateTreePropertyPath(AIntCond.ID, TEXT("Left")));
+
+		// This state should be selected as the sent event fullfils both transition's condition and state's enter condition.
+		UStateTreeState& StateB = Root.AddChildState(FName(TEXT("B")));
+		StateB.bHasRequiredEventToEnter  = true;
+		StateB.RequiredEventToEnter.PayloadStruct = FStateTreeTest_PropertyStructA::StaticStruct();
+		auto& TaskB = StateB.AddTask<FTestTask_Stand>(FName(TEXT("TaskB")));
+		TStateTreeEditorNode<FStateTreeCompareIntCondition>& BIntCond = StateB.AddEnterCondition<FStateTreeCompareIntCondition>(EGenericAICheck::Equal);
+		BIntCond.GetInstanceData().Right = 1;
+		EditorData.AddPropertyBinding(
+			FStateTreePropertyPath(StateB.GetEventID(), PathToPayloadMember.GetSegments()),
+			FStateTreePropertyPath(BIntCond.ID, TEXT("Left")));
+
+		// This state should be selected only initially when there's not event in the queue.
+		UStateTreeState& StateInitial = Root.AddChildState(FName(TEXT("Initial")));
+		auto& TaskInitial = StateInitial.AddTask<FTestTask_Stand>(FName(TEXT("TaskInitial")));
+		FStateTreeTransition& TransA = StateInitial.AddTransition(EStateTreeTransitionTrigger::OnEvent, FGameplayTag(), EStateTreeTransitionType::GotoState, &StateA);
+		TransA.RequiredEvent.PayloadStruct = FStateTreeTest_PropertyStructA::StaticStruct();
+		TStateTreeEditorNode<FStateTreeCompareIntCondition>& TransAIntCond = TransA.AddCondition<FStateTreeCompareIntCondition>(EGenericAICheck::Equal);
+		TransAIntCond.GetInstanceData().Right = 1;
+		EditorData.AddPropertyBinding(
+			FStateTreePropertyPath(TransA.GetEventID(), PathToPayloadMember.GetSegments()),
+			FStateTreePropertyPath(TransAIntCond.ID, TEXT("Left")));
+
+		FStateTreeTransition& TransB = StateInitial.AddTransition(EStateTreeTransitionTrigger::OnEvent, FGameplayTag(), EStateTreeTransitionType::GotoState, &StateB);
+		TransB.RequiredEvent.PayloadStruct = FStateTreeTest_PropertyStructA::StaticStruct();
+		TStateTreeEditorNode<FStateTreeCompareIntCondition>& TransBIntCond = TransB.AddCondition<FStateTreeCompareIntCondition>(EGenericAICheck::Equal);
+		TransBIntCond.GetInstanceData().Right = 1;
+		EditorData.AddPropertyBinding(
+			FStateTreePropertyPath(TransB.GetEventID(), PathToPayloadMember.GetSegments()),
+			FStateTreePropertyPath(TransBIntCond.ID, TEXT("Left")));
+
+		FStateTreeCompilerLog Log;
+		FStateTreeCompiler Compiler(Log);
+		const bool bResult = Compiler.Compile(StateTree);
+
+		AITEST_TRUE("StateTree should get compiled", bResult);
+
+		EStateTreeRunStatus Status = EStateTreeRunStatus::Unset;
+		FStateTreeInstanceData InstanceData;
+		FTestStateTreeExecutionContext Exec(StateTree, StateTree, InstanceData);
+		const bool bInitSucceeded = Exec.IsValid();
+		AITEST_TRUE("StateTree should init", bInitSucceeded);
+
+		const FString EnterStateStr(TEXT("EnterState"));
+
+		Status = Exec.Start();
+		AITEST_TRUE("StateTree TaskInitial should not enter state", Exec.Expect(TaskInitial.GetName(), EnterStateStr));
+		Exec.SendEvent(UE::StateTree::Tests::FNativeGameplayTags::Get().TestTag, FConstStructView::Make(FStateTreeTest_PropertyStructA{0}));
+		Exec.SendEvent(UE::StateTree::Tests::FNativeGameplayTags::Get().TestTag, FConstStructView::Make(FStateTreeTest_PropertyStructA{1}));
+		Status = Exec.Tick(0.1f);
+
+		AITEST_FALSE("StateTree TaskA should not enter state", Exec.Expect(TaskA.GetName(), EnterStateStr));
+		AITEST_TRUE("StateTree TaskB should enter state", Exec.Expect(TaskB.GetName(), EnterStateStr));
+		Exec.LogClear();
+
+		return true;
+	}
+};
+IMPLEMENT_AI_INSTANT_TEST(FStateTreeTest_PassingTransitionEventToStateSelection, "System.StateTree.PassingTransitionEventToStateSelection");
+
 struct FStateTreeTest_PropertyPathOffset : FAITestBase
 {
 	virtual bool InstantTest() override
diff --git a/Engine/Plugins/Runtime/StateTree/Source/StateTreeTestSuite/Private/StateTreeTestTypes.h b/Engine/Plugins/Runtime/StateTree/Source/StateTreeTestSuite/Private/StateTreeTestTypes.h
--- a/Engine/Plugins/Runtime/StateTree/Source/StateTreeTestSuite/Private/StateTreeTestTypes.h
+++ b/Engine/Plugins/Runtime/StateTree/Source/StateTreeTestSuite/Private/StateTreeTestTypes.h
@@ -369,6 +369,15 @@ struct FStateTreeTestRunContext
 };
 
 
+USTRUCT()
+struct FStateTreeTest_PropertyStructA
+{
+	GENERATED_BODY()
+
+	UPROPERTY(EditAnywhere, Category = "")
+	int32 A = 0;
+};
+
 USTRUCT()
 struct FStateTreeTest_PropertyStructB
 {
