StateTree: Fix template error

diff --git a/Engine/Plugins/Runtime/StateTree/Source/StateTreeModule/Public/StateTreeExecutionContext.h b/Engine/Plugins/Runtime/StateTree/Source/StateTreeModule/Public/StateTreeExecutionContext.h
--- a/Engine/Plugins/Runtime/StateTree/Source/StateTreeModule/Public/StateTreeExecutionContext.h
+++ b/Engine/Plugins/Runtime/StateTree/Source/StateTreeModule/Public/StateTreeExecutionContext.h
@@ -194,7 +194,7 @@ public:
 
 	/** Iterates over all events. Can only be used during StateTree tick. Expects a lambda which takes const FStateTreeSharedEvent& Event, and returns EStateTreeLoopEvents. */
 	template<typename TFunc>
-	TEnableIf<TIsInvocable<TFunc, FStateTreeSharedEvent>::Value, void>::Type ForEachEvent(TFunc&& Function) const
+	typename TEnableIf<TIsInvocable<TFunc, FStateTreeSharedEvent>::Value, void>::Type ForEachEvent(TFunc&& Function) const
 	{
 		for (const FStateTreeSharedEvent& Event : EventsToProcess)
 		{
@@ -207,7 +207,7 @@ public:
 
 	/** Iterates over all events. Can only be used during StateTree tick. Expects a lambda which takes const FStateTreeEvent& Event, and returns EStateTreeLoopEvents. Less preferable than FStateTreeSharedEvent version. */
 	template<typename TFunc>
-	TEnableIf<TIsInvocable<TFunc, FStateTreeEvent>::Value, void>::Type ForEachEvent(TFunc&& Function) const
+	typename TEnableIf<TIsInvocable<TFunc, FStateTreeEvent>::Value, void>::Type ForEachEvent(TFunc&& Function) const
 	{
 		for (const FStateTreeSharedEvent& Event : EventsToProcess)
 		{
