Fix state parameter and task parameter appearing twice in binding selection in state tree editor
#rnx

diff --git a/Engine/Plugins/Runtime/StateTree/Source/StateTreeEditorModule/Private/StateTreeEditorData.cpp b/Engine/Plugins/Runtime/StateTree/Source/StateTreeEditorModule/Private/StateTreeEditorData.cpp
--- a/Engine/Plugins/Runtime/StateTree/Source/StateTreeEditorModule/Private/StateTreeEditorData.cpp
+++ b/Engine/Plugins/Runtime/StateTree/Source/StateTreeEditorModule/Private/StateTreeEditorData.cpp
@@ -297,12 +297,6 @@ void UStateTreeEditorData::GetAccessibleStructs(const TConstArrayView<const USta
 					}
 
 					// Not at target yet, collect all bindable source accessible so far.
-					if (Desc.DataSource == EStateTreeBindableStructSource::Task
-						|| Desc.DataSource == EStateTreeBindableStructSource::StateParameter)
-					{
-						BindableDescs.Add(Desc);
-					}
-
 					switch (Desc.DataSource)
 					{
 						case EStateTreeBindableStructSource::StateParameter:
