Invalid EGenericAICheck value for State Tree Condition Evaluation.
Hid the value in the editor for unapplicable conditions.

diff --git a/Engine/Plugins/Runtime/StateTree/Source/StateTreeModule/Public/Conditions/StateTreeCommonConditions.h b/Engine/Plugins/Runtime/StateTree/Source/StateTreeModule/Public/Conditions/StateTreeCommonConditions.h
--- a/Engine/Plugins/Runtime/StateTree/Source/StateTreeModule/Public/Conditions/StateTreeCommonConditions.h
+++ b/Engine/Plugins/Runtime/StateTree/Source/StateTreeModule/Public/Conditions/StateTreeCommonConditions.h
@@ -44,7 +44,7 @@ struct STATETREEMODULE_API FStateTreeCompareIntCondition : public FStateTreeCond
 	UPROPERTY(EditAnywhere, Category = "Parameter")
 	bool bInvert = false;
 
-	UPROPERTY(EditAnywhere, Category = "Parameter")
+	UPROPERTY(EditAnywhere, Category = "Parameter", meta = (InvalidEnumValues = "IsTrue"))
 	EGenericAICheck Operator = EGenericAICheck::Equal;
 };
 
@@ -84,7 +84,7 @@ struct STATETREEMODULE_API FStateTreeCompareFloatCondition : public FStateTreeCo
 	UPROPERTY(EditAnywhere, Category = "Parameter")
 	bool bInvert = false;
 
-	UPROPERTY(EditAnywhere, Category = "Parameter")
+	UPROPERTY(EditAnywhere, Category = "Parameter", meta = (InvalidEnumValues = "IsTrue"))
 	EGenericAICheck Operator = EGenericAICheck::Equal;
 };
 
@@ -211,7 +211,7 @@ struct STATETREEMODULE_API FStateTreeCompareDistanceCondition : public FStateTre
 	UPROPERTY(EditAnywhere, Category = "Condition")
 	bool bInvert = false;
 
-	UPROPERTY(EditAnywhere, Category = "Condition")
+	UPROPERTY(EditAnywhere, Category = "Condition", meta = (InvalidEnumValues = "IsTrue"))
 	EGenericAICheck Operator = EGenericAICheck::Equal;
 };
 
