State Tree Editor UI Property Destruction triggers assertions in some functions in StateTreeUIDebuggerExtensions

diff --git a/Engine/Plugins/Runtime/StateTree/Source/StateTreeEditorModule/Private/Debugger/StateTreeDebuggerUIExtensions.cpp b/Engine/Plugins/Runtime/StateTree/Source/StateTreeEditorModule/Private/Debugger/StateTreeDebuggerUIExtensions.cpp
--- a/Engine/Plugins/Runtime/StateTree/Source/StateTreeEditorModule/Private/Debugger/StateTreeDebuggerUIExtensions.cpp
+++ b/Engine/Plugins/Runtime/StateTree/Source/StateTreeEditorModule/Private/Debugger/StateTreeDebuggerUIExtensions.cpp
@@ -209,31 +209,45 @@ ECheckBoxState GetTransitionBreakpointCheckState(const TSharedPtr<IPropertyHandl
 
 ECheckBoxState GetTransitionEnabledCheckState(const TSharedPtr<IPropertyHandle>& StructPropertyHandle)
 {
+	ECheckBoxState CommonState = ECheckBoxState::Checked;
+
+	if (!StructPropertyHandle.IsValid() || !StructPropertyHandle->IsValidHandle())
+	{
+		return CommonState;
+	}
+
 	const TSharedPtr<IPropertyHandle> EnabledProperty = StructPropertyHandle->GetChildHandle(GET_MEMBER_NAME_CHECKED(FStateTreeTransition, bTransitionEnabled));
 	if (ensureMsgf(EnabledProperty.IsValid() && EnabledProperty->IsValidHandle(),
-		TEXT("The property needs to be exposed to the Editor to allow us to use a property handle")))
+		TEXT("The property is missing keywords in its UPROPERTY macro to be exposed to the Editor, or doesn't have a UPROPERTY macro.")))
 	{
 		bool bEnabled = false;
 
 		switch (EnabledProperty->GetValue(bEnabled))
 		{
 		case FPropertyAccess::MultipleValues:
-			return ECheckBoxState::Undetermined;
+			CommonState = ECheckBoxState::Undetermined;
+			break;
 		case FPropertyAccess::Success:
-			return bEnabled ? ECheckBoxState::Checked : ECheckBoxState::Unchecked;
+			CommonState = bEnabled ? ECheckBoxState::Checked : ECheckBoxState::Unchecked;
+			break;
 		case FPropertyAccess::Fail:
 		default:
-			return ECheckBoxState::Checked;
+			break;
 		}
 	}
-	return ECheckBoxState::Checked;
+	return CommonState;
 }
 
 void OnTransitionEnableToggled(const TSharedPtr<IPropertyHandle>& StructPropertyHandle)
 {
+	if (!StructPropertyHandle.IsValid() || !StructPropertyHandle->IsValidHandle())
+	{
+		return;
+	}
+
 	const TSharedPtr<IPropertyHandle> EnabledProperty = StructPropertyHandle->GetChildHandle(GET_MEMBER_NAME_CHECKED(FStateTreeTransition, bTransitionEnabled));
 	if (ensureMsgf(EnabledProperty.IsValid() && EnabledProperty->IsValidHandle(),
-		TEXT("The property needs to be exposed to the Editor to allow us to use a property handle")))
+		TEXT("The property is missing keywords in its UPROPERTY macro to be exposed to the Editor, or doesn't have a UPROPERTY macro.")))
 	{
 		FScopedTransaction Transaction(LOCTEXT("OnStateEnableToggled", "Toggled State Enabled"));
 
