Fixing unreachable code shipping compile error

#rnx

diff --git a/Engine/Plugins/Runtime/StateTree/Source/StateTreeEditorModule/Private/SStateTreeViewRow.cpp b/Engine/Plugins/Runtime/StateTree/Source/StateTreeEditorModule/Private/SStateTreeViewRow.cpp
--- a/Engine/Plugins/Runtime/StateTree/Source/StateTreeEditorModule/Private/SStateTreeViewRow.cpp
+++ b/Engine/Plugins/Runtime/StateTree/Source/StateTreeEditorModule/Private/SStateTreeViewRow.cpp
@@ -1137,8 +1137,6 @@ FText SStateTreeViewRow::GetTransitionsIcon(const UStateTreeState& State, const
 		default:
 			return FText::GetEmpty();
 	}
-
-	return FText::GetEmpty();
 }
 
 EVisibility SStateTreeViewRow::GetTransitionsVisibility(const UStateTreeState& State, const EStateTreeTransitionTrigger Trigger) const
