Create state tree task Run Parallel tree allowing a state to run a subtree while child state can still to be selected and entered.

diff --git a/Engine/Plugins/Runtime/StateTree/Source/StateTreeEditorModule/Private/StateTreeState.cpp b/Engine/Plugins/Runtime/StateTree/Source/StateTreeEditorModule/Private/StateTreeState.cpp
--- a/Engine/Plugins/Runtime/StateTree/Source/StateTreeEditorModule/Private/StateTreeState.cpp
+++ b/Engine/Plugins/Runtime/StateTree/Source/StateTreeEditorModule/Private/StateTreeState.cpp
@@ -534,3 +534,15 @@ FStateTreeStateLink UStateTreeState::GetLinkToState() const
 	Link.ID = ID;
 	return Link;
 }
+
+TSubclassOf<UStateTreeSchema> UStateTreeState::GetSchema() const
+{
+	if (const UStateTreeEditorData* EditorData = GetTypedOuter<UStateTreeEditorData>())
+	{
+		if (EditorData->Schema)
+		{
+			return EditorData->Schema->GetClass();
+		}
+	}
+	return nullptr;
+}
diff --git a/Engine/Plugins/Runtime/StateTree/Source/StateTreeEditorModule/Public/StateTreeState.h b/Engine/Plugins/Runtime/StateTree/Source/StateTreeEditorModule/Public/StateTreeState.h
--- a/Engine/Plugins/Runtime/StateTree/Source/StateTreeEditorModule/Public/StateTreeState.h
+++ b/Engine/Plugins/Runtime/StateTree/Source/StateTreeEditorModule/Public/StateTreeState.h
@@ -2,6 +2,7 @@
 
 #pragma once
 
+#include "IStateTreeSchemaProvider.h"
 #include "StateTreeEditorNode.h"
 #include "StateTreeEditorTypes.h"
 #include "StateTreeEvents.h"
@@ -187,7 +188,7 @@ struct STATETREEEDITORMODULE_API FStateTreeStateParameters
  * Editor representation of a state in StateTree
  */
 UCLASS(BlueprintType, EditInlineNew, CollapseCategories)
-class STATETREEEDITORMODULE_API UStateTreeState : public UObject
+class STATETREEEDITORMODULE_API UStateTreeState : public UObject, public IStateTreeSchemaProvider
 {
 	GENERATED_BODY()
 
@@ -295,6 +296,14 @@ public:
 
 	// ~StateTree Builder API
 
+	// IStateTreeSchemaProvider API
+	/**
+	* @return Class of schema used by the state tree containing this state.
+	*/
+	virtual TSubclassOf<UStateTreeSchema> GetSchema() const override;
+
+	// ~IStateTreeSchemaProvider API
+
 	/** Display name of the State */
 	UPROPERTY(EditDefaultsOnly, Category = "State")
 	FName Name;
diff --git a/Engine/Plugins/Runtime/StateTree/Source/StateTreeModule/Private/StateTreeExecutionContext.cpp b/Engine/Plugins/Runtime/StateTree/Source/StateTreeModule/Private/StateTreeExecutionContext.cpp
--- a/Engine/Plugins/Runtime/StateTree/Source/StateTreeModule/Private/StateTreeExecutionContext.cpp
+++ b/Engine/Plugins/Runtime/StateTree/Source/StateTreeModule/Private/StateTreeExecutionContext.cpp
@@ -121,6 +121,25 @@ FStateTreeExecutionContext::FStateTreeExecutionContext(UObject& InOwner, const U
 }
 
 
+FStateTreeExecutionContext::FStateTreeExecutionContext(const FStateTreeExecutionContext& InContextToCopy, const UStateTree& InStateTree, FStateTreeInstanceData& InInstanceData)
+	: FStateTreeExecutionContext(InContextToCopy.Owner, InStateTree, InInstanceData, InContextToCopy.CollectExternalDataDelegate)
+{
+	const bool bIsSameSchema = RootStateTree.GetSchema()->GetClass() == InContextToCopy.GetStateTree()->GetSchema()->GetClass();
+	if (bIsSameSchema)
+	{
+		for (const FStateTreeExternalDataDesc& TargetDataDesc : GetContextDataDescs())
+		{
+			const int32 TargetIndex = TargetDataDesc.Handle.DataHandle.GetIndex();
+			ContextAndExternalDataViews[TargetIndex] = InContextToCopy.ContextAndExternalDataViews[TargetIndex];
+		}
+	}
+	else
+	{
+		STATETREE_LOG(Error, TEXT("%hs: '%s' using StateTree '%s' trying to run subtree '%s' but their schemas don't match"),
+			__FUNCTION__, *GetNameSafe(&Owner), *GetFullNameSafe(InContextToCopy.GetStateTree()), *GetFullNameSafe(&RootStateTree));
+	}
+}
+
 FStateTreeExecutionContext::~FStateTreeExecutionContext()
 {
 	// Mark external data indices as invalid
diff --git a/Engine/Plugins/Runtime/StateTree/Source/StateTreeModule/Private/Tasks/StateTreeRunParallelStateTreeTask.cpp b/Engine/Plugins/Runtime/StateTree/Source/StateTreeModule/Private/Tasks/StateTreeRunParallelStateTreeTask.cpp
new file mode 100644
--- /dev/null
+++ b/Engine/Plugins/Runtime/StateTree/Source/StateTreeModule/Private/Tasks/StateTreeRunParallelStateTreeTask.cpp
@@ -0,0 +1,72 @@
+// Copyright Epic Games, Inc. All Rights Reserved.
+
+#include "Tasks/StateTreeRunParallelStateTreeTask.h"
+
+#include "StateTreeExecutionContext.h"
+
+FStateTreeRunParallelStateTreeTask::FStateTreeRunParallelStateTreeTask()
+{
+	bShouldCopyBoundPropertiesOnTick = false;
+	bShouldCopyBoundPropertiesOnExitState = false;
+}
+
+EStateTreeRunStatus FStateTreeRunParallelStateTreeTask::EnterState(FStateTreeExecutionContext& Context, const FStateTreeTransitionResult& Transitions) const
+{
+	FInstanceDataType& InstanceData = Context.GetInstanceData(*this);
+	if (!InstanceData.StateTree.IsValid())
+	{
+		return EStateTreeRunStatus::Failed;
+	}
+
+	FStateTreeExecutionContext ParallelTreeContext(Context, *InstanceData.StateTree.GetStateTree(), InstanceData.TreeInstanceData);
+	if (!ParallelTreeContext.IsValid())
+	{
+		return EStateTreeRunStatus::Failed;
+	}
+
+	return ParallelTreeContext.Start(&InstanceData.StateTree.GetParameters());
+}
+
+EStateTreeRunStatus FStateTreeRunParallelStateTreeTask::Tick(FStateTreeExecutionContext& Context, const float DeltaTime) const
+{
+	FInstanceDataType& InstanceData = Context.GetInstanceData(*this);
+	if (!InstanceData.StateTree.IsValid())
+	{
+		return EStateTreeRunStatus::Failed;
+	}
+
+	FStateTreeExecutionContext ParallelTreeContext(Context, *InstanceData.StateTree.GetStateTree(), InstanceData.TreeInstanceData);
+	if (!ParallelTreeContext.IsValid())
+	{
+		return EStateTreeRunStatus::Failed;
+	}
+
+	return ParallelTreeContext.Tick(DeltaTime);
+}
+
+void FStateTreeRunParallelStateTreeTask::ExitState(FStateTreeExecutionContext& Context, const FStateTreeTransitionResult& Transition) const
+{
+	FInstanceDataType& InstanceData = Context.GetInstanceData(*this);
+	if (!InstanceData.StateTree.IsValid())
+	{
+		return;
+	}
+
+	FStateTreeExecutionContext ParallelTreeContext(Context, *InstanceData.StateTree.GetStateTree(), InstanceData.TreeInstanceData);
+	if (!ParallelTreeContext.IsValid())
+	{
+		return;
+	}
+
+	ParallelTreeContext.Stop();
+}
+
+#if WITH_EDITOR
+void FStateTreeRunParallelStateTreeTask::PostEditInstanceDataChangeChainProperty(const FPropertyChangedChainEvent& PropertyChangedEvent, FStateTreeDataView InstanceDataView)
+{
+	if (PropertyChangedEvent.Property && PropertyChangedEvent.Property->GetFName() == GET_MEMBER_NAME_CHECKED(FStateTreeRunParallelStateTreeTaskInstanceData, StateTree))
+	{
+		InstanceDataView.GetMutable<FInstanceDataType>().StateTree.SyncParameters();
+	}
+}
+#endif // WITH_EDITOR
\ No newline at end of file
diff --git a/Engine/Plugins/Runtime/StateTree/Source/StateTreeModule/Private/Tasks/StateTreeRunParallelStateTreeTask.h b/Engine/Plugins/Runtime/StateTree/Source/StateTreeModule/Private/Tasks/StateTreeRunParallelStateTreeTask.h
new file mode 100644
--- /dev/null
+++ b/Engine/Plugins/Runtime/StateTree/Source/StateTreeModule/Private/Tasks/StateTreeRunParallelStateTreeTask.h
@@ -0,0 +1,47 @@
+// Copyright Epic Games, Inc. All Rights Reserved.
+
+#pragma once
+
+#include "StateTreeInstanceData.h"
+#include "StateTreeReference.h"
+#include "StateTreeTaskBase.h"
+
+#include "StateTreeRunParallelStateTreeTask.generated.h"
+
+USTRUCT()
+struct FStateTreeRunParallelStateTreeTaskInstanceData
+{
+	GENERATED_BODY()
+
+	/** State tree and parameters that will be run when this task is started. */
+	UPROPERTY(EditAnywhere, Category = "Parameter")
+	FStateTreeReference StateTree;
+
+	UPROPERTY(Transient)
+	FStateTreeInstanceData TreeInstanceData;
+};
+
+/**
+* Task that will run another state tree in the current state while allowing the current tree to continue selection and process of child state.
+* It will succeed, fail or run depending on the result of the parallel tree.
+* Less efficient then Linked Asset state, it has the advantage of allowing multiple trees to run in parallel.
+*/
+USTRUCT(meta = (DisplayName = "Run Parallel Tree", Category = "Common"))
+struct FStateTreeRunParallelStateTreeTask : public FStateTreeTaskCommonBase
+{
+	GENERATED_BODY()
+	using FInstanceDataType = FStateTreeRunParallelStateTreeTaskInstanceData;
+	
+	FStateTreeRunParallelStateTreeTask();
+
+protected:
+	virtual const UStruct* GetInstanceDataType() const override { return FInstanceDataType::StaticStruct(); }
+
+	virtual EStateTreeRunStatus EnterState(FStateTreeExecutionContext& Context, const FStateTreeTransitionResult& Transitions) const override;
+	virtual EStateTreeRunStatus Tick(FStateTreeExecutionContext& Context, const float DeltaTime) const override;
+	virtual void ExitState(FStateTreeExecutionContext& Context, const FStateTreeTransitionResult& Transition) const override;
+
+#if WITH_EDITOR
+	virtual void PostEditInstanceDataChangeChainProperty(const FPropertyChangedChainEvent& PropertyChangedEvent, FStateTreeDataView InstanceDataView) override;
+#endif // WITH_EDITOR
+};
\ No newline at end of file
diff --git a/Engine/Plugins/Runtime/StateTree/Source/StateTreeModule/Public/StateTreeExecutionContext.h b/Engine/Plugins/Runtime/StateTree/Source/StateTreeModule/Public/StateTreeExecutionContext.h
--- a/Engine/Plugins/Runtime/StateTree/Source/StateTreeModule/Public/StateTreeExecutionContext.h
+++ b/Engine/Plugins/Runtime/StateTree/Source/StateTreeModule/Public/StateTreeExecutionContext.h
@@ -91,6 +91,8 @@ struct STATETREEMODULE_API FStateTreeExecutionContext
 {
 public:
 	FStateTreeExecutionContext(UObject& InOwner, const UStateTree& InStateTree, FStateTreeInstanceData& InInstanceData, const FOnCollectStateTreeExternalData& CollectExternalDataCallback = {});
+	/** Construct an execution context from a parent context and another tree. Useful to run a subtree from the parent context with the same schema. */
+	FStateTreeExecutionContext(const FStateTreeExecutionContext& InContextToCopy, const UStateTree& InStateTree, FStateTreeInstanceData& InInstanceData);
 	virtual ~FStateTreeExecutionContext();
 
 	/** Updates data view of the parameters by using the default values defined in the StateTree asset. */
