StateTree: Compilation fix when TStateTreePropertyRefExternalHandle references a single value

diff --git a/Engine/Plugins/Runtime/StateTree/Source/StateTreeModule/Public/StateTreePropertyRef.h b/Engine/Plugins/Runtime/StateTree/Source/StateTreeModule/Public/StateTreePropertyRef.h
--- a/Engine/Plugins/Runtime/StateTree/Source/StateTreeModule/Public/StateTreePropertyRef.h
+++ b/Engine/Plugins/Runtime/StateTree/Source/StateTreeModule/Public/StateTreePropertyRef.h
@@ -5,6 +5,7 @@
 #include "StateTreeIndexTypes.h"
 #include "StateTreeExecutionContext.h"
 #include "StateTreePropertyRefHelpers.h"
+#include "Templates/Tuple.h"
 #include "StateTreePropertyRef.generated.h"
 
 struct FStateTreePropertyRef;
@@ -235,7 +236,7 @@ struct FStateTreePropertyRefExternalHandle
 	template <class TRef>
 	TRef* GetMutablePtr() const
 	{
-		return GetMutablePtrTuple<TRef>().Key;
+		return GetMutablePtrTuple<TRef>().Get<0>();
 	}
 
 	/** @return a tuple of pointers of the given types to the property if possible, nullptr otherwise. */
@@ -257,7 +258,7 @@ struct FStateTreePropertyRefExternalHandle
 			return TTuple<TRef*...>{};
 		}
 
-		return UE::StateTree::PropertyRefHelpers::GetMutablePtrTupleToProperty<TRef...>(PropertyRef, InstanceStorage, Frame, ParentFrame);
+		return UE::StateTree::PropertyRefHelpers::GetMutablePtrTupleToProperty<TRef...>(PropertyRef, InstanceStorage, *Frame, ParentFrame);
 	}
 
 protected:
