StateTree: Fixed missing template keyword in dependent template call

diff --git a/Engine/Plugins/Runtime/StateTree/Source/StateTreeModule/Public/StateTreePropertyRef.h b/Engine/Plugins/Runtime/StateTree/Source/StateTreeModule/Public/StateTreePropertyRef.h
--- a/Engine/Plugins/Runtime/StateTree/Source/StateTreeModule/Public/StateTreePropertyRef.h
+++ b/Engine/Plugins/Runtime/StateTree/Source/StateTreeModule/Public/StateTreePropertyRef.h
@@ -236,7 +236,7 @@ struct FStateTreePropertyRefExternalHandle
 	template <class TRef>
 	TRef* GetMutablePtr() const
 	{
-		return GetMutablePtrTuple<TRef>().Get<0>();
+		return GetMutablePtrTuple<TRef>().template Get<0>();
 	}
 
 	/** @return a tuple of pointers of the given types to the property if possible, nullptr otherwise. */
