Add postload callback to state tree node in editor
Use callback to sync parameters in FStateTreeRunParallelStateTreeTask
#rnx

diff --git a/Engine/Plugins/Runtime/StateTree/Source/StateTreeEditorModule/Private/StateTreeEditorData.cpp b/Engine/Plugins/Runtime/StateTree/Source/StateTreeEditorModule/Private/StateTreeEditorData.cpp
--- a/Engine/Plugins/Runtime/StateTree/Source/StateTreeEditorModule/Private/StateTreeEditorData.cpp
+++ b/Engine/Plugins/Runtime/StateTree/Source/StateTreeEditorModule/Private/StateTreeEditorData.cpp
@@ -147,6 +147,7 @@ void UStateTreeEditorData::PostLoad()
 	FixObjectNodes();
 	FixDuplicateIDs();
 	UpdateBindingsInstanceStructs();
+	CallPostLoadOnNodes();
 }
 
 void UStateTreeEditorData::PostEditChangeChainProperty(FPropertyChangedChainEvent& PropertyChangedEvent)
@@ -758,6 +759,25 @@ void UStateTreeEditorData::UpdateBindingsInstanceStructs()
 	}
 }
 
+void UStateTreeEditorData::CallPostLoadOnNodes()
+{
+	for (FStateTreeEditorNode& EvaluatorEditorNode : Evaluators)
+	{
+		if (FStateTreeNodeBase* EvaluatorNode = EvaluatorEditorNode.Node.GetMutablePtr<FStateTreeNodeBase>())
+		{
+			EvaluatorNode->PostLoad(EvaluatorEditorNode.GetInstance());
+		}
+	}
+
+	for (FStateTreeEditorNode& GlobalTaskEditorNode : GlobalTasks)
+	{
+		if (FStateTreeNodeBase* GlobalTaskNode = GlobalTaskEditorNode.Node.GetMutablePtr<FStateTreeNodeBase>())
+		{
+			GlobalTaskNode->PostLoad(GlobalTaskEditorNode.GetInstance());
+		}
+	}
+}
+
 EStateTreeVisitor UStateTreeEditorData::VisitStateNodes(const UStateTreeState& State, TFunctionRef<EStateTreeVisitor(const UStateTreeState* State, const FGuid& ID, const FName& Name, const EStateTreeNodeType NodeType, const UScriptStruct* NodeStruct, const UStruct* InstanceStruct)> InFunc) const
 {
 	bool bContinue = true;
diff --git a/Engine/Plugins/Runtime/StateTree/Source/StateTreeEditorModule/Private/StateTreeState.cpp b/Engine/Plugins/Runtime/StateTree/Source/StateTreeEditorModule/Private/StateTreeState.cpp
--- a/Engine/Plugins/Runtime/StateTree/Source/StateTreeEditorModule/Private/StateTreeState.cpp
+++ b/Engine/Plugins/Runtime/StateTree/Source/StateTreeEditorModule/Private/StateTreeState.cpp
@@ -396,6 +396,29 @@ void UStateTreeState::PostLoad()
 	
 #endif // WITH_EDITORONLY_DATA
 
+#if WITH_EDITOR
+	for (FStateTreeEditorNode& EnterConditionEditorNode : EnterConditions)
+	{
+		if (FStateTreeNodeBase* ConditionNode = EnterConditionEditorNode.Node.GetMutablePtr<FStateTreeNodeBase>())
+		{
+			ConditionNode->PostLoad(EnterConditionEditorNode.GetInstance());
+		}
+	}
+
+	for (FStateTreeEditorNode& TaskEditorNode : Tasks)
+	{
+		if (FStateTreeNodeBase* TaskNode = TaskEditorNode.Node.GetMutablePtr<FStateTreeNodeBase>())
+		{
+			TaskNode->PostLoad(TaskEditorNode.GetInstance());
+		}
+	}
+
+	if (FStateTreeNodeBase* SingleTaskNode = SingleTask.Node.GetMutablePtr<FStateTreeNodeBase>())
+	{
+		SingleTaskNode->PostLoad(SingleTask.GetInstance());
+	}
+#endif // WITH_EDITOR
+
 }
 
 void UStateTreeState::UpdateParametersFromLinkedSubtree()
diff --git a/Engine/Plugins/Runtime/StateTree/Source/StateTreeEditorModule/Public/StateTreeEditorData.h b/Engine/Plugins/Runtime/StateTree/Source/StateTreeEditorModule/Public/StateTreeEditorData.h
--- a/Engine/Plugins/Runtime/StateTree/Source/StateTreeEditorModule/Public/StateTreeEditorData.h
+++ b/Engine/Plugins/Runtime/StateTree/Source/StateTreeEditorModule/Public/StateTreeEditorData.h
@@ -252,6 +252,7 @@ private:
 	void FixObjectNodes();
 	void FixDuplicateIDs();
 	void UpdateBindingsInstanceStructs();
+	void CallPostLoadOnNodes();
 
 #if WITH_EDITORONLY_DATA
 	FDelegateHandle OnObjectsReinstancedHandle;
diff --git a/Engine/Plugins/Runtime/StateTree/Source/StateTreeModule/Private/StateTree.cpp b/Engine/Plugins/Runtime/StateTree/Source/StateTreeModule/Private/StateTree.cpp
--- a/Engine/Plugins/Runtime/StateTree/Source/StateTreeModule/Private/StateTree.cpp
+++ b/Engine/Plugins/Runtime/StateTree/Source/StateTreeModule/Private/StateTree.cpp
@@ -331,6 +331,29 @@ void UStateTree::PostLoad()
 {
 	Super::PostLoad();
 
+#if WITH_EDITOR
+	for (int32 NodeIndex = 0; NodeIndex < Nodes.Num(); ++NodeIndex)
+	{
+		FStructView NodeView = Nodes[NodeIndex];
+		if (FStateTreeNodeBase* Node = NodeView.GetPtr<FStateTreeNodeBase>())
+		{
+			if (Node->InstanceTemplateIndex.IsValid())
+			{
+				FStateTreeInstanceData& SourceInstanceData = NodeView.GetPtr<FStateTreeConditionBase>() != nullptr ? SharedInstanceData : DefaultInstanceData;
+				if (SourceInstanceData.IsObject(Node->InstanceTemplateIndex.Get()))
+				{
+					Node->PostLoad(SourceInstanceData.GetMutableObject(Node->InstanceTemplateIndex.Get()));
+				}
+				else
+				{
+					Node->PostLoad(SourceInstanceData.GetMutableStruct(Node->InstanceTemplateIndex.Get()));
+				}
+			}
+			
+		}
+	}
+#endif // WITH_EDITOR
+
 	const int32 CurrentVersion = GetLinkerCustomVersion(FStateTreeCustomVersion::GUID);
 
 	if (CurrentVersion < FStateTreeCustomVersion::LatestVersion)
diff --git a/Engine/Plugins/Runtime/StateTree/Source/StateTreeModule/Private/Tasks/StateTreeRunParallelStateTreeTask.cpp b/Engine/Plugins/Runtime/StateTree/Source/StateTreeModule/Private/Tasks/StateTreeRunParallelStateTreeTask.cpp
--- a/Engine/Plugins/Runtime/StateTree/Source/StateTreeModule/Private/Tasks/StateTreeRunParallelStateTreeTask.cpp
+++ b/Engine/Plugins/Runtime/StateTree/Source/StateTreeModule/Private/Tasks/StateTreeRunParallelStateTreeTask.cpp
@@ -69,4 +69,12 @@ void FStateTreeRunParallelStateTreeTask::PostEditInstanceDataChangeChainProperty
 		InstanceDataView.GetMutable<FInstanceDataType>().StateTree.SyncParameters();
 	}
 }
+
+void FStateTreeRunParallelStateTreeTask::PostLoad(FStateTreeDataView InstanceDataView)
+{
+	if (FInstanceDataType* DataType = InstanceDataView.GetMutablePtr<FInstanceDataType>())
+	{
+		DataType->StateTree.SyncParameters();
+	}
+}
 #endif // WITH_EDITOR
\ No newline at end of file
diff --git a/Engine/Plugins/Runtime/StateTree/Source/StateTreeModule/Private/Tasks/StateTreeRunParallelStateTreeTask.h b/Engine/Plugins/Runtime/StateTree/Source/StateTreeModule/Private/Tasks/StateTreeRunParallelStateTreeTask.h
--- a/Engine/Plugins/Runtime/StateTree/Source/StateTreeModule/Private/Tasks/StateTreeRunParallelStateTreeTask.h
+++ b/Engine/Plugins/Runtime/StateTree/Source/StateTreeModule/Private/Tasks/StateTreeRunParallelStateTreeTask.h
@@ -43,5 +43,6 @@ protected:
 
 #if WITH_EDITOR
 	virtual void PostEditInstanceDataChangeChainProperty(const FPropertyChangedChainEvent& PropertyChangedEvent, FStateTreeDataView InstanceDataView) override;
+	virtual void PostLoad(FStateTreeDataView InstanceDataView) override;
 #endif // WITH_EDITOR
 };
\ No newline at end of file
diff --git a/Engine/Plugins/Runtime/StateTree/Source/StateTreeModule/Public/StateTreeNodeBase.h b/Engine/Plugins/Runtime/StateTree/Source/StateTreeModule/Public/StateTreeNodeBase.h
--- a/Engine/Plugins/Runtime/StateTree/Source/StateTreeModule/Public/StateTreeNodeBase.h
+++ b/Engine/Plugins/Runtime/StateTree/Source/StateTreeModule/Public/StateTreeNodeBase.h
@@ -76,6 +76,12 @@ PRAGMA_ENABLE_DEPRECATION_WARNINGS
 	 * @param InstanceData view to the instance data, can be struct or class.
 	 */
 	virtual void PostEditInstanceDataChangeChainProperty(const FPropertyChangedChainEvent& PropertyChangedEvent, FStateTreeDataView InstanceDataView) {}
+
+	/**
+	* Called after the state tree asset that contains this node is loaded from disk.
+	* @param InstanceData view to the instance data, can be struct or class.
+	*/
+	virtual void PostLoad(FStateTreeDataView InstanceDataView) {}
 #endif
 
 	/** Name of the node. */
