Fix LinkedAsset schema check that was inversed leading to subtree with looser requirement being rejected
#rnx

diff --git a/Engine/Plugins/Runtime/StateTree/Source/StateTreeModule/Private/StateTreeExecutionContext.cpp b/Engine/Plugins/Runtime/StateTree/Source/StateTreeModule/Private/StateTreeExecutionContext.cpp
--- a/Engine/Plugins/Runtime/StateTree/Source/StateTreeModule/Private/StateTreeExecutionContext.cpp
+++ b/Engine/Plugins/Runtime/StateTree/Source/StateTreeModule/Private/StateTreeExecutionContext.cpp
@@ -3325,7 +3325,7 @@ bool FStateTreeExecutionContext::SelectStateInternal(
 				}
 
 				// The linked state tree should have compatible context requirements.
-				if (!RootStateTree.HasCompatibleContextData(*NextState.LinkedAsset))
+				if (!NextState.LinkedAsset->HasCompatibleContextData(RootStateTree))
 				{
 					STATETREE_LOG(Error, TEXT("%hs: The linked State Tree '%s' does not have compatible schema, trying to select state %s from '%s'.  '%s' using StateTree '%s'."),
 						__FUNCTION__, *GetFullNameSafe(NextState.LinkedAsset), *GetSafeStateName(CurrentFrame, NextStateHandle), *GetStateStatusString(Exec), *GetNameSafe(&Owner), *GetFullNameSafe(CurrentFrame.StateTree));
