Add meta tag SchemaCanBeOverriden on StateTreeReference to make IStateSchemaProvider an opt in feature

diff --git a/Engine/Plugins/Runtime/GameplayStateTree/Source/GameplayStateTreeModule/Public/BehaviorTree/Tasks/BTTask_RunStateTree.h b/Engine/Plugins/Runtime/GameplayStateTree/Source/GameplayStateTreeModule/Public/BehaviorTree/Tasks/BTTask_RunStateTree.h
--- a/Engine/Plugins/Runtime/GameplayStateTree/Source/GameplayStateTreeModule/Public/BehaviorTree/Tasks/BTTask_RunStateTree.h
+++ b/Engine/Plugins/Runtime/GameplayStateTree/Source/GameplayStateTreeModule/Public/BehaviorTree/Tasks/BTTask_RunStateTree.h
@@ -32,7 +32,7 @@ protected:
 	virtual TSubclassOf<UStateTreeSchema> GetSchema() const override;
 
 	/** State tree that will be run when the task is selected. */
-	UPROPERTY(EditAnywhere, Category = Task, meta = (Schema = "/Script/GameplayStateTreeModule.StateTreeAIComponentSchema"))
+	UPROPERTY(EditAnywhere, Category = Task, meta = (Schema = "/Script/GameplayStateTreeModule.StateTreeAIComponentSchema", SchemaCanBeOverriden))
 	FStateTreeReference StateTreeRef;
 
 	UPROPERTY(Transient)
diff --git a/Engine/Plugins/Runtime/GameplayStateTree/Source/GameplayStateTreeModule/Public/Components/StateTreeComponent.h b/Engine/Plugins/Runtime/GameplayStateTree/Source/GameplayStateTreeModule/Public/Components/StateTreeComponent.h
--- a/Engine/Plugins/Runtime/GameplayStateTree/Source/GameplayStateTreeModule/Public/Components/StateTreeComponent.h
+++ b/Engine/Plugins/Runtime/GameplayStateTree/Source/GameplayStateTreeModule/Public/Components/StateTreeComponent.h
@@ -100,7 +100,7 @@ protected:
 	TObjectPtr<UStateTree> StateTree_DEPRECATED;
 #endif
 
-	UPROPERTY(EditAnywhere, Category = AI, meta=(Schema="/Script/GameplayStateTreeModule.StateTreeComponentSchema"))
+	UPROPERTY(EditAnywhere, Category = AI, meta=(Schema="/Script/GameplayStateTreeModule.StateTreeComponentSchema", SchemaCanBeOverriden))
 	FStateTreeReference StateTreeRef;
 
 	/** If true, the StateTree logic is started on begin play. Otherwise StartLogic() needs to be called. */
diff --git a/Engine/Plugins/Runtime/StateTree/Source/StateTreeEditorModule/Private/Customizations/StateTreeReferenceDetails.cpp b/Engine/Plugins/Runtime/StateTree/Source/StateTreeEditorModule/Private/Customizations/StateTreeReferenceDetails.cpp
--- a/Engine/Plugins/Runtime/StateTree/Source/StateTreeEditorModule/Private/Customizations/StateTreeReferenceDetails.cpp
+++ b/Engine/Plugins/Runtime/StateTree/Source/StateTreeEditorModule/Private/Customizations/StateTreeReferenceDetails.cpp
@@ -202,16 +202,19 @@ void FStateTreeReferenceDetails::SyncParameters(const UStateTree* StateTreeToSyn
 
 FString FStateTreeReferenceDetails::GetSchemaPath(IPropertyHandle& StructPropertyHandle)
 {
-	TArray<UObject*> OuterObjects;
-	StructPropertyHandle.GetOuterObjects(OuterObjects);
-
-	for (const UObject* OuterObject : OuterObjects)
+	if (StructPropertyHandle.HasMetaData(UE::StateTree::SchemaCanBeOverridenTag))
 	{
-		if (const IStateTreeSchemaProvider* SchemaProvider = Cast<IStateTreeSchemaProvider>(OuterObject))
+		TArray<UObject*> OuterObjects;
+		StructPropertyHandle.GetOuterObjects(OuterObjects);
+
+		for (const UObject* OuterObject : OuterObjects)
 		{
-			if (const UClass* SchemaClass = SchemaProvider->GetSchema().Get())
+			if (const IStateTreeSchemaProvider* SchemaProvider = Cast<IStateTreeSchemaProvider>(OuterObject))
 			{
-				return SchemaClass->GetPathName();
+				if (const UClass* SchemaClass = SchemaProvider->GetSchema().Get())
+				{
+					return SchemaClass->GetPathName();
+				}
 			}
 		}
 	}
diff --git a/Engine/Plugins/Runtime/StateTree/Source/StateTreeModule/Private/Tasks/StateTreeRunParallelStateTreeTask.h b/Engine/Plugins/Runtime/StateTree/Source/StateTreeModule/Private/Tasks/StateTreeRunParallelStateTreeTask.h
--- a/Engine/Plugins/Runtime/StateTree/Source/StateTreeModule/Private/Tasks/StateTreeRunParallelStateTreeTask.h
+++ b/Engine/Plugins/Runtime/StateTree/Source/StateTreeModule/Private/Tasks/StateTreeRunParallelStateTreeTask.h
@@ -14,7 +14,7 @@ struct FStateTreeRunParallelStateTreeTaskInstanceData
 	GENERATED_BODY()
 
 	/** State tree and parameters that will be run when this task is started. */
-	UPROPERTY(EditAnywhere, Category = "Parameter")
+	UPROPERTY(EditAnywhere, Category = "Parameter", meta=(SchemaCanBeOverriden))
 	FStateTreeReference StateTree;
 
 	UPROPERTY(Transient)
diff --git a/Engine/Plugins/Runtime/StateTree/Source/StateTreeModule/Public/IStateTreeSchemaProvider.h b/Engine/Plugins/Runtime/StateTree/Source/StateTreeModule/Public/IStateTreeSchemaProvider.h
--- a/Engine/Plugins/Runtime/StateTree/Source/StateTreeModule/Public/IStateTreeSchemaProvider.h
+++ b/Engine/Plugins/Runtime/StateTree/Source/StateTreeModule/Public/IStateTreeSchemaProvider.h
@@ -17,6 +17,10 @@ class UStateTreeSchemaProvider : public UInterface
 
 /**
 * Implementing this interface allows derived class to override the schema used to filter valid state trees for a FStateTreeReference.
+* The state tree reference property needs to be marked with SchemaCanBeOverriden metatag.
+* Ex:
+*	UPROPERTY(EditAnywhere, Category = AI, meta=(Schema="/Script/GameplayStateTreeModule.StateTreeComponentSchema", SchemaCanBeOverriden))
+*	FStateTreeReference StateTreeRef;
 */
 class IStateTreeSchemaProvider
 {
diff --git a/Engine/Plugins/Runtime/StateTree/Source/StateTreeModule/Public/StateTreeTypes.h b/Engine/Plugins/Runtime/StateTree/Source/StateTreeModule/Public/StateTreeTypes.h
--- a/Engine/Plugins/Runtime/StateTree/Source/StateTreeModule/Public/StateTreeTypes.h
+++ b/Engine/Plugins/Runtime/StateTree/Source/StateTreeModule/Public/StateTreeTypes.h
@@ -21,6 +21,8 @@ namespace UE::StateTree
 	inline constexpr int32 MaxConditionIndent = 4;
 
 	inline const FName SchemaTag(TEXT("Schema"));
+
+	inline const FName SchemaCanBeOverridenTag(TEXT("SchemaCanBeOverriden"));
 }; // UE::StateTree
 
 enum class EStateTreeRunStatus : uint8;
