[SStateTreeViewRow] removed unreachable code

diff --git a/Engine/Plugins/Runtime/StateTree/Source/StateTreeEditorModule/Private/SStateTreeViewRow.cpp b/Engine/Plugins/Runtime/StateTree/Source/StateTreeEditorModule/Private/SStateTreeViewRow.cpp
--- a/Engine/Plugins/Runtime/StateTree/Source/StateTreeEditorModule/Private/SStateTreeViewRow.cpp
+++ b/Engine/Plugins/Runtime/StateTree/Source/StateTreeEditorModule/Private/SStateTreeViewRow.cpp
@@ -1135,7 +1135,7 @@ FText SStateTreeViewRow::GetTransitionsIcon(const UStateTreeState& State, const
 		case IconLevelUp:
 			return FEditorFontGlyphs::Level_Up;
 		default:
-			return FText::GetEmpty();
+			break;
 	}
 }
 
