[SStateTreeViewRow] removed unreachable code

diff --git a/Engine/Plugins/Runtime/StateTree/Source/StateTreeEditorModule/Private/SStateTreeViewRow.cpp b/Engine/Plugins/Runtime/StateTree/Source/StateTreeEditorModule/Private/SStateTreeViewRow.cpp
--- a/Engine/Plugins/Runtime/StateTree/Source/StateTreeEditorModule/Private/SStateTreeViewRow.cpp
+++ b/Engine/Plugins/Runtime/StateTree/Source/StateTreeEditorModule/Private/SStateTreeViewRow.cpp
@@ -1126,17 +1126,20 @@ FText SStateTreeViewRow::GetTransitionsIcon(const UStateTreeState& State, const
 		IconType = IconLevelUp;
 	}
 
-	switch (IconType)
+	if (IconType == IconRightArrow)
 	{
-		case IconRightArrow:
-			return FEditorFontGlyphs::Long_Arrow_Right;
-		case IconDownArrow:
-			return FEditorFontGlyphs::Long_Arrow_Down;
-		case IconLevelUp:
-			return FEditorFontGlyphs::Level_Up;
-		default:
-			return FText::GetEmpty();
+		return FEditorFontGlyphs::Long_Arrow_Right;
 	}
+	else if (IconType == IconDownArrow)
+	{
+		return FEditorFontGlyphs::Long_Arrow_Down;
+	}
+	else if (IconType == IconLevelUp)
+	{
+		return FEditorFontGlyphs::Level_Up;
+	}
+
+	return FText::GetEmpty();
 }
 
 EVisibility SStateTreeViewRow::GetTransitionsVisibility(const UStateTreeState& State, const EStateTreeTransitionTrigger Trigger) const
