Browse to BP and Open BP Editor buttons for State Tree tasks/conditions
Hide buttons when inapplicable

diff --git a/Engine/Plugins/Runtime/StateTree/Source/StateTreeEditorModule/Private/Customizations/StateTreeEditorNodeDetails.cpp b/Engine/Plugins/Runtime/StateTree/Source/StateTreeEditorModule/Private/Customizations/StateTreeEditorNodeDetails.cpp
--- a/Engine/Plugins/Runtime/StateTree/Source/StateTreeEditorModule/Private/Customizations/StateTreeEditorNodeDetails.cpp
+++ b/Engine/Plugins/Runtime/StateTree/Source/StateTreeEditorModule/Private/Customizations/StateTreeEditorNodeDetails.cpp
@@ -9,6 +9,7 @@
 #include "StateTreeEditor.h"
 #include "StateTreeEditorData.h"
 #include "StateTreePropertyRef.h"
+#include "Widgets/Input/SButton.h"
 #include "Widgets/Input/SComboButton.h"
 #include "Framework/MultiBox/MultiBoxBuilder.h"
 #include "Styling/SlateIconFinder.h"
@@ -36,6 +37,7 @@
 #include "Widgets/Notifications/SNotificationList.h"
 #include "Debugger/StateTreeDebuggerUIExtensions.h"
 #include "StateTreePropertyBindings.h"
+#include "Engine/BlueprintGeneratedClass.h"
 
 #define LOCTEXT_NAMESPACE "StateTreeEditor"
 
@@ -629,7 +631,40 @@ void FStateTreeEditorNodeDetails::CustomizeHeader(TSharedRef<class IPropertyHand
 				[
 					StructPropertyHandle->CreateDefaultPropertyButtonWidgets()
 				]
-
+				// Browse To Button
+				+ SHorizontalBox::Slot()
+				.AutoWidth()
+				.VAlign(VAlign_Center)
+				[
+					SNew(SButton)
+					.ButtonStyle(FAppStyle::Get(), "SimpleButton")
+					.Visibility(this, &FStateTreeEditorNodeDetails::IsBrowseToNodeBlueprintVisible)
+					.OnClicked(this, &FStateTreeEditorNodeDetails::OnBrowseToNodeBlueprint)
+					.ToolTipText(LOCTEXT("BrowseToCurrentNodeBP", "Browse to the current node blueprint in Content Browser"))
+					.ContentPadding(.0f)
+					[
+						SNew(SImage)
+						.Image(FAppStyle::GetBrush("Icons.BrowseContent"))
+						.ColorAndOpacity(FSlateColor::UseForeground())
+					]
+				]
+				// Edit Button
+				+ SHorizontalBox::Slot()
+				.AutoWidth()
+				.VAlign(VAlign_Center)
+				[
+						SNew(SButton)
+						.ButtonStyle(FAppStyle::Get(), "SimpleButton")
+						.OnClicked(this, &FStateTreeEditorNodeDetails::OnEditNodeBlueprint)
+						.Visibility(this, &FStateTreeEditorNodeDetails::IsEditNodeBlueprintVisible)
+						.ToolTipText(LOCTEXT("EditCurrentNodeBP", "Edit the current node blueprint in Editor"))
+						.ContentPadding(.0f)
+						[
+							SNew(SImage)
+							.Image(FAppStyle::GetBrush("Icons.Edit"))
+							.ColorAndOpacity(FSlateColor::UseForeground())
+						]
+				]
 				+ SHorizontalBox::Slot()
 				.HAlign(HAlign_Right)
 				.FillWidth(1.0f)
@@ -1845,6 +1880,81 @@ void FStateTreeEditorNodeDetails::RestoreExpansionState()
 	}
 }
 
+FReply FStateTreeEditorNodeDetails::OnBrowseToNodeBlueprint() const
+{
+	const UObject* InstanceObject = nullptr;
+	const FPropertyAccess::Result AccessResult = InstanceObjectProperty->GetValue(InstanceObject);
+	if (AccessResult == FPropertyAccess::Success)
+	{
+		check(InstanceObject);
+		if (const UBlueprintGeneratedClass* BlueprintGeneratedClass = Cast<UBlueprintGeneratedClass>(InstanceObject->GetClass()))
+		{
+			//If the blueprint asset has been cooked, UBlueprint Object will be set to null and we need to browse to its BlueprintGeneratedClass
+			GEditor->SyncBrowserToObject(BlueprintGeneratedClass->ClassGeneratedBy ? BlueprintGeneratedClass->ClassGeneratedBy.Get() : BlueprintGeneratedClass);
+		}
+	}
+
+	return FReply::Handled();
+}
+
+FReply FStateTreeEditorNodeDetails::OnEditNodeBlueprint() const
+{
+	const UObject* InstanceObject = nullptr;
+	const FPropertyAccess::Result AccessResult = InstanceObjectProperty->GetValue(InstanceObject);
+	if (AccessResult == FPropertyAccess::Success)
+	{
+		check(InstanceObject);
+		const UBlueprintGeneratedClass* BlueprintGeneratedClass = Cast<UBlueprintGeneratedClass>(InstanceObject->GetClass());
+		if (BlueprintGeneratedClass && BlueprintGeneratedClass->ClassGeneratedBy)
+		{
+			//Cooked blueprint asset is not editable
+			GEditor->GetEditorSubsystem<UAssetEditorSubsystem>()->OpenEditorForAsset(BlueprintGeneratedClass->ClassGeneratedBy);
+		}
+	}
+
+	return FReply::Handled();
+}
+
+EVisibility FStateTreeEditorNodeDetails::IsBrowseToNodeBlueprintVisible() const
+{
+	const UObject* InstanceObject = nullptr;
+	const FPropertyAccess::Result AccessResult = InstanceObjectProperty->GetValue(InstanceObject);
+	if (AccessResult == FPropertyAccess::Success)
+	{
+		//The read could be null with an Success AccessResult in updating visibility 
+		if (InstanceObject)
+		{
+			if (const UBlueprintGeneratedClass* BlueprintGeneratedClass = Cast<UBlueprintGeneratedClass>(InstanceObject->GetClass()))
+			{
+				return EVisibility::Visible;
+			}
+		}
+	}
+
+	return EVisibility::Collapsed;
+}
+
+EVisibility FStateTreeEditorNodeDetails::IsEditNodeBlueprintVisible() const
+{
+	const UObject* InstanceObject = nullptr;
+	const FPropertyAccess::Result AccessResult = InstanceObjectProperty->GetValue(InstanceObject);
+	if (AccessResult == FPropertyAccess::Success)
+	{
+		//The read could be null with an Success AccessResult in updating visibility 
+		if (InstanceObject)
+		{
+			const UBlueprintGeneratedClass* BlueprintGeneratedClass = Cast<UBlueprintGeneratedClass>(InstanceObject->GetClass());
+			if (BlueprintGeneratedClass && BlueprintGeneratedClass->ClassGeneratedBy)
+			{
+				//Cooked blueprint asset is not editable
+				return EVisibility::Visible;
+			}
+		}
+	}
+
+	return EVisibility::Collapsed;
+}
+
 void FStateTreeEditorNodeDetails::ExpandAll(const TArray<TSharedPtr<FStateTreeNodeTypeItem>>& Items)
 {
 	for (const TSharedPtr<FStateTreeNodeTypeItem>& Item : Items)
diff --git a/Engine/Plugins/Runtime/StateTree/Source/StateTreeEditorModule/Private/Customizations/StateTreeEditorNodeDetails.h b/Engine/Plugins/Runtime/StateTree/Source/StateTreeEditorModule/Private/Customizations/StateTreeEditorNodeDetails.h
--- a/Engine/Plugins/Runtime/StateTree/Source/StateTreeEditorModule/Private/Customizations/StateTreeEditorNodeDetails.h
+++ b/Engine/Plugins/Runtime/StateTree/Source/StateTreeEditorModule/Private/Customizations/StateTreeEditorNodeDetails.h
@@ -122,6 +122,11 @@ private:
 	void SaveExpansionState();
 	void RestoreExpansionState();
 	
+	FReply OnBrowseToNodeBlueprint() const;
+	FReply OnEditNodeBlueprint() const;
+	EVisibility IsBrowseToNodeBlueprintVisible() const;
+	EVisibility IsEditNodeBlueprintVisible() const;
+
 	TSharedPtr<FStateTreeNodeTypeItem> RootNode;
 	TSharedPtr<FStateTreeNodeTypeItem> FilteredRootNode;
 	
