StateTree: Fix active instance data access when during selection.

diff --git a/Engine/Plugins/Runtime/StateTree/Source/StateTreeModule/Private/StateTreeExecutionContext.cpp b/Engine/Plugins/Runtime/StateTree/Source/StateTreeModule/Private/StateTreeExecutionContext.cpp
--- a/Engine/Plugins/Runtime/StateTree/Source/StateTreeModule/Private/StateTreeExecutionContext.cpp
+++ b/Engine/Plugins/Runtime/StateTree/Source/StateTreeModule/Private/StateTreeExecutionContext.cpp
@@ -1059,7 +1059,7 @@ bool FStateTreeExecutionContext::IsHandleSourceValid(const FStateTreeExecutionFr
 	case EStateTreeDataSourceType::ActiveInstanceData:
 	case EStateTreeDataSourceType::ActiveInstanceDataObject:
 		return CurrentFrame.ActiveInstanceIndexBase.IsValid()
-			&& CurrentFrame.ActiveStates.Contains(Handle.GetState())
+			&& CurrentFrame.ActiveStates.Contains(Handle.GetState(), CurrentFrame.NumCurrentlyActiveStates)
 			&& InstanceDataStorage->IsValidIndex(CurrentFrame.ActiveInstanceIndexBase.Get() + Handle.GetIndex());
 		
 	case EStateTreeDataSourceType::SharedInstanceData:
@@ -1084,7 +1084,7 @@ bool FStateTreeExecutionContext::IsHandleSourceValid(const FStateTreeExecutionFr
 
 	case EStateTreeDataSourceType::StateParameterData:
 		return CurrentFrame.ActiveInstanceIndexBase.IsValid()
-			&& CurrentFrame.ActiveStates.Contains(Handle.GetState())
+			&& CurrentFrame.ActiveStates.Contains(Handle.GetState(), CurrentFrame.NumCurrentlyActiveStates)
 			&& InstanceDataStorage->IsValidIndex(CurrentFrame.ActiveInstanceIndexBase.Get() + Handle.GetIndex());
 
 	case EStateTreeDataSourceType::TransitionEvent:
@@ -1466,6 +1466,7 @@ EStateTreeRunStatus FStateTreeExecutionContext::EnterState(FStateTreeTransitionR
 					__FUNCTION__, *GetStateStatusString(Exec), *GetNameSafe(&Owner), *GetFullNameSafe(&RootStateTree));
 				break;
 			}
+			CurrentFrame.NumCurrentlyActiveStates = static_cast<uint8>(CurrentFrame.ActiveStates.Num());
 
 			if (State.Type == EStateTreeStateType::Linked
 				|| State.Type == EStateTreeStateType::LinkedAsset)
@@ -3064,6 +3065,8 @@ bool FStateTreeExecutionContext::SelectState(const FStateTreeExecutionFrame& Cur
 	{
 		LastFrame.ActiveStates.Push(ParentStates[Index]);
 	}
+	// Existing state's data is safe to access during select.
+	LastFrame.NumCurrentlyActiveStates = static_cast<uint8>(LastFrame.ActiveStates.Num());
 
 	FStateSelectionResult InitialSelection;
 
@@ -3235,16 +3238,16 @@ bool FStateTreeExecutionContext::SelectStateInternal(
 		}
 
 		// Check if we're still tracking on the current active frame and state.
-		// If not, mark the ActiveInstanceIndexBase as invalid because we no longer can access the instance data during selection.
-		const FStateTreeIndex16 PrevActiveInstanceIndexBase = CurrentFrame.ActiveInstanceIndexBase; 
+		// If we are, update the NumCurrentlyActiveStates to indicate that this state's instance data can be accessed. 
+		const uint8 PrevNumCurrentlyActiveStates = CurrentFrame.NumCurrentlyActiveStates; 
 		if (CurrentFrame.ActiveInstanceIndexBase.IsValid()
 			&& CurrentFrameInActiveFrames)
 		{
 			const int32 CurrentStateIndex = CurrentFrame.ActiveStates.Num() - 1;
 			const FStateTreeStateHandle MatchingActiveHandle = CurrentFrameInActiveFrames->ActiveStates.GetStateSafe(CurrentStateIndex);
-			if (MatchingActiveHandle != NextStateHandle)
+			if (MatchingActiveHandle == NextStateHandle)
 			{
-				CurrentFrame.ActiveInstanceIndexBase = FStateTreeIndex16();
+				CurrentFrame.NumCurrentlyActiveStates = static_cast<uint8>(CurrentFrame.ActiveStates.Num());
 			}
 		}
 		
@@ -3541,7 +3544,7 @@ bool FStateTreeExecutionContext::SelectStateInternal(
 		}
 
 		// State could not be selected, restore.
-		CurrentFrame.ActiveInstanceIndexBase = PrevActiveInstanceIndexBase;
+		CurrentFrame.NumCurrentlyActiveStates = PrevNumCurrentlyActiveStates;
 		CurrentFrame.ActiveStates.Pop();
 	}
 
diff --git a/Engine/Plugins/Runtime/StateTree/Source/StateTreeModule/Public/StateTreeExecutionTypes.h b/Engine/Plugins/Runtime/StateTree/Source/StateTreeModule/Public/StateTreeExecutionTypes.h
--- a/Engine/Plugins/Runtime/StateTree/Source/StateTreeModule/Public/StateTreeExecutionTypes.h
+++ b/Engine/Plugins/Runtime/StateTree/Source/StateTreeModule/Public/StateTreeExecutionTypes.h
@@ -355,6 +355,20 @@ struct STATETREEMODULE_API FStateTreeActiveStates
 		return false;
 	}
 
+	/** Returns true of the array contains specified state within MaxNumStatesToCheck states. */
+	bool Contains(const FStateTreeStateHandle StateHandle, const uint8 MaxNumStatesToCheck) const
+	{
+		const int32 Num = (int32)FMath::Min(NumStates, MaxNumStatesToCheck);
+		for (int32 Index = 0; Index < Num; Index++)
+		{
+			if (States[Index] == StateHandle)
+			{
+				return true;
+			}
+		}
+		return false;
+	}
+
 	/** Returns index of a state, searching in reverse order. */
 	int32 IndexOfReverse(const FStateTreeStateHandle StateHandle) const
 	{
@@ -577,6 +591,9 @@ struct STATETREEMODULE_API FStateTreeExecutionFrame
 	UPROPERTY()
 	FStateTreeDataHandle GlobalParameterDataHandle = FStateTreeDataHandle::Invalid; 
 
+	/** Number of states in ActiveStates which have instance data. Used during state selection to decide which active state data is safe to access. */
+	uint8 NumCurrentlyActiveStates = 0;
+	
 	/** If true, the global tasks of the State Tree should be handle in this frame. */
 	UPROPERTY()
 	uint8 bIsGlobalFrame : 1 = false;
diff --git a/Engine/Plugins/Runtime/StateTree/Source/StateTreeTestSuite/Private/StateTreeTest.cpp b/Engine/Plugins/Runtime/StateTree/Source/StateTreeTestSuite/Private/StateTreeTest.cpp
--- a/Engine/Plugins/Runtime/StateTree/Source/StateTreeTestSuite/Private/StateTreeTest.cpp
+++ b/Engine/Plugins/Runtime/StateTree/Source/StateTreeTestSuite/Private/StateTreeTest.cpp
@@ -940,6 +940,63 @@ struct FStateTreeTest_TransitionNextSelectableState : FAITestBase
 };
 IMPLEMENT_AI_INSTANT_TEST(FStateTreeTest_TransitionNextSelectableState, "System.StateTree.Transition.NextSelectableState");
 
+
+struct FStateTreeTest_TransitionNextWithParentData : FAITestBase
+{
+	virtual bool InstantTest() override
+	{
+		UStateTree& StateTree = UE::StateTree::Tests::NewStateTree(&GetWorld());
+		UStateTreeEditorData& EditorData = *Cast<UStateTreeEditorData>(StateTree.EditorData);
+
+		UStateTreeState& Root =	EditorData.AddSubTree(FName(TEXT("Root")));
+		UStateTreeState& State0 = Root.AddChildState(FName(TEXT("State0")));
+		UStateTreeState& State1 = Root.AddChildState(FName(TEXT("State1")));
+		UStateTreeState& State1A = State1.AddChildState(FName(TEXT("State1A")));
+
+		auto& RootTask = Root.AddTask<FTestTask_B>(FName(TEXT("RootTask")));
+		RootTask.GetInstanceData().bBoolB = true;
+
+		auto& Task0 = State0.AddTask<FTestTask_Stand>(FName(TEXT("Task0")));
+		State0.AddTransition(EStateTreeTransitionTrigger::OnStateCompleted, EStateTreeTransitionType::NextState);
+
+		auto& Task1A = State1A.AddTask<FTestTask_Stand>(FName(TEXT("Task1A")));
+		auto& BoolCond1 = State1A.AddEnterCondition<FStateTreeCompareBoolCondition>();
+
+		EditorData.AddPropertyBinding(RootTask, TEXT("bBoolB"), BoolCond1, TEXT("bLeft"));
+		BoolCond1.GetInstanceData().bRight = true;
+
+		FStateTreeCompilerLog Log;
+		FStateTreeCompiler Compiler(Log);
+		const bool bResult = Compiler.Compile(StateTree);
+		AITEST_TRUE("StateTree should get compiled", bResult);
+
+		FStateTreeInstanceData InstanceData;
+		FTestStateTreeExecutionContext Exec(StateTree, StateTree, InstanceData);
+		const bool bInitSucceeded = Exec.IsValid();
+		AITEST_TRUE("StateTree should init", bInitSucceeded);
+
+		const FString TickStr(TEXT("Tick"));
+		const FString EnterStateStr(TEXT("EnterState"));
+		const FString ExitStateStr(TEXT("ExitState"));
+		const FString StateCompletedStr(TEXT("StateCompleted"));
+
+		// Start and enter state
+		Exec.Start();
+		AITEST_TRUE("StateTree Task0 should enter state", Exec.Expect(Task0.GetName(), EnterStateStr));
+		Exec.LogClear();
+
+		// Transition from State0 and tries to select State1.
+		// This tests that data from current shared active states (Root) is available during state selection.
+		Exec.Tick(0.1f);
+		AITEST_TRUE("StateTree Task0 should complete", Exec.Expect(Task0.GetName(), StateCompletedStr));
+		AITEST_TRUE("StateTree Task1A should enter state", Exec.Expect(Task1A.GetName(), EnterStateStr));
+		Exec.LogClear();
+
+		return true;
+	}
+};
+IMPLEMENT_AI_INSTANT_TEST(FStateTreeTest_TransitionNextWithParentData, "System.StateTree.Transition.NextWithParentData");
+
 struct FStateTreeTest_LastConditionWithIndent : FAITestBase
 {
 	virtual bool InstantTest() override
diff --git a/Engine/Plugins/Runtime/StateTree/Source/StateTreeTestSuite/Private/StateTreeTestTypes.h b/Engine/Plugins/Runtime/StateTree/Source/StateTreeTestSuite/Private/StateTreeTestTypes.h
--- a/Engine/Plugins/Runtime/StateTree/Source/StateTreeTestSuite/Private/StateTreeTestTypes.h
+++ b/Engine/Plugins/Runtime/StateTree/Source/StateTreeTestSuite/Private/StateTreeTestTypes.h
@@ -147,9 +147,17 @@ struct FTestTask_B : public FStateTreeTaskBase
 	using FInstanceDataType = FTestTask_BInstanceData;
 
 	FTestTask_B() = default;
+	FTestTask_B(const FName InName) { Name = InName; }
 	virtual ~FTestTask_B() override {}
 	
 	virtual const UStruct* GetInstanceDataType() const override { return FInstanceDataType::StaticStruct(); }
+
+	virtual EStateTreeRunStatus EnterState(FStateTreeExecutionContext& Context, const FStateTreeTransitionResult& Transition) const override
+	{
+		FTestStateTreeExecutionContext& TestContext = static_cast<FTestStateTreeExecutionContext&>(Context);
+		TestContext.Log(Name,  TEXT("EnterState"));
+		return EStateTreeRunStatus::Running;
+	}
 };
 
 USTRUCT()
