StateTree: Make sure PostLoad() is called on objects that are required for automatic compiling on version change.

diff --git a/Engine/Plugins/Runtime/StateTree/Source/StateTreeEditorModule/Private/StateTreeEditorData.cpp b/Engine/Plugins/Runtime/StateTree/Source/StateTreeEditorModule/Private/StateTreeEditorData.cpp
--- a/Engine/Plugins/Runtime/StateTree/Source/StateTreeEditorModule/Private/StateTreeEditorData.cpp
+++ b/Engine/Plugins/Runtime/StateTree/Source/StateTreeEditorModule/Private/StateTreeEditorData.cpp
@@ -143,6 +143,18 @@ void UStateTreeEditorData::OnParametersChanged(const UStateTree& StateTree)
 void UStateTreeEditorData::PostLoad()
 {
 	Super::PostLoad();
+
+	// Ensure the schema and states have had their PostLoad() fixed applied as we may need them in the later calls (or StateTree compile which might be calling this).
+	if (Schema)
+	{
+		Schema->ConditionalPostLoad();
+	}
+	VisitHierarchy([](UStateTreeState& State, UStateTreeState* ParentState) mutable 
+	{
+		State.ConditionalPostLoad();
+		return EStateTreeVisitor::Continue;
+	});
+
 	ReparentStates();
 	FixObjectNodes();
 	FixDuplicateIDs();
