StateTree: Added linked state overrides.
- Added ability to describe override table for linked state tree references
- This allows to override parts of the tree e.g. based on character type
- Fixed visibility of the State parameter button when in fixed mode

diff --git a/Engine/Plugins/Runtime/GameplayStateTree/Source/GameplayStateTreeModule/Private/Components/StateTreeComponent.cpp b/Engine/Plugins/Runtime/GameplayStateTree/Source/GameplayStateTreeModule/Private/Components/StateTreeComponent.cpp
--- a/Engine/Plugins/Runtime/GameplayStateTree/Source/GameplayStateTreeModule/Private/Components/StateTreeComponent.cpp
+++ b/Engine/Plugins/Runtime/GameplayStateTree/Source/GameplayStateTreeModule/Private/Components/StateTreeComponent.cpp
@@ -65,6 +65,7 @@ bool UStateTreeComponent::CollectExternalData(const FStateTreeExecutionContext&
 
 bool UStateTreeComponent::SetContextRequirements(FStateTreeExecutionContext& Context, bool bLogErrors)
 {
+	Context.SetLinkedStateTreeOverrides(&LinkedStateTreeOverrides);
 	Context.SetCollectExternalDataCallback(FOnCollectStateTreeExternalData::CreateUObject(this, &UStateTreeComponent::CollectExternalData));
 	return UStateTreeComponentSchema::SetContextRequirements(*this, Context);
 }
diff --git a/Engine/Plugins/Runtime/GameplayStateTree/Source/GameplayStateTreeModule/Public/Components/StateTreeComponent.h b/Engine/Plugins/Runtime/GameplayStateTree/Source/GameplayStateTreeModule/Public/Components/StateTreeComponent.h
--- a/Engine/Plugins/Runtime/GameplayStateTree/Source/GameplayStateTreeModule/Public/Components/StateTreeComponent.h
+++ b/Engine/Plugins/Runtime/GameplayStateTree/Source/GameplayStateTreeModule/Public/Components/StateTreeComponent.h
@@ -100,9 +100,17 @@ protected:
 	TObjectPtr<UStateTree> StateTree_DEPRECATED;
 #endif
 
+	/** State Tree asset to run on the component. */
 	UPROPERTY(EditAnywhere, Category = AI, meta=(Schema="/Script/GameplayStateTreeModule.StateTreeComponentSchema", SchemaCanBeOverriden))
 	FStateTreeReference StateTreeRef;
 
+	/**
+	 * Overrides for linked State Trees. This table is used to override State Tree references on linked states.
+	 * If a linked state's tag is exact match of the tag specified on the table, the reference from the table is used instead.
+	 */
+	UPROPERTY(EditAnywhere, Category = AI, meta=(Schema="/Script/GameplayStateTreeModule.StateTreeComponentSchema"))
+	FStateTreeReferenceOverrides LinkedStateTreeOverrides;
+	
 	/** If true, the StateTree logic is started on begin play. Otherwise StartLogic() needs to be called. */
 	UPROPERTY(EditAnywhere, Category = AI)
 	bool bStartLogicAutomatically = true;
diff --git a/Engine/Plugins/Runtime/StateTree/Source/StateTreeEditorModule/Private/Customizations/StateTreeReferenceOverridesDetails.cpp b/Engine/Plugins/Runtime/StateTree/Source/StateTreeEditorModule/Private/Customizations/StateTreeReferenceOverridesDetails.cpp
new file mode 100644
--- /dev/null
+++ b/Engine/Plugins/Runtime/StateTree/Source/StateTreeEditorModule/Private/Customizations/StateTreeReferenceOverridesDetails.cpp
@@ -0,0 +1,68 @@
+﻿// Copyright Epic Games, Inc. All Rights Reserved.
+
+#include "StateTreeReferenceOverridesDetails.h"
+#include "IDetailChildrenBuilder.h"
+#include "PropertyCustomizationHelpers.h"
+#include "StateTreeReference.h"
+#include "ScopedTransaction.h"
+
+#define LOCTEXT_NAMESPACE "StateTreeEditor"
+
+TSharedRef<IPropertyTypeCustomization> FStateTreeReferenceOverridesDetails::MakeInstance()
+{
+	return MakeShareable(new FStateTreeReferenceOverridesDetails);
+}
+
+void FStateTreeReferenceOverridesDetails::CustomizeHeader(const TSharedRef<IPropertyHandle> InStructPropertyHandle, FDetailWidgetRow& InHeaderRow, IPropertyTypeCustomizationUtils& InCustomizationUtils)
+{
+	StructPropertyHandle = InStructPropertyHandle;
+	PropUtils = InCustomizationUtils.GetPropertyUtilities();
+
+	OverrideItemsHandle = StructPropertyHandle->GetChildHandle(GET_MEMBER_NAME_CHECKED(FStateTreeReferenceOverrides, OverrideItems));
+	check(OverrideItemsHandle);
+	
+	InHeaderRow
+	.NameContent()
+	[
+		StructPropertyHandle->CreatePropertyNameWidget()
+	]
+	.ValueContent()
+	[
+		OverrideItemsHandle->CreatePropertyValueWidget()
+	]
+	.ShouldAutoExpand(true);
+}
+
+void FStateTreeReferenceOverridesDetails::CustomizeChildren(const TSharedRef<IPropertyHandle> InStructPropertyHandle, IDetailChildrenBuilder& InChildrenBuilder, IPropertyTypeCustomizationUtils& InCustomizationUtils)
+{
+	check(OverrideItemsHandle);
+	
+	const TSharedRef<FDetailArrayBuilder> NestedTreeOverridesBuilder = MakeShareable(new FDetailArrayBuilder(OverrideItemsHandle.ToSharedRef(), /*InGenerateHeader*/false, /*InDisplayResetToDefault*/ true, /*InDisplayElementNum*/ false));
+	NestedTreeOverridesBuilder->OnGenerateArrayElementWidget(FOnGenerateArrayElementWidget::CreateLambda(
+		[](TSharedRef<IPropertyHandle> PropertyHandle, int32 ArrayIndex, IDetailChildrenBuilder& ChildrenBuilder)
+		{
+
+			TSharedPtr<IPropertyHandle> StateTagHandle = PropertyHandle->GetChildHandle(GET_MEMBER_NAME_CHECKED(FStateTreeReferenceOverrideItem, StateTag));
+			check(StateTagHandle);
+			TSharedPtr<IPropertyHandle> StateTreeReferenceHandle = PropertyHandle->GetChildHandle(GET_MEMBER_NAME_CHECKED(FStateTreeReferenceOverrideItem, StateTreeReference));
+			check(StateTreeReferenceHandle);
+
+			IDetailPropertyRow& PropertyRow = ChildrenBuilder.AddProperty(StateTreeReferenceHandle.ToSharedRef());
+
+			TSharedPtr<SWidget> NameWidget;
+			TSharedPtr<SWidget> ValueWidget;
+			PropertyRow.GetDefaultWidgets(NameWidget, ValueWidget, /*bAddWidgetDecoration*/true);
+
+			PropertyRow.CustomWidget(/*bShowChildren*/true)
+			.NameContent()
+			[
+				StateTagHandle->CreatePropertyValueWidgetWithCustomization(nullptr)
+			]
+			.ValueContent()
+			[
+				ValueWidget.ToSharedRef()
+			];
+		}));
+	InChildrenBuilder.AddCustomBuilder(NestedTreeOverridesBuilder);}
+
+#undef LOCTEXT_NAMESPACE
diff --git a/Engine/Plugins/Runtime/StateTree/Source/StateTreeEditorModule/Private/Customizations/StateTreeReferenceOverridesDetails.h b/Engine/Plugins/Runtime/StateTree/Source/StateTreeEditorModule/Private/Customizations/StateTreeReferenceOverridesDetails.h
new file mode 100644
--- /dev/null
+++ b/Engine/Plugins/Runtime/StateTree/Source/StateTreeEditorModule/Private/Customizations/StateTreeReferenceOverridesDetails.h
@@ -0,0 +1,28 @@
+﻿// Copyright Epic Games, Inc. All Rights Reserved.
+
+#pragma once
+
+#include "IPropertyTypeCustomization.h"
+
+class IPropertyHandle;
+class IPropertyUtilities;
+
+/**
+ * Type customization for FNestedStateTreeReference.
+ */
+class FStateTreeReferenceOverridesDetails : public IPropertyTypeCustomization
+{
+public:
+	/** Makes a new instance of this detail layout class for a specific detail view requesting it */
+	static TSharedRef<IPropertyTypeCustomization> MakeInstance();
+
+	/** IPropertyTypeCustomization interface */
+	virtual void CustomizeHeader(TSharedRef<IPropertyHandle> InStructPropertyHandle, FDetailWidgetRow& InHeaderRow, IPropertyTypeCustomizationUtils& InCustomizationUtils) override;
+	virtual void CustomizeChildren(TSharedRef<IPropertyHandle> InStructPropertyHandle, IDetailChildrenBuilder& InChildrenBuilder, IPropertyTypeCustomizationUtils& InCustomizationUtils) override;
+
+private:
+	TSharedPtr<IPropertyHandle> StructPropertyHandle;
+	TSharedPtr<IPropertyHandle> OverrideItemsHandle;
+	
+	TSharedPtr<IPropertyUtilities> PropUtils;
+};
diff --git a/Engine/Plugins/Runtime/StateTree/Source/StateTreeEditorModule/Private/Customizations/StateTreeStateDetails.cpp b/Engine/Plugins/Runtime/StateTree/Source/StateTreeEditorModule/Private/Customizations/StateTreeStateDetails.cpp
--- a/Engine/Plugins/Runtime/StateTree/Source/StateTreeEditorModule/Private/Customizations/StateTreeStateDetails.cpp
+++ b/Engine/Plugins/Runtime/StateTree/Source/StateTreeEditorModule/Private/Customizations/StateTreeStateDetails.cpp
@@ -36,8 +36,12 @@ void FStateTreeStateDetails::CustomizeDetails(IDetailLayoutBuilder& DetailBuilde
 		}
 	}
 	const UStateTreeSchema* Schema = EditorData ? EditorData->Schema : nullptr;
-
+	const FString SchemaPath = Schema ? Schema->GetClass()->GetPathName() : FString();
+	
 	const TSharedPtr<IPropertyHandle> IDProperty = DetailBuilder.GetProperty(GET_MEMBER_NAME_CHECKED(UStateTreeState, ID));
+	const TSharedPtr<IPropertyHandle> NameProperty = DetailBuilder.GetProperty(GET_MEMBER_NAME_CHECKED(UStateTreeState, Name));
+	const TSharedPtr<IPropertyHandle> TagProperty = DetailBuilder.GetProperty(GET_MEMBER_NAME_CHECKED(UStateTreeState, Tag));
+	const TSharedPtr<IPropertyHandle> ColorRefProperty = DetailBuilder.GetProperty(GET_MEMBER_NAME_CHECKED(UStateTreeState, ColorRef));
 	const TSharedPtr<IPropertyHandle> EnabledProperty = DetailBuilder.GetProperty(GET_MEMBER_NAME_CHECKED(UStateTreeState, bEnabled));
 	const TSharedPtr<IPropertyHandle> TasksProperty = DetailBuilder.GetProperty(GET_MEMBER_NAME_CHECKED(UStateTreeState, Tasks));
 	const TSharedPtr<IPropertyHandle> SingleTaskProperty = DetailBuilder.GetProperty(GET_MEMBER_NAME_CHECKED(UStateTreeState, SingleTask));
@@ -48,8 +52,14 @@ void FStateTreeStateDetails::CustomizeDetails(IDetailLayoutBuilder& DetailBuilde
 	const TSharedPtr<IPropertyHandle> LinkedAssetProperty = DetailBuilder.GetProperty(GET_MEMBER_NAME_CHECKED(UStateTreeState, LinkedAsset));
 	const TSharedPtr<IPropertyHandle> ParametersProperty = DetailBuilder.GetProperty(GET_MEMBER_NAME_CHECKED(UStateTreeState, Parameters));
 	const TSharedPtr<IPropertyHandle> SelectionBehaviorProperty = DetailBuilder.GetProperty(GET_MEMBER_NAME_CHECKED(UStateTreeState, SelectionBehavior));
+	const TSharedPtr<IPropertyHandle> RequiredEventToEnterProperty = DetailBuilder.GetProperty(GET_MEMBER_NAME_CHECKED(UStateTreeState, RequiredEventToEnter));
 
+	
+
+	// Never show enabled
 	EnabledProperty->MarkHiddenByCustomization();
+
+	// Show ID only for debugging
 	if (UE::StateTree::Editor::GbDisplayItemIds == false)
 	{
 		IDProperty->MarkHiddenByCustomization();
@@ -58,26 +68,72 @@ void FStateTreeStateDetails::CustomizeDetails(IDetailLayoutBuilder& DetailBuilde
 	uint8 StateTypeValue = 0;
 	TypeProperty->GetValue(StateTypeValue);
 	const EStateTreeStateType StateType = (EStateTreeStateType)StateTypeValue;
-	
+
 	IDetailCategoryBuilder& StateCategory = DetailBuilder.EditCategory(TEXT("State"), LOCTEXT("StateDetailsState", "State"));
 	StateCategory.SetSortOrder(0);
 
 	StateCategory.HeaderContent(UE::StateTreeEditor::DebuggerExtensions::CreateStateWidget(DetailBuilder, EditorData));
 
-	if (StateType != EStateTreeStateType::Linked)
+	// Name
+	NameProperty->MarkHiddenByCustomization();
+	StateCategory.AddProperty(NameProperty);
+
+	// Tag
+	TagProperty->MarkHiddenByCustomization();
+	StateCategory.AddProperty(TagProperty);
+
+	// Color
+	ColorRefProperty->MarkHiddenByCustomization();
+	StateCategory.AddProperty(ColorRefProperty);
+
+	// Type
+	TypeProperty->MarkHiddenByCustomization();
+	StateCategory.AddProperty(TypeProperty);
+
+	// Per state type properties
+	SelectionBehaviorProperty->MarkHiddenByCustomization();
+	LinkedSubtreeProperty->MarkHiddenByCustomization();
+	LinkedAssetProperty->MarkHiddenByCustomization();
+
+	if (StateType == EStateTreeStateType::State)
 	{
-		LinkedSubtreeProperty->MarkHiddenByCustomization();
+		StateCategory.AddProperty(SelectionBehaviorProperty);
 	}
-	if (StateType != EStateTreeStateType::LinkedAsset)
+	else if (StateType == EStateTreeStateType::Linked)
 	{
-		LinkedAssetProperty->MarkHiddenByCustomization();
+		StateCategory.AddProperty(LinkedSubtreeProperty);
 	}
-	
-	if (StateType == EStateTreeStateType::Linked || StateType == EStateTreeStateType::LinkedAsset)
+	else if (StateType == EStateTreeStateType::LinkedAsset)
 	{
-		SelectionBehaviorProperty->MarkHiddenByCustomization();
+		// Custom widget for the linked asset, to add filtering to the assets.
+		IDetailPropertyRow& Row = StateCategory.AddProperty(LinkedAssetProperty);
+		Row.CustomWidget()
+		.NameContent()
+		[
+			LinkedAssetProperty->CreatePropertyNameWidget()
+		]
+		.ValueContent()
+		[
+			SNew(SObjectPropertyEntryBox)
+			.PropertyHandle(LinkedAssetProperty)
+			.AllowedClass(UStateTree::StaticClass())
+			.ThumbnailPool(DetailBuilder.GetPropertyUtilities()->GetThumbnailPool())
+			.OnShouldFilterAsset_Lambda([SchemaPath](const FAssetData& InAssetData)
+			{
+				return !SchemaPath.IsEmpty() && !InAssetData.TagsAndValues.ContainsKeyValue(UE::StateTree::SchemaTag, SchemaPath);
+			})
+		];
 	}
-	
+
+	// Parameters
+	ParametersProperty->MarkHiddenByCustomization();
+	StateCategory.AddProperty(ParametersProperty);
+
+	// Event
+	RequiredEventToEnterProperty->MarkHiddenByCustomization();
+	StateCategory.AddProperty(RequiredEventToEnterProperty);
+
+	// Enter conditions
 	const FName EnterConditionsCategoryName(TEXT("Enter Conditions"));
 	if (Schema && Schema->AllowEnterConditions())
 	{
@@ -88,6 +144,7 @@ void FStateTreeStateDetails::CustomizeDetails(IDetailLayoutBuilder& DetailBuilde
 		DetailBuilder.EditCategory(EnterConditionsCategoryName).SetCategoryVisibility(false);
 	}
 
+	// Tasks
 	if ((StateType == EStateTreeStateType::State || StateType == EStateTreeStateType::Subtree))
 	{
 		if (Schema && Schema->AllowMultipleTasks())
@@ -114,6 +171,7 @@ void FStateTreeStateDetails::CustomizeDetails(IDetailLayoutBuilder& DetailBuilde
 		TasksProperty->MarkHiddenByCustomization();
 	}
 
+	// Transitions
 	MakeArrayCategory(DetailBuilder, "Transitions", LOCTEXT("StateDetailsTransitions", "Transitions"), 4, TransitionsProperty);
 
 	// Refresh the UI when the type changes.	
diff --git a/Engine/Plugins/Runtime/StateTree/Source/StateTreeEditorModule/Private/StateTreeCompiler.cpp b/Engine/Plugins/Runtime/StateTree/Source/StateTreeEditorModule/Private/StateTreeCompiler.cpp
--- a/Engine/Plugins/Runtime/StateTree/Source/StateTreeEditorModule/Private/StateTreeCompiler.cpp
+++ b/Engine/Plugins/Runtime/StateTree/Source/StateTreeEditorModule/Private/StateTreeCompiler.cpp
@@ -190,10 +190,16 @@ bool FStateTreeCompiler::Compile(UStateTree& InStateTree)
 	{
 		return false;
 	}
-
 	// Cleanup existing state
 	StateTree->ResetCompiled();
 
+	if (!EditorData->Schema)
+	{
+		Log.Reportf(EMessageSeverity::Error, TEXT("Missing Schema. Please set valid schema in the State Tree Asset settings."));
+		return false;
+	}
+	Schema = EditorData->Schema;
+
 	if (!BindingsCompiler.Init(StateTree->PropertyBindings, Log))
 	{
 		StateTree->ResetCompiled();
@@ -406,6 +412,7 @@ bool FStateTreeCompiler::CreateStateRecursive(UStateTreeState& State, const FSta
 	const int32 StateIdx = StateTree->States.AddDefaulted();
 	FCompactStateTreeState& CompactState = StateTree->States[StateIdx];
 	CompactState.Name = State.Name;
+	CompactState.Tag = State.Tag;
 	CompactState.Parent = Parent;
 	CompactState.bEnabled = State.bEnabled;
 
@@ -841,6 +848,28 @@ bool FStateTreeCompiler::CreateStateTransitions()
 				return false;
 			}
 
+			// Linked asset must have same schema.
+			const UStateTreeSchema* LinkedAssetSchema = SourceState->LinkedAsset ? SourceState->LinkedAsset->GetSchema() : nullptr;
+
+			if (!LinkedAssetSchema)
+			{
+				Log.Reportf(EMessageSeverity::Error,
+					TEXT("Linked State Tree asset must have valid schema."));
+				return false;
+			}
+			
+			check(Schema);
+			if (LinkedAssetSchema->GetClass() != Schema->GetClass())
+			{
+				Log.Reportf(EMessageSeverity::Error,
+					TEXT("Linked State Tree asset '%s' must have same schema class as this asset. Linked asset has '%s', expected '%s'."),
+					*GetFullNameSafe(SourceState->LinkedAsset),
+					*LinkedAssetSchema->GetClass()->GetDisplayNameText().ToString(),
+					*Schema->GetClass()->GetDisplayNameText().ToString()
+				);
+				return false;
+			}
+			
 			CompactState.LinkedAsset = SourceState->LinkedAsset;
 		}
 		
diff --git a/Engine/Plugins/Runtime/StateTree/Source/StateTreeEditorModule/Private/StateTreeEditorModule.cpp b/Engine/Plugins/Runtime/StateTree/Source/StateTreeEditorModule/Private/StateTreeEditorModule.cpp
--- a/Engine/Plugins/Runtime/StateTree/Source/StateTreeEditorModule/Private/StateTreeEditorModule.cpp
+++ b/Engine/Plugins/Runtime/StateTree/Source/StateTreeEditorModule/Private/StateTreeEditorModule.cpp
@@ -9,11 +9,13 @@
 #include "Customizations/StateTreeEditorDataDetails.h"
 #include "Customizations/StateTreeEditorNodeDetails.h"
 #include "Customizations/StateTreeReferenceDetails.h"
+#include "Customizations/StateTreeReferenceOverridesDetails.h"
 #include "Customizations/StateTreeStateDetails.h"
 #include "Customizations/StateTreeStateLinkDetails.h"
 #include "Customizations/StateTreeStateParametersDetails.h"
 #include "Customizations/StateTreeTransitionDetails.h"
 #include "Customizations/StateTreeEventDescDetails.h"
+#include "Customizations/StateTreeBlueprintPropertyRefDetails.h"
 #include "PropertyEditorModule.h"
 #include "StateTree.h"
 #include "StateTreeCompiler.h"
@@ -25,7 +27,6 @@
 #include "StateTreeEditorStyle.h"
 #include "StateTreeNodeClassCache.h"
 #include "Modules/ModuleManager.h"
-#include "Customizations/StateTreeBlueprintPropertyRefDetails.h"
 
 #define LOCTEXT_NAMESPACE "StateTreeEditor"
 
@@ -92,6 +93,7 @@ void FStateTreeEditorModule::StartupModule()
 	PropertyModule.RegisterCustomPropertyTypeLayout("StateTreeStateParameters", FOnGetPropertyTypeCustomizationInstance::CreateStatic(&FStateTreeStateParametersDetails::MakeInstance));
 	PropertyModule.RegisterCustomPropertyTypeLayout("StateTreeAnyEnum", FOnGetPropertyTypeCustomizationInstance::CreateStatic(&FStateTreeAnyEnumDetails::MakeInstance));
 	PropertyModule.RegisterCustomPropertyTypeLayout("StateTreeReference", FOnGetPropertyTypeCustomizationInstance::CreateStatic(&FStateTreeReferenceDetails::MakeInstance));
+	PropertyModule.RegisterCustomPropertyTypeLayout("StateTreeReferenceOverrides", FOnGetPropertyTypeCustomizationInstance::CreateStatic(&FStateTreeReferenceOverridesDetails::MakeInstance));
 	PropertyModule.RegisterCustomPropertyTypeLayout("StateTreeEditorColorRef", FOnGetPropertyTypeCustomizationInstance::CreateStatic(&FStateTreeEditorColorRefDetails::MakeInstance));
 	PropertyModule.RegisterCustomPropertyTypeLayout("StateTreeEditorColor", FOnGetPropertyTypeCustomizationInstance::CreateStatic(&FStateTreeEditorColorDetails::MakeInstance));
 	PropertyModule.RegisterCustomPropertyTypeLayout("StateTreeBlueprintPropertyRef", FOnGetPropertyTypeCustomizationInstance::CreateStatic(&FStateTreeBlueprintPropertyRefDetails::MakeInstance));
@@ -125,7 +127,13 @@ void FStateTreeEditorModule::ShutdownModule()
 		PropertyModule.UnregisterCustomPropertyTypeLayout("StateTreeEditorNode");
 		PropertyModule.UnregisterCustomPropertyTypeLayout("StateTreeStateParameters");
 		PropertyModule.UnregisterCustomPropertyTypeLayout("StateTreeAnyEnum");
+		PropertyModule.UnregisterCustomPropertyTypeLayout("StateTreeReference");
+		PropertyModule.UnregisterCustomPropertyTypeLayout("StateTreeReferenceOverrides");
+		PropertyModule.UnregisterCustomPropertyTypeLayout("StateTreeEditorColorRef");
+		PropertyModule.UnregisterCustomPropertyTypeLayout("StateTreeEditorColor");
 		PropertyModule.UnregisterCustomPropertyTypeLayout("StateTreeBlueprintPropertyRef");
+		PropertyModule.UnregisterCustomClassLayout("StateTreeState");
+		PropertyModule.UnregisterCustomClassLayout("StateTreeEditorData");
 		PropertyModule.NotifyCustomizationModuleChanged();
 	}
 }
diff --git a/Engine/Plugins/Runtime/StateTree/Source/StateTreeEditorModule/Private/StateTreeState.cpp b/Engine/Plugins/Runtime/StateTree/Source/StateTreeEditorModule/Private/StateTreeState.cpp
--- a/Engine/Plugins/Runtime/StateTree/Source/StateTreeEditorModule/Private/StateTreeState.cpp
+++ b/Engine/Plugins/Runtime/StateTree/Source/StateTreeEditorModule/Private/StateTreeState.cpp
@@ -105,7 +105,7 @@ void UStateTreeState::PreEditChange(FEditPropertyChain& PropertyAboutToChange)
 		if (Type == EStateTreeStateType::Linked
 			|| Type == EStateTreeStateType::LinkedAsset)
 		{
-			Parameters.Reset();
+			Parameters.ResetParametersAndOverrides();
 		}
 	}
 }
@@ -430,7 +430,7 @@ void UStateTreeState::UpdateParametersFromLinkedSubtree()
 	}
 	else
 	{
-		Parameters.Reset();
+		Parameters.ResetParametersAndOverrides();
 	}
 }
 
diff --git a/Engine/Plugins/Runtime/StateTree/Source/StateTreeEditorModule/Public/StateTreeCompiler.h b/Engine/Plugins/Runtime/StateTree/Source/StateTreeEditorModule/Public/StateTreeCompiler.h
--- a/Engine/Plugins/Runtime/StateTree/Source/StateTreeEditorModule/Public/StateTreeCompiler.h
+++ b/Engine/Plugins/Runtime/StateTree/Source/StateTreeEditorModule/Public/StateTreeCompiler.h
@@ -14,6 +14,7 @@ struct FStateTreeStateHandle;
 class UStateTree;
 class UStateTreeState;
 class UStateTreeEditorData;
+class UStateTreeSchema;
 struct FStateTreeEditorNode;
 struct FStateTreeStateLink;
 struct FStateTreeNodeBase;
@@ -63,6 +64,7 @@ private:
 	FStateTreeCompilerLog& Log;
 	UStateTree* StateTree = nullptr;
 	UStateTreeEditorData* EditorData = nullptr;
+	UStateTreeSchema* Schema = nullptr;
 	TMap<FGuid, int32> IDToNode;
 	TMap<FGuid, int32> IDToState;
 	TMap<FGuid, int32> IDToTransition;
diff --git a/Engine/Plugins/Runtime/StateTree/Source/StateTreeEditorModule/Public/StateTreeState.h b/Engine/Plugins/Runtime/StateTree/Source/StateTreeEditorModule/Public/StateTreeState.h
--- a/Engine/Plugins/Runtime/StateTree/Source/StateTreeEditorModule/Public/StateTreeState.h
+++ b/Engine/Plugins/Runtime/StateTree/Source/StateTreeEditorModule/Public/StateTreeState.h
@@ -9,6 +9,7 @@
 #include "StateTreeState.generated.h"
 
 class UStateTreeState;
+class UStateTree;
 
 /**
  * Editor representation of an event description.
@@ -160,11 +161,11 @@ struct STATETREEEDITORMODULE_API FStateTreeStateParameters
 {
 	GENERATED_BODY()
 
-	void Reset()
+	void ResetParametersAndOverrides()
 	{
+		// Reset just the parameters, keep the bFixedLayout intact.
 		Parameters.Reset();
 		PropertyOverrides.Reset();
-		bFixedLayout = false;
 	}
 
 	/** Removes overrides that do appear in Parameters. */
@@ -271,6 +272,14 @@ public:
 		return static_cast<TStateTreeEditorNode<T>&>(TaskItem);
 	}
 
+	/** Sets linked asset and updates parameters to match the linked asset. */
+	void SetLinkedStateAsset(UStateTree* InLinkedAsset)
+	{
+		check(Type == EStateTreeStateType::LinkedAsset);
+		LinkedAsset = InLinkedAsset;
+		UpdateParametersFromLinkedSubtree();
+	}
+	
 	/**
 	 * Adds Transition.
 	 * @return reference to the new Transition.
@@ -304,10 +313,16 @@ public:
 
 	// ~IStateTreeSchemaProvider API
 
+	// Note: these properties are customized out in FStateTreeStateDetails, adding a new property might require to adjust the customization.
+	
 	/** Display name of the State */
 	UPROPERTY(EditDefaultsOnly, Category = "State")
 	FName Name;
 
+	/** GameplayTag describing the State */
+	UPROPERTY(EditDefaultsOnly, Category = "State")
+	FGameplayTag Tag;
+
 	/** Display color of the State */
 	UPROPERTY(EditDefaultsOnly, Category = "State", DisplayName = "Color")
 	FStateTreeEditorColorRef ColorRef;
diff --git a/Engine/Plugins/Runtime/StateTree/Source/StateTreeModule/Private/StateTreeExecutionContext.cpp b/Engine/Plugins/Runtime/StateTree/Source/StateTreeModule/Private/StateTreeExecutionContext.cpp
--- a/Engine/Plugins/Runtime/StateTree/Source/StateTreeModule/Private/StateTreeExecutionContext.cpp
+++ b/Engine/Plugins/Runtime/StateTree/Source/StateTreeModule/Private/StateTreeExecutionContext.cpp
@@ -4,6 +4,7 @@
 #include "StateTreeTaskBase.h"
 #include "StateTreeEvaluatorBase.h"
 #include "StateTreeConditionBase.h"
+#include "StateTreeReference.h"
 #include "Containers/StaticArray.h"
 #include "Debugger/StateTreeTrace.h"
 #include "Debugger/StateTreeTraceTypes.h"
@@ -120,7 +121,6 @@ FStateTreeExecutionContext::FStateTreeExecutionContext(UObject& InOwner, const U
 	}
 }
 
-
 FStateTreeExecutionContext::FStateTreeExecutionContext(const FStateTreeExecutionContext& InContextToCopy, const UStateTree& InStateTree, FStateTreeInstanceData& InInstanceData)
 	: FStateTreeExecutionContext(InContextToCopy.Owner, InStateTree, InInstanceData, InContextToCopy.CollectExternalDataDelegate)
 {
@@ -155,6 +155,30 @@ void FStateTreeExecutionContext::SetCollectExternalDataCallback(const FOnCollect
 	CollectExternalDataDelegate = Callback;
 }
 
+void FStateTreeExecutionContext::SetLinkedStateTreeOverrides(const FStateTreeReferenceOverrides* InLinkedStateTreeOverrides)
+{
+	LinkedStateTreeOverrides = InLinkedStateTreeOverrides;
+}
+
+const FStateTreeReference* FStateTreeExecutionContext::GetLinkedStateTreeOverrideForTag(const FGameplayTag StateTag) const
+{
+	if (!LinkedStateTreeOverrides)
+	{
+		return nullptr;
+	}
+	
+	for (const FStateTreeReferenceOverrideItem& Item : LinkedStateTreeOverrides->GetOverrideItems())
+	{
+		if (Item.GetStateTag() == StateTag)
+		{
+			return &Item.GetStateTreeReference();
+		}
+	}
+
+	return nullptr;
+}
+
+
 bool FStateTreeExecutionContext::AreContextDataViewsValid() const
 {
 	if (!IsValid())
@@ -820,8 +844,22 @@ void FStateTreeExecutionContext::UpdateInstanceData(TConstArrayView<FStateTreeEx
 				{
 					// Linked state's instance data is the parameters.
 					check(State.ParameterDataHandle.IsValid());
-					const FConstStructView ParamsInstanceData = NextFrame.StateTree->DefaultInstanceData.GetStruct(State.ParameterTemplateIndex.Get());
-					InstanceStructs[BaseIndex + State.ParameterDataHandle.GetIndex()] = ParamsInstanceData;
+
+					const FCompactStateTreeParameters* Params = nullptr;
+					if (FInstancedStruct* TempParamsInstanceData = FindInstanceTempData(NextFrame, State.ParameterDataHandle))
+					{
+						// If we have temp data for the parameters, then setup the instance data with just a type, so that we can steal the temp data below (TempInstanceStructs).
+						// We expect overridden linked assets to hit this code path. 
+						InstanceStructs[BaseIndex + State.ParameterDataHandle.GetIndex()] = FConstStructView(TempParamsInstanceData->GetScriptStruct());
+						Params = TempParamsInstanceData->GetPtr<const FCompactStateTreeParameters>();
+					}
+					else
+					{
+						// If not temp data, use the states default values.
+						const FConstStructView ParamsInstanceData = NextFrame.StateTree->DefaultInstanceData.GetStruct(State.ParameterTemplateIndex.Get());
+						InstanceStructs[BaseIndex + State.ParameterDataHandle.GetIndex()] = ParamsInstanceData;
+						Params = ParamsInstanceData.GetPtr<const FCompactStateTreeParameters>();
+					}
 
 					if (State.Type == EStateTreeStateType::Linked
 						|| State.Type == EStateTreeStateType::LinkedAsset)
@@ -830,8 +868,6 @@ void FStateTreeExecutionContext::UpdateInstanceData(TConstArrayView<FStateTreeEx
 						check(State.ParameterDataHandle.GetSource() == EStateTreeDataSourceType::StateParameterData);
 						checkf(!NextStateParameterDataHandle.IsValid(), TEXT("NextStateParameterDataIndex not should be set yet when we encounter a linked state."));
 						NextStateParameterDataHandle = State.ParameterDataHandle;
-					
-						const FCompactStateTreeParameters* Params = ParamsInstanceData.GetPtr<const FCompactStateTreeParameters>();
 						NextStateParameterDataStruct = Params ? Params->Parameters.GetPropertyBagStruct() : nullptr;
 					}
 				}
@@ -3155,6 +3191,18 @@ bool FStateTreeExecutionContext::SelectStateInternal(
 		return false;
 	}
 
+	// Look up linked state overrides
+	const UStateTree* NextLinkedStateAssetOverride = nullptr;
+	const FInstancedPropertyBag* NextLinkedStateParameterOverride = nullptr;
+	if (NextState.Type == EStateTreeStateType::LinkedAsset)
+	{
+		if (const FStateTreeReference* Override = GetLinkedStateTreeOverrideForTag(NextState.Tag))
+		{
+			NextLinkedStateAssetOverride = Override->GetStateTree();
+			NextLinkedStateParameterOverride = &Override->GetParameters();
+		}
+	}
+
 	if (NextState.ParameterDataHandle.IsValid())
 	{
 		// Instantiate state parameters if not done yet.
@@ -3162,14 +3210,30 @@ bool FStateTreeExecutionContext::SelectStateInternal(
 		if (!NextStateParametersView.IsValid())
 		{
 			// Allocate temporary instance for parameters if the state has params.
-			const FConstStructView DefaultStateParamsInstanceData = CurrentFrame.StateTree->DefaultInstanceData.GetStruct(NextState.ParameterTemplateIndex.Get());
-			const FCompactStateTreeParameters& DefaultStateParams = DefaultStateParamsInstanceData.Get<const FCompactStateTreeParameters>();
-			if (DefaultStateParams.Parameters.IsValid())
+			if (NextLinkedStateParameterOverride)
 			{
-				FStateTreeDataView TempStateParametersView = AddTemporaryInstance(CurrentFrame, FStateTreeIndex16::Invalid, NextState.ParameterDataHandle, DefaultStateParamsInstanceData);
-				check(TempStateParametersView.IsValid());
-				FCompactStateTreeParameters& StateParams = TempStateParametersView.GetMutable<FCompactStateTreeParameters>();
-				NextStateParametersView = FStateTreeDataView(StateParams.Parameters.GetMutableValue());
+				// Create from an override.
+				if (NextLinkedStateParameterOverride->IsValid())
+				{
+					FStateTreeDataView TempStateParametersView = AddTemporaryInstance(CurrentFrame, FStateTreeIndex16::Invalid, NextState.ParameterDataHandle, FConstStructView(TBaseStructure<FCompactStateTreeParameters>::Get()));
+					check(TempStateParametersView.IsValid());
+					FCompactStateTreeParameters& StateParams = TempStateParametersView.GetMutable<FCompactStateTreeParameters>();
+					StateParams.Parameters = *NextLinkedStateParameterOverride;
+					NextStateParametersView = FStateTreeDataView(StateParams.Parameters.GetMutableValue());
+				}
+			}
+			else
+			{
+				// Create from template in the asset.
+				const FConstStructView DefaultStateParamsInstanceData = CurrentFrame.StateTree->DefaultInstanceData.GetStruct(NextState.ParameterTemplateIndex.Get());
+				const FCompactStateTreeParameters& DefaultStateParams = DefaultStateParamsInstanceData.Get<const FCompactStateTreeParameters>();
+				if (DefaultStateParams.Parameters.IsValid())
+				{
+					FStateTreeDataView TempStateParametersView = AddTemporaryInstance(CurrentFrame, FStateTreeIndex16::Invalid, NextState.ParameterDataHandle, DefaultStateParamsInstanceData);
+					check(TempStateParametersView.IsValid());
+					FCompactStateTreeParameters& StateParams = TempStateParametersView.GetMutable<FCompactStateTreeParameters>();
+					NextStateParametersView = FStateTreeDataView(StateParams.Parameters.GetMutableValue());
+				}
 			}
 		}
 
@@ -3179,6 +3243,7 @@ bool FStateTreeExecutionContext::SelectStateInternal(
 			&& NextState.ParameterBindingsBatch.IsValid())
 		{
 			// Note: the parameters are for the current (linked) state, stored in current frame.
+			// The copy can fail, if the overridden parameters do not match, this is by design.
 			CopyBatchWithValidation(CurrentParentFrame, CurrentFrame, NextStateParametersView, NextState.ParameterBindingsBatch);
 		}
 	}
@@ -3318,7 +3383,17 @@ bool FStateTreeExecutionContext::SelectStateInternal(
 		}
 		else if (NextState.Type == EStateTreeStateType::LinkedAsset)
 		{
-			if (NextState.LinkedAsset)
+			const UStateTree* LinkedAsset = NextState.LinkedAsset;
+			if (NextLinkedStateAssetOverride)
+			{
+				STATETREE_LOG(VeryVerbose, TEXT("%hs: In state '%s', overriding linked asset '%s' with '%s'. '%s' using StateTree '%s'."),
+					__FUNCTION__, *GetSafeStateName(CurrentFrame, NextStateHandle),
+					*GetFullNameSafe(NextState.LinkedAsset), *GetFullNameSafe(NextLinkedStateAssetOverride),
+					*GetNameSafe(&Owner), *GetFullNameSafe(CurrentFrame.StateTree));
+				LinkedAsset = NextLinkedStateAssetOverride;
+			}
+			
+			if (LinkedAsset)
 			{
 				if (OutSelectionResult.IsFull())
 				{
@@ -3328,15 +3403,15 @@ bool FStateTreeExecutionContext::SelectStateInternal(
 				}
 
 				// The linked state tree should have compatible context requirements.
-				if (!NextState.LinkedAsset->HasCompatibleContextData(RootStateTree))
+				if (!LinkedAsset->HasCompatibleContextData(RootStateTree))
 				{
 					STATETREE_LOG(Error, TEXT("%hs: The linked State Tree '%s' does not have compatible schema, trying to select state %s from '%s'.  '%s' using StateTree '%s'."),
-						__FUNCTION__, *GetFullNameSafe(NextState.LinkedAsset), *GetSafeStateName(CurrentFrame, NextStateHandle), *GetStateStatusString(Exec), *GetNameSafe(&Owner), *GetFullNameSafe(CurrentFrame.StateTree));
+						__FUNCTION__, *GetFullNameSafe(LinkedAsset), *GetSafeStateName(CurrentFrame, NextStateHandle), *GetStateStatusString(Exec), *GetNameSafe(&Owner), *GetFullNameSafe(CurrentFrame.StateTree));
 					return false;
 				}
 				
 				FStateTreeExecutionFrame NewFrame;
-				NewFrame.StateTree = NextState.LinkedAsset;
+				NewFrame.StateTree = LinkedAsset;
 				NewFrame.RootState = FStateTreeStateHandle::Root;
 				NewFrame.bIsGlobalFrame = true;
 
diff --git a/Engine/Plugins/Runtime/StateTree/Source/StateTreeModule/Public/StateTreeExecutionContext.h b/Engine/Plugins/Runtime/StateTree/Source/StateTreeModule/Public/StateTreeExecutionContext.h
--- a/Engine/Plugins/Runtime/StateTree/Source/StateTreeModule/Public/StateTreeExecutionContext.h
+++ b/Engine/Plugins/Runtime/StateTree/Source/StateTreeModule/Public/StateTreeExecutionContext.h
@@ -11,12 +11,14 @@
 struct FGameplayTag;
 struct FInstancedPropertyBag;
 struct FStateTreeExecutionContext;
+struct FStateTreeReferenceOverrides;
 struct FStateTreeEvaluatorBase;
 struct FStateTreeTaskBase;
 struct FStateTreeConditionBase;
 struct FStateTreeEvent;
 struct FStateTreeTransitionRequest;
 struct FStateTreeInstanceDebugId;
+struct FStateTreeReference;
 
 /**
  * Delegate used by the execution context to collect external data views for a given StateTree asset.
@@ -108,6 +110,11 @@ public:
 
 	/** Sets callback used to collect external data views during State Tree execution. */
 	void SetCollectExternalDataCallback(const FOnCollectStateTreeExternalData& Callback);
+
+	/** */
+	void SetLinkedStateTreeOverrides(const FStateTreeReferenceOverrides* InLinkedStateTreeOverrides);
+
+	const FStateTreeReference* GetLinkedStateTreeOverrideForTag(const FGameplayTag StateTag) const;
 	
 	/** @return the StateTree asset in use. */
 	const UStateTree* GetStateTree() const { return &RootStateTree; }
@@ -664,6 +671,9 @@ protected:
 	/** Data storage of the instance data, cached for less indirections. */
 	FStateTreeInstanceStorage* InstanceDataStorage = nullptr;
 
+	/** Pointer to linked state tree overrides. */
+	const FStateTreeReferenceOverrides* LinkedStateTreeOverrides = nullptr;
+	
 	/** Data view of the context data. */
 	TArray<FStateTreeDataView, TConcurrentLinearArrayAllocator<FDefaultBlockAllocationTag>> ContextAndExternalDataViews;
 
diff --git a/Engine/Plugins/Runtime/StateTree/Source/StateTreeModule/Public/StateTreeReference.h b/Engine/Plugins/Runtime/StateTree/Source/StateTreeModule/Public/StateTreeReference.h
--- a/Engine/Plugins/Runtime/StateTree/Source/StateTreeModule/Public/StateTreeReference.h
+++ b/Engine/Plugins/Runtime/StateTree/Source/StateTreeModule/Public/StateTreeReference.h
@@ -3,6 +3,7 @@
 #pragma once
 
 #include "PropertyBag.h"
+#include "GameplayTagContainer.h"
 #include "StateTreeReference.generated.h"
 
 class UStateTree;
@@ -111,6 +112,71 @@ struct TStructOpsTypeTraits<FStateTreeReference> : public TStructOpsTypeTraitsBa
 };
 
 
-#if UE_ENABLE_INCLUDE_ORDER_DEPRECATED_IN_5_2
-#include "StateTreeTypes.h"
-#endif
+/**
+ * Item describing a state tree override for a state with a specific tag.
+ */
+USTRUCT()
+struct STATETREEMODULE_API FStateTreeReferenceOverrideItem
+{
+	GENERATED_BODY()
+
+	FStateTreeReferenceOverrideItem() = default;
+	FStateTreeReferenceOverrideItem(const FGameplayTag InStateTag, const FStateTreeReference& InStateTreeReference)
+		: StateTag(InStateTag)
+		, StateTreeReference(InStateTreeReference)
+	{
+	}
+	
+	FGameplayTag GetStateTag() const
+	{
+		return StateTag;
+	}
+
+	const FStateTreeReference& GetStateTreeReference() const
+	{
+		return StateTreeReference;
+	}
+	
+private:
+
+	/** Exact tag used to match against a tag on a linked State Tree state. */
+	UPROPERTY(EditAnywhere, Category = "")
+	FGameplayTag StateTag;
+
+	/** State Tree and parameters to replace the linked state asset with. */
+	UPROPERTY(EditAnywhere, Category = "")
+	FStateTreeReference StateTreeReference;
+	
+	friend class FStateTreeReferenceOverridesDetails;
+};
+
+/**
+ * Overrides for linked State Trees. This table is used to override State Tree references on linked states.
+ * If a linked state's tag is exact match of the tag specified on the table, the reference from the table is used instead.
+ */
+USTRUCT()
+struct STATETREEMODULE_API FStateTreeReferenceOverrides
+{
+	GENERATED_BODY()
+
+	void Reset()
+	{
+		OverrideItems.Reset();	
+	}
+	
+	void AddOverride(const FGameplayTag StateTag, FStateTreeReference& StateTreeReference)
+	{
+		OverrideItems.Emplace(StateTag, StateTreeReference);
+	}
+
+	TConstArrayView<FStateTreeReferenceOverrideItem> GetOverrideItems() const
+	{
+		return OverrideItems;		
+	}
+	
+private:
+	UPROPERTY(EditAnywhere, Category = "")
+	TArray<FStateTreeReferenceOverrideItem> OverrideItems;
+
+	friend class FStateTreeReferenceOverridesDetails;
+};
\ No newline at end of file
diff --git a/Engine/Plugins/Runtime/StateTree/Source/StateTreeModule/Public/StateTreeTypes.h b/Engine/Plugins/Runtime/StateTree/Source/StateTreeModule/Public/StateTreeTypes.h
--- a/Engine/Plugins/Runtime/StateTree/Source/StateTreeModule/Public/StateTreeTypes.h
+++ b/Engine/Plugins/Runtime/StateTree/Source/StateTreeModule/Public/StateTreeTypes.h
@@ -618,13 +618,17 @@ struct STATETREEMODULE_API FCompactStateTreeState
 	UPROPERTY()
 	FName Name;
 
-	/** Linked state handle if the state type is linked state. */
+	/** GameplayTag describing the State */
 	UPROPERTY()
-	FStateTreeStateHandle LinkedState = FStateTreeStateHandle::Invalid; 
+	FGameplayTag Tag;
 
 	UPROPERTY()
 	TObjectPtr<UStateTree> LinkedAsset = nullptr;
-	
+
+	/** Linked state handle if the state type is linked state. */
+	UPROPERTY()
+	FStateTreeStateHandle LinkedState = FStateTreeStateHandle::Invalid; 
+
 	/** Parent state handle, invalid if root state. */
 	UPROPERTY()
 	FStateTreeStateHandle Parent = FStateTreeStateHandle::Invalid;
diff --git a/Engine/Plugins/Runtime/StateTree/Source/StateTreeTestSuite/Private/StateTreeTest.cpp b/Engine/Plugins/Runtime/StateTree/Source/StateTreeTestSuite/Private/StateTreeTest.cpp
--- a/Engine/Plugins/Runtime/StateTree/Source/StateTreeTestSuite/Private/StateTreeTest.cpp
+++ b/Engine/Plugins/Runtime/StateTree/Source/StateTreeTestSuite/Private/StateTreeTest.cpp
@@ -10,6 +10,7 @@
 #include "Engine/World.h"
 #include "Async/ParallelFor.h"
 #include "GameplayTagsManager.h"
+#include "StateTreeReference.h"
 
 #include UE_INLINE_GENERATED_CPP_BY_NAME(StateTreeTest)
 
@@ -2881,6 +2882,100 @@ struct FStateTreeTest_DeferredStop_ExitTask : FStateTreeTest_DeferredStop
 IMPLEMENT_AI_INSTANT_TEST(FStateTreeTest_DeferredStop_ExitTask, "System.StateTree.DeferredStop.ExitTask");
 
 
+// Test nested tree overrides
+struct FStateTreeTest_NestedOverride : FAITestBase
+{
+	virtual bool InstantTest() override
+	{
+		FStateTreeCompilerLog Log;
+		
+		const FGameplayTag Tag = UE::StateTree::Tests::FNativeGameplayTags::Get().TestTag;
+
+		// Asset 2
+		UStateTree& StateTree2 = UE::StateTree::Tests::NewStateTree(&GetWorld());
+		UStateTreeEditorData& EditorData2 = *Cast<UStateTreeEditorData>(StateTree2.EditorData);
+		EditorData2.RootParameters.Parameters.AddProperty(FName(TEXT("Int")), EPropertyBagPropertyType::Int32);
+		UStateTreeState& Root2 = EditorData2.AddSubTree(FName(TEXT("Root2")));
+		TStateTreeEditorNode<FTestTask_Stand>& TaskRoot2 = Root2.AddTask<FTestTask_Stand>(FName(TEXT("TaskRoot2")));
+
+		FStateTreeCompiler Compiler2(Log);
+		const bool bResult2 = Compiler2.Compile(StateTree2);
+		AITEST_TRUE("StateTree2 should get compiled", bResult2);
+
+		
+		// Asset 3
+		UStateTree& StateTree3 = UE::StateTree::Tests::NewStateTree(&GetWorld());
+		UStateTreeEditorData& EditorData3 = *Cast<UStateTreeEditorData>(StateTree3.EditorData);
+		EditorData3.RootParameters.Parameters.AddProperty(FName(TEXT("Float")), EPropertyBagPropertyType::Float); // Different parameters
+		UStateTreeState& Root3 = EditorData3.AddSubTree(FName(TEXT("Root3")));
+		TStateTreeEditorNode<FTestTask_Stand>& TaskRoot3 = Root3.AddTask<FTestTask_Stand>(FName(TEXT("TaskRoot3")));
+
+		FStateTreeCompiler Compiler3(Log);
+		const bool bResult3 = Compiler3.Compile(StateTree3);
+		AITEST_TRUE("StateTree3 should get compiled", bResult3);
+
+		// Main asset
+		UStateTree& StateTree = UE::StateTree::Tests::NewStateTree(&GetWorld());
+		UStateTreeEditorData& EditorData = *Cast<UStateTreeEditorData>(StateTree.EditorData);
+
+		EditorData.RootParameters.Parameters.AddProperty(FName(TEXT("Int")), EPropertyBagPropertyType::Int32);
+		
+		UStateTreeState& Root = EditorData.AddSubTree(FName(TEXT("Root1")));
+		UStateTreeState& StateA = Root.AddChildState(FName(TEXT("A1")), EStateTreeStateType::LinkedAsset);
+		StateA.Tag = Tag;
+		StateA.SetLinkedStateAsset(&StateTree2);
+
+		FStateTreeCompiler Compiler(Log);
+		const bool bResult = Compiler.Compile(StateTree);
+		AITEST_TRUE("StateTree should get compiled", bResult);
+
+		const FString TickStr(TEXT("Tick"));
+		const FString EnterStateStr(TEXT("EnterState"));
+		const FString ExitStateStr(TEXT("ExitState"));
+
+		// Without overrides
+		{
+			EStateTreeRunStatus Status = EStateTreeRunStatus::Unset;
+			FStateTreeInstanceData InstanceData;
+			FTestStateTreeExecutionContext Exec(StateTree, StateTree, InstanceData);
+			const bool bInitSucceeded = Exec.IsValid();
+			AITEST_TRUE("StateTree should init", bInitSucceeded);
+
+			Status = Exec.Start();
+			AITEST_EQUAL("Start should complete with Running", Status, EStateTreeRunStatus::Running);
+			AITEST_TRUE("StateTree should enter TaskRoot2", Exec.Expect(TaskRoot2.GetName(), EnterStateStr));
+
+			Exec.LogClear();
+		}
+
+		// With overrides
+		{
+			EStateTreeRunStatus Status = EStateTreeRunStatus::Unset;
+			FStateTreeInstanceData InstanceData;
+
+			FStateTreeReferenceOverrides Overrides;
+			FStateTreeReference OverrideRef;
+			OverrideRef.SetStateTree(&StateTree3);
+			Overrides.AddOverride(Tag, OverrideRef);
+			
+			FTestStateTreeExecutionContext Exec(StateTree, StateTree, InstanceData);
+			Exec.SetLinkedStateTreeOverrides(&Overrides);
+			
+			const bool bInitSucceeded = Exec.IsValid();
+			AITEST_TRUE("StateTree should init", bInitSucceeded);
+			
+			Status = Exec.Start();
+			AITEST_EQUAL("Start should complete with Running", Status, EStateTreeRunStatus::Running);
+			AITEST_TRUE("StateTree should enter TaskRoot3", Exec.Expect(TaskRoot3.GetName(), EnterStateStr));
+
+			Exec.LogClear();
+		}
+
+		return true;
+	}
+};
+IMPLEMENT_AI_INSTANT_TEST(FStateTreeTest_NestedOverride, "System.StateTree.NestedOverride");
+
 UE_ENABLE_OPTIMIZATION_SHIP
 
 #undef LOCTEXT_NAMESPACE
