StateTree: Improve state picker.
- imropoved the state picker to display the states in hieracrchy

diff --git a/Engine/Plugins/Runtime/StateTree/Source/StateTreeEditorModule/Private/Customizations/StateTreeStateLinkDetails.cpp b/Engine/Plugins/Runtime/StateTree/Source/StateTreeEditorModule/Private/Customizations/StateTreeStateLinkDetails.cpp
--- a/Engine/Plugins/Runtime/StateTree/Source/StateTreeEditorModule/Private/Customizations/StateTreeStateLinkDetails.cpp
+++ b/Engine/Plugins/Runtime/StateTree/Source/StateTreeEditorModule/Private/Customizations/StateTreeStateLinkDetails.cpp
@@ -10,10 +10,15 @@
 #include "Widgets/Images/SImage.h"
 #include "Framework/MultiBox/MultiBoxBuilder.h"
 #include "ScopedTransaction.h"
+#include "StateTreeEditorStyle.h"
 #include "StateTreePropertyHelpers.h"
+#include "Widgets/Input/SSearchBox.h"
 
 #define LOCTEXT_NAMESPACE "StateTreeEditor"
 
+TMap<FObjectKey, FStateTreeStateLinkDetails::FStateExpansionState> FStateTreeStateLinkDetails::StateExpansionStates;
+
+
 TSharedRef<IPropertyTypeCustomization> FStateTreeStateLinkDetails::MakeInstance()
 {
 	return MakeShareable(new FStateTreeStateLinkDetails);
@@ -47,8 +52,8 @@ void FStateTreeStateLinkDetails::CustomizeHeader(TSharedRef<class IPropertyHandl
 		.ValueContent()
 		.VAlign(VAlign_Center)
 		[
-			SNew(SComboButton)
-			.OnGetMenuContent(this, &FStateTreeStateLinkDetails::OnGetStateContent)
+			SAssignNew(ComboButton, SComboButton)
+			.OnGetMenuContent(this, &FStateTreeStateLinkDetails::GenerateStatePicker)
 			.ButtonContent()
 			[
 				SNew(SHorizontalBox)
@@ -58,25 +63,34 @@ void FStateTreeStateLinkDetails::CustomizeHeader(TSharedRef<class IPropertyHandl
 				.VAlign(VAlign_Center)
 				.Padding(0,0,4,0)
 				[
-					SNew(SBox)
-					[
-						SNew(SImage)
-						.ToolTipText(LOCTEXT("MissingState", "The specified state cannot be found."))
-						.Visibility_Lambda([this]()
-						{
-							return IsValidLink() ? EVisibility::Collapsed : EVisibility::Visible;
-						})
-						.Image(FAppStyle::GetBrush("Icons.ErrorWithColor"))
-					]
+					SNew(SImage)
+					.ToolTipText(LOCTEXT("MissingState", "The specified state cannot be found."))
+					.Visibility_Lambda([this]()
+					{
+						return IsValidLink() ? EVisibility::Collapsed : EVisibility::Visible;
+					})
+					.Image(FAppStyle::GetBrush("Icons.ErrorWithColor"))
 				]
 
+				+SHorizontalBox::Slot()
+				.HAlign(HAlign_Left)
+				.VAlign(VAlign_Center)
+				.Padding(0, 2.0f, 4.0f, 2.0f)
+				.AutoWidth()
+				[
+					SNew(SImage)
+					.DesiredSizeOverride(FVector2D(16.0f, 16.0f))
+					.Image(this, &FStateTreeStateLinkDetails::GetCurrentStateIcon)
+					.ColorAndOpacity(this, &FStateTreeStateLinkDetails::GetCurrentStateColor)
+				]
+				
 				+SHorizontalBox::Slot()
 				.AutoWidth()
 				.VAlign(VAlign_Center)
 				[
 					SNew(STextBlock)
 					.Text(this, &FStateTreeStateLinkDetails::GetCurrentStateDesc)
-					.Font(IDetailLayoutBuilder::GetDetailFont())
+					.Font(this, &FStateTreeStateLinkDetails::GetCurrentStateFont)
 				]
 			]
 		];
@@ -93,7 +107,7 @@ void FStateTreeStateLinkDetails::OnIdentifierChanged(const UStateTree& StateTree
 	CacheStates();
 }
 
-void FStateTreeStateLinkDetails::CacheStates(const UStateTreeState* State)
+void FStateTreeStateLinkDetails::CacheStates(TSharedPtr<FStateTreeStateItem> ParentNode, const UStateTreeState* State)
 {
 	if (State == nullptr)
 	{
@@ -113,213 +127,549 @@ void FStateTreeStateLinkDetails::CacheStates(const UStateTreeState* State)
 	
 	if (bShouldAdd)
 	{
-		CachedNames.Add(State->Name);
-		CachedIDs.Add(State->ID);
+		TSharedRef<FStateTreeStateItem> StateItem = MakeShared<FStateTreeStateItem>();
+		StateItem->Desc = FText::FromName(State->Name);
+		StateItem->TransitionType = EStateTreeTransitionType::GotoState;
+		StateItem->StateID = State->ID;
+		StateItem->Color = FLinearColor(1.f, 1.f, 1.f, 0.25f);
+		StateItem->bIsSubTree = State->Type == EStateTreeStateType::Subtree;
+
+		// Figure out icon.
+		if (State->SelectionBehavior == EStateTreeStateSelectionBehavior::None)
+		{
+			StateItem->Icon = FStateTreeEditorStyle::Get().GetBrush("StateTreeEditor.SelectNone");
+		}
+		else if (State->SelectionBehavior == EStateTreeStateSelectionBehavior::TryEnterState)
+		{
+			StateItem->Icon = FStateTreeEditorStyle::Get().GetBrush("StateTreeEditor.TryEnterState");			
+		}
+		else if (State->SelectionBehavior == EStateTreeStateSelectionBehavior::TrySelectChildrenInOrder)
+		{
+			if (State->Children.IsEmpty()
+				|| State->Type == EStateTreeStateType::Linked
+				|| State->Type == EStateTreeStateType::LinkedAsset)
+			{
+				StateItem->Icon = FStateTreeEditorStyle::Get().GetBrush("StateTreeEditor.TryEnterState");			
+			}
+			else
+			{
+				StateItem->Icon = FStateTreeEditorStyle::Get().GetBrush("StateTreeEditor.TrySelectChildrenInOrder");
+			}
+		}
+		else if (State->SelectionBehavior == EStateTreeStateSelectionBehavior::TryFollowTransitions)
+		{
+			StateItem->Icon = FStateTreeEditorStyle::Get().GetBrush("StateTreeEditor.TryFollowTransitions");
+		}
+
+		// Linked states
+		if (State->Type == EStateTreeStateType::Linked)
+		{
+			StateItem->bIsLinked = true;
+			StateItem->LinkedDesc = FText::FromName(State->LinkedSubtree.Name);
+		}
+		else if (State->Type == EStateTreeStateType::LinkedAsset)
+		{
+			StateItem->bIsLinked = true;
+			StateItem->LinkedDesc = FText::FromString(GetNameSafe(State->LinkedAsset.Get()));
+		}
+		
+		ParentNode->Children.Add(StateItem);
+
+		ParentNode = StateItem;
 	}
 
 	for (UStateTreeState* ChildState : State->Children)
 	{
-		CacheStates(ChildState);
+		CacheStates(ParentNode, ChildState);
 	}
 }
 
 void FStateTreeStateLinkDetails::CacheStates()
 {
-	CachedNames.Reset();
-	CachedIDs.Reset();
+	RootItem = MakeShared<FStateTreeStateItem>();
 
+	if (!bDirectStatesOnly)
+	{
+		RootItem->Children.Add(MakeShared<FStateTreeStateItem>(
+			LOCTEXT("TransitionNone", "None"),
+			LOCTEXT("TransitionNoneTooltip", "No transition."),
+			FStateTreeEditorStyle::Get().GetBrush("StateTreeEditor.Transition.None"),
+			EStateTreeTransitionType::None));
+
+		RootItem->Children.Add(MakeShared<FStateTreeStateItem>(
+			LOCTEXT("TransitionNextState", "Next State"),
+			LOCTEXT("TransitionNextTooltip", "Goto next sibling State."),
+			FStateTreeEditorStyle::Get().GetBrush("StateTreeEditor.Transition.Next"),
+			EStateTreeTransitionType::NextState));
+
+		RootItem->Children.Add(MakeShared<FStateTreeStateItem>(
+			LOCTEXT("TransitionNextSelectableState", "Next Selectable State"),
+			LOCTEXT("TransitionNextSelectableTooltip", "Goto next sibling state, whose enter conditions pass."),
+			FStateTreeEditorStyle::Get().GetBrush("StateTreeEditor.Transition.Next"),
+			EStateTreeTransitionType::NextSelectableState));
+
+		RootItem->Children.Add(MakeShared<FStateTreeStateItem>(
+			LOCTEXT("TransitionTreeSucceeded", "Tree Succeeded"),
+			LOCTEXT("TransitionTreeSuccessTooltip", "Complete tree with success."),
+			FStateTreeEditorStyle::Get().GetBrush("StateTreeEditor.Transition.Succeeded"),
+			EStateTreeTransitionType::Succeeded));
+
+		RootItem->Children.Add(MakeShared<FStateTreeStateItem>(
+			LOCTEXT("TransitionTreeFailed", "Tree Failed"),
+			LOCTEXT("TransitionTreeFailedTooltip", "Complete tree with failure."),
+			FStateTreeEditorStyle::Get().GetBrush("StateTreeEditor.Transition.Failed"),
+			EStateTreeTransitionType::Failed));
+	}
+	
 	TArray<UObject*> OuterObjects;
 	StructProperty->GetOuterObjects(OuterObjects);
 	for (int32 ObjectIdx = 0; ObjectIdx < OuterObjects.Num(); ObjectIdx++)
 	{
-		UStateTree* OuterStateTree = OuterObjects[ObjectIdx]->GetTypedOuter<UStateTree>();
-		if (OuterStateTree)
+		if (const UStateTree* OuterStateTree = OuterObjects[ObjectIdx]->GetTypedOuter<UStateTree>())
 		{
-			if (UStateTreeEditorData* TreeData = Cast<UStateTreeEditorData>(OuterStateTree->EditorData))
+			if (const UStateTreeEditorData* TreeData = Cast<UStateTreeEditorData>(OuterStateTree->EditorData))
 			{
+				WeakStateTree = OuterStateTree;
 				for (const UStateTreeState* SubTree : TreeData->SubTrees)
 				{
-					CacheStates(SubTree);
+					CacheStates(RootItem, SubTree);
 				}
 			}
 			break;
 		}
 	}
 
+	FilteredRootItem = RootItem;
 }
 
-void FStateTreeStateLinkDetails::OnStateComboChange(int Idx)
+TSharedRef<SWidget> FStateTreeStateLinkDetails::GenerateStatePicker()
 {
-	if (NameProperty && IDProperty)
+	check(ComboButton);
+	
+	CacheStates();
+
+	TSharedRef<SWidget> MenuWidget = 
+		SNew(SBox)
+		.MinDesiredWidth(300.f)
+		.MaxDesiredHeight(400.f)
+		.Padding(2)	
+		[
+			SNew(SVerticalBox)
+			+ SVerticalBox::Slot()
+			.HAlign(HAlign_Fill)
+			.VAlign(VAlign_Top)
+			.Padding(4, 2, 4, 2)
+			.AutoHeight()
+			[
+				SAssignNew(SearchBox, SSearchBox)
+				.OnTextChanged(this, &FStateTreeStateLinkDetails::OnSearchBoxTextChanged)
+			]
+			+ SVerticalBox::Slot()
+			[
+				SAssignNew(StateItemTree, STreeView<TSharedPtr<FStateTreeStateItem>>)
+				.SelectionMode(ESelectionMode::Single)
+				.ItemHeight(20.0f)
+				.TreeItemsSource(&FilteredRootItem->Children)
+				.OnGenerateRow(this, &FStateTreeStateLinkDetails::GenerateStateItemRow)
+				.OnGetChildren(this, &FStateTreeStateLinkDetails::GetStateItemChildren)
+				.OnSelectionChanged(this, &FStateTreeStateLinkDetails::OnStateItemSelected)
+				.OnExpansionChanged(this, &FStateTreeStateLinkDetails::OnStateItemExpansionChanged)
+			]
+		];
+
+	// Restore category expansion state from previous use.
+	RestoreExpansionState();
+	
+	// Expand and select currently selected item.
+	TArray<TSharedPtr<FStateTreeStateItem>> Path;
+	if (GetCurrentStateItem(Path))
+	{
+		// Expand all categories up to the selected item.
+		bIsRestoringExpansion = true;
+		for (int32 Index = 0; Index < Path.Num() - 1; Index++)
+		{
+			StateItemTree->SetItemExpansion(Path[Index], true);
+		}
+		bIsRestoringExpansion = false;
+		
+		StateItemTree->SetItemSelection(Path.Last(), true);
+		StateItemTree->RequestScrollIntoView(Path.Last());
+	}
+	
+	ComboButton->SetMenuContentWidgetToFocus(SearchBox);
+
+	return MenuWidget;
+}
+
+FSlateFontInfo FStateTreeStateLinkDetails::GetItemFont(TSharedPtr<FStateTreeStateItem> Item)
+{
+	FSlateFontInfo Font = FAppStyle::Get().GetFontStyle("PropertyWindow.NormalFont");
+	if (Item->TransitionType != EStateTreeTransitionType::GotoState)
+	{
+		Font = FAppStyle::Get().GetFontStyle("PropertyWindow.ItalicFont");
+	}
+	else if (Item->bIsSubTree)
+	{
+		Font = FAppStyle::Get().GetFontStyle("PropertyWindow.BoldFont");
+	}
+	return Font;	
+}
+
+TSharedRef<ITableRow> FStateTreeStateLinkDetails::GenerateStateItemRow(TSharedPtr<FStateTreeStateItem> Item, const TSharedRef<STableViewBase>& OwnerTable)
+{
+	const FSlateFontInfo Font = GetItemFont(Item);
+
+	TSharedRef<SHorizontalBox> Container = SNew(SHorizontalBox);
+
+	// Icon
+	Container->AddSlot()
+		.HAlign(HAlign_Left)
+		.VAlign(VAlign_Center)
+		.Padding(0, 2.0f, 4.0f, 2.0f)
+		.AutoWidth()
+		[
+			SNew(SImage)
+			.Visibility(Item->Icon ? EVisibility::Visible : EVisibility::Collapsed)
+			.DesiredSizeOverride(FVector2D(16.0f, 16.0f))
+			.Image(Item->Icon)
+			.ColorAndOpacity(Item->Color)
+		];
+
+	// Name
+	Container->AddSlot()
+		.HAlign(HAlign_Fill)
+		.VAlign(VAlign_Center)
+		.AutoWidth()
+		[
+			SNew(STextBlock)
+			.Font(Font)
+			.Text(Item->Desc)
+			.ToolTipText(Item->TooltipText)
+			.HighlightText_Lambda([this]() { return SearchBox.IsValid() ? SearchBox->GetText() : FText::GetEmpty(); })
+		];
+
+	// Link
+	if (Item->bIsLinked)
+	{
+		// Link icon
+		Container->AddSlot()
+			.HAlign(HAlign_Left)
+			.VAlign(VAlign_Center)
+			.Padding(4.f, 0.f)
+			.AutoWidth()
+			[
+				SNew(SImage)
+				.ColorAndOpacity(FSlateColor::UseSubduedForeground())
+				.Image(FStateTreeEditorStyle::Get().GetBrush("StateTreeEditor.StateLinked"))
+		];
+
+		// Linked name
+		Container->AddSlot()
+			.HAlign(HAlign_Fill)
+			.VAlign(VAlign_Center)
+			.AutoWidth()
+			[
+				SNew(STextBlock)
+				.Font(FAppStyle::Get().GetFontStyle("PropertyWindow.NormalFont"))
+				.ColorAndOpacity(FSlateColor::UseSubduedForeground())
+				.Text(Item->LinkedDesc)
+			];
+	}
+	
+	return SNew(STableRow<TSharedPtr<FStateTreeStateItem>>, OwnerTable)
+		[
+			Container
+		];
+}
+
+void FStateTreeStateLinkDetails::GetStateItemChildren(TSharedPtr<FStateTreeStateItem> Item, TArray<TSharedPtr<FStateTreeStateItem>>& OutItems) const
+{
+	if (Item.IsValid())
+	{
+		OutItems = Item->Children;
+	}
+}
+
+void FStateTreeStateLinkDetails::OnStateItemSelected(TSharedPtr<FStateTreeStateItem> SelectedItem, ESelectInfo::Type Type)
+{
+	// Skip selection from code
+	if (Type == ESelectInfo::Direct)
+	{
+		return;
+	}
+	
+	if (SelectedItem
+		&& NameProperty
+		&& IDProperty)
 	{
 		FScopedTransaction Transaction(FText::Format(LOCTEXT("SetPropertyValue", "Set {0}"), StructProperty->GetPropertyDisplayName()));
-		if (Idx >= 0)
+
+		LinkTypeProperty->SetValue((uint8)SelectedItem->TransitionType);
+
+		if (SelectedItem->TransitionType == EStateTreeTransitionType::GotoState)
 		{
-			LinkTypeProperty->SetValue((uint8)EStateTreeTransitionType::GotoState);
-			NameProperty->SetValue(CachedNames[Idx], EPropertyValueSetFlags::NotTransactable);
-			UE::StateTree::PropertyHelpers::SetStructValue<FGuid>(IDProperty, CachedIDs[Idx], EPropertyValueSetFlags::NotTransactable);
+			NameProperty->SetValue(FName(SelectedItem->Desc.ToString()), EPropertyValueSetFlags::NotTransactable);
+			UE::StateTree::PropertyHelpers::SetStructValue<FGuid>(IDProperty, SelectedItem->StateID, EPropertyValueSetFlags::NotTransactable);
 		}
 		else
 		{
-			switch (Idx)
-			{
-			case ComboSucceeded:
-				LinkTypeProperty->SetValue((uint8)EStateTreeTransitionType::Succeeded);
-				break;
-			case ComboFailed:
-				LinkTypeProperty->SetValue((uint8)EStateTreeTransitionType::Failed);
-				break;
-			case ComboNextState:
-				LinkTypeProperty->SetValue((uint8)EStateTreeTransitionType::NextState);
-				break;
-			case ComboNextSelectableState:
-				LinkTypeProperty->SetValue((uint8)EStateTreeTransitionType::NextSelectableState);
-				break;
-			case ComboNotSet:
-			default:
-				LinkTypeProperty->SetValue((uint8)EStateTreeTransitionType::None);
-				break;
-			}
 			// Clear name and id.
 			NameProperty->SetValue(FName(), EPropertyValueSetFlags::NotTransactable);
 			UE::StateTree::PropertyHelpers::SetStructValue<FGuid>(IDProperty, FGuid(), EPropertyValueSetFlags::NotTransactable);
 		}
 	}
+
+	check(ComboButton);
+	ComboButton->SetIsOpen(false);
 }
 
-TSharedRef<SWidget> FStateTreeStateLinkDetails::OnGetStateContent() const
+void FStateTreeStateLinkDetails::OnStateItemExpansionChanged(TSharedPtr<FStateTreeStateItem> ExpandedItem, bool bInExpanded) const
 {
-	FMenuBuilder MenuBuilder(true, NULL);
+	// Do not save expansion state we're restoring expansion state, or when showing filtered results. 
+	if (bIsRestoringExpansion || FilteredRootItem != RootItem)
+	{
+		return;
+	}
 
-	if (!bDirectStatesOnly)
+	if (ExpandedItem.IsValid() && ExpandedItem->StateID.IsValid())
 	{
-		const FUIAction NotSetItemAction(FExecuteAction::CreateSP(const_cast<FStateTreeStateLinkDetails*>(this), &FStateTreeStateLinkDetails::OnStateComboChange, ComboNotSet));
-		MenuBuilder.AddMenuEntry(LOCTEXT("TransitionNone", "None"), LOCTEXT("TransitionNoneTooltip", "No transition."), FSlateIcon(), NotSetItemAction);
+		FStateExpansionState& ExpansionState = StateExpansionStates.FindOrAdd(FObjectKey(WeakStateTree.Get()));
+		if (bInExpanded)
+		{
+			ExpansionState.CollapsedStates.Remove(ExpandedItem->StateID);
+		}
+		else
+		{
+			ExpansionState.CollapsedStates.Add(ExpandedItem->StateID);
+		}
+	}
+}
 
-		const FUIAction NextItemAction(FExecuteAction::CreateSP(const_cast<FStateTreeStateLinkDetails*>(this), &FStateTreeStateLinkDetails::OnStateComboChange, ComboNextState));
-		MenuBuilder.AddMenuEntry(LOCTEXT("TransitionNextState", "Next State"), LOCTEXT("TransitionNextTooltip", "Goto next sibling State."), FSlateIcon(), NextItemAction);
+void FStateTreeStateLinkDetails::OnSearchBoxTextChanged(const FText& NewText)
+{
+	if (!StateItemTree.IsValid())
+	{
+		return;
+	}
+	
+	FilteredRootItem.Reset();
 
-		const FUIAction NextSelectableItemAction(FExecuteAction::CreateSP(const_cast<FStateTreeStateLinkDetails*>(this), &FStateTreeStateLinkDetails::OnStateComboChange, ComboNextSelectableState));
-		MenuBuilder.AddMenuEntry(LOCTEXT("TransitionNextSelectableState", "Next Selectable State"), LOCTEXT("TransitionNextSelectableTooltip", "Goto next sibling state, whose enter conditions pass."), FSlateIcon(), NextSelectableItemAction);
+	TArray<FString> FilterStrings;
+	NewText.ToString().ParseIntoArrayWS(FilterStrings);
+	FilterStrings.RemoveAll([](const FString& String) { return String.IsEmpty(); });
+	
+	if (FilterStrings.IsEmpty())
+	{
+		// Show all when there's no filter string.
+		FilteredRootItem = RootItem;
+		StateItemTree->SetTreeItemsSource(&FilteredRootItem->Children);
+		RestoreExpansionState();
+		StateItemTree->RequestTreeRefresh();
+		return;
+	}
 
-		const FUIAction SucceededItemAction(FExecuteAction::CreateSP(const_cast<FStateTreeStateLinkDetails*>(this), &FStateTreeStateLinkDetails::OnStateComboChange, ComboSucceeded));
-		MenuBuilder.AddMenuEntry(LOCTEXT("TransitionTreeSucceeded", "Tree Succeeded"), LOCTEXT("TransitionTreeSuccessTooltip", "Complete tree with success."), FSlateIcon(), SucceededItemAction);
+	FilteredRootItem = MakeShared<FStateTreeStateItem>();
+	FilterStateItemChildren(FilterStrings, /*bParentMatches*/false, RootItem->Children, FilteredRootItem->Children);
 
-		const FUIAction FailedItemAction(FExecuteAction::CreateSP(const_cast<FStateTreeStateLinkDetails*>(this), &FStateTreeStateLinkDetails::OnStateComboChange, ComboFailed));
-		MenuBuilder.AddMenuEntry(LOCTEXT("TransitionTreeFailed", "Tree Failed"), LOCTEXT("TransitionTreeFailedTooltip", "Complete tree with failure."), FSlateIcon(), FailedItemAction);
+	StateItemTree->SetTreeItemsSource(&FilteredRootItem->Children);
+	ExpandAll(FilteredRootItem->Children);
+
+	// Update selection based on filtered items.
+	TArray<TSharedPtr<FStateTreeStateItem>> Path;
+	const EStateTreeTransitionType TransitionType = GetTransitionType().Get(EStateTreeTransitionType::Failed);
+	FGuid StateID;
+	if (IDProperty
+		&& UE::StateTree::PropertyHelpers::GetStructValue<FGuid>(IDProperty, StateID) == FPropertyAccess::Success)
+	{
+		// Use RootItem (not filtered) to get any result even during search.
+		if (FindStateByIDRecursive(FilteredRootItem, TransitionType, StateID, Path))
+		{
+			StateItemTree->SetItemSelection(Path.Last(), true);
+		}
 	}
 
-	if (CachedNames.Num() > 0)
+	StateItemTree->RequestTreeRefresh();
+}
+
+
+int32 FStateTreeStateLinkDetails::FilterStateItemChildren(const TArray<FString>& FilterStrings, const bool bParentMatches,
+															const TArray<TSharedPtr<FStateTreeStateItem>>& SourceArray,
+															TArray<TSharedPtr<FStateTreeStateItem>>& OutDestArray)
+{
+	int32 NumFound = 0;
+
+	auto MatchFilter = [&FilterStrings](const TSharedPtr<FStateTreeStateItem>& SourceItem)
+	{
+		const FString ItemName = SourceItem->Desc.ToString();
+		for (const FString& Filter : FilterStrings)
+		{
+			if (ItemName.Contains(Filter))
+			{
+				return true;
+			}
+		}
+		return false;
+	};
+
+	for (const TSharedPtr<FStateTreeStateItem>& SourceItem : SourceArray)
 	{
-		MenuBuilder.BeginSection(FName(), LOCTEXT("TransitionGotoStateSection", "Go to State"));
+		// Check if our name matches the filters
+		// If bParentMatches is true, the search matched a parent category.
+		const bool bMatchesFilters = bParentMatches || MatchFilter(SourceItem);
 
-		for (int32 Idx = 0; Idx < CachedNames.Num(); Idx++)
+		int32 NumChildren = 0;
+		if (bMatchesFilters)
 		{
-			FUIAction ItemAction(FExecuteAction::CreateSP(const_cast<FStateTreeStateLinkDetails*>(this), &FStateTreeStateLinkDetails::OnStateComboChange, Idx));
-			MenuBuilder.AddMenuEntry(FText::FromName(CachedNames[Idx]), FText::Format(LOCTEXT("TransitionGotoStateTooltip", "Go to State {0}."), FText::FromName(CachedNames[Idx])), FSlateIcon(), ItemAction);
+			NumChildren++;
 		}
 
-		MenuBuilder.EndSection();
+		// if we don't match, then we still want to check all our children
+		TArray<TSharedPtr<FStateTreeStateItem>> FilteredChildren;
+		NumChildren += FilterStateItemChildren(FilterStrings, bMatchesFilters, SourceItem->Children, FilteredChildren);
+
+		// then add this item to the destination array
+		if (NumChildren > 0)
+		{
+			TSharedPtr<FStateTreeStateItem>& NewItem = OutDestArray.Add_GetRef(MakeShared<FStateTreeStateItem>());
+			*NewItem = *SourceItem;
+			NewItem->Children = FilteredChildren;
+
+			NumFound += NumChildren;
+		}
 	}
 
-	return MenuBuilder.MakeWidget();
+	return NumFound;
 }
 
-FText FStateTreeStateLinkDetails::GetCurrentStateDesc() const
+void FStateTreeStateLinkDetails::RestoreExpansionState()
 {
-	const EStateTreeTransitionType TransitionType = GetTransitionType().Get(EStateTreeTransitionType::Failed);
+	check(StateItemTree.IsValid());
+	
+	FStateExpansionState& ExpansionState = StateExpansionStates.FindOrAdd(FObjectKey(WeakStateTree.Get()));
+
+	bIsRestoringExpansion = true;
+
+	// Default state is expanded.
+	ExpandAll(FilteredRootItem->Children);
 
-	FText Result;
-	switch (TransitionType)
-	{
-	case EStateTreeTransitionType::None:
-		Result = LOCTEXT("TransitionNone", "None");
-		break;
-	case EStateTreeTransitionType::NextState:
-		Result = LOCTEXT("TransitionNextState", "Next State");
-		break;
-	case EStateTreeTransitionType::NextSelectableState:
-		Result = LOCTEXT("TransitionNextSelectableState", "Next Selectable State");
-		break;
-	case EStateTreeTransitionType::Succeeded:
-		Result = LOCTEXT("TransitionTreeSucceeded", "Tree Succeeded");
-		break;
-	case EStateTreeTransitionType::Failed:
-		Result = LOCTEXT("TransitionTreeFailed", "Tree Failed");
-		break;
-	case EStateTreeTransitionType::GotoState:
+	// Collapse the ones that are specifically collapsed.
+	for (const FGuid& StateID : ExpansionState.CollapsedStates)
+	{
+		TArray<TSharedPtr<FStateTreeStateItem>> Path;
+		if (FindStateByIDRecursive(FilteredRootItem, EStateTreeTransitionType::GotoState, StateID, Path))
 		{
-			FName OldName;
-			if (NameProperty && IDProperty)
-			{
-				NameProperty->GetValue(OldName);
+			StateItemTree->SetItemExpansion(Path.Last(), false);
+		}
+	}
 
-				FGuid StateID;
-				if (UE::StateTree::PropertyHelpers::GetStructValue<FGuid>(IDProperty, StateID) == FPropertyAccess::Success)
-				{
-					if (!StateID.IsValid())
-					{
-						return LOCTEXT("None", "None");
-					}
-					else
-					{
-						for (int32 Idx = 0; Idx < CachedIDs.Num(); Idx++)
-						{
-							if (CachedIDs[Idx] == StateID)
-							{
-								return FText::FromName(CachedNames[Idx]);
-							}
-						}
-					}
-				}
-			}
+	bIsRestoringExpansion = false;
+}
 
-			return FText::FromName(OldName);
+void FStateTreeStateLinkDetails::ExpandAll(const TArray<TSharedPtr<FStateTreeStateItem>>& Items)
+{
+	for (const TSharedPtr<FStateTreeStateItem>& Item : Items)
+	{
+		StateItemTree->SetItemExpansion(Item, true);
+		ExpandAll(Item->Children);
+	}
+}
+
+bool FStateTreeStateLinkDetails::FindStateByIDRecursive(const TSharedPtr<FStateTreeStateItem>& Item, const EStateTreeTransitionType TransitionType, const FGuid StateID, TArray<TSharedPtr<FStateTreeStateItem>>& OutPath)
+{
+	OutPath.Push(Item);
+
+	if (!Item->Desc.IsEmpty()
+		&& Item->TransitionType == TransitionType
+		&& Item->StateID == StateID)
+	{
+		return true;
+	}
+
+	for (const TSharedPtr<FStateTreeStateItem>& ChildItem : Item->Children)
+	{
+		if (FindStateByIDRecursive(ChildItem, TransitionType, StateID, OutPath))
+		{
+			return true;
 		}
-		break;
-	default:
-		ensureMsgf(false, TEXT("FStateTreeStateLinkDetails: Unhnandled enum %s"), *UEnum::GetValueAsString(TransitionType));
-		Result = LOCTEXT("TransitionInvalid", "Invalid");
-		break;
 	}
 
-	return Result;
+	OutPath.Pop();
+
+	return false;
 }
 
-bool FStateTreeStateLinkDetails::IsValidLink() const
+bool FStateTreeStateLinkDetails::GetCurrentStateItem(TArray<TSharedPtr<FStateTreeStateItem>>& OutPath) const
 {
+	OutPath.Reset();
 	const EStateTreeTransitionType TransitionType = GetTransitionType().Get(EStateTreeTransitionType::Failed);
+	FGuid StateID;
+	if (IDProperty
+		&& UE::StateTree::PropertyHelpers::GetStructValue<FGuid>(IDProperty, StateID) == FPropertyAccess::Success)
+	{
+		// Use RootItem (not filtered) to get any result even during search.
+		return FindStateByIDRecursive(RootItem, TransitionType, StateID, OutPath);
+	}
+	return false;
+}
 
-	bool bIsValid = false;
-	switch (TransitionType)
-	{
-	// all these are unconditionally valid		
-	case EStateTreeTransitionType::None: // fall through on purpose
-	case EStateTreeTransitionType::NextState: // fall through on purpose
-	case EStateTreeTransitionType::NextSelectableState: // fall through on purpose
-	case EStateTreeTransitionType::Succeeded: // fall through on purpose
-	case EStateTreeTransitionType::Failed:
-		bIsValid = true;
-		break;
-	case EStateTreeTransitionType::GotoState:
+FText FStateTreeStateLinkDetails::GetCurrentStateDesc() const
+{
+	TArray<TSharedPtr<FStateTreeStateItem>> Path;
+	if (GetCurrentStateItem(Path))
+	{
+		return Path.Last()->Desc;
+	}
+	// Could not find item, try to display old name as a hint what is missing.
+	if (NameProperty)
+	{
+		FName OldName;
+		if (NameProperty->GetValue(OldName) == FPropertyAccess::Success)
 		{
-			bIsValid = false;
-			if (NameProperty && IDProperty)
-			{
-				FGuid StateID;
-				if (UE::StateTree::PropertyHelpers::GetStructValue<FGuid>(IDProperty, StateID) == FPropertyAccess::Success)
-				{
-					bIsValid = !StateID.IsValid() || CachedIDs.Contains(StateID);
-				}
-			}
+			return FText::FromName(OldName);
 		}
-		break;
-	default:
-		bIsValid = false;
-		break;
 	}
+	return LOCTEXT("TransitionInvalid", "Invalid");
+}
+
+const FSlateBrush* FStateTreeStateLinkDetails::GetCurrentStateIcon() const
+{
+	TArray<TSharedPtr<FStateTreeStateItem>> Path;
+	if (GetCurrentStateItem(Path))
+	{
+		return Path.Last()->Icon;
+	}
+	return nullptr;
+}
+
+FSlateColor FStateTreeStateLinkDetails::GetCurrentStateColor() const
+{
+	TArray<TSharedPtr<FStateTreeStateItem>> Path;
+	if (GetCurrentStateItem(Path))
+	{
+		return Path.Last()->Color;
+	}
+	return {};
+}
+
+FSlateFontInfo FStateTreeStateLinkDetails::GetCurrentStateFont() const
+{
+	TArray<TSharedPtr<FStateTreeStateItem>> Path;
+	if (GetCurrentStateItem(Path))
+	{
+		return GetItemFont(Path.Last());
+	}
+	return FAppStyle::Get().GetFontStyle("NormalText");
+}
 
-	return bIsValid;
+bool FStateTreeStateLinkDetails::IsValidLink() const
+{
+	TArray<TSharedPtr<FStateTreeStateItem>> Path;
+	if (GetCurrentStateItem(Path))
+	{
+		return true;
+	}
+	// The state is missing.
+	return false;
 }
 
 TOptional<EStateTreeTransitionType> FStateTreeStateLinkDetails::GetTransitionType() const
diff --git a/Engine/Plugins/Runtime/StateTree/Source/StateTreeEditorModule/Private/Customizations/StateTreeStateLinkDetails.h b/Engine/Plugins/Runtime/StateTree/Source/StateTreeEditorModule/Private/Customizations/StateTreeStateLinkDetails.h
--- a/Engine/Plugins/Runtime/StateTree/Source/StateTreeEditorModule/Private/Customizations/StateTreeStateLinkDetails.h
+++ b/Engine/Plugins/Runtime/StateTree/Source/StateTreeEditorModule/Private/Customizations/StateTreeStateLinkDetails.h
@@ -3,12 +3,16 @@
 #pragma once
 
 #include "IPropertyTypeCustomization.h"
+#include "StateTreeTypes.h"
+#include "Widgets/Views/STreeView.h"
 
 enum class EStateTreeTransitionType : uint8;
 template <typename OptionalType> struct TOptional;
 
 class IPropertyHandle;
 class SWidget;
+class SSearchBox;
+class SComboButton;
 class UStateTree;
 class UStateTreeState;
 
@@ -27,36 +31,84 @@ public:
 	virtual void CustomizeChildren(TSharedRef<class IPropertyHandle> StructPropertyHandle, class IDetailChildrenBuilder& StructBuilder, IPropertyTypeCustomizationUtils& StructCustomizationUtils) override;
 
 private:
+	// Stores info about state
+	struct FStateTreeStateItem : TSharedFromThis<FStateTreeStateItem>
+	{
+		FStateTreeStateItem() = default;
+		
+		FStateTreeStateItem(const FText& InDesc, const FText& InTooltipText, const FSlateBrush* InIcon, const EStateTreeTransitionType InType)
+			: Desc(InDesc)
+			, TooltipText(InTooltipText)
+			, TransitionType(InType)
+			, Icon(InIcon)
+		{
+		}
 
-	// Special indices for the combo selection.
-	inline static constexpr int ComboNotSet = -1;
-	inline static constexpr int ComboSucceeded = -2;
-	inline static constexpr int ComboFailed = -3;
-	inline static constexpr int ComboNextState = -4;
-	inline static constexpr int ComboNextSelectableState = -5;
+		FText Desc;
+		FText TooltipText;
+		EStateTreeTransitionType TransitionType = EStateTreeTransitionType::None;
+		FGuid StateID = {};
+		FSlateColor Color = FSlateColor(FLinearColor::White);
+		const FSlateBrush* Icon = nullptr;
+		bool bIsSubTree = false;
+		bool bIsLinked = false;
+		FText LinkedDesc;
+		TArray<TSharedPtr<FStateTreeStateItem>> Children;
+	};
+
+	// Stores per session node expansion state for a node type.
+	struct FStateExpansionState
+	{
+		TSet<FGuid> CollapsedStates;
+	};
 
 	void CacheStates();
-	void CacheStates(const UStateTreeState* State);
+	void CacheStates(TSharedPtr<FStateTreeStateItem> ParentNode, const UStateTreeState* State);
 	void OnIdentifierChanged(const UStateTree& StateTree);
 
-	void OnStateComboChange(int Idx);
-	TSharedRef<SWidget> OnGetStateContent() const;
+	TSharedRef<SWidget> GenerateStatePicker();
 	FText GetCurrentStateDesc() const;
+	const FSlateBrush* GetCurrentStateIcon() const;
+	FSlateColor GetCurrentStateColor() const;
+	FSlateFontInfo GetCurrentStateFont() const;
 	bool IsValidLink() const;
 	TOptional<EStateTreeTransitionType> GetTransitionType() const;
+	static FSlateFontInfo GetItemFont(TSharedPtr<FStateTreeStateItem> Item);
+	TSharedRef<ITableRow> GenerateStateItemRow(TSharedPtr<FStateTreeStateItem> Item, const TSharedRef<STableViewBase>& OwnerTable);
+	void GetStateItemChildren(TSharedPtr<FStateTreeStateItem> Item, TArray<TSharedPtr<FStateTreeStateItem>>& OutItems) const;
+	void OnStateItemSelected(TSharedPtr<FStateTreeStateItem> SelectedItem, ESelectInfo::Type);
+	void OnStateItemExpansionChanged(TSharedPtr<FStateTreeStateItem> ExpandedItem, bool bInExpanded) const;
+	void OnSearchBoxTextChanged(const FText& NewText);
+	static int32 FilterStateItemChildren(const TArray<FString>& FilterStrings, const bool bParentMatches, const TArray<TSharedPtr<FStateTreeStateItem>>& SourceArray, TArray<TSharedPtr<FStateTreeStateItem>>& OutDestArray);
+	void ExpandAll(const TArray<TSharedPtr<FStateTreeStateItem>>& Items);
+	
+	static bool FindStateByIDRecursive(const TSharedPtr<FStateTreeStateItem>& Item, EStateTreeTransitionType TransitionType, const FGuid StateID, TArray<TSharedPtr<FStateTreeStateItem>>& OutPath);
+	bool GetCurrentStateItem(TArray<TSharedPtr<FStateTreeStateItem>>& OutPath) const;
+	
+	void RestoreExpansionState();
 
-	TArray<FName> CachedNames;
-	TArray<FGuid> CachedIDs;
+	TSharedPtr<FStateTreeStateItem> RootItem;
+	TSharedPtr<FStateTreeStateItem> FilteredRootItem;
 
 	TSharedPtr<IPropertyHandle> NameProperty;
 	TSharedPtr<IPropertyHandle> IDProperty;
 	TSharedPtr<IPropertyHandle> LinkTypeProperty;
 
+	TSharedPtr<SComboButton> ComboButton;
+	TSharedPtr<SSearchBox> SearchBox;
+	
+	TSharedPtr<STreeView<TSharedPtr<FStateTreeStateItem>>> StateItemTree;
+	bool bIsRestoringExpansion = false;
+	TWeakObjectPtr<const UStateTree> WeakStateTree = nullptr;
+
+	// Save expansion state for each base node type. The expansion state does not persist between editor sessions. 
+	static TMap<FObjectKey, FStateExpansionState> StateExpansionStates;
+
 	// If set, hide selecting meta states like Next or (tree) Succeeded.
 	bool bDirectStatesOnly = false;
 	// If set, allow to select only states marked as subtrees.
 	bool bSubtreesOnly = false;
 	
-	class IPropertyUtilities* PropUtils;
+	class IPropertyUtilities* PropUtils = nullptr;
 	TSharedPtr<IPropertyHandle> StructProperty;
 };
diff --git a/Engine/Plugins/Runtime/StateTree/Source/StateTreeEditorModule/Private/SStateTreeViewRow.cpp b/Engine/Plugins/Runtime/StateTree/Source/StateTreeEditorModule/Private/SStateTreeViewRow.cpp
--- a/Engine/Plugins/Runtime/StateTree/Source/StateTreeEditorModule/Private/SStateTreeViewRow.cpp
+++ b/Engine/Plugins/Runtime/StateTree/Source/StateTreeEditorModule/Private/SStateTreeViewRow.cpp
@@ -807,7 +807,9 @@ const FSlateBrush* SStateTreeViewRow::GetSelectorIcon() const
 		}
 		else if (State->SelectionBehavior == EStateTreeStateSelectionBehavior::TrySelectChildrenInOrder)
 		{
-			if (State->Children.IsEmpty())
+			if (State->Children.IsEmpty()
+				|| State->Type == EStateTreeStateType::Linked
+				|| State->Type == EStateTreeStateType::LinkedAsset)
 			{
 				// Backwards compatible behavior
 				return FStateTreeEditorStyle::Get().GetBrush("StateTreeEditor.TryEnterState");			
diff --git a/Engine/Plugins/Runtime/StateTree/Source/StateTreeEditorModule/Private/StateTreeEditorStyle.cpp b/Engine/Plugins/Runtime/StateTree/Source/StateTreeEditorModule/Private/StateTreeEditorStyle.cpp
--- a/Engine/Plugins/Runtime/StateTree/Source/StateTreeEditorModule/Private/StateTreeEditorStyle.cpp
+++ b/Engine/Plugins/Runtime/StateTree/Source/StateTreeEditorModule/Private/StateTreeEditorStyle.cpp
@@ -258,6 +258,11 @@ FStateTreeEditorStyle::FStateTreeEditorStyle()
 		Set("StateTreeEditor.EnableStates", new IMAGE_BRUSH("Icons/Empty_16x", CoreStyleConstants::Icon16x16));
 		Set("StateTreeEditor.Debugger.Breakpoint.EnabledAndValid", new IMAGE_BRUSH_SVG( "Starship/Blueprints/Breakpoint_Valid", CoreStyleConstants::Icon16x16, FStyleColors::AccentRed));
 		Set("StateTreeEditor.Debugger.ResumeDebuggerAnalysis", new IMAGE_BRUSH_SVG("Starship/Common/Timeline", CoreStyleConstants::Icon16x16));
+
+		Set("StateTreeEditor.Transition.None", new CORE_IMAGE_BRUSH_SVG("Starship/Common/x-circle", CoreStyleConstants::Icon16x16, FSlateColor::UseSubduedForeground()));
+		Set("StateTreeEditor.Transition.Next", new CORE_IMAGE_BRUSH_SVG("Starship/Common/arrow-down", CoreStyleConstants::Icon16x16, FSlateColor::UseSubduedForeground()));
+		Set("StateTreeEditor.Transition.Succeeded", new CORE_IMAGE_BRUSH_SVG("Starship/Common/check", CoreStyleConstants::Icon16x16, FStyleColors::AccentGreen));
+		Set("StateTreeEditor.Transition.Failed", new CORE_IMAGE_BRUSH_SVG("Starship/Common/close", CoreStyleConstants::Icon16x16, FStyleColors::AccentRed));
 	}
 
 	{
@@ -278,8 +283,6 @@ FStateTreeEditorStyle::FStateTreeEditorStyle()
 		Set("StateTreeEditor.StateSubtree", new IMAGE_BRUSH_SVG("Icons/State_Subtree", CoreStyleConstants::Icon16x16));
 		Set("StateTreeEditor.StateLinked", new IMAGE_BRUSH_SVG("Icons/State_Linked", CoreStyleConstants::Icon16x16));
 	}
-
-
 }
 
 void FStateTreeEditorStyle::Register()
