CIS fix in 29.40 rollback original fix as the one I have made earlier today is conflicting with this one but it was in 29.20 branches which was for 5.4 release.

diff --git a/Engine/Plugins/Runtime/StateTree/Source/StateTreeEditorModule/Private/SStateTreeViewRow.cpp b/Engine/Plugins/Runtime/StateTree/Source/StateTreeEditorModule/Private/SStateTreeViewRow.cpp
--- a/Engine/Plugins/Runtime/StateTree/Source/StateTreeEditorModule/Private/SStateTreeViewRow.cpp
+++ b/Engine/Plugins/Runtime/StateTree/Source/StateTreeEditorModule/Private/SStateTreeViewRow.cpp
@@ -1128,17 +1128,16 @@ FText SStateTreeViewRow::GetTransitionsIcon(const UStateTreeState& State, const
 		IconType = IconLevelUp;
 	}
 
-	if (IconType == IconRightArrow)
+	switch (IconType)
 	{
-		return FEditorFontGlyphs::Long_Arrow_Right;
-	}
-	else if (IconType == IconDownArrow)
-	{
-		return FEditorFontGlyphs::Long_Arrow_Down;
-	}
-	else if (IconType == IconLevelUp)
-	{
-		return FEditorFontGlyphs::Level_Up;
+		case IconRightArrow:
+			return FEditorFontGlyphs::Long_Arrow_Right;
+		case IconDownArrow:
+			return FEditorFontGlyphs::Long_Arrow_Down;
+		case IconLevelUp:
+			return FEditorFontGlyphs::Level_Up;
+		default:
+			return FText::GetEmpty();
 	}
 }
 
