[State Tree Editor] Buttons for creating new blueprint tasks and conditions from existing blueprint classes

diff --git a/Engine/Plugins/Runtime/StateTree/Source/StateTreeEditorModule/Private/StateTreeEditor.cpp b/Engine/Plugins/Runtime/StateTree/Source/StateTreeEditorModule/Private/StateTreeEditor.cpp
--- a/Engine/Plugins/Runtime/StateTree/Source/StateTreeEditorModule/Private/StateTreeEditor.cpp
+++ b/Engine/Plugins/Runtime/StateTree/Source/StateTreeEditorModule/Private/StateTreeEditor.cpp
@@ -1,9 +1,19 @@
 // Copyright Epic Games, Inc. All Rights Reserved.
 
 #include "StateTreeEditor.h"
+#include "AssetRegistry/AssetRegistryModule.h"
+#include "Engine/BlueprintGeneratedClass.h"
+#include "Blueprint/StateTreeTaskBlueprintBase.h"
+#include "Blueprint/StateTreeConditionBlueprintBase.h"
+#include "ContentBrowserModule.h"
+#include "ClassViewerFilter.h"
+#include "Kismet2/KismetEditorUtilities.h"
+#include "Kismet2/BlueprintEditorUtils.h"
 #include "Customizations/StateTreeBindingExtension.h"
 #include "DetailsViewArgs.h"
 #include "IDetailsView.h"
+#include "IContentBrowserSingleton.h"
+#include "Framework/Application/SlateApplication.h"
 #include "Framework/Commands/UICommandList.h"
 #include "IMessageLogListing.h"
 #include "MessageLogModule.h"
@@ -816,6 +826,25 @@ namespace UE::StateTree::Editor
 		return EditorDataHash;
 	}
 
+	template <typename ClassType, typename = typename TEnableIf<TIsDerivedFrom<ClassType, UStateTreeNodeBlueprintBase>::Value>::Type>
+	class FEditorNodeClassFilter : public IClassViewerFilter
+	{
+	public:
+		virtual bool IsClassAllowed(const FClassViewerInitializationOptions& InInitOptions, const UClass* InClass, TSharedRef< FClassViewerFilterFuncs > InFilterFuncs) override
+		{
+			check(InClass);
+			return InClass->IsChildOf(ClassType::StaticClass());
+		}
+
+		virtual bool IsUnloadedClassAllowed(const FClassViewerInitializationOptions& InInitOptions, const TSharedRef< const IUnloadedBlueprintData > InUnloadedClassData, TSharedRef< FClassViewerFilterFuncs > InFilterFuncs) override
+		{
+			return InUnloadedClassData->IsChildOf(ClassType::StaticClass());
+		}
+	};
+
+	using FStateTreeTaskBPClassFilter = FEditorNodeClassFilter<UStateTreeTaskBlueprintBase>;
+	using FStateTreeConditionBPClassFilter = FEditorNodeClassFilter<UStateTreeConditionBlueprintBase>;
+
 }; // UE::StateTree::Editor
 
 void FStateTreeEditor::BindCommands()
@@ -868,16 +897,16 @@ void FStateTreeEditor::RegisterToolbar()
 		ToolBar = UToolMenus::Get()->RegisterMenu(MenuName, ParentName, EMultiBoxType::ToolBar);
 	}
 
-	const FToolMenuInsert InsertAfterAssetSection("Asset", EToolMenuInsertType::After);
+	static const FToolMenuInsert InsertAfterAssetSection("Asset", EToolMenuInsertType::After);
 
-	FToolMenuSection& Section = ToolBar->AddSection("Compile", TAttribute<FText>(), InsertAfterAssetSection);
+	FToolMenuSection& CompileSection = ToolBar->AddSection("Compile", TAttribute<FText>(), InsertAfterAssetSection);
 	
-	Section.AddDynamicEntry("CompileCommands", FNewToolMenuSectionDelegate::CreateLambda([](FToolMenuSection& InSection)
+	CompileSection.AddDynamicEntry("CompileCommands", FNewToolMenuSectionDelegate::CreateLambda([](FToolMenuSection& InSection)
 	{
 		const UStateTreeToolMenuContext* Context = InSection.FindContext<UStateTreeToolMenuContext>();
 		if (Context && Context->StateTreeEditor.IsValid())
 		{
-			TSharedPtr<FStateTreeEditor> StateTreeEditor = Context->StateTreeEditor.Pin();
+			const TSharedPtr<FStateTreeEditor> StateTreeEditor = Context->StateTreeEditor.Pin();
 			if (StateTreeEditor.IsValid())
 			{
 				const FStateTreeEditorCommands& Commands = FStateTreeEditorCommands::Get();
@@ -900,6 +929,40 @@ void FStateTreeEditor::RegisterToolbar()
 			}
 		}
 	}));
+
+	static const FToolMenuInsert InsertAfterCompileSection("Compile", EToolMenuInsertType::After);
+
+	FToolMenuSection& CreateNewNodeSection = ToolBar->AddSection("CreateNewTaskAndCondition", TAttribute<FText>(), InsertAfterCompileSection);
+
+	CreateNewNodeSection.AddDynamicEntry("CreateNewTaskAndCondition", FNewToolMenuSectionDelegate::CreateLambda([](FToolMenuSection& InSection)
+	{
+		const UStateTreeToolMenuContext* Context = InSection.FindContext<UStateTreeToolMenuContext>();
+		if (Context && Context->StateTreeEditor.IsValid())
+		{
+			const TSharedPtr<FStateTreeEditor> StateTreeEditor = Context->StateTreeEditor.Pin();
+			if (StateTreeEditor.IsValid())
+			{
+				const TSharedRef<FStateTreeEditor> StateTreeEditorRef = StateTreeEditor.ToSharedRef();
+				FToolMenuEntry& CreateNewTaskDropdown = InSection.AddEntry(FToolMenuEntry::InitComboButton(
+					"CreateNewTaskComboButton",
+					FUIAction(),
+					FOnGetContent::CreateSP(StateTreeEditorRef, &FStateTreeEditor::GenerateTaskBPBaseClassesMenu),
+					LOCTEXT("CreateNewTask_Title", "New Task"),
+					LOCTEXT("CreateNewTask_ToolbarTooltip", "Create a new Blueprint State Tree Task"),
+					TAttribute<FSlateIcon>(StateTreeEditorRef, &FStateTreeEditor::GetNewTaskButtonImage)
+				));
+
+				FToolMenuEntry& CreateNewConditionDropdown = InSection.AddEntry(FToolMenuEntry::InitComboButton(
+					"CreateNewConditionComboButton",
+					FUIAction(),
+					FOnGetContent::CreateSP(StateTreeEditorRef, &FStateTreeEditor::GenerateConditionBPBaseClassesMenu),
+					LOCTEXT("CreateNewCondition_Title", "New Condition"),
+					LOCTEXT("CreateNewCondition_ToolbarTooltip", "Create a new Blueprint State Tree Condition"),
+					TAttribute<FSlateIcon>(StateTreeEditorRef, &FStateTreeEditor::GetNewConditionButtonImage)
+				));
+			}
+		}
+	}));
 }
 
 void FStateTreeEditor::Compile()
@@ -1000,6 +1063,86 @@ FSlateIcon FStateTreeEditor::GetCompileStatusImage() const
 	return FSlateIcon(FAppStyle::GetAppStyleSetName(), CompileStatusBackground, NAME_None, CompileStatusGood);
 }
 
+FSlateIcon FStateTreeEditor::GetNewConditionButtonImage() const
+{
+	//placeholder
+	return FSlateIcon(FAppStyle::GetAppStyleSetName(), NAME_None, NAME_None, NAME_None);
+}
+
+TSharedRef<SWidget> FStateTreeEditor::GenerateConditionBPBaseClassesMenu() const
+{
+	FClassViewerInitializationOptions Options;
+	Options.NameTypeToDisplay = EClassViewerNameTypeToDisplay::DisplayName;
+	Options.ClassFilters.Add(MakeShareable(new UE::StateTree::Editor::FStateTreeConditionBPClassFilter));
+
+	FOnClassPicked OnPicked(FOnClassPicked::CreateSP(this, &FStateTreeEditor::OnNodeBPBaseClassPicked));
+
+	return FModuleManager::LoadModuleChecked<FClassViewerModule>("ClassViewer").CreateClassViewer(Options, OnPicked);
+}
+
+FSlateIcon FStateTreeEditor::GetNewTaskButtonImage() const
+{
+	//placeholder
+	return FSlateIcon(FAppStyle::GetAppStyleSetName(), NAME_None, NAME_None, NAME_None);
+}
+
+TSharedRef<SWidget> FStateTreeEditor::GenerateTaskBPBaseClassesMenu() const
+{
+	FClassViewerInitializationOptions Options;
+	Options.NameTypeToDisplay = EClassViewerNameTypeToDisplay::DisplayName;
+	Options.ClassFilters.Add(MakeShareable(new UE::StateTree::Editor::FStateTreeTaskBPClassFilter));
+
+	FOnClassPicked OnPicked(FOnClassPicked::CreateSP(this, &FStateTreeEditor::OnNodeBPBaseClassPicked));
+
+	return FModuleManager::LoadModuleChecked<FClassViewerModule>("ClassViewer").CreateClassViewer(Options, OnPicked);
+}
+
+void FStateTreeEditor::OnNodeBPBaseClassPicked(UClass* NodeClass) const
+{
+	check(NodeClass);
+	
+	if (!StateTree)
+	{
+		return;
+	}
+
+	const FString ClassName = FBlueprintEditorUtils::GetClassNameWithoutSuffix(NodeClass);
+	const FString PathName = FPaths::GetPath(StateTree->GetOutermost()->GetPathName());
+
+	// Now that we've generated some reasonable default locations/names for the package, allow the user to have the final say
+	// before we create the package and initialize the blueprint inside of it.
+	FSaveAssetDialogConfig SaveAssetDialogConfig;
+	SaveAssetDialogConfig.DialogTitleOverride = LOCTEXT("SaveAssetDialogTitle", "Save Asset As");
+	SaveAssetDialogConfig.DefaultPath = PathName;
+	SaveAssetDialogConfig.DefaultAssetName = ClassName + TEXT("_New");
+	SaveAssetDialogConfig.ExistingAssetPolicy = ESaveAssetDialogExistingAssetPolicy::Disallow;
+
+	const FContentBrowserModule& ContentBrowserModule = FModuleManager::LoadModuleChecked<FContentBrowserModule>("ContentBrowser");
+	const FString SaveObjectPath = ContentBrowserModule.Get().CreateModalSaveAssetDialog(SaveAssetDialogConfig);
+	if (!SaveObjectPath.IsEmpty())
+	{
+		const FString SavePackageName = FPackageName::ObjectPathToPackageName(SaveObjectPath);
+		const FString SavePackagePath = FPaths::GetPath(SavePackageName);
+		const FString SaveAssetName = FPaths::GetBaseFilename(SavePackageName);
+
+		if (UPackage* Package = CreatePackage(*SavePackageName))
+		{
+			// Create and init a new Blueprint
+			if (UBlueprint* NewBP = FKismetEditorUtilities::CreateBlueprint(NodeClass, Package, FName(*SaveAssetName), BPTYPE_Normal, UBlueprint::StaticClass(), UBlueprintGeneratedClass::StaticClass()))
+			{
+				GEditor->GetEditorSubsystem<UAssetEditorSubsystem>()->OpenEditorForAsset(NewBP);
+
+				// Notify the asset registry
+				FAssetRegistryModule::AssetCreated(NewBP);
+
+				Package->MarkPackageDirty();
+			}
+		}
+	}
+
+	FSlateApplication::Get().DismissAllMenus();
+}
+
 void FStateTreeEditor::UpdateAsset()
 {
 	if (!StateTree)
diff --git a/Engine/Plugins/Runtime/StateTree/Source/StateTreeEditorModule/Private/StateTreeEditor.h b/Engine/Plugins/Runtime/StateTree/Source/StateTreeEditorModule/Private/StateTreeEditor.h
--- a/Engine/Plugins/Runtime/StateTree/Source/StateTreeEditorModule/Private/StateTreeEditor.h
+++ b/Engine/Plugins/Runtime/StateTree/Source/StateTreeEditorModule/Private/StateTreeEditor.h
@@ -8,6 +8,7 @@
 class FSpawnTabArgs;
 class IMessageToken;
 class UStateTreeState;
+class UStateTreeNodeBlueprintBase;
 
 class IDetailsView;
 class UStateTree;
@@ -82,6 +83,14 @@ private:
 	bool CanCompile() const;
 	FSlateIcon GetCompileStatusImage() const;
 
+	FSlateIcon GetNewTaskButtonImage() const;
+	TSharedRef<SWidget> GenerateTaskBPBaseClassesMenu() const;
+
+	FSlateIcon GetNewConditionButtonImage() const;
+	TSharedRef<SWidget> GenerateConditionBPBaseClassesMenu() const;
+	
+	void OnNodeBPBaseClassPicked(UClass* NodeClass) const;
+
 	void UpdateAsset();
 
 	void HandleModelAssetChanged();
