Add support for state tree override in StateTreeRunParallelStateTreeTask
#rnx

diff --git a/Engine/Plugins/Runtime/StateTree/Source/StateTreeModule/Private/StateTreeExecutionContext.cpp b/Engine/Plugins/Runtime/StateTree/Source/StateTreeModule/Private/StateTreeExecutionContext.cpp
--- a/Engine/Plugins/Runtime/StateTree/Source/StateTreeModule/Private/StateTreeExecutionContext.cpp
+++ b/Engine/Plugins/Runtime/StateTree/Source/StateTreeModule/Private/StateTreeExecutionContext.cpp
@@ -124,6 +124,7 @@ FStateTreeExecutionContext::FStateTreeExecutionContext(UObject& InOwner, const U
 FStateTreeExecutionContext::FStateTreeExecutionContext(const FStateTreeExecutionContext& InContextToCopy, const UStateTree& InStateTree, FStateTreeInstanceData& InInstanceData)
 	: FStateTreeExecutionContext(InContextToCopy.Owner, InStateTree, InInstanceData, InContextToCopy.CollectExternalDataDelegate)
 {
+	LinkedStateTreeOverrides = InContextToCopy.LinkedStateTreeOverrides;
 	const bool bIsSameSchema = RootStateTree.GetSchema()->GetClass() == InContextToCopy.GetStateTree()->GetSchema()->GetClass();
 	if (bIsSameSchema)
 	{
diff --git a/Engine/Plugins/Runtime/StateTree/Source/StateTreeModule/Private/Tasks/StateTreeRunParallelStateTreeTask.cpp b/Engine/Plugins/Runtime/StateTree/Source/StateTreeModule/Private/Tasks/StateTreeRunParallelStateTreeTask.cpp
--- a/Engine/Plugins/Runtime/StateTree/Source/StateTreeModule/Private/Tasks/StateTreeRunParallelStateTreeTask.cpp
+++ b/Engine/Plugins/Runtime/StateTree/Source/StateTreeModule/Private/Tasks/StateTreeRunParallelStateTreeTask.cpp
@@ -13,29 +13,31 @@ FStateTreeRunParallelStateTreeTask::FStateTreeRunParallelStateTreeTask()
 EStateTreeRunStatus FStateTreeRunParallelStateTreeTask::EnterState(FStateTreeExecutionContext& Context, const FStateTreeTransitionResult& Transitions) const
 {
 	FInstanceDataType& InstanceData = Context.GetInstanceData(*this);
-	if (!InstanceData.StateTree.IsValid())
+	const FStateTreeReference& StateTreeToRun = GetStateTreeToRun(Context, InstanceData);
+	if (!StateTreeToRun.IsValid())
 	{
 		return EStateTreeRunStatus::Failed;
 	}
 
-	FStateTreeExecutionContext ParallelTreeContext(Context, *InstanceData.StateTree.GetStateTree(), InstanceData.TreeInstanceData);
+	InstanceData.RunningStateTree = StateTreeToRun.GetStateTree();
+	FStateTreeExecutionContext ParallelTreeContext(Context, *InstanceData.RunningStateTree, InstanceData.TreeInstanceData);
 	if (!ParallelTreeContext.IsValid())
 	{
 		return EStateTreeRunStatus::Failed;
 	}
 
-	return ParallelTreeContext.Start(&InstanceData.StateTree.GetParameters());
+	return ParallelTreeContext.Start(&StateTreeToRun.GetParameters());
 }
 
 EStateTreeRunStatus FStateTreeRunParallelStateTreeTask::Tick(FStateTreeExecutionContext& Context, const float DeltaTime) const
 {
 	FInstanceDataType& InstanceData = Context.GetInstanceData(*this);
-	if (!InstanceData.StateTree.IsValid())
+	if (!InstanceData.RunningStateTree)
 	{
 		return EStateTreeRunStatus::Failed;
 	}
 
-	FStateTreeExecutionContext ParallelTreeContext(Context, *InstanceData.StateTree.GetStateTree(), InstanceData.TreeInstanceData);
+	FStateTreeExecutionContext ParallelTreeContext(Context, *InstanceData.RunningStateTree, InstanceData.TreeInstanceData);
 	if (!ParallelTreeContext.IsValid())
 	{
 		return EStateTreeRunStatus::Failed;
@@ -47,12 +49,12 @@ EStateTreeRunStatus FStateTreeRunParallelStateTreeTask::Tick(FStateTreeExecution
 void FStateTreeRunParallelStateTreeTask::ExitState(FStateTreeExecutionContext& Context, const FStateTreeTransitionResult& Transition) const
 {
 	FInstanceDataType& InstanceData = Context.GetInstanceData(*this);
-	if (!InstanceData.StateTree.IsValid())
+	if (!InstanceData.RunningStateTree)
 	{
 		return;
 	}
 
-	FStateTreeExecutionContext ParallelTreeContext(Context, *InstanceData.StateTree.GetStateTree(), InstanceData.TreeInstanceData);
+	FStateTreeExecutionContext ParallelTreeContext(Context, *InstanceData.RunningStateTree, InstanceData.TreeInstanceData);
 	if (!ParallelTreeContext.IsValid())
 	{
 		return;
@@ -61,6 +63,19 @@ void FStateTreeRunParallelStateTreeTask::ExitState(FStateTreeExecutionContext& C
 	ParallelTreeContext.Stop();
 }
 
+const FStateTreeReference& FStateTreeRunParallelStateTreeTask::GetStateTreeToRun(FStateTreeExecutionContext& Context, FInstanceDataType& InstanceData) const
+{
+	if (StateTreeOverrideTag.IsValid())
+	{
+		if (const FStateTreeReference* Override = Context.GetLinkedStateTreeOverrideForTag(StateTreeOverrideTag))
+		{
+			return *Override;
+		}
+	}
+
+	return InstanceData.StateTree;
+}
+
 #if WITH_EDITOR
 void FStateTreeRunParallelStateTreeTask::PostEditInstanceDataChangeChainProperty(const FPropertyChangedChainEvent& PropertyChangedEvent, FStateTreeDataView InstanceDataView)
 {
diff --git a/Engine/Plugins/Runtime/StateTree/Source/StateTreeModule/Private/Tasks/StateTreeRunParallelStateTreeTask.h b/Engine/Plugins/Runtime/StateTree/Source/StateTreeModule/Private/Tasks/StateTreeRunParallelStateTreeTask.h
--- a/Engine/Plugins/Runtime/StateTree/Source/StateTreeModule/Private/Tasks/StateTreeRunParallelStateTreeTask.h
+++ b/Engine/Plugins/Runtime/StateTree/Source/StateTreeModule/Private/Tasks/StateTreeRunParallelStateTreeTask.h
@@ -19,6 +19,9 @@ struct FStateTreeRunParallelStateTreeTaskInstanceData
 
 	UPROPERTY(Transient)
 	FStateTreeInstanceData TreeInstanceData;
+
+	UPROPERTY(Transient)
+	TObjectPtr<const UStateTree> RunningStateTree = nullptr;
 };
 
 /**
@@ -45,4 +48,10 @@ protected:
 	virtual void PostEditInstanceDataChangeChainProperty(const FPropertyChangedChainEvent& PropertyChangedEvent, FStateTreeDataView InstanceDataView) override;
 	virtual void PostLoad(FStateTreeDataView InstanceDataView) override;
 #endif // WITH_EDITOR
+
+	const FStateTreeReference& GetStateTreeToRun(FStateTreeExecutionContext& Context, FInstanceDataType& InstanceData) const;
+
+	/** If set the task will look at the linked state tree override to replace the state tree it's running. */
+	UPROPERTY(EditAnywhere, Category = Parameter)
+	FGameplayTag StateTreeOverrideTag;
 };
\ No newline at end of file
