StateTree: Made parent state required events and enter conditions to be evaluated when transition leads directly to it's child. (contributed)

diff --git a/Engine/Plugins/Runtime/StateTree/Source/StateTreeEditorModule/Private/Customizations/StateTreeStateDetails.cpp b/Engine/Plugins/Runtime/StateTree/Source/StateTreeEditorModule/Private/Customizations/StateTreeStateDetails.cpp
--- a/Engine/Plugins/Runtime/StateTree/Source/StateTreeEditorModule/Private/Customizations/StateTreeStateDetails.cpp
+++ b/Engine/Plugins/Runtime/StateTree/Source/StateTreeEditorModule/Private/Customizations/StateTreeStateDetails.cpp
@@ -54,9 +54,10 @@ void FStateTreeStateDetails::CustomizeDetails(IDetailLayoutBuilder& DetailBuilde
 	const TSharedPtr<IPropertyHandle> ParametersProperty = DetailBuilder.GetProperty(GET_MEMBER_NAME_CHECKED(UStateTreeState, Parameters));
 	const TSharedPtr<IPropertyHandle> SelectionBehaviorProperty = DetailBuilder.GetProperty(GET_MEMBER_NAME_CHECKED(UStateTreeState, SelectionBehavior));
 	const TSharedPtr<IPropertyHandle> RequiredEventToEnterProperty = DetailBuilder.GetProperty(GET_MEMBER_NAME_CHECKED(UStateTreeState, RequiredEventToEnter));
-
+	const TSharedPtr<IPropertyHandle> CheckPrerequisitesProperty = DetailBuilder.GetProperty(GET_MEMBER_NAME_CHECKED(UStateTreeState, bCheckPrerequisitesWhenActivatingChildDirectly));
 	
 
+
 	// Never show enabled
 	EnabledProperty->MarkHiddenByCustomization();
 
@@ -126,6 +127,9 @@ void FStateTreeStateDetails::CustomizeDetails(IDetailLayoutBuilder& DetailBuilde
 		];
 	}
 
+	CheckPrerequisitesProperty->MarkHiddenByCustomization();
+	StateCategory.AddProperty(CheckPrerequisitesProperty);
+
 	// Parameters
 	ParametersProperty->MarkHiddenByCustomization();
 	StateCategory.AddProperty(ParametersProperty);
diff --git a/Engine/Plugins/Runtime/StateTree/Source/StateTreeEditorModule/Private/StateTreeCompiler.cpp b/Engine/Plugins/Runtime/StateTree/Source/StateTreeEditorModule/Private/StateTreeCompiler.cpp
--- a/Engine/Plugins/Runtime/StateTree/Source/StateTreeEditorModule/Private/StateTreeCompiler.cpp
+++ b/Engine/Plugins/Runtime/StateTree/Source/StateTreeEditorModule/Private/StateTreeCompiler.cpp
@@ -415,6 +415,7 @@ bool FStateTreeCompiler::CreateStateRecursive(UStateTreeState& State, const FSta
 	CompactState.Tag = State.Tag;
 	CompactState.Parent = Parent;
 	CompactState.bEnabled = State.bEnabled;
+	CompactState.bCheckPrerequisitesWhenActivatingChildDirectly = State.bCheckPrerequisitesWhenActivatingChildDirectly;
 
 	CompactState.Type = State.Type;
 	CompactState.SelectionBehavior = State.SelectionBehavior;
diff --git a/Engine/Plugins/Runtime/StateTree/Source/StateTreeEditorModule/Private/StateTreeState.cpp b/Engine/Plugins/Runtime/StateTree/Source/StateTreeEditorModule/Private/StateTreeState.cpp
--- a/Engine/Plugins/Runtime/StateTree/Source/StateTreeEditorModule/Private/StateTreeState.cpp
+++ b/Engine/Plugins/Runtime/StateTree/Source/StateTreeEditorModule/Private/StateTreeState.cpp
@@ -393,6 +393,11 @@ void UStateTreeState::PostLoad()
 			}
 		}
 	}
+
+	if (CurrentVersion < FStateTreeCustomVersion::AddedCheckingParentsPrerequisites)
+	{
+		bCheckPrerequisitesWhenActivatingChildDirectly = false;
+	}
 	
 #endif // WITH_EDITORONLY_DATA
 
diff --git a/Engine/Plugins/Runtime/StateTree/Source/StateTreeEditorModule/Public/StateTreeState.h b/Engine/Plugins/Runtime/StateTree/Source/StateTreeEditorModule/Public/StateTreeState.h
--- a/Engine/Plugins/Runtime/StateTree/Source/StateTreeEditorModule/Public/StateTreeState.h
+++ b/Engine/Plugins/Runtime/StateTree/Source/StateTreeEditorModule/Public/StateTreeState.h
@@ -343,6 +343,10 @@ public:
 	UPROPERTY(EditDefaultsOnly, Category = "State")
 	TObjectPtr<UStateTree> LinkedAsset = nullptr;
 
+	/** Should state's required event and enter conditions be evaluated when transition leads directly to it's child. */
+	UPROPERTY(EditDefaultsOnly, Category = "State")
+	bool bCheckPrerequisitesWhenActivatingChildDirectly = true;
+
 	/** Parameters of this state. If the state is linked to another state or asset, the parameters are for the linked state. */
 	UPROPERTY(EditDefaultsOnly, Category = "State")
 	FStateTreeStateParameters Parameters;
diff --git a/Engine/Plugins/Runtime/StateTree/Source/StateTreeModule/Private/StateTreeExecutionContext.cpp b/Engine/Plugins/Runtime/StateTree/Source/StateTreeModule/Private/StateTreeExecutionContext.cpp
--- a/Engine/Plugins/Runtime/StateTree/Source/StateTreeModule/Private/StateTreeExecutionContext.cpp
+++ b/Engine/Plugins/Runtime/StateTree/Source/StateTreeModule/Private/StateTreeExecutionContext.cpp
@@ -3032,24 +3032,25 @@ bool FStateTreeExecutionContext::SelectState(const FStateTreeExecutionFrame& Cur
 	}
 
 	// Walk towards the root from current state.
-	TStaticArray<FStateTreeStateHandle, FStateTreeActiveStates::MaxStates> ParentStates;
-	int32 NumParentStates = 0;
+	TArray<FStateTreeStateHandle, TInlineAllocator<FStateTreeActiveStates::MaxStates>> PathToNextState;
 	FStateTreeStateHandle CurrState = NextState;
 	while (CurrState.IsValid())
 	{
-		if (NumParentStates == FStateTreeActiveStates::MaxStates)
+		if (PathToNextState.Num() == FStateTreeActiveStates::MaxStates)
 		{
 			STATETREE_LOG(Error, TEXT("%hs: Reached max execution depth when trying to select state %s from '%s'.  '%s' using StateTree '%s'."),
 				__FUNCTION__, *GetSafeStateName(CurrentFrame, NextState), *GetStateStatusString(Exec), *GetNameSafe(&Owner), *GetFullNameSafe(&RootStateTree));
 			return false;
 		}
 		// Store the states that are in between the 'NextState' and common ancestor. 
-		ParentStates[NumParentStates++] = CurrState;
+		PathToNextState.Push(CurrState);
 		CurrState = CurrentFrame.StateTree->States[CurrState.Index].Parent;
 	}
 
+	Algo::Reverse(PathToNextState);
+
 	const UStateTree* NextStateTree = CurrentFrame.StateTree;
-	const FStateTreeStateHandle NextRootState = ParentStates[NumParentStates - 1]; 
+	const FStateTreeStateHandle NextRootState = PathToNextState[0]; 
 
 	// Find the frame that the next state belongs to.
 	int32 CurrentFrameIndex = INDEX_NONE;
@@ -3097,14 +3098,27 @@ bool FStateTreeExecutionContext::SelectState(const FStateTreeExecutionFrame& Cur
 	const int32 LastFrameIndex = OutSelectionResult.FramesNum() - 1;
 	FStateTreeExecutionFrame& LastFrame = OutSelectionResult.GetSelectedFrames()[LastFrameIndex];
 
-	LastFrame.ActiveStates.Reset();
-	for (int32 Index = NumParentStates - 1; Index > 0; Index--)
+	// Find index of the first state to be evaluated.
+	int32 FirstNewStateIndex = 0;
+	if (CurrentFrameIndex != INDEX_NONE)
 	{
-		LastFrame.ActiveStates.Push(ParentStates[Index]);
+		// If LastFrame.ActiveStates is a subset of PathToNextState (e.g when someone use "TryEnter" selection behavior and then make a transition to it's child or if one is reentering the same state).
+		// In such case loop below won't break on anything and FirstNewStateIndex will be incorrectly 0, thus we initialize it to be right after the shorter range.
+		FirstNewStateIndex = FMath::Max(0, FMath::Min(PathToNextState.Num(), LastFrame.ActiveStates.Num()) - 1);
+		for (int32 Index = 0; Index < FMath::Min(PathToNextState.Num(), LastFrame.ActiveStates.Num()); ++Index)
+		{
+			if (LastFrame.ActiveStates[Index] != PathToNextState[Index])
+			{
+				FirstNewStateIndex = Index;
+				break;
+			}
+		}
 	}
 	// Existing state's data is safe to access during select.
 	LastFrame.NumCurrentlyActiveStates = static_cast<uint8>(LastFrame.ActiveStates.Num());
 
+	LastFrame.ActiveStates.SetNum(FirstNewStateIndex);
+
 	FStateSelectionResult InitialSelection;
 
 	if (Fallback == EStateTreeSelectionFallback::NextSelectableSibling)
@@ -3113,18 +3127,21 @@ bool FStateTreeExecutionContext::SelectState(const FStateTreeExecutionFrame& Cur
 	}
 	
 	// We take copy of the last frame and assign it later, as SelectStateInternal() might change the array and invalidate the pointer.
-	const FStateTreeExecutionFrame* CurrentParentFrame = LastFrameIndex > 0 ? &OutSelectionResult.GetSelectedFrames()[LastFrameIndex - 1] : nullptr; 
-	if (SelectStateInternal(CurrentParentFrame, OutSelectionResult.GetSelectedFrames()[LastFrameIndex], CurrentFrameInActiveFrames , NextState, OutSelectionResult, TransitionEvent))
+	const FStateTreeExecutionFrame* CurrentParentFrame = LastFrameIndex > 0 ? &OutSelectionResult.GetSelectedFrames()[LastFrameIndex - 1] : nullptr;
+
+	// Path from the first new state up to the NextState
+	TConstArrayView<FStateTreeStateHandle> NewStatesPathToNextState(&PathToNextState[FirstNewStateIndex], PathToNextState.Num() - FirstNewStateIndex);
+
+	if (SelectStateInternal(CurrentParentFrame, OutSelectionResult.GetSelectedFrames()[LastFrameIndex], CurrentFrameInActiveFrames, NewStatesPathToNextState, OutSelectionResult, TransitionEvent))
 	{
 		return true;
 	}
 
 	// Failed to Select Next State, handle fallback here
 	// Return true on the first next sibling that gets selected successfully
-	if (Fallback == EStateTreeSelectionFallback::NextSelectableSibling && NumParentStates >= 2)
+	if (Fallback == EStateTreeSelectionFallback::NextSelectableSibling && PathToNextState.Num() >= 2)
 	{
-		// InBetweenStates is in reversed order (i.e. from leaf to root)
-		const FStateTreeStateHandle Parent = ParentStates[1];
+		const FStateTreeStateHandle Parent = PathToNextState.Last(1);
 		if (Parent.IsValid())
 		{
 			const FCompactStateTreeState& ParentState = CurrentFrame.StateTree->States[Parent.Index];
@@ -3139,7 +3156,7 @@ bool FStateTreeExecutionContext::SelectState(const FStateTreeExecutionFrame& Cur
 	
 				// We take copy of the last frame and assign it later, as SelectStateInternal() might change the array and invalidate the pointer.
 				CurrentParentFrame = LastFrameIndex > 0 ? &OutSelectionResult.GetSelectedFrames()[LastFrameIndex - 1] : nullptr; 
-				if (SelectStateInternal(CurrentParentFrame, OutSelectionResult.GetSelectedFrames()[LastFrameIndex], CurrentFrameInActiveFrames, ChildStateHandle, OutSelectionResult))
+				if (SelectStateInternal(CurrentParentFrame, OutSelectionResult.GetSelectedFrames()[LastFrameIndex], CurrentFrameInActiveFrames, {ChildStateHandle}, OutSelectionResult))
 				{
 					return true;
 				}
@@ -3154,7 +3171,7 @@ bool FStateTreeExecutionContext::SelectStateInternal(
 	const FStateTreeExecutionFrame* CurrentParentFrame,
 	FStateTreeExecutionFrame& CurrentFrame,
 	const FStateTreeExecutionFrame* CurrentFrameInActiveFrames,
-	const FStateTreeStateHandle NextStateHandle,
+	TConstArrayView<FStateTreeStateHandle> PathToNextState,
 	FStateSelectionResult& OutSelectionResult,
 	const FStateTreeSharedEvent* TransitionEvent)
 {
@@ -3162,6 +3179,8 @@ bool FStateTreeExecutionContext::SelectStateInternal(
 
 	const FStateTreeExecutionState& Exec = GetExecState();
 
+	check(!PathToNextState.IsEmpty());
+	const FStateTreeStateHandle NextStateHandle = PathToNextState[0];
 	if (!NextStateHandle.IsValid())
 	{
 		// Trying to select non-existing state.
@@ -3249,13 +3268,15 @@ bool FStateTreeExecutionContext::SelectStateInternal(
 		}
 	}
 
+	const bool bIsDestinationState = PathToNextState.Num() < 2;
+	const bool bShouldPrerequisitesBeChecked = bIsDestinationState || NextState.bCheckPrerequisitesWhenActivatingChildDirectly;
 	TArray<const FStateTreeSharedEvent*, TInlineAllocator<FStateTreeEventQueue::MaxActiveEvents>> StateSelectionEvents;
 	if (NextState.EventDataIndex.IsValid())
 	{
 		check(NextState.RequiredEventToEnter.IsValid());
 
-		TArrayView<FStateTreeSharedEvent> EventsQueue = GetMutableEventsToProcessView();
-		if (TransitionEvent && TransitionEvent->IsValid())
+		// Use the same event as performed transition unless it didn't lead to this state as only state selected by the transition should get it's event.
+		if (TransitionEvent && TransitionEvent->IsValid() && bIsDestinationState)
 		{
 			if (NextState.RequiredEventToEnter.DoesEventMatchDesc(*TransitionEvent->Get()))
 			{
@@ -3264,6 +3285,7 @@ bool FStateTreeExecutionContext::SelectStateInternal(
 		}
 		else
 		{
+			TArrayView<FStateTreeSharedEvent> EventsQueue = GetMutableEventsToProcessView();
 			for (FStateTreeSharedEvent& Event : EventsQueue)
 			{
 				check(Event.IsValid());
@@ -3272,6 +3294,12 @@ bool FStateTreeExecutionContext::SelectStateInternal(
 					StateSelectionEvents.Emplace(&Event);
 				}
 			}
+
+			// Couldn't find matching state's event, but it's marked as not required. Adding an empty event which allows us to continue the state selection.
+			if (!bShouldPrerequisitesBeChecked && StateSelectionEvents.IsEmpty())
+			{
+				StateSelectionEvents.Emplace();
+			}
 		}
 	}
 	else
@@ -3286,14 +3314,17 @@ bool FStateTreeExecutionContext::SelectStateInternal(
 			CurrentlyProcessedStateSelectionEvents->Events[NextState.Depth] = *StateSelectionEvent;
 		}
 		
-		// Check that the state can be entered
-		STATETREE_TRACE_PHASE_BEGIN(EStateTreeUpdatePhase::EnterConditions);
-		const bool bEnterConditionsPassed = TestAllConditions(CurrentParentFrame, CurrentFrame, NextState.EnterConditionsBegin, NextState.EnterConditionsNum);
-		STATETREE_TRACE_PHASE_END(EStateTreeUpdatePhase::EnterConditions);
-
-		if (!bEnterConditionsPassed)
+		if (bShouldPrerequisitesBeChecked)
 		{
-			continue;
+			// Check that the state can be entered
+			STATETREE_TRACE_PHASE_BEGIN(EStateTreeUpdatePhase::EnterConditions);
+			const bool bEnterConditionsPassed = TestAllConditions(CurrentParentFrame, CurrentFrame, NextState.EnterConditionsBegin, NextState.EnterConditionsNum);
+			STATETREE_TRACE_PHASE_END(EStateTreeUpdatePhase::EnterConditions);
+
+			if (!bEnterConditionsPassed)
+			{
+				continue;
+			}
 		}
 
 		if (!CurrentFrame.ActiveStates.Push(NextStateHandle))
@@ -3317,7 +3348,15 @@ bool FStateTreeExecutionContext::SelectStateInternal(
 			}
 		}
 		
-		if (NextState.Type == EStateTreeStateType::Linked)
+		if (!bIsDestinationState)
+		{
+			// Next child state is already known. Passing TransitionEvent further so state selected directly by transition can use it.
+			if (SelectStateInternal(CurrentParentFrame, CurrentFrame, CurrentFrameInActiveFrames, PathToNextState.Mid(1), OutSelectionResult, TransitionEvent))
+			{
+				return true;
+			}
+		}
+		else if (NextState.Type == EStateTreeStateType::Linked)
 		{
 			if (NextState.LinkedState.IsValid())
 			{
@@ -3369,7 +3408,7 @@ bool FStateTreeExecutionContext::SelectStateInternal(
 				OutSelectionResult.PushFrame(NewFrame);
 
 				// If State is linked, proceed to the linked state.
-				if (SelectStateInternal(&CurrentFrame, OutSelectionResult.GetSelectedFrames().Last(), ExistingFrame, NewFrame.RootState, OutSelectionResult))
+				if (SelectStateInternal(&CurrentFrame, OutSelectionResult.GetSelectedFrames().Last(), ExistingFrame, {NewFrame.RootState}, OutSelectionResult))
 				{
 					return true;
 				}
@@ -3472,7 +3511,7 @@ bool FStateTreeExecutionContext::SelectStateInternal(
 				OutSelectionResult.PushFrame(NewFrame);
 
 				// If State is linked, proceed to the linked state.
-				if (SelectStateInternal(&CurrentFrame, OutSelectionResult.GetSelectedFrames().Last(), ExistingFrame, NewFrame.RootState, OutSelectionResult))
+				if (SelectStateInternal(&CurrentFrame, OutSelectionResult.GetSelectedFrames().Last(), ExistingFrame, {NewFrame.RootState}, OutSelectionResult))
 				{
 					return true;
 				}
@@ -3604,7 +3643,7 @@ bool FStateTreeExecutionContext::SelectStateInternal(
 				// If the state has children, proceed to select children.
 				for (uint16 ChildState = NextState.ChildrenBegin; ChildState < NextState.ChildrenEnd; ChildState = CurrentStateTree->States[ChildState].GetNextSibling())
 				{
-					if (SelectStateInternal(CurrentParentFrame, CurrentFrame, CurrentFrameInActiveFrames, FStateTreeStateHandle(ChildState), OutSelectionResult))
+					if (SelectStateInternal(CurrentParentFrame, CurrentFrame, CurrentFrameInActiveFrames, {FStateTreeStateHandle(ChildState)}, OutSelectionResult))
 					{
 						// Selection succeeded
 						return true;
diff --git a/Engine/Plugins/Runtime/StateTree/Source/StateTreeModule/Public/StateTree.h b/Engine/Plugins/Runtime/StateTree/Source/StateTreeModule/Public/StateTree.h
--- a/Engine/Plugins/Runtime/StateTree/Source/StateTreeModule/Public/StateTree.h
+++ b/Engine/Plugins/Runtime/StateTree/Source/StateTreeModule/Public/StateTree.h
@@ -51,6 +51,8 @@ struct STATETREEMODULE_API FStateTreeCustomVersion
 		StoringGlobalParametersInInstanceStorage,
 		// Added binding to events
 		AddedBindingToEvents,
+		// Added checking parent states' prerequisites when activating child state directly.
+		AddedCheckingParentsPrerequisites,
 
 		// -----<new versions can be added above this line>-------------------------------------------------
 		VersionPlusOne,
diff --git a/Engine/Plugins/Runtime/StateTree/Source/StateTreeModule/Public/StateTreeExecutionContext.h b/Engine/Plugins/Runtime/StateTree/Source/StateTreeModule/Public/StateTreeExecutionContext.h
--- a/Engine/Plugins/Runtime/StateTree/Source/StateTreeModule/Public/StateTreeExecutionContext.h
+++ b/Engine/Plugins/Runtime/StateTree/Source/StateTreeModule/Public/StateTreeExecutionContext.h
@@ -587,7 +587,7 @@ protected:
 		const FStateTreeExecutionFrame* CurrentParentFrame,
 		FStateTreeExecutionFrame& CurrentFrame,
 		const FStateTreeExecutionFrame* CurrentFrameInActiveFrames,
-		const FStateTreeStateHandle NextStateHandle,
+		TConstArrayView<FStateTreeStateHandle> PathToNextState,
 		FStateSelectionResult& OutSelectionResult,
 		const FStateTreeSharedEvent* TransitionEvent = nullptr);
 
diff --git a/Engine/Plugins/Runtime/StateTree/Source/StateTreeModule/Public/StateTreeTypes.h b/Engine/Plugins/Runtime/StateTree/Source/StateTreeModule/Public/StateTreeTypes.h
--- a/Engine/Plugins/Runtime/StateTree/Source/StateTreeModule/Public/StateTreeTypes.h
+++ b/Engine/Plugins/Runtime/StateTree/Source/StateTreeModule/Public/StateTreeTypes.h
@@ -597,12 +597,6 @@ USTRUCT()
 struct STATETREEMODULE_API FCompactStateTreeState
 {
 	GENERATED_BODY()
-
-	FCompactStateTreeState()
-		: bHasTransitionTasks(false)
-		, bEnabled(true)
-	{
-	}
 	
 	/** @return Index to the next sibling state. */
 	uint16 GetNextSibling() const { return ChildrenEnd; }
@@ -696,11 +690,15 @@ struct STATETREEMODULE_API FCompactStateTreeState
 
 	/** True if the state contains tasks that should be called during transition handling. */
 	UPROPERTY()
-	uint8 bHasTransitionTasks : 1;
+	uint8 bHasTransitionTasks : 1 = false;
+
+	/** Should state's required event and enter conditions be evaluated when transition leads directly to it's child. */
+	UPROPERTY()
+	uint8 bCheckPrerequisitesWhenActivatingChildDirectly : 1 = false;
 
 	/** True if the state is Enabled (i.e. not explicitly marked as disabled). */
 	UPROPERTY()
-	uint8 bEnabled : 1;
+	uint8 bEnabled : 1 = true;
 };
 
 USTRUCT()
