Fix warning V595: The 'SourceProperty' pointer was utilized before it was verified against nullptr.
#rnx

diff --git a/Engine/Plugins/Runtime/StateTree/Source/StateTreeEditorModule/Private/Customizations/StateTreeBindingExtension.cpp b/Engine/Plugins/Runtime/StateTree/Source/StateTreeEditorModule/Private/Customizations/StateTreeBindingExtension.cpp
--- a/Engine/Plugins/Runtime/StateTree/Source/StateTreeEditorModule/Private/Customizations/StateTreeBindingExtension.cpp
+++ b/Engine/Plugins/Runtime/StateTree/Source/StateTreeEditorModule/Private/Customizations/StateTreeBindingExtension.cpp
@@ -628,6 +628,11 @@ struct FCachedBindingData : public TSharedFromThis<FCachedBindingData>
 
 	bool CanAcceptPropertyOrChildren(const FProperty* SourceProperty, TConstArrayView<FBindingChainElement> InBindingChain)
 	{
+		if (!SourceProperty)
+		{
+			return false;
+		}
+
 		ConditionallyUpdateData();
 
 		if (UE::StateTree::PropertyRefHelpers::IsPropertyRef(*PropertyHandle->GetProperty()))
@@ -641,7 +646,7 @@ struct FCachedBindingData : public TSharedFromThis<FCachedBindingData>
 			}
 		}
 
-		return SourceProperty && IsPropertyBindable(*SourceProperty);
+		return IsPropertyBindable(*SourceProperty);
 	}
 
 	static bool ArePropertyAndContextStructCompatible(const UStruct* SourceStruct, const FProperty* TargetProperty)
