[Backout] - CL32723276
Original CL Desc
-----------------------------------------------------------------
StateTree: Made parent state required events and enter conditions to be evaluated when transition leads directly to it's child. (contributed)

diff --git a/Engine/Plugins/Runtime/StateTree/Source/StateTreeEditorModule/Private/Customizations/StateTreeStateDetails.cpp b/Engine/Plugins/Runtime/StateTree/Source/StateTreeEditorModule/Private/Customizations/StateTreeStateDetails.cpp
--- a/Engine/Plugins/Runtime/StateTree/Source/StateTreeEditorModule/Private/Customizations/StateTreeStateDetails.cpp
+++ b/Engine/Plugins/Runtime/StateTree/Source/StateTreeEditorModule/Private/Customizations/StateTreeStateDetails.cpp
@@ -54,9 +54,8 @@ void FStateTreeStateDetails::CustomizeDetails(IDetailLayoutBuilder& DetailBuilde
 	const TSharedPtr<IPropertyHandle> ParametersProperty = DetailBuilder.GetProperty(GET_MEMBER_NAME_CHECKED(UStateTreeState, Parameters));
 	const TSharedPtr<IPropertyHandle> SelectionBehaviorProperty = DetailBuilder.GetProperty(GET_MEMBER_NAME_CHECKED(UStateTreeState, SelectionBehavior));
 	const TSharedPtr<IPropertyHandle> RequiredEventToEnterProperty = DetailBuilder.GetProperty(GET_MEMBER_NAME_CHECKED(UStateTreeState, RequiredEventToEnter));
-	const TSharedPtr<IPropertyHandle> CheckPrerequisitesProperty = DetailBuilder.GetProperty(GET_MEMBER_NAME_CHECKED(UStateTreeState, bCheckPrerequisitesWhenActivatingChildDirectly));
-	
 
+	
 
 	// Never show enabled
 	EnabledProperty->MarkHiddenByCustomization();
@@ -127,9 +126,6 @@ void FStateTreeStateDetails::CustomizeDetails(IDetailLayoutBuilder& DetailBuilde
 		];
 	}
 
-	CheckPrerequisitesProperty->MarkHiddenByCustomization();
-	StateCategory.AddProperty(CheckPrerequisitesProperty);
-
 	// Parameters
 	ParametersProperty->MarkHiddenByCustomization();
 	StateCategory.AddProperty(ParametersProperty);
diff --git a/Engine/Plugins/Runtime/StateTree/Source/StateTreeEditorModule/Private/StateTreeCompiler.cpp b/Engine/Plugins/Runtime/StateTree/Source/StateTreeEditorModule/Private/StateTreeCompiler.cpp
--- a/Engine/Plugins/Runtime/StateTree/Source/StateTreeEditorModule/Private/StateTreeCompiler.cpp
+++ b/Engine/Plugins/Runtime/StateTree/Source/StateTreeEditorModule/Private/StateTreeCompiler.cpp
@@ -415,7 +415,6 @@ bool FStateTreeCompiler::CreateStateRecursive(UStateTreeState& State, const FSta
 	CompactState.Tag = State.Tag;
 	CompactState.Parent = Parent;
 	CompactState.bEnabled = State.bEnabled;
-	CompactState.bCheckPrerequisitesWhenActivatingChildDirectly = State.bCheckPrerequisitesWhenActivatingChildDirectly;
 
 	CompactState.Type = State.Type;
 	CompactState.SelectionBehavior = State.SelectionBehavior;
diff --git a/Engine/Plugins/Runtime/StateTree/Source/StateTreeEditorModule/Private/StateTreeState.cpp b/Engine/Plugins/Runtime/StateTree/Source/StateTreeEditorModule/Private/StateTreeState.cpp
--- a/Engine/Plugins/Runtime/StateTree/Source/StateTreeEditorModule/Private/StateTreeState.cpp
+++ b/Engine/Plugins/Runtime/StateTree/Source/StateTreeEditorModule/Private/StateTreeState.cpp
@@ -393,11 +393,6 @@ void UStateTreeState::PostLoad()
 			}
 		}
 	}
-
-	if (CurrentVersion < FStateTreeCustomVersion::AddedCheckingParentsPrerequisites)
-	{
-		bCheckPrerequisitesWhenActivatingChildDirectly = false;
-	}
 	
 #endif // WITH_EDITORONLY_DATA
 
diff --git a/Engine/Plugins/Runtime/StateTree/Source/StateTreeEditorModule/Public/StateTreeState.h b/Engine/Plugins/Runtime/StateTree/Source/StateTreeEditorModule/Public/StateTreeState.h
--- a/Engine/Plugins/Runtime/StateTree/Source/StateTreeEditorModule/Public/StateTreeState.h
+++ b/Engine/Plugins/Runtime/StateTree/Source/StateTreeEditorModule/Public/StateTreeState.h
@@ -343,10 +343,6 @@ public:
 	UPROPERTY(EditDefaultsOnly, Category = "State")
 	TObjectPtr<UStateTree> LinkedAsset = nullptr;
 
-	/** Should state's required event and enter conditions be evaluated when transition leads directly to it's child. */
-	UPROPERTY(EditDefaultsOnly, Category = "State")
-	bool bCheckPrerequisitesWhenActivatingChildDirectly = true;
-
 	/** Parameters of this state. If the state is linked to another state or asset, the parameters are for the linked state. */
 	UPROPERTY(EditDefaultsOnly, Category = "State")
 	FStateTreeStateParameters Parameters;
diff --git a/Engine/Plugins/Runtime/StateTree/Source/StateTreeModule/Private/StateTreeExecutionContext.cpp b/Engine/Plugins/Runtime/StateTree/Source/StateTreeModule/Private/StateTreeExecutionContext.cpp
--- a/Engine/Plugins/Runtime/StateTree/Source/StateTreeModule/Private/StateTreeExecutionContext.cpp
+++ b/Engine/Plugins/Runtime/StateTree/Source/StateTreeModule/Private/StateTreeExecutionContext.cpp
@@ -3032,25 +3032,24 @@ bool FStateTreeExecutionContext::SelectState(const FStateTreeExecutionFrame& Cur
 	}
 
 	// Walk towards the root from current state.
-	TArray<FStateTreeStateHandle, TInlineAllocator<FStateTreeActiveStates::MaxStates>> PathToNextState;
+	TStaticArray<FStateTreeStateHandle, FStateTreeActiveStates::MaxStates> ParentStates;
+	int32 NumParentStates = 0;
 	FStateTreeStateHandle CurrState = NextState;
 	while (CurrState.IsValid())
 	{
-		if (PathToNextState.Num() == FStateTreeActiveStates::MaxStates)
+		if (NumParentStates == FStateTreeActiveStates::MaxStates)
 		{
 			STATETREE_LOG(Error, TEXT("%hs: Reached max execution depth when trying to select state %s from '%s'.  '%s' using StateTree '%s'."),
 				__FUNCTION__, *GetSafeStateName(CurrentFrame, NextState), *GetStateStatusString(Exec), *GetNameSafe(&Owner), *GetFullNameSafe(&RootStateTree));
 			return false;
 		}
 		// Store the states that are in between the 'NextState' and common ancestor. 
-		PathToNextState.Push(CurrState);
+		ParentStates[NumParentStates++] = CurrState;
 		CurrState = CurrentFrame.StateTree->States[CurrState.Index].Parent;
 	}
 
-	Algo::Reverse(PathToNextState);
-
 	const UStateTree* NextStateTree = CurrentFrame.StateTree;
-	const FStateTreeStateHandle NextRootState = PathToNextState[0]; 
+	const FStateTreeStateHandle NextRootState = ParentStates[NumParentStates - 1]; 
 
 	// Find the frame that the next state belongs to.
 	int32 CurrentFrameIndex = INDEX_NONE;
@@ -3098,27 +3097,14 @@ bool FStateTreeExecutionContext::SelectState(const FStateTreeExecutionFrame& Cur
 	const int32 LastFrameIndex = OutSelectionResult.FramesNum() - 1;
 	FStateTreeExecutionFrame& LastFrame = OutSelectionResult.GetSelectedFrames()[LastFrameIndex];
 
-	// Find index of the first state to be evaluated.
-	int32 FirstNewStateIndex = 0;
-	if (CurrentFrameIndex != INDEX_NONE)
+	LastFrame.ActiveStates.Reset();
+	for (int32 Index = NumParentStates - 1; Index > 0; Index--)
 	{
-		// If LastFrame.ActiveStates is a subset of PathToNextState (e.g when someone use "TryEnter" selection behavior and then make a transition to it's child or if one is reentering the same state).
-		// In such case loop below won't break on anything and FirstNewStateIndex will be incorrectly 0, thus we initialize it to be right after the shorter range.
-		FirstNewStateIndex = FMath::Max(0, FMath::Min(PathToNextState.Num(), LastFrame.ActiveStates.Num()) - 1);
-		for (int32 Index = 0; Index < FMath::Min(PathToNextState.Num(), LastFrame.ActiveStates.Num()); ++Index)
-		{
-			if (LastFrame.ActiveStates[Index] != PathToNextState[Index])
-			{
-				FirstNewStateIndex = Index;
-				break;
-			}
-		}
+		LastFrame.ActiveStates.Push(ParentStates[Index]);
 	}
 	// Existing state's data is safe to access during select.
 	LastFrame.NumCurrentlyActiveStates = static_cast<uint8>(LastFrame.ActiveStates.Num());
 
-	LastFrame.ActiveStates.SetNum(FirstNewStateIndex);
-
 	FStateSelectionResult InitialSelection;
 
 	if (Fallback == EStateTreeSelectionFallback::NextSelectableSibling)
@@ -3127,21 +3113,18 @@ bool FStateTreeExecutionContext::SelectState(const FStateTreeExecutionFrame& Cur
 	}
 	
 	// We take copy of the last frame and assign it later, as SelectStateInternal() might change the array and invalidate the pointer.
-	const FStateTreeExecutionFrame* CurrentParentFrame = LastFrameIndex > 0 ? &OutSelectionResult.GetSelectedFrames()[LastFrameIndex - 1] : nullptr;
-
-	// Path from the first new state up to the NextState
-	TConstArrayView<FStateTreeStateHandle> NewStatesPathToNextState(&PathToNextState[FirstNewStateIndex], PathToNextState.Num() - FirstNewStateIndex);
-
-	if (SelectStateInternal(CurrentParentFrame, OutSelectionResult.GetSelectedFrames()[LastFrameIndex], CurrentFrameInActiveFrames, NewStatesPathToNextState, OutSelectionResult, TransitionEvent))
+	const FStateTreeExecutionFrame* CurrentParentFrame = LastFrameIndex > 0 ? &OutSelectionResult.GetSelectedFrames()[LastFrameIndex - 1] : nullptr; 
+	if (SelectStateInternal(CurrentParentFrame, OutSelectionResult.GetSelectedFrames()[LastFrameIndex], CurrentFrameInActiveFrames , NextState, OutSelectionResult, TransitionEvent))
 	{
 		return true;
 	}
 
 	// Failed to Select Next State, handle fallback here
 	// Return true on the first next sibling that gets selected successfully
-	if (Fallback == EStateTreeSelectionFallback::NextSelectableSibling && PathToNextState.Num() >= 2)
+	if (Fallback == EStateTreeSelectionFallback::NextSelectableSibling && NumParentStates >= 2)
 	{
-		const FStateTreeStateHandle Parent = PathToNextState.Last(1);
+		// InBetweenStates is in reversed order (i.e. from leaf to root)
+		const FStateTreeStateHandle Parent = ParentStates[1];
 		if (Parent.IsValid())
 		{
 			const FCompactStateTreeState& ParentState = CurrentFrame.StateTree->States[Parent.Index];
@@ -3156,7 +3139,7 @@ bool FStateTreeExecutionContext::SelectState(const FStateTreeExecutionFrame& Cur
 	
 				// We take copy of the last frame and assign it later, as SelectStateInternal() might change the array and invalidate the pointer.
 				CurrentParentFrame = LastFrameIndex > 0 ? &OutSelectionResult.GetSelectedFrames()[LastFrameIndex - 1] : nullptr; 
-				if (SelectStateInternal(CurrentParentFrame, OutSelectionResult.GetSelectedFrames()[LastFrameIndex], CurrentFrameInActiveFrames, {ChildStateHandle}, OutSelectionResult))
+				if (SelectStateInternal(CurrentParentFrame, OutSelectionResult.GetSelectedFrames()[LastFrameIndex], CurrentFrameInActiveFrames, ChildStateHandle, OutSelectionResult))
 				{
 					return true;
 				}
@@ -3171,7 +3154,7 @@ bool FStateTreeExecutionContext::SelectStateInternal(
 	const FStateTreeExecutionFrame* CurrentParentFrame,
 	FStateTreeExecutionFrame& CurrentFrame,
 	const FStateTreeExecutionFrame* CurrentFrameInActiveFrames,
-	TConstArrayView<FStateTreeStateHandle> PathToNextState,
+	const FStateTreeStateHandle NextStateHandle,
 	FStateSelectionResult& OutSelectionResult,
 	const FStateTreeSharedEvent* TransitionEvent)
 {
@@ -3179,8 +3162,6 @@ bool FStateTreeExecutionContext::SelectStateInternal(
 
 	const FStateTreeExecutionState& Exec = GetExecState();
 
-	check(!PathToNextState.IsEmpty());
-	const FStateTreeStateHandle NextStateHandle = PathToNextState[0];
 	if (!NextStateHandle.IsValid())
 	{
 		// Trying to select non-existing state.
@@ -3268,15 +3249,13 @@ bool FStateTreeExecutionContext::SelectStateInternal(
 		}
 	}
 
-	const bool bIsDestinationState = PathToNextState.Num() < 2;
-	const bool bShouldPrerequisitesBeChecked = bIsDestinationState || NextState.bCheckPrerequisitesWhenActivatingChildDirectly;
 	TArray<const FStateTreeSharedEvent*, TInlineAllocator<FStateTreeEventQueue::MaxActiveEvents>> StateSelectionEvents;
 	if (NextState.EventDataIndex.IsValid())
 	{
 		check(NextState.RequiredEventToEnter.IsValid());
 
-		// Use the same event as performed transition unless it didn't lead to this state as only state selected by the transition should get it's event.
-		if (TransitionEvent && TransitionEvent->IsValid() && bIsDestinationState)
+		TArrayView<FStateTreeSharedEvent> EventsQueue = GetMutableEventsToProcessView();
+		if (TransitionEvent && TransitionEvent->IsValid())
 		{
 			if (NextState.RequiredEventToEnter.DoesEventMatchDesc(*TransitionEvent->Get()))
 			{
@@ -3285,7 +3264,6 @@ bool FStateTreeExecutionContext::SelectStateInternal(
 		}
 		else
 		{
-			TArrayView<FStateTreeSharedEvent> EventsQueue = GetMutableEventsToProcessView();
 			for (FStateTreeSharedEvent& Event : EventsQueue)
 			{
 				check(Event.IsValid());
@@ -3294,12 +3272,6 @@ bool FStateTreeExecutionContext::SelectStateInternal(
 					StateSelectionEvents.Emplace(&Event);
 				}
 			}
-
-			// Couldn't find matching state's event, but it's marked as not required. Adding an empty event which allows us to continue the state selection.
-			if (!bShouldPrerequisitesBeChecked && StateSelectionEvents.IsEmpty())
-			{
-				StateSelectionEvents.Emplace();
-			}
 		}
 	}
 	else
@@ -3314,17 +3286,14 @@ bool FStateTreeExecutionContext::SelectStateInternal(
 			CurrentlyProcessedStateSelectionEvents->Events[NextState.Depth] = *StateSelectionEvent;
 		}
 		
-		if (bShouldPrerequisitesBeChecked)
-		{
-			// Check that the state can be entered
-			STATETREE_TRACE_PHASE_BEGIN(EStateTreeUpdatePhase::EnterConditions);
-			const bool bEnterConditionsPassed = TestAllConditions(CurrentParentFrame, CurrentFrame, NextState.EnterConditionsBegin, NextState.EnterConditionsNum);
-			STATETREE_TRACE_PHASE_END(EStateTreeUpdatePhase::EnterConditions);
+		// Check that the state can be entered
+		STATETREE_TRACE_PHASE_BEGIN(EStateTreeUpdatePhase::EnterConditions);
+		const bool bEnterConditionsPassed = TestAllConditions(CurrentParentFrame, CurrentFrame, NextState.EnterConditionsBegin, NextState.EnterConditionsNum);
+		STATETREE_TRACE_PHASE_END(EStateTreeUpdatePhase::EnterConditions);
 
-			if (!bEnterConditionsPassed)
-			{
-				continue;
-			}
+		if (!bEnterConditionsPassed)
+		{
+			continue;
 		}
 
 		if (!CurrentFrame.ActiveStates.Push(NextStateHandle))
@@ -3348,15 +3317,7 @@ bool FStateTreeExecutionContext::SelectStateInternal(
 			}
 		}
 		
-		if (!bIsDestinationState)
-		{
-			// Next child state is already known. Passing TransitionEvent further so state selected directly by transition can use it.
-			if (SelectStateInternal(CurrentParentFrame, CurrentFrame, CurrentFrameInActiveFrames, PathToNextState.Mid(1), OutSelectionResult, TransitionEvent))
-			{
-				return true;
-			}
-		}
-		else if (NextState.Type == EStateTreeStateType::Linked)
+		if (NextState.Type == EStateTreeStateType::Linked)
 		{
 			if (NextState.LinkedState.IsValid())
 			{
@@ -3408,7 +3369,7 @@ bool FStateTreeExecutionContext::SelectStateInternal(
 				OutSelectionResult.PushFrame(NewFrame);
 
 				// If State is linked, proceed to the linked state.
-				if (SelectStateInternal(&CurrentFrame, OutSelectionResult.GetSelectedFrames().Last(), ExistingFrame, {NewFrame.RootState}, OutSelectionResult))
+				if (SelectStateInternal(&CurrentFrame, OutSelectionResult.GetSelectedFrames().Last(), ExistingFrame, NewFrame.RootState, OutSelectionResult))
 				{
 					return true;
 				}
@@ -3511,7 +3472,7 @@ bool FStateTreeExecutionContext::SelectStateInternal(
 				OutSelectionResult.PushFrame(NewFrame);
 
 				// If State is linked, proceed to the linked state.
-				if (SelectStateInternal(&CurrentFrame, OutSelectionResult.GetSelectedFrames().Last(), ExistingFrame, {NewFrame.RootState}, OutSelectionResult))
+				if (SelectStateInternal(&CurrentFrame, OutSelectionResult.GetSelectedFrames().Last(), ExistingFrame, NewFrame.RootState, OutSelectionResult))
 				{
 					return true;
 				}
@@ -3643,7 +3604,7 @@ bool FStateTreeExecutionContext::SelectStateInternal(
 				// If the state has children, proceed to select children.
 				for (uint16 ChildState = NextState.ChildrenBegin; ChildState < NextState.ChildrenEnd; ChildState = CurrentStateTree->States[ChildState].GetNextSibling())
 				{
-					if (SelectStateInternal(CurrentParentFrame, CurrentFrame, CurrentFrameInActiveFrames, {FStateTreeStateHandle(ChildState)}, OutSelectionResult))
+					if (SelectStateInternal(CurrentParentFrame, CurrentFrame, CurrentFrameInActiveFrames, FStateTreeStateHandle(ChildState), OutSelectionResult))
 					{
 						// Selection succeeded
 						return true;
diff --git a/Engine/Plugins/Runtime/StateTree/Source/StateTreeModule/Public/StateTree.h b/Engine/Plugins/Runtime/StateTree/Source/StateTreeModule/Public/StateTree.h
--- a/Engine/Plugins/Runtime/StateTree/Source/StateTreeModule/Public/StateTree.h
+++ b/Engine/Plugins/Runtime/StateTree/Source/StateTreeModule/Public/StateTree.h
@@ -51,8 +51,6 @@ struct STATETREEMODULE_API FStateTreeCustomVersion
 		StoringGlobalParametersInInstanceStorage,
 		// Added binding to events
 		AddedBindingToEvents,
-		// Added checking parent states' prerequisites when activating child state directly.
-		AddedCheckingParentsPrerequisites,
 
 		// -----<new versions can be added above this line>-------------------------------------------------
 		VersionPlusOne,
diff --git a/Engine/Plugins/Runtime/StateTree/Source/StateTreeModule/Public/StateTreeExecutionContext.h b/Engine/Plugins/Runtime/StateTree/Source/StateTreeModule/Public/StateTreeExecutionContext.h
--- a/Engine/Plugins/Runtime/StateTree/Source/StateTreeModule/Public/StateTreeExecutionContext.h
+++ b/Engine/Plugins/Runtime/StateTree/Source/StateTreeModule/Public/StateTreeExecutionContext.h
@@ -587,7 +587,7 @@ protected:
 		const FStateTreeExecutionFrame* CurrentParentFrame,
 		FStateTreeExecutionFrame& CurrentFrame,
 		const FStateTreeExecutionFrame* CurrentFrameInActiveFrames,
-		TConstArrayView<FStateTreeStateHandle> PathToNextState,
+		const FStateTreeStateHandle NextStateHandle,
 		FStateSelectionResult& OutSelectionResult,
 		const FStateTreeSharedEvent* TransitionEvent = nullptr);
 
diff --git a/Engine/Plugins/Runtime/StateTree/Source/StateTreeModule/Public/StateTreeTypes.h b/Engine/Plugins/Runtime/StateTree/Source/StateTreeModule/Public/StateTreeTypes.h
--- a/Engine/Plugins/Runtime/StateTree/Source/StateTreeModule/Public/StateTreeTypes.h
+++ b/Engine/Plugins/Runtime/StateTree/Source/StateTreeModule/Public/StateTreeTypes.h
@@ -597,6 +597,12 @@ USTRUCT()
 struct STATETREEMODULE_API FCompactStateTreeState
 {
 	GENERATED_BODY()
+
+	FCompactStateTreeState()
+		: bHasTransitionTasks(false)
+		, bEnabled(true)
+	{
+	}
 	
 	/** @return Index to the next sibling state. */
 	uint16 GetNextSibling() const { return ChildrenEnd; }
@@ -690,15 +696,11 @@ struct STATETREEMODULE_API FCompactStateTreeState
 
 	/** True if the state contains tasks that should be called during transition handling. */
 	UPROPERTY()
-	uint8 bHasTransitionTasks : 1 = false;
-
-	/** Should state's required event and enter conditions be evaluated when transition leads directly to it's child. */
-	UPROPERTY()
-	uint8 bCheckPrerequisitesWhenActivatingChildDirectly : 1 = false;
+	uint8 bHasTransitionTasks : 1;
 
 	/** True if the state is Enabled (i.e. not explicitly marked as disabled). */
 	UPROPERTY()
-	uint8 bEnabled : 1 = true;
+	uint8 bEnabled : 1;
 };
 
 USTRUCT()
