StateTree: Fix FStateTreePropertyRef example to guide correct type usage
- Fix FStateTreePropertyRef types in Env Query Task
- Fix too eager PVS warning

diff --git a/Engine/Plugins/Runtime/GameplayStateTree/Source/GameplayStateTreeModule/Public/Tasks/StateTreeRunEnvQueryTask.h b/Engine/Plugins/Runtime/GameplayStateTree/Source/GameplayStateTreeModule/Public/Tasks/StateTreeRunEnvQueryTask.h
--- a/Engine/Plugins/Runtime/GameplayStateTree/Source/GameplayStateTreeModule/Public/Tasks/StateTreeRunEnvQueryTask.h
+++ b/Engine/Plugins/Runtime/GameplayStateTree/Source/GameplayStateTreeModule/Public/Tasks/StateTreeRunEnvQueryTask.h
@@ -16,7 +16,7 @@ struct FStateTreeRunEnvQueryInstanceData
 	GENERATED_BODY()
 
 	// Result of the query. If an array is binded, it will output all the created values otherwise it will output the best one.
-	UPROPERTY(EditAnywhere, Category = Out, meta = (RefType = "Vector, Actor", CanRefToArray))
+	UPROPERTY(EditAnywhere, Category = Out, meta = (RefType = "/Script/CoreUObject.Vector, /Script/Engine.Actor", CanRefToArray))
 	FStateTreePropertyRef Result;
 
 	// The query will be run with this actor has the owner object.
diff --git a/Engine/Plugins/Runtime/StateTree/Source/StateTreeModule/Private/StateTreeExecutionContext.cpp b/Engine/Plugins/Runtime/StateTree/Source/StateTreeModule/Private/StateTreeExecutionContext.cpp
--- a/Engine/Plugins/Runtime/StateTree/Source/StateTreeModule/Private/StateTreeExecutionContext.cpp
+++ b/Engine/Plugins/Runtime/StateTree/Source/StateTreeModule/Private/StateTreeExecutionContext.cpp
@@ -2937,7 +2937,7 @@ bool FStateTreeExecutionContext::TriggerTransitions()
 					}
 				}
 
-				if (NextTransition.Priority != EStateTreeTransitionPriority::None)
+				if (NextTransition.Priority != EStateTreeTransitionPriority::None) //-V547
 				{
 					break;
 				}
@@ -2945,7 +2945,7 @@ bool FStateTreeExecutionContext::TriggerTransitions()
 		}
 
 		// Handle the case where no transition was found.
-		if (NextTransition.Priority == EStateTreeTransitionPriority::None)
+		if (NextTransition.Priority == EStateTreeTransitionPriority::None) //-V547
 		{
 			STATETREE_LOG_AND_TRACE(Verbose, TEXT("Could not trigger completion transition, jump back to root state."));
 
diff --git a/Engine/Plugins/Runtime/StateTree/Source/StateTreeModule/Public/StateTreePropertyRef.h b/Engine/Plugins/Runtime/StateTree/Source/StateTreeModule/Public/StateTreePropertyRef.h
--- a/Engine/Plugins/Runtime/StateTree/Source/StateTreeModule/Public/StateTreePropertyRef.h
+++ b/Engine/Plugins/Runtime/StateTree/Source/StateTreeModule/Public/StateTreePropertyRef.h
@@ -98,15 +98,16 @@ namespace UE::StateTree::PropertyRefHelpers
  *
  * Meta specifiers for the type:
  *  - RefType = "<type>"
- *		- Specifies a comma separated list of type of property to reference.
- *		- Supported types are: bool, byte, int32, int64, float, double, Name, String, Text, UObject pointers, and structs.
- *		- If multiple types are specified, GetMutablePtrTuple can be used to access the correct type.
+ *		- Specifies a comma separated list of type of property to reference
+ *		- Supported types are: bool, byte, int32, int64, float, double, Name, String, Text, UObject pointers, and structs
+ *		- Structs and Objects must use full path name
+ *		- If multiple types are specified, GetMutablePtrTuple can be used to access the correct type
  *  - IsRefToArray
  *		- If specified, the reference is to an TArray<RefType>
  *	- CanRefToArray
  *		- If specified, the reference can bind to a Reftype or TArray<RefType>
  *  - Optional
- *		- If specified, the reference can be left unbound, otherwise the compiler report error if the reference is not bound.
+ *		- If specified, the reference can be left unbound, otherwise the compiler report error if the reference is not bound
  *
  * Example:
  *
@@ -123,7 +124,7 @@ namespace UE::StateTree::PropertyRefHelpers
  *	FStateTreePropertyRef RefToArrayOfTests;
  *
  *  // Reference to Vector, TArray<FVector>, AActor*, TArray<AActor*>
- *	UPROPERTY(EditAnywhere, meta = (RefType = "Vector, AActor", CanRefToArray))
+ *	UPROPERTY(EditAnywhere, meta = (RefType = "/Script/CoreUObject.Vector, /Script/Engine.Actor", CanRefToArray))
  *	FStateTreePropertyRef RefToLocationLikeTypes;
  */
 USTRUCT()
