StateTree: Added automatic description to tasks and conditions.

diff --git a/Engine/Plugins/Runtime/StateTree/Source/StateTreeEditorModule/Private/Customizations/StateTreeEditorNodeDetails.cpp b/Engine/Plugins/Runtime/StateTree/Source/StateTreeEditorModule/Private/Customizations/StateTreeEditorNodeDetails.cpp
--- a/Engine/Plugins/Runtime/StateTree/Source/StateTreeEditorModule/Private/Customizations/StateTreeEditorNodeDetails.cpp
+++ b/Engine/Plugins/Runtime/StateTree/Source/StateTreeEditorModule/Private/Customizations/StateTreeEditorNodeDetails.cpp
@@ -38,6 +38,10 @@
 #include "Debugger/StateTreeDebuggerUIExtensions.h"
 #include "StateTreePropertyBindings.h"
 #include "Engine/BlueprintGeneratedClass.h"
+#include "Widgets/Text/SRichTextBlock.h"
+#include "Styling/SlateTypes.h"
+#include "TextStyleDecorator.h"
+#include "Widgets/Layout/SWidgetSwitcher.h"
 
 #define LOCTEXT_NAMESPACE "StateTreeEditor"
 
@@ -474,10 +478,10 @@ void FStateTreeEditorNodeDetails::CustomizeHeader(TSharedRef<class IPropertyHand
 	StructProperty = StructPropertyHandle;
 	PropUtils = StructCustomizationUtils.GetPropertyUtilities();
 
-	NodeProperty = StructProperty->GetChildHandle(TEXT("Node"));
-	InstanceProperty = StructProperty->GetChildHandle(TEXT("Instance"));
-	InstanceObjectProperty = StructProperty->GetChildHandle(TEXT("InstanceObject"));
-	IDProperty = StructProperty->GetChildHandle(TEXT("ID"));
+	NodeProperty = StructProperty->GetChildHandle(GET_MEMBER_NAME_CHECKED(FStateTreeEditorNode, Node));
+	InstanceProperty = StructProperty->GetChildHandle(GET_MEMBER_NAME_CHECKED(FStateTreeEditorNode, Instance));
+	InstanceObjectProperty = StructProperty->GetChildHandle(GET_MEMBER_NAME_CHECKED(FStateTreeEditorNode, InstanceObject));
+	IDProperty = StructProperty->GetChildHandle(GET_MEMBER_NAME_CHECKED(FStateTreeEditorNode, ID));
 
 	IndentProperty = StructProperty->GetChildHandle(TEXT("ConditionIndent"));
 	OperandProperty = StructProperty->GetChildHandle(TEXT("ConditionOperand"));
@@ -497,7 +501,7 @@ void FStateTreeEditorNodeDetails::CustomizeHeader(TSharedRef<class IPropertyHand
 
 	const FString BaseClassName = StructProperty->GetMetaData(BaseClassMetaName);
 	BaseClass = UClass::TryFindTypeSlow<UClass>(BaseClassName);
-	
+
 	const FIsResetToDefaultVisible IsResetVisible = FIsResetToDefaultVisible::CreateSP(this, &FStateTreeEditorNodeDetails::ShouldResetToDefault);
 	const FResetToDefaultHandler ResetHandler = FResetToDefaultHandler::CreateSP(this, &FStateTreeEditorNodeDetails::ResetToDefault);
 	const FResetToDefaultOverride ResetOverride = FResetToDefaultOverride::Create(IsResetVisible, ResetHandler);
@@ -512,165 +516,181 @@ void FStateTreeEditorNodeDetails::CustomizeHeader(TSharedRef<class IPropertyHand
 		.VAlign(VAlign_Center)
 		[
 			SNew(SHorizontalBox)
+			// Indent
 			+ SHorizontalBox::Slot()
-			.FillWidth(2.0f)
+			.AutoWidth()
+			.VAlign(VAlign_Center)
 			[
-				SNew(SHorizontalBox)
-				// Indent
-				+ SHorizontalBox::Slot()
-				.AutoWidth()
-				.VAlign(VAlign_Center)
+				SNew(SBox)
+				.WidthOverride(this, &FStateTreeEditorNodeDetails::GetIndentSize)
+				.Visibility(this, &FStateTreeEditorNodeDetails::IsConditionVisible)
 				[
-					SNew(SBox)
-					.WidthOverride(this, &FStateTreeEditorNodeDetails::GetIndentSize)
-					.Visibility(this, &FStateTreeEditorNodeDetails::IsConditionVisible)
+					SNew(SComboButton)
+					.ComboButtonStyle(FStateTreeEditorStyle::Get(), "StateTree.Node.Indent.ComboBox")
+					.ContentPadding(2.0f)
+					.HasDownArrow(false)
+					.VAlign(VAlign_Center)
+					.OnGetMenuContent(this, &FStateTreeEditorNodeDetails::OnGetIndentContent)
+					.ButtonContent()
 					[
-						SNew(SComboButton)
-						.ComboButtonStyle(FStateTreeEditorStyle::Get(), "StateTree.Node.Indent.ComboBox")
-						.ContentPadding(2.0f)
-						.HasDownArrow(false)
-						.VAlign(VAlign_Center)
-						.OnGetMenuContent(this, &FStateTreeEditorNodeDetails::OnGetIndentContent)
-						.ButtonContent()
-						[
-							SNew(STextBlock)
-							.Text(FText::FromString(TEXT(">")))
-							.TextStyle(FStateTreeEditorStyle::Get(), "Details.Normal")
-							.ToolTipText(LOCTEXT("IndentTooltip", "Indent of the expression row, controls parentheses and evaluation order."))
-						]
+						SNew(STextBlock)
+						.Text(FText::FromString(TEXT(">")))
+						.TextStyle(FStateTreeEditorStyle::Get(), "Details.Normal")
+						.ToolTipText(LOCTEXT("IndentTooltip", "Indent of the expression row, controls parentheses and evaluation order."))
 					]
 				]
-				// Operand
-				+ SHorizontalBox::Slot()
-				.AutoWidth()
+			]
+			// Operand
+			+ SHorizontalBox::Slot()
+			.AutoWidth()
+			.VAlign(VAlign_Center)
+			[
+				SNew(SBox)
+				.WidthOverride(30.0f)
+				.Padding(FMargin(2, 4))
 				.VAlign(VAlign_Center)
+				.Visibility(this, &FStateTreeEditorNodeDetails::IsConditionVisible)
 				[
-					SNew(SBox)
-					.WidthOverride(30.0f)
-					.Padding(FMargin(2, 4))
+					SNew(SComboButton)
+					.IsEnabled(TAttribute<bool>(this, &FStateTreeEditorNodeDetails::IsOperandEnabled))
+					.ComboButtonStyle(FStateTreeEditorStyle::Get(), "StateTree.Node.Operand.ComboBox")
+					.ButtonColorAndOpacity(this, &FStateTreeEditorNodeDetails::GetOperandColor)
+					.HasDownArrow(false)
+					.HAlign(HAlign_Center)
 					.VAlign(VAlign_Center)
-					.Visibility(this, &FStateTreeEditorNodeDetails::IsConditionVisible)
+					.OnGetMenuContent(this, &FStateTreeEditorNodeDetails::OnGetOperandContent)
+					.ButtonContent()
 					[
-						SNew(SComboButton)
-						.IsEnabled(TAttribute<bool>(this, &FStateTreeEditorNodeDetails::IsOperandEnabled))
-						.ComboButtonStyle(FStateTreeEditorStyle::Get(), "StateTree.Node.Operand.ComboBox")
-						.ButtonColorAndOpacity(this, &FStateTreeEditorNodeDetails::GetOperandColor)
-						.HasDownArrow(false)
-						.HAlign(HAlign_Center)
-						.VAlign(VAlign_Center)
-						.OnGetMenuContent(this, &FStateTreeEditorNodeDetails::OnGetOperandContent)
-						.ButtonContent()
-						[
-							SNew(STextBlock)
-							.TextStyle(FStateTreeEditorStyle::Get(), "StateTree.Node.Operand")
-							.Text(this, &FStateTreeEditorNodeDetails::GetOperandText)
-						]
+						SNew(STextBlock)
+						.TextStyle(FStateTreeEditorStyle::Get(), "StateTree.Node.Operand")
+						.Text(this, &FStateTreeEditorNodeDetails::GetOperandText)
 					]
 				]
-				// Open parens
-				+ SHorizontalBox::Slot()
-				.AutoWidth()
-				.VAlign(VAlign_Center)
-				[
-					SNew(STextBlock)
-					.TextStyle(FStateTreeEditorStyle::Get(), "StateTree.Node.Parens")
-					.Text(this, &FStateTreeEditorNodeDetails::GetOpenParens)
-					.Visibility(this, &FStateTreeEditorNodeDetails::IsConditionVisible)
-				]
-				// Name & type selection
-				+ SHorizontalBox::Slot()
-				.AutoWidth()
-				.Padding(FMargin(FMargin(4.0f, 0.0f, 0.0f, 0.0f)))
-				.VAlign(VAlign_Center)
+			]
+			// Open parens
+			+ SHorizontalBox::Slot()
+			.AutoWidth()
+			.VAlign(VAlign_Center)
+			[
+				SNew(STextBlock)
+				.TextStyle(FStateTreeEditorStyle::Get(), "StateTree.Node.Parens")
+				.Text(this, &FStateTreeEditorNodeDetails::GetOpenParens)
+				.Visibility(this, &FStateTreeEditorNodeDetails::IsConditionVisible)
+			]
+			// Name & type selection
+			+ SHorizontalBox::Slot()
+			.FillContentWidth(0.0f, 1.0f) // no growing, allow shrink
+			.Padding(FMargin(FMargin(4.0f, 0.0f, 0.0f, 0.0f)))
+			.VAlign(VAlign_Center)
+			[
+				SAssignNew(ComboButton, SComboButton)
+				.OnGetMenuContent(this, &FStateTreeEditorNodeDetails::GeneratePicker)
+				.ToolTipText(this, &FStateTreeEditorNodeDetails::GetNodePickerTooltip)
+				.ContentPadding(FMargin(2, 2, 2, 2))
+				.ButtonContent()
 				[
-					SAssignNew(ComboButton, SComboButton)
-					.OnGetMenuContent(this, &FStateTreeEditorNodeDetails::GeneratePicker)
-					.ToolTipText(this, &FStateTreeEditorNodeDetails::GetDisplayValueString)
-					.ContentPadding(FMargin(2, 2, 2, 2))
-					.ButtonContent()
+					SNew(SHorizontalBox)
+					+ SHorizontalBox::Slot()
+					.AutoWidth()
+					.VAlign(VAlign_Center)
+					.Padding(0, 0, 4, 0)
 					[
-						SNew(SHorizontalBox)
-						+ SHorizontalBox::Slot()
-						.AutoWidth()
-						.VAlign(VAlign_Center)
-						.Padding(0, 0, 4, 0)
-						[
-							SNew(STextBlock)
-							.Text(FEditorFontGlyphs::Paper_Plane)
-							.TextStyle(FStateTreeEditorStyle::Get(), "StateTree.DetailsIcon")
-							.Visibility(this, &FStateTreeEditorNodeDetails::IsTaskVisible)
-						]
-						+ SHorizontalBox::Slot()
-						.VAlign(VAlign_Center)
-						.Padding(0, 0, 8, 0)
+						SNew(STextBlock)
+						.Text(FEditorFontGlyphs::Paper_Plane)
+						.TextStyle(FStateTreeEditorStyle::Get(), "StateTree.DetailsIcon")
+						.Visibility(this, &FStateTreeEditorNodeDetails::IsTaskVisible)
+					]
+					+ SHorizontalBox::Slot()
+					.VAlign(VAlign_Center)
+					.FillWidth(1)
+					.Padding(0, 0, 0, 0)
+					[
+						// Show the description as rich text, when clicked switch to editable text to edit the name.
+						SNew(SBorder)
+						.BorderImage(FStyleDefaults::GetNoBrush())
+						.Padding(0)
+						.OnMouseButtonDown(this, &FStateTreeEditorNodeDetails::OnDescriptionClicked)
 						[
-							SNew(SEditableText)
-							.IsEnabled(TAttribute<bool>(this, &FStateTreeEditorNodeDetails::IsNameEnabled))
-							.Text(this, &FStateTreeEditorNodeDetails::GetName)
-							.OnTextCommitted(this, &FStateTreeEditorNodeDetails::OnNameCommitted)
-							.SelectAllTextWhenFocused(true)
-							.RevertTextOnEscape(true)
-							.Style(FStateTreeEditorStyle::Get(), "StateTree.Node.Name")
-							.Visibility(this, &FStateTreeEditorNodeDetails::IsNameVisible)
+							SAssignNew(NameSwitcher, SWidgetSwitcher)
+							.WidgetIndex(0)
+							+ SWidgetSwitcher::Slot()
+							[
+								SNew(SBox)
+								.Padding(FMargin(1,0))
+								[
+									SNew(SRichTextBlock)
+									.Text(this, &FStateTreeEditorNodeDetails::GetNodeDescription)
+									.TextStyle(&FStateTreeEditorStyle::Get().GetWidgetStyle<FTextBlockStyle>("StateTree.Node.Normal"))
+									.OverflowPolicy(ETextOverflowPolicy::Ellipsis)
+									.Visibility(this, &FStateTreeEditorNodeDetails::IsNodeDescriptionVisible)
+									+SRichTextBlock::Decorator(FTextStyleDecorator::Create(TEXT(""), FStateTreeEditorStyle::Get().GetWidgetStyle<FTextBlockStyle>("StateTree.Node.Normal")))
+									+SRichTextBlock::Decorator(FTextStyleDecorator::Create(TEXT("b"), FStateTreeEditorStyle::Get().GetWidgetStyle<FTextBlockStyle>("StateTree.Node.Bold")))
+									+SRichTextBlock::Decorator(FTextStyleDecorator::Create(TEXT("s"), FStateTreeEditorStyle::Get().GetWidgetStyle<FTextBlockStyle>("StateTree.Node.Subdued")))
+								]
+							]
+							+ SWidgetSwitcher::Slot()
+							[
+								SAssignNew(NameEdit, SEditableText)
+								.Style(&FStateTreeEditorStyle::Get().GetWidgetStyle<FEditableTextStyle>("StateTree.Node.Editable"))
+								.Text(this, &FStateTreeEditorNodeDetails::GetName)
+								.OnTextCommitted(this, &FStateTreeEditorNodeDetails::OnNameCommitted)
+								.SelectAllTextWhenFocused(true)
+								.RevertTextOnEscape(true)
+								.Visibility(this, &FStateTreeEditorNodeDetails::IsNodeDescriptionVisible)
+							]
 						]
 					]
 				]
-				// Close parens
-				+ SHorizontalBox::Slot()
-				.AutoWidth()
-				.VAlign(VAlign_Center)
-				[
-					SNew(STextBlock)
-					.TextStyle(FStateTreeEditorStyle::Get(), "StateTree.Node.Parens")
-					.Text(this, &FStateTreeEditorNodeDetails::GetCloseParens)
-					.Visibility(this, &FStateTreeEditorNodeDetails::IsConditionVisible)
-				]
+			]
+			// Close parens
+			+ SHorizontalBox::Slot()
+			.AutoWidth()
+			.VAlign(VAlign_Center)
+			[
+				SNew(STextBlock)
+				.TextStyle(FStateTreeEditorStyle::Get(), "StateTree.Node.Parens")
+				.Text(this, &FStateTreeEditorNodeDetails::GetCloseParens)
+				.Visibility(this, &FStateTreeEditorNodeDetails::IsConditionVisible)
+			]
 
-				+ SHorizontalBox::Slot()
-				.AutoWidth()
+			+ SHorizontalBox::Slot()
+			.AutoWidth()
+			[
+				StructPropertyHandle->CreateDefaultPropertyButtonWidgets()
+			]
+			// Browse To Button
+			+ SHorizontalBox::Slot()
+			.AutoWidth()
+			.VAlign(VAlign_Center)
+			[
+				SNew(SButton)
+				.ButtonStyle(FAppStyle::Get(), "SimpleButton")
+				.Visibility(this, &FStateTreeEditorNodeDetails::IsBrowseToNodeBlueprintVisible)
+				.OnClicked(this, &FStateTreeEditorNodeDetails::OnBrowseToNodeBlueprint)
+				.ToolTipText(LOCTEXT("BrowseToCurrentNodeBP", "Browse to the current node blueprint in Content Browser"))
+				.ContentPadding(.0f)
 				[
-					StructPropertyHandle->CreateDefaultPropertyButtonWidgets()
+					SNew(SImage)
+					.Image(FAppStyle::GetBrush("Icons.BrowseContent"))
+					.ColorAndOpacity(FSlateColor::UseForeground())
 				]
-				// Browse To Button
-				+ SHorizontalBox::Slot()
-				.AutoWidth()
-				.VAlign(VAlign_Center)
-				[
+			]
+			// Edit Button
+			+ SHorizontalBox::Slot()
+			.AutoWidth()
+			.VAlign(VAlign_Center)
+			[
 					SNew(SButton)
 					.ButtonStyle(FAppStyle::Get(), "SimpleButton")
-					.Visibility(this, &FStateTreeEditorNodeDetails::IsBrowseToNodeBlueprintVisible)
-					.OnClicked(this, &FStateTreeEditorNodeDetails::OnBrowseToNodeBlueprint)
-					.ToolTipText(LOCTEXT("BrowseToCurrentNodeBP", "Browse to the current node blueprint in Content Browser"))
+					.OnClicked(this, &FStateTreeEditorNodeDetails::OnEditNodeBlueprint)
+					.Visibility(this, &FStateTreeEditorNodeDetails::IsEditNodeBlueprintVisible)
+					.ToolTipText(LOCTEXT("EditCurrentNodeBP", "Edit the current node blueprint in Editor"))
 					.ContentPadding(.0f)
 					[
 						SNew(SImage)
-						.Image(FAppStyle::GetBrush("Icons.BrowseContent"))
+						.Image(FAppStyle::GetBrush("Icons.Edit"))
 						.ColorAndOpacity(FSlateColor::UseForeground())
 					]
-				]
-				// Edit Button
-				+ SHorizontalBox::Slot()
-				.AutoWidth()
-				.VAlign(VAlign_Center)
-				[
-						SNew(SButton)
-						.ButtonStyle(FAppStyle::Get(), "SimpleButton")
-						.OnClicked(this, &FStateTreeEditorNodeDetails::OnEditNodeBlueprint)
-						.Visibility(this, &FStateTreeEditorNodeDetails::IsEditNodeBlueprintVisible)
-						.ToolTipText(LOCTEXT("EditCurrentNodeBP", "Edit the current node blueprint in Editor"))
-						.ContentPadding(.0f)
-						[
-							SNew(SImage)
-							.Image(FAppStyle::GetBrush("Icons.Edit"))
-							.ColorAndOpacity(FSlateColor::UseForeground())
-						]
-				]
-				+ SHorizontalBox::Slot()
-				.HAlign(HAlign_Right)
-				.FillWidth(1.0f)
-				[
-					UE::StateTreeEditor::DebuggerExtensions::CreateEditorNodeWidget(StructPropertyHandle, EditorData)
-				]
 			]
 		]
 		.OverrideResetToDefault(ResetOverride)
@@ -789,11 +809,6 @@ void FStateTreeEditorNodeDetails::OnPasteNode()
 					UE::StateTreeEditor::Internal::InstantiateStructSubobjects(*OuterObject, EditorNode->Instance);
 				}
 				
-				if (FStateTreeNodeBase* Node = EditorNode->Node.GetMutablePtr<FStateTreeNodeBase>())
-				{
-					Node->Name = FName(Node->Name.ToString() + TEXT(" Copy"));
-				}
-
 				const FGuid OldStructID = EditorNode->ID; 
 				EditorNode->ID = FGuid::NewGuid();
 
@@ -1276,43 +1291,50 @@ EVisibility FStateTreeEditorNodeDetails::IsTaskVisible() const
 	return UE::StateTreeEditor::EditorNodeUtils::IsTaskVisible(StructProperty);
 }
 
-FText FStateTreeEditorNodeDetails::GetName() const
+FReply FStateTreeEditorNodeDetails::OnDescriptionClicked(const FGeometry& MyGeometry, const FPointerEvent& MouseEvent) const
 {
-	check(StructProperty);
+	if (NameSwitcher && NameEdit)
+	{
+		if (NameSwitcher->GetActiveWidgetIndex() == 0)
+		{
+			// Enter edit mode
+			NameSwitcher->SetActiveWidgetIndex(1);
+
+			// Focus on name edit.
+			FReply Reply = FReply::Handled();
+			Reply.SetUserFocus(NameEdit.ToSharedRef());
+			return Reply;
+		}
+	}
 
+	return FReply::Unhandled();
+}
+
+FText FStateTreeEditorNodeDetails::GetNodeDescription() const
+{
+	check(StructProperty);
+	if (!EditorData)
+	{
+		return FText::GetEmpty();
+	}
+	
 	// Multiple names do not make sense, just if only one node is selected.
 	TArray<void*> RawNodeData;
 	StructProperty->AccessRawData(RawNodeData);
 	if (RawNodeData.Num() == 1)
 	{
-		// Dig out name from the struct without knowing the type.
-		FStateTreeEditorNode* Node = static_cast<FStateTreeEditorNode*>(RawNodeData[0]);
-
-		// Make sure the associated property has valid data (e.g. while moving an node in array)
-		const UScriptStruct* ScriptStruct = Node != nullptr ? Node->Node.GetScriptStruct() : nullptr;
-		if (ScriptStruct != nullptr)
+		FText Description = LOCTEXT("EmptyNodeStyled", "<s>None</>");
+		if (const FStateTreeEditorNode* Node = static_cast<FStateTreeEditorNode*>(RawNodeData[0]))
 		{
-			if (FProperty* NameProperty = ScriptStruct->FindPropertyByName(TEXT("Name")))
-			{
-				if (NameProperty->IsA(FNameProperty::StaticClass()))
-				{
-					void* Ptr = const_cast<void*>(static_cast<const void*>(Node->Node.GetMemory()));
-					const FName NameValue = *NameProperty->ContainerPtrToValuePtr<FName>(Ptr);
-					if (NameValue.IsNone())
-					{
-						return GetDisplayValueString();
-					}
-					return FText::FromName(NameValue);
-				}
-			}
+			return EditorData->GetNodeDescription(*Node, EStateTreeNodeFormatting::RichText);
 		}
-		return LOCTEXT("Empty", "Empty");
+		return Description;
 	}
 
-	return LOCTEXT("MultipleSelected", "Multiple Selected");
+	return LOCTEXT("MultipleSelected", "<s>Multiple Selected</>");
 }
 
-EVisibility FStateTreeEditorNodeDetails::IsNameVisible() const
+EVisibility FStateTreeEditorNodeDetails::IsNodeDescriptionVisible() const
 {
 	const UScriptStruct* ScriptStruct = nullptr;
 	if (const FStateTreeEditorNode* Node = UE::StateTreeEditor::EditorNodeUtils::GetCommonNode(StructProperty))
@@ -1333,6 +1355,37 @@ EVisibility FStateTreeEditorNodeDetails::IsNameVisible() const
 	return EVisibility::Visible;
 }
 
+FText FStateTreeEditorNodeDetails::GetName() const
+{
+	check(StructProperty);
+
+	// Multiple names do not make sense, just if only one node is selected.
+	TArray<void*> RawNodeData;
+	StructProperty->AccessRawData(RawNodeData);
+	if (RawNodeData.Num() == 1)
+	{
+		if (const FStateTreeEditorNode* Node = static_cast<FStateTreeEditorNode*>(RawNodeData[0]))
+		{
+			if (const FStateTreeNodeBase* BaseNode = Node->Node.GetPtr<FStateTreeNodeBase>())
+			{
+				if (!BaseNode->Name.IsNone())
+				{
+					return FText::FromName(BaseNode->Name);
+				}
+				const FText Desc = EditorData->GetNodeDescription(*Node, EStateTreeNodeFormatting::Text);
+				if (!Desc.IsEmpty())
+				{
+					return Desc;
+				}
+			}
+		}
+
+		return FText::GetEmpty();
+	}
+
+	return LOCTEXT("MultipleSelected", "Multiple Selected");
+}
+
 void FStateTreeEditorNodeDetails::OnNameCommitted(const FText& NewText, ETextCommit::Type InTextCommit) const
 {
 	check(StructProperty);
@@ -1353,20 +1406,12 @@ void FStateTreeEditorNodeDetails::OnNameCommitted(const FText& NewText, ETextCom
 		
 		for (void* Data : RawNodeData)
 		{
+			// Set Name
 			if (FStateTreeEditorNode* Node = static_cast<FStateTreeEditorNode*>(Data))
 			{
-				// Set Name
-				if (const UScriptStruct* ScriptStruct = Node->Node.GetScriptStruct())
+				if (FStateTreeNodeBase* BaseNode = Node->Node.GetMutablePtr<FStateTreeNodeBase>())
 				{
-					if (FProperty* NameProperty = ScriptStruct->FindPropertyByName(TEXT("Name")))
-					{
-						if (NameProperty->IsA(FNameProperty::StaticClass()))
-						{
-							void* Ptr = const_cast<void*>(static_cast<const void*>(Node->Node.GetMemory()));
-							FName& NameValue = *NameProperty->ContainerPtrToValuePtr<FName>(Ptr);
-							NameValue = FName(NewName);
-						}
-					}
+					BaseNode->Name = FName(NewName);
 				}
 			}
 		}
@@ -1385,17 +1430,39 @@ void FStateTreeEditorNodeDetails::OnNameCommitted(const FText& NewText, ETextCom
 
 		StructProperty->NotifyFinishedChangingProperties();
 	}
-}
 
-bool FStateTreeEditorNodeDetails::IsNameEnabled() const
-{
-	// Can only edit if we have valid instantiated type.
-	const FStateTreeEditorNode* Node = UE::StateTreeEditor::EditorNodeUtils::GetCommonNode(StructProperty);
-	return Node && Node->Node.IsValid();
+	// Switch back to rich view.
+	NameSwitcher->SetActiveWidgetIndex(0);
 }
 
-FText FStateTreeEditorNodeDetails::GetDisplayValueString() const
+FText FStateTreeEditorNodeDetails::GetNodePickerTooltip() const
 {
+	check(StructProperty);
+	if (!EditorData)
+	{
+		return FText::GetEmpty();
+	}
+
+	FTextBuilder TextBuilder;
+
+	// Append full description.
+	TArray<void*> RawNodeData;
+	StructProperty->AccessRawData(RawNodeData);
+	if (RawNodeData.Num() == 1)
+	{
+		FText Description = LOCTEXT("EmptyNodeStyled", "<s>None</>");
+		if (const FStateTreeEditorNode* Node = static_cast<FStateTreeEditorNode*>(RawNodeData[0]))
+		{
+			TextBuilder.AppendLine(EditorData->GetNodeDescription(*Node));
+		}
+	}
+
+	if (TextBuilder.GetNumLines() > 0)
+	{
+		TextBuilder.AppendLine(FText::GetEmpty());
+	}
+	
+	// Text describing the type.
 	if (const FStateTreeEditorNode* Node = UE::StateTreeEditor::EditorNodeUtils::GetCommonNode(StructProperty))
 	{
 		if (const UScriptStruct* ScriptStruct = Node->Node.GetScriptStruct())
@@ -1407,37 +1474,17 @@ FText FStateTreeEditorNodeDetails::GetDisplayValueString() const
 				if (Node->InstanceObject != nullptr
 					&& Node->InstanceObject->GetClass() != nullptr)
 				{
-					return Node->InstanceObject->GetClass()->GetDisplayNameText();
+					TextBuilder.AppendLine(Node->InstanceObject->GetClass()->GetDisplayNameText());
 				}
 			}
 			else
 			{
-				return ScriptStruct->GetDisplayNameText();
+				TextBuilder.AppendLine(ScriptStruct->GetDisplayNameText());
 			}
 		}
 	}
-	return FText();
-}
 
-const FSlateBrush* FStateTreeEditorNodeDetails::GetDisplayValueIcon() const
-{
-	if (const FStateTreeEditorNode* Node = UE::StateTreeEditor::EditorNodeUtils::GetCommonNode(StructProperty))
-	{
-		if (const UScriptStruct* ScriptStruct = Node->Node.GetScriptStruct())
-		{
-			if (ScriptStruct->IsChildOf(FStateTreeBlueprintEvaluatorWrapper::StaticStruct())
-				|| ScriptStruct->IsChildOf(FStateTreeBlueprintTaskWrapper::StaticStruct())
-				|| ScriptStruct->IsChildOf(FStateTreeBlueprintConditionWrapper::StaticStruct()))
-			{
-				if (Node->InstanceObject != nullptr)
-				{
-					return FSlateIconFinder::FindIconBrushForClass(Node->InstanceObject->GetClass());
-				}
-			}
-		}
-	}
-	
-	return FSlateIconFinder::FindIconBrushForClass(UScriptStruct::StaticClass());
+	return TextBuilder.ToText();
 }
 
 void FStateTreeEditorNodeDetails::SortNodeTypesFunctionItemsRecursive(TArray<TSharedPtr<FStateTreeNodeTypeItem>>& Items)
@@ -2128,8 +2175,6 @@ void FStateTreeEditorNodeDetails::OnStructPicked(const UScriptStruct* InStruct)
 						if (InStruct->IsChildOf(FStateTreeTaskBase::StaticStruct()))
 						{
 							FStateTreeTaskBase& Task = Node->Node.GetMutable<FStateTreeTaskBase>();
-							Task.Name = FName(InStruct->GetDisplayNameText().ToString());
-
 							if (const UScriptStruct* InstanceType = Cast<const UScriptStruct>(Task.GetInstanceDataType()))
 							{
 								Node->Instance.InitializeAs(InstanceType);
@@ -2142,8 +2187,6 @@ void FStateTreeEditorNodeDetails::OnStructPicked(const UScriptStruct* InStruct)
 						else if (InStruct->IsChildOf(FStateTreeEvaluatorBase::StaticStruct()))
 						{
 							FStateTreeEvaluatorBase& Eval = Node->Node.GetMutable<FStateTreeEvaluatorBase>();
-							Eval.Name = FName(InStruct->GetDisplayNameText().ToString());
-
 							if (const UScriptStruct* InstanceType = Cast<const UScriptStruct>(Eval.GetInstanceDataType()))
 							{
 								Node->Instance.InitializeAs(InstanceType);
@@ -2156,8 +2199,6 @@ void FStateTreeEditorNodeDetails::OnStructPicked(const UScriptStruct* InStruct)
 						else if (InStruct->IsChildOf(FStateTreeConditionBase::StaticStruct()))
 						{
 							FStateTreeConditionBase& Cond = Node->Node.GetMutable<FStateTreeConditionBase>();
-							Cond.Name = FName(InStruct->GetDisplayNameText().ToString());
-
 							if (const UScriptStruct* InstanceType = Cast<const UScriptStruct>(Cond.GetInstanceDataType()))
 							{
 								Node->Instance.InitializeAs(InstanceType);
@@ -2223,7 +2264,6 @@ void FStateTreeEditorNodeDetails::OnClassPicked(const UClass* InClass) const
 						Node->Node.InitializeAs(FStateTreeBlueprintTaskWrapper::StaticStruct());
 						FStateTreeBlueprintTaskWrapper& Task = Node->Node.GetMutable<FStateTreeBlueprintTaskWrapper>();
 						Task.TaskClass = const_cast<UClass*>(InClass);
-						Task.Name = FName(InClass->GetDisplayNameText().ToString());
 						
 						Node->InstanceObject = NewObject<UObject>(Outer, InClass);
 
@@ -2234,7 +2274,6 @@ void FStateTreeEditorNodeDetails::OnClassPicked(const UClass* InClass) const
 						Node->Node.InitializeAs(FStateTreeBlueprintEvaluatorWrapper::StaticStruct());
 						FStateTreeBlueprintEvaluatorWrapper& Eval = Node->Node.GetMutable<FStateTreeBlueprintEvaluatorWrapper>();
 						Eval.EvaluatorClass = const_cast<UClass*>(InClass);
-						Eval.Name = FName(InClass->GetDisplayNameText().ToString());
 						
 						Node->InstanceObject = NewObject<UObject>(Outer, InClass);
 
@@ -2245,7 +2284,6 @@ void FStateTreeEditorNodeDetails::OnClassPicked(const UClass* InClass) const
 						Node->Node.InitializeAs(FStateTreeBlueprintConditionWrapper::StaticStruct());
 						FStateTreeBlueprintConditionWrapper& Cond = Node->Node.GetMutable<FStateTreeBlueprintConditionWrapper>();
 						Cond.ConditionClass = const_cast<UClass*>(InClass);
-						Cond.Name = FName(InClass->GetDisplayNameText().ToString());
 
 						Node->InstanceObject = NewObject<UObject>(Outer, InClass);
 
diff --git a/Engine/Plugins/Runtime/StateTree/Source/StateTreeEditorModule/Private/Customizations/StateTreeEditorNodeDetails.h b/Engine/Plugins/Runtime/StateTree/Source/StateTreeEditorModule/Private/Customizations/StateTreeEditorNodeDetails.h
--- a/Engine/Plugins/Runtime/StateTree/Source/StateTreeEditorModule/Private/Customizations/StateTreeEditorNodeDetails.h
+++ b/Engine/Plugins/Runtime/StateTree/Source/StateTreeEditorModule/Private/Customizations/StateTreeEditorNodeDetails.h
@@ -15,6 +15,8 @@ class IPropertyUtilities;
 class SComboButton;
 class SWidget;
 class SSearchBox;
+class SWidgetSwitcher;
+class SEditableText;
 class UStateTree;
 class UStateTreeState;
 class UStateTreeEditorData;
@@ -71,13 +73,14 @@ private:
 	EVisibility IsConditionVisible() const;
 	EVisibility IsTaskVisible() const;
 
+	FReply OnDescriptionClicked(const FGeometry& MyGeometry, const FPointerEvent& MouseEvent) const;
+	FText GetNodeDescription() const;
+	EVisibility IsNodeDescriptionVisible() const;
+	
 	FText GetName() const;
-	EVisibility IsNameVisible() const;
 	void OnNameCommitted(const FText& NewText, ETextCommit::Type InTextCommit) const;
-	bool IsNameEnabled() const;
 
-	FText GetDisplayValueString() const;
-	const FSlateBrush* GetDisplayValueIcon() const;
+	FText GetNodePickerTooltip() const;
 
 	TSharedRef<SWidget> GeneratePicker();
 	void OnStructPicked(const UScriptStruct* InStruct) const;
@@ -135,6 +138,8 @@ private:
 	TSharedPtr<SComboButton> ComboButton;
 	TSharedPtr<SSearchBox> SearchBox;
 	TSharedPtr<STreeView<TSharedPtr<FStateTreeNodeTypeItem>>> NodeTypeTree;
+	TSharedPtr<SWidgetSwitcher> NameSwitcher;
+	TSharedPtr<SEditableText> NameEdit;
 	bool bIsRestoringExpansion = false;
 
 	UStateTreeEditorData* EditorData = nullptr;
diff --git a/Engine/Plugins/Runtime/StateTree/Source/StateTreeEditorModule/Private/SStateTreeViewRow.cpp b/Engine/Plugins/Runtime/StateTree/Source/StateTreeEditorModule/Private/SStateTreeViewRow.cpp
--- a/Engine/Plugins/Runtime/StateTree/Source/StateTreeEditorModule/Private/SStateTreeViewRow.cpp
+++ b/Engine/Plugins/Runtime/StateTree/Source/StateTreeEditorModule/Private/SStateTreeViewRow.cpp
@@ -2,22 +2,20 @@
 
 #include "SStateTreeViewRow.h"
 #include "SStateTreeView.h"
-
-
 #include "EditorFontGlyphs.h"
 #include "StateTreeEditor.h"
 #include "StateTreeEditorData.h"
 #include "StateTreeEditorStyle.h"
-
 #include "Widgets/Text/SInlineEditableTextBlock.h"
 #include "Widgets/Layout/SScrollBox.h"
 #include "Widgets/Layout/SSpacer.h"
-
 #include "StateTree.h"
 #include "StateTreeState.h"
 #include "StateTreeTaskBase.h"
 #include "StateTreeViewModel.h"
 #include "Widgets/Views/SListView.h"
+#include "TextStyleDecorator.h"
+#include "Widgets/Text/SRichTextBlock.h"
 
 #define LOCTEXT_NAMESPACE "StateTreeEditor"
 
@@ -40,7 +38,7 @@ void SStateTreeViewRow::Construct(const FArguments& InArgs, const TSharedRef<STa
 	StateTreeViewModel = InStateTreeViewModel;
 	WeakState = InState;
 	const UStateTreeState* State = InState.Get();
-	WeakTreeData = State != nullptr ? State->GetTypedOuter<UStateTreeEditorData>() : nullptr;
+	WeakEditorData = State != nullptr ? State->GetTypedOuter<UStateTreeEditorData>() : nullptr;
 
 	ConstructInternal(STableRow::FArguments()
 		.Padding(5.0f)
@@ -540,18 +538,24 @@ void SStateTreeViewRow::Construct(const FArguments& InArgs, const TSharedRef<STa
 TSharedRef<SHorizontalBox> SStateTreeViewRow::CreateTasksWidget()
 {
 	const TSharedRef<SHorizontalBox> TasksBox = SNew(SHorizontalBox);
-	
+	const UStateTreeEditorData* EditorData = WeakEditorData.Get();
 	const UStateTreeState* State = WeakState.Get();
-	if (State == nullptr || State->Tasks.IsEmpty())
+	
+	if (!EditorData || !State || State->Tasks.IsEmpty())
 	{
 		return TasksBox;
 	}
 
-	TWeakObjectPtr<UStateTreeEditorData> WeakEditorData = State->GetTypedOuter<UStateTreeEditorData>();
+	const int32 NumTasks = State->Tasks.Num();
 
-	for (int32 TaskIndex = 0; TaskIndex < State->Tasks.Num(); TaskIndex++)
+	// The task descriptions can get long. Make some effort to limit how long they can get. 
+	constexpr float ReferenceWidth = 1000.0f;
+	const float MaxTaskWidth = FMath::RoundToFloat(FMath::Max(150.0f, ReferenceWidth / static_cast<float>(FMath::Max(1, NumTasks))));
+
+	for (int32 TaskIndex = 0; TaskIndex < NumTasks; TaskIndex++)
 	{
-		if (const FStateTreeTaskBase* Task = State->Tasks[TaskIndex].Node.GetPtr<FStateTreeTaskBase>())
+		const FStateTreeEditorNode& TaskNode = State->Tasks[TaskIndex];
+		if (const FStateTreeTaskBase* Task = TaskNode.Node.GetPtr<FStateTreeTaskBase>())
 		{
 			FGuid TaskId = State->Tasks[TaskIndex].ID;
 			auto IsTaskEnabledFunc = [WeakState=WeakState, TaskIndex]
@@ -567,7 +571,7 @@ TSharedRef<SHorizontalBox> SStateTreeViewRow::CreateTasksWidget()
 					return true;
 				};
 
-			auto IsTaskBreakpointEnabledFunc = [WeakEditorData, TaskId]
+			auto IsTaskBreakpointEnabledFunc = [WeakEditorData = WeakEditorData, TaskId]
 				{
 #if WITH_STATETREE_DEBUGGER
 					const UStateTreeEditorData* EditorData = WeakEditorData.Get();
@@ -579,7 +583,7 @@ TSharedRef<SHorizontalBox> SStateTreeViewRow::CreateTasksWidget()
 					return EVisibility::Hidden;
 				};
 			
-			auto GetTaskBreakpointTooltipFunc = [WeakEditorData, TaskId]
+			auto GetTaskBreakpointTooltipFunc = [WeakEditorData = WeakEditorData, TaskId]
 				{
 #if WITH_STATETREE_DEBUGGER
 					if (const UStateTreeEditorData* EditorData = WeakEditorData.Get())
@@ -605,16 +609,6 @@ TSharedRef<SHorizontalBox> SStateTreeViewRow::CreateTasksWidget()
 					return FText::GetEmpty();
 				};
 
-			FText TaskName;
-			if (UE::StateTree::Editor::GbDisplayItemIds)
-			{
-				TaskName = FText::FromString(FString::Printf(TEXT("%s (%s)"), *Task->Name.ToString(), *LexToString(TaskId)));
-			}
-			else
-			{
-				TaskName = FText::FromName(Task->Name);
-			}
-
 			TasksBox->AddSlot()
 				.AutoWidth()
 				.VAlign(VAlign_Fill)
@@ -629,12 +623,19 @@ TSharedRef<SHorizontalBox> SStateTreeViewRow::CreateTasksWidget()
 						SNew(SOverlay)
 						+ SOverlay::Slot()
 						[
-							SNew(STextBlock)
-							.Margin(FMargin(4.f, 0.f))
-							.Text(TaskName)
-							.TextStyle(FStateTreeEditorStyle::Get(), "StateTree.Task.Title")
-							.IsEnabled_Lambda(IsTaskEnabledFunc)
-							.ToolTipText(FText::FromName(Task->Name))
+							SNew(SBox)
+							.MaxDesiredWidth(MaxTaskWidth)
+							[
+								SNew(SRichTextBlock)
+								.Margin(FMargin(6.f, 0.f))
+								.Text(this, &SStateTreeViewRow::GetTaskDesc, TaskId, EStateTreeNodeFormatting::RichText)
+								.ToolTipText(this, &SStateTreeViewRow::GetTaskDesc, TaskId, EStateTreeNodeFormatting::Text)
+								.TextStyle(&FStateTreeEditorStyle::Get().GetWidgetStyle<FTextBlockStyle>("StateTree.Task.Title"))
+								.OverflowPolicy(ETextOverflowPolicy::Ellipsis)
+								+SRichTextBlock::Decorator(FTextStyleDecorator::Create(TEXT(""), FStateTreeEditorStyle::Get().GetWidgetStyle<FTextBlockStyle>("StateTree.Task.Title")))
+								+SRichTextBlock::Decorator(FTextStyleDecorator::Create(TEXT("b"), FStateTreeEditorStyle::Get().GetWidgetStyle<FTextBlockStyle>("StateTree.Task.Title.Bold")))
+								+SRichTextBlock::Decorator(FTextStyleDecorator::Create(TEXT("s"), FStateTreeEditorStyle::Get().GetWidgetStyle<FTextBlockStyle>("StateTree.Task.Title.Subdued")))
+							]
 						]
 						+ SOverlay::Slot()
 						[
@@ -675,7 +676,7 @@ void SStateTreeViewRow::RequestRename() const
 FSlateColor SStateTreeViewRow::GetTitleColor() const
 {
 	const UStateTreeState* State = WeakState.Get();
-	const UStateTreeEditorData* EditorData = WeakTreeData.Get();
+	const UStateTreeEditorData* EditorData = WeakEditorData.Get();
 
 	if (State != nullptr && EditorData != nullptr)
 	{
@@ -755,7 +756,7 @@ EVisibility SStateTreeViewRow::GetStateBreakpointVisibility() const
 {
 #if WITH_STATETREE_DEBUGGER
 	const UStateTreeState* State = WeakState.Get();
-	const UStateTreeEditorData* EditorData = WeakTreeData.Get();
+	const UStateTreeEditorData* EditorData = WeakEditorData.Get();
 	if (State != nullptr && EditorData != nullptr)
 	{
 		return (EditorData != nullptr && EditorData->HasAnyBreakpoint(State->ID)) ? EVisibility::Visible : EVisibility::Hidden;
@@ -768,7 +769,7 @@ FText SStateTreeViewRow::GetStateBreakpointTooltipText() const
 {
 #if WITH_STATETREE_DEBUGGER
 	const UStateTreeState* State = WeakState.Get();
-	const UStateTreeEditorData* EditorData = WeakTreeData.Get();
+	const UStateTreeEditorData* EditorData = WeakEditorData.Get();
 	if (State != nullptr && EditorData != nullptr)
 	{
 		const bool bHasBreakpointOnEnter = EditorData->HasBreakpoint(State->ID, EStateTreeBreakpointType::OnEnter);
@@ -879,6 +880,32 @@ FText SStateTreeViewRow::GetStateTypeTooltip() const
 	return FText::GetEmpty();
 }
 
+FText SStateTreeViewRow::GetTaskDesc(FGuid TaskID, EStateTreeNodeFormatting Formatting) const
+{
+	FText TaskName;
+	const UStateTreeState* State = WeakState.Get();
+	const UStateTreeEditorData* EditorData = WeakEditorData.Get();
+	if (EditorData != nullptr
+		&& State != nullptr)
+	{
+		const FStateTreeEditorNode* TaskNode = State->Tasks.FindByPredicate([&TaskID](const FStateTreeEditorNode& Node)
+		{
+			return Node.ID == TaskID;
+		});
+		if (TaskNode)
+		{
+			if (UE::StateTree::Editor::GbDisplayItemIds)
+			{
+				TaskName = FText::Format(LOCTEXT("TaskNameWithID", "{0} ({1})"), EditorData->GetNodeDescription(*TaskNode, Formatting), FText::AsCultureInvariant(*LexToString(TaskID)));
+			}
+			else
+			{
+				TaskName = EditorData->GetNodeDescription(*TaskNode, Formatting);
+			}
+		}
+	}
+	return TaskName;
+}
 
 EVisibility SStateTreeViewRow::GetTasksVisibility() const
 {
@@ -973,7 +1000,7 @@ FText SStateTreeViewRow::GetLinkDescription(const FStateTreeStateLink& Link)
 FText SStateTreeViewRow::GetTransitionsDesc(const UStateTreeState& State, const EStateTreeTransitionTrigger Trigger, const FTransitionDescFilterOptions FilterOptions) const
 {
 	TArray<FText> DescItems;
-	const UStateTreeEditorData* TreeEditorData = WeakTreeData.Get();
+	const UStateTreeEditorData* TreeEditorData = WeakEditorData.Get();
 
 	for (const FStateTreeTransition& Transition : State.Transitions)
 	{
@@ -1063,7 +1090,7 @@ FText SStateTreeViewRow::GetTransitionsIcon(const UStateTreeState& State, const
 	};
 	uint8 IconType = IconNone;
 	
-	const UStateTreeEditorData* TreeEditorData = WeakTreeData.Get();
+	const UStateTreeEditorData* EditorData = WeakEditorData.Get();
 	
 	for (const FStateTreeTransition& Transition : State.Transitions)
 	{
@@ -1076,7 +1103,7 @@ FText SStateTreeViewRow::GetTransitionsIcon(const UStateTreeState& State, const
 
 #if WITH_STATETREE_DEBUGGER
 		// Apply filter for transitions with/without breakpoint
-		const bool bHasBreakpoint = TreeEditorData != nullptr && TreeEditorData->HasBreakpoint(Transition.ID, EStateTreeBreakpointType::OnTransition);
+		const bool bHasBreakpoint = EditorData != nullptr && EditorData->HasBreakpoint(Transition.ID, EStateTreeBreakpointType::OnTransition);
 		if ((FilterOptions.WithBreakpoint == ETransitionDescRequirement::RequiredTrue && bHasBreakpoint == false)
 			|| (FilterOptions.WithBreakpoint == ETransitionDescRequirement::RequiredFalse && bHasBreakpoint))
 		{
@@ -1238,13 +1265,13 @@ EVisibility SStateTreeViewRow::GetTransitionsVisibility(const UStateTreeState& S
 EVisibility SStateTreeViewRow::GetTransitionsBreakpointVisibility(const UStateTreeState& State, const EStateTreeTransitionTrigger Trigger) const
 {
 #if WITH_STATETREE_DEBUGGER
-	if (const UStateTreeEditorData* TreeEditorData = WeakTreeData.Get())
+	if (const UStateTreeEditorData* EditorData = WeakEditorData.Get())
 	{
 		for (const FStateTreeTransition& Transition : State.Transitions)
 		{
 			if (Transition.bTransitionEnabled && EnumHasAnyFlags(Trigger, Transition.Trigger))
 			{
-				if (TreeEditorData->HasBreakpoint(Transition.ID, EStateTreeBreakpointType::OnTransition))
+				if (EditorData->HasBreakpoint(Transition.ID, EStateTreeBreakpointType::OnTransition))
 				{
 					return GetTransitionsVisibility(State, Trigger);
 				}
diff --git a/Engine/Plugins/Runtime/StateTree/Source/StateTreeEditorModule/Private/SStateTreeViewRow.h b/Engine/Plugins/Runtime/StateTree/Source/StateTreeEditorModule/Private/SStateTreeViewRow.h
--- a/Engine/Plugins/Runtime/StateTree/Source/StateTreeEditorModule/Private/SStateTreeViewRow.h
+++ b/Engine/Plugins/Runtime/StateTree/Source/StateTreeEditorModule/Private/SStateTreeViewRow.h
@@ -8,6 +8,7 @@
 class STableViewBase;
 class UStateTreeState;
 enum class EStateTreeTransitionTrigger : uint8;
+enum EStateTreeNodeFormatting : uint8;
 struct FStateTreeStateLink;
 
 class UStateTreeEditorData;
@@ -45,6 +46,8 @@ private:
 	FText GetSelectorTooltip() const;
 	FText GetStateTypeTooltip() const;
 
+	FText GetTaskDesc(FGuid TaskID, EStateTreeNodeFormatting Formatting) const;
+
 	EVisibility GetTasksVisibility() const;
 
 	EVisibility GetLinkedStateVisibility() const;
@@ -112,6 +115,6 @@ private:
 	
 	TSharedPtr<FStateTreeViewModel> StateTreeViewModel;
 	TWeakObjectPtr<UStateTreeState> WeakState;
-	TWeakObjectPtr<UStateTreeEditorData> WeakTreeData;
+	TWeakObjectPtr<UStateTreeEditorData> WeakEditorData;
 	TSharedPtr<SInlineEditableTextBlock> NameTextBlock;
 };
diff --git a/Engine/Plugins/Runtime/StateTree/Source/StateTreeEditorModule/Private/StateTreeCompiler.cpp b/Engine/Plugins/Runtime/StateTree/Source/StateTreeEditorModule/Private/StateTreeCompiler.cpp
--- a/Engine/Plugins/Runtime/StateTree/Source/StateTreeEditorModule/Private/StateTreeCompiler.cpp
+++ b/Engine/Plugins/Runtime/StateTree/Source/StateTreeEditorModule/Private/StateTreeCompiler.cpp
@@ -1089,7 +1089,7 @@ bool FStateTreeCompiler::CreateCondition(UStateTreeState& State, const FStateTre
 
 	FStateTreeBindableStructDesc StructDesc;
 	StructDesc.ID = CondNode.ID;
-	StructDesc.Name = CondNode.Node.GetScriptStruct()->GetFName();
+	StructDesc.Name = CondNode.GetName();
 	StructDesc.DataSource = EStateTreeBindableStructSource::Condition;
 
 	// Check that item has valid instance initialized.
@@ -1110,6 +1110,9 @@ bool FStateTreeCompiler::CreateCondition(UStateTreeState& State, const FStateTre
 	Cond.Operand = Operand;
 	Cond.DeltaIndent = DeltaIndent;
 
+	// Update condition name as description for runtime.
+	Cond.Name = CondNode.GetName();
+
 	FStateTreeDataView InstanceDataView;
 	
 	if (CondNode.Instance.IsValid())
@@ -1120,7 +1123,6 @@ bool FStateTreeCompiler::CreateCondition(UStateTreeState& State, const FStateTre
 
 		// Create binding source struct descriptor.
 		StructDesc.Struct = CondNode.Instance.GetScriptStruct();
-		StructDesc.Name = Cond.Name;
 
 		if (const auto Validation = UE::StateTree::Compiler::IsValidIndex16(InstanceIndex); Validation.DidFail())
 		{
@@ -1143,7 +1145,6 @@ bool FStateTreeCompiler::CreateCondition(UStateTreeState& State, const FStateTre
 		
 		// Create binding source struct descriptor.
 		StructDesc.Struct = Instance->GetClass();
-		StructDesc.Name = Cond.Name;
 
 		if (const auto Validation = UE::StateTree::Compiler::IsValidIndex16(InstanceIndex); Validation.DidFail())
 		{
@@ -1268,7 +1269,7 @@ bool FStateTreeCompiler::CreateTask(UStateTreeState* State, const FStateTreeEdit
 	// Create binding source struct descriptor.
 	FStateTreeBindableStructDesc StructDesc;
 	StructDesc.ID = TaskNode.ID;
-	StructDesc.Name = TaskNode.Node.GetScriptStruct()->GetFName();
+	StructDesc.Name = TaskNode.GetName();
 	StructDesc.DataSource = EStateTreeBindableStructSource::Task;
 
 	// Check that node has valid instance initialized.
@@ -1287,6 +1288,9 @@ bool FStateTreeCompiler::CreateTask(UStateTreeState* State, const FStateTreeEdit
 	FStateTreeTaskBase& Task = Node.GetMutable<FStateTreeTaskBase>();
 	FStateTreeDataView InstanceDataView;
 
+	// Update task name as description for runtime.
+	Task.Name = TaskNode.GetName();
+	
 	if (TaskNode.Instance.IsValid())
 	{
 		// Struct Instance
@@ -1295,7 +1299,6 @@ bool FStateTreeCompiler::CreateTask(UStateTreeState* State, const FStateTreeEdit
 
 		// Create binding source struct descriptor.
 		StructDesc.Struct = TaskNode.Instance.GetScriptStruct();
-		StructDesc.Name = Task.Name;
 
 		if (const auto Validation = UE::StateTree::Compiler::IsValidIndex16(InstanceIndex); Validation.DidFail())
 		{
@@ -1318,7 +1321,6 @@ bool FStateTreeCompiler::CreateTask(UStateTreeState* State, const FStateTreeEdit
 
 		// Create binding source struct descriptor.
 		StructDesc.Struct = Instance->GetClass();
-		StructDesc.Name = Task.Name;
 
 		if (const auto Validation = UE::StateTree::Compiler::IsValidIndex16(InstanceIndex); Validation.DidFail())
 		{
@@ -1381,7 +1383,7 @@ bool FStateTreeCompiler::CreateEvaluator(const FStateTreeEditorNode& EvalNode, c
 	// Create binding source struct descriptor.
 	FStateTreeBindableStructDesc StructDesc;
     StructDesc.ID = EvalNode.ID;
-    StructDesc.Name = EvalNode.Node.GetScriptStruct()->GetFName();
+	StructDesc.Name = EvalNode.GetName();
 	StructDesc.DataSource = EStateTreeBindableStructSource::Evaluator;
 
     // Check that node has valid instance initialized.
@@ -1399,7 +1401,10 @@ bool FStateTreeCompiler::CreateEvaluator(const FStateTreeEditorNode& EvalNode, c
 	
 	FStateTreeEvaluatorBase& Eval = Node.GetMutable<FStateTreeEvaluatorBase>();
 	FStateTreeDataView InstanceDataView;
-	
+
+	// Update eval name as description for runtime.
+	Eval.Name = EvalNode.GetName();
+
 	if (EvalNode.Instance.IsValid())
 	{
 		// Struct Instance
@@ -1408,7 +1413,6 @@ bool FStateTreeCompiler::CreateEvaluator(const FStateTreeEditorNode& EvalNode, c
 
 		// Create binding source struct descriptor.
 		StructDesc.Struct = EvalNode.Instance.GetScriptStruct();
-		StructDesc.Name = Eval.Name;
 
 		if (const auto Validation = UE::StateTree::Compiler::IsValidIndex16(InstanceIndex); Validation.DidFail())
 		{
@@ -1431,7 +1435,6 @@ bool FStateTreeCompiler::CreateEvaluator(const FStateTreeEditorNode& EvalNode, c
 		
 		// Create binding source struct descriptor.
 		StructDesc.Struct = Instance->GetClass();
-		StructDesc.Name = Eval.Name;
 
 		if (const auto Validation = UE::StateTree::Compiler::IsValidIndex16(InstanceIndex); Validation.DidFail())
 		{
diff --git a/Engine/Plugins/Runtime/StateTree/Source/StateTreeEditorModule/Private/StateTreeEditorData.cpp b/Engine/Plugins/Runtime/StateTree/Source/StateTreeEditorModule/Private/StateTreeEditorData.cpp
--- a/Engine/Plugins/Runtime/StateTree/Source/StateTreeEditorModule/Private/StateTreeEditorData.cpp
+++ b/Engine/Plugins/Runtime/StateTree/Source/StateTreeEditorModule/Private/StateTreeEditorData.cpp
@@ -17,6 +17,8 @@
 
 #include UE_INLINE_GENERATED_CPP_BY_NAME(StateTreeEditorData)
 
+#define LOCTEXT_NAMESPACE "StateTreeEditor"
+
 UStateTreeEditorData::UStateTreeEditorData()
 {
 	FStateTreeEditorColor DefaultColor;
@@ -196,11 +198,6 @@ void UStateTreeEditorData::PostEditChangeChainProperty(FPropertyChangedChainEven
 				const int32 ArrayIndex = PropertyChangedEvent.GetArrayIndex(MemberProperty->GetFName().ToString());
 				if (Evaluators.IsValidIndex(ArrayIndex))
 				{
-					if (FStateTreeEvaluatorBase* Eval = Evaluators[ArrayIndex].Node.GetMutablePtr<FStateTreeEvaluatorBase>())
-					{
-						Eval->Name = FName(Eval->Name.ToString() + TEXT(" Duplicate"));
-					}
-					
 					const FGuid OldStructID = Evaluators[ArrayIndex].ID;
 					Evaluators[ArrayIndex].ID = FGuid::NewGuid();
 					EditorBindings.CopyBindings(OldStructID, Evaluators[ArrayIndex].ID);
@@ -211,11 +208,6 @@ void UStateTreeEditorData::PostEditChangeChainProperty(FPropertyChangedChainEven
 				const int32 ArrayIndex = PropertyChangedEvent.GetArrayIndex(MemberProperty->GetFName().ToString());
 				if (GlobalTasks.IsValidIndex(ArrayIndex))
 				{
-					if (FStateTreeTaskBase* Task = GlobalTasks[ArrayIndex].Node.GetMutablePtr<FStateTreeTaskBase>())
-					{
-						Task->Name = FName(Task->Name.ToString() + TEXT(" Duplicate"));
-					}
-					
 					const FGuid OldStructID = GlobalTasks[ArrayIndex].ID;
 					GlobalTasks[ArrayIndex].ID = FGuid::NewGuid();
 					EditorBindings.CopyBindings(OldStructID, GlobalTasks[ArrayIndex].ID);
@@ -616,6 +608,38 @@ void UStateTreeEditorData::FixObjectNodes()
 	}
 }
 
+FText UStateTreeEditorData::GetNodeDescription(const FStateTreeEditorNode& Node, const EStateTreeNodeFormatting Formatting) const
+{
+	if (const FStateTreeNodeBase* NodePtr = Node.Node.GetPtr<FStateTreeNodeBase>())
+	{
+		// If the node has name override, return it.
+		if (!NodePtr->Name.IsNone())
+		{
+			return FText::FromName(NodePtr->Name);
+		}
+
+		// If the node has automatic description, return it.
+		const FStateTreeBindingLookup BindingLookup(this);
+		const FStateTreeDataView InstanceData = Node.GetInstance();
+		if (InstanceData.IsValid())
+		{
+			
+			const FText Description = NodePtr->GetDescription(Node.ID, InstanceData, BindingLookup, Formatting);
+			if (!Description.IsEmpty())
+			{
+				return Description;
+			}
+		}
+
+		// As last resort, return node's display name.
+		check(Node.Node.GetScriptStruct());
+		return Node.Node.GetScriptStruct()->GetDisplayNameText();
+	}
+
+	// The node is not initialized.
+	return LOCTEXT("EmptyNode", "None");
+}
+
 void UStateTreeEditorData::FixDuplicateIDs()
 {
 	// Around version 5.1-5.3 we had issue that copy/paste or some duplication methods could create nodes with duplicate IDs.
@@ -918,7 +942,7 @@ EStateTreeVisitor UStateTreeEditorData::VisitStateNodes(const UStateTreeState& S
 			{
 				FStateTreeBindableStructDesc Desc;
 				Desc.Struct = Cond->GetInstanceDataType();
-				Desc.Name = Cond->Name;
+				Desc.Name = Node.GetName();
 				Desc.ID = Node.ID;
 				Desc.DataSource = EStateTreeBindableStructSource::Condition;
 
@@ -939,7 +963,7 @@ EStateTreeVisitor UStateTreeEditorData::VisitStateNodes(const UStateTreeState& S
 			{
 				FStateTreeBindableStructDesc Desc;
 				Desc.Struct = Task->GetInstanceDataType();
-				Desc.Name = Task->Name;
+				Desc.Name = Node.GetName();
 				Desc.ID = Node.ID;
 				Desc.DataSource = EStateTreeBindableStructSource::Task;
 
@@ -957,7 +981,7 @@ EStateTreeVisitor UStateTreeEditorData::VisitStateNodes(const UStateTreeState& S
 		{
 			FStateTreeBindableStructDesc Desc;
 			Desc.Struct = Task->GetInstanceDataType();
-			Desc.Name = Task->Name;
+			Desc.Name = State.SingleTask.GetName();
 			Desc.ID = State.SingleTask.ID;
 			Desc.DataSource = EStateTreeBindableStructSource::Task;
 
@@ -994,7 +1018,7 @@ EStateTreeVisitor UStateTreeEditorData::VisitStateNodes(const UStateTreeState& S
 				{
 					FStateTreeBindableStructDesc Desc;
 					Desc.Struct = Cond->GetInstanceDataType();
-					Desc.Name = Cond->Name;
+					Desc.Name = Node.GetName();
 					Desc.ID = Node.ID;
 					Desc.DataSource = EStateTreeBindableStructSource::Condition;
 
@@ -1109,7 +1133,7 @@ EStateTreeVisitor UStateTreeEditorData::VisitGlobalNodes(TFunctionRef<EStateTree
 		{
 			FStateTreeBindableStructDesc Desc;
 			Desc.Struct = Evaluator->GetInstanceDataType();
-			Desc.Name = Evaluator->Name;
+			Desc.Name = Node.GetName();
 			Desc.ID = Node.ID;
 			Desc.DataSource = EStateTreeBindableStructSource::Evaluator;
 
@@ -1127,7 +1151,7 @@ EStateTreeVisitor UStateTreeEditorData::VisitGlobalNodes(TFunctionRef<EStateTree
 		{
 			FStateTreeBindableStructDesc Desc;
 			Desc.Struct = Task->GetInstanceDataType();
-			Desc.Name = Task->Name;
+			Desc.Name = Node.GetName();
 			Desc.ID = Node.ID;
 			Desc.DataSource = EStateTreeBindableStructSource::GlobalTask;
 
@@ -1214,3 +1238,5 @@ void UStateTreeEditorData::AddPropertyBinding(const FStateTreeEditorPropertyPath
 	EditorBindings.AddPropertyBinding(UE::StateTree::Private::ConvertEditorPath(SourcePath), UE::StateTree::Private::ConvertEditorPath(TargetPath));
 }
 PRAGMA_ENABLE_DEPRECATION_WARNINGS
+
+#undef LOCTEXT_NAMESPACE
\ No newline at end of file
diff --git a/Engine/Plugins/Runtime/StateTree/Source/StateTreeEditorModule/Private/StateTreeEditorPropertyBindings.cpp b/Engine/Plugins/Runtime/StateTree/Source/StateTreeEditorModule/Private/StateTreeEditorPropertyBindings.cpp
--- a/Engine/Plugins/Runtime/StateTree/Source/StateTreeEditorModule/Private/StateTreeEditorPropertyBindings.cpp
+++ b/Engine/Plugins/Runtime/StateTree/Source/StateTreeEditorModule/Private/StateTreeEditorPropertyBindings.cpp
@@ -4,9 +4,12 @@
 #include "StateTreePropertyBindingCompiler.h"
 #include "Misc/EnumerateRange.h"
 #include "PropertyPathHelpers.h"
+#include "StateTreeNodeBase.h"
 
 #include UE_INLINE_GENERATED_CPP_BY_NAME(StateTreeEditorPropertyBindings)
 
+#define LOCTEXT_NAMESPACE "StateTreeEditor"
+
 UStateTreeEditorPropertyBindingsOwner::UStateTreeEditorPropertyBindingsOwner(const FObjectInitializer& ObjectInitializer)
 	: Super(ObjectInitializer)
 {
@@ -214,7 +217,7 @@ void FStateTreeEditorPropertyBindings::RemoveUnusedBindings(const TMap<FGuid, co
 
 //////////////////////////////////////////////////////////////////////////
 
-FStateTreeBindingLookup::FStateTreeBindingLookup(IStateTreeEditorPropertyBindingsOwner* InBindingOwner)
+FStateTreeBindingLookup::FStateTreeBindingLookup(const IStateTreeEditorPropertyBindingsOwner* InBindingOwner)
 	: BindingOwner(InBindingOwner)
 {
 }
@@ -228,7 +231,7 @@ const FStateTreePropertyPath* FStateTreeBindingLookup::GetPropertyBindingSource(
 	return EditorBindings->GetPropertyBindingSource(InTargetPath);
 }
 
-FText FStateTreeBindingLookup::GetPropertyPathDisplayName(const FStateTreePropertyPath& InPath) const
+FText FStateTreeBindingLookup::GetPropertyPathDisplayName(const FStateTreePropertyPath& InPath, EStateTreeNodeFormatting Formatting) const
 {
 	check(BindingOwner);
 	const FStateTreeEditorPropertyBindings* EditorBindings = BindingOwner->GetPropertyEditorBindings();
@@ -242,11 +245,73 @@ FText FStateTreeBindingLookup::GetPropertyPathDisplayName(const FStateTreeProper
 		Result = Struct.Name.ToString();
 	}
 
-	Result += TEXT(".") + InPath.ToString();
+	if (!InPath.IsPathEmpty())
+	{
+		Result += TEXT(".") + InPath.ToString();
+	}
 
 	return FText::FromString(Result);
 }
 
+FText FStateTreeBindingLookup::GetBindingSourceDisplayName(const FStateTreePropertyPath& InTargetPath, EStateTreeNodeFormatting Formatting) const
+{
+	check(BindingOwner);
+	const FStateTreeEditorPropertyBindings* EditorBindings = BindingOwner->GetPropertyEditorBindings();
+	check(EditorBindings);
+
+	// Check if the target property is bound, if so, return binding description.
+	if (const FStateTreePropertyPath* SourcePath = GetPropertyBindingSource(InTargetPath))
+	{
+		return GetPropertyPathDisplayName(*SourcePath, Formatting);
+	}
+
+	// Check if it's bound to context data.
+	const UStruct* TargetStruct = nullptr;
+	const FProperty* TargetProperty = nullptr;
+	EStateTreePropertyUsage Usage = EStateTreePropertyUsage::Invalid;
+	
+	FStateTreeBindableStructDesc TargetStructDesc;
+	if (BindingOwner->GetStructByID(InTargetPath.GetStructID(), TargetStructDesc))
+	{
+		TArray<FStateTreePropertyPathIndirection> Indirection;
+		if (InTargetPath.ResolveIndirections(TargetStructDesc.Struct, Indirection)
+			&& Indirection.Num() > 0)
+		{
+			const FStateTreePropertyPathIndirection& Leaf = Indirection.Last(); 
+			TargetProperty = Leaf.GetProperty();
+			if (TargetProperty)
+			{
+				Usage = UE::StateTree::GetUsageFromMetaData(TargetProperty);
+			}
+			if (const FStructProperty* StructProperty = CastField<FStructProperty>(TargetProperty))
+			{
+				TargetStruct = StructProperty->Struct;
+			}
+			if (const FObjectPropertyBase* ObjectProperty = CastField<FObjectPropertyBase>(TargetProperty))
+			{
+				TargetStruct = ObjectProperty->PropertyClass;
+			}
+		}
+	}
+
+	if (Usage == EStateTreePropertyUsage::Context)
+	{
+		if (TargetStruct)
+		{
+			const FStateTreeBindableStructDesc Desc = BindingOwner->FindContextData(TargetStruct, TargetProperty->GetName());
+			if (Desc.IsValid())
+			{
+				// Connected
+				return FText::FromName(Desc.Name);
+			}
+		}
+		return LOCTEXT("Unlinked", "???");
+	}
+
+	// Not a binding nor context data.
+	return FText::GetEmpty();
+}
+
 const FProperty* FStateTreeBindingLookup::GetPropertyPathLeafProperty(const FStateTreePropertyPath& InPath) const
 {
 	check(BindingOwner);
@@ -266,3 +331,5 @@ const FProperty* FStateTreeBindingLookup::GetPropertyPathLeafProperty(const FSta
 
 	return Result;
 }
+
+#undef LOCTEXT_NAMESPACE
\ No newline at end of file
diff --git a/Engine/Plugins/Runtime/StateTree/Source/StateTreeEditorModule/Private/StateTreeEditorStyle.cpp b/Engine/Plugins/Runtime/StateTree/Source/StateTreeEditorModule/Private/StateTreeEditorStyle.cpp
--- a/Engine/Plugins/Runtime/StateTree/Source/StateTreeEditorModule/Private/StateTreeEditorStyle.cpp
+++ b/Engine/Plugins/Runtime/StateTree/Source/StateTreeEditorModule/Private/StateTreeEditorStyle.cpp
@@ -104,15 +104,22 @@ FStateTreeEditorStyle::FStateTreeEditorStyle()
 
 	// Task
 	{
-		const FTextBlockStyle TaskTitle = FTextBlockStyle(NormalText)
+		Set("StateTree.Task.Title", FTextBlockStyle(NormalText)
 			.SetFont(DEFAULT_FONT("Regular", 10))
-			.SetColorAndOpacity(FLinearColor(230.0f / 255.0f, 230.0f / 255.0f, 230.0f / 255.0f, 0.85f));
-		Set("StateTree.Task.Title", TaskTitle);
+			.SetColorAndOpacity(FLinearColor(230.0f / 255.0f, 230.0f / 255.0f, 230.0f / 255.0f, 0.85f)));
+
+		Set("StateTree.Task.Title.Bold", FTextBlockStyle(NormalText)
+			.SetFont(DEFAULT_FONT("Bold", 10))
+			.SetColorAndOpacity(FLinearColor(230.0f / 255.0f, 230.0f / 255.0f, 230.0f / 255.0f, 0.85f)));
+
+		Set("StateTree.Task.Title.Subdued", FTextBlockStyle(NormalText)
+			.SetFont(DEFAULT_FONT("Regular", 10))
+			.SetColorAndOpacity(FLinearColor(230.0f / 255.0f, 230.0f / 255.0f, 230.0f / 255.0f, 0.4f)));
 
 		// Tasks to be show up a bit darker than the state
 		Set("StateTree.Task.Rect", new FSlateColorBrush(FLinearColor(FVector3f(0.67f))));
 	}
-	
+
 	// Details rich text
 	{
 		Set("Details.Normal", FTextBlockStyle(NormalText)
@@ -197,11 +204,29 @@ FStateTreeEditorStyle::FStateTreeEditorStyle()
 	
 	Set("StateTree.Node.Indent.ComboBox", FComboButtonStyle(ComboButtonStyle).SetButtonStyle(IndentButton));
 
-	const FEditableTextStyle& NormalEditableText = FCoreStyle::Get().GetWidgetStyle<FEditableTextStyle>("NormalEditableText");
-	FEditableTextStyle NameEditStyle(NormalEditableText);
-	NameEditStyle.Font.Size = 10;
-	Set("StateTree.Node.Name", NameEditStyle);
+	
+	// Node text styles
+	{
+		FEditableTextStyle EditableTextStyle = FEditableTextStyle(FAppStyle::GetWidgetStyle<FEditableTextStyle>("NormalEditableText"));
+		EditableTextStyle.Font = FAppStyle::GetFontStyle(TEXT("PropertyWindow.NormalFont"));
+		EditableTextStyle.Font.Size = 10.0f;
+		Set("StateTree.Node.Editable", EditableTextStyle);
 
+		Set("StateTree.Node.Normal", FTextBlockStyle(NormalText)
+			.SetFont(FAppStyle::GetFontStyle(TEXT("PropertyWindow.NormalFont")))
+			.SetFontSize(10));
+
+		Set("StateTree.Node.Bold", FTextBlockStyle(NormalText)
+			.SetFont(FAppStyle::GetFontStyle(TEXT("PropertyWindow.BoldFont")))
+			.SetFontSize(10));
+
+		Set("StateTree.Node.Subdued", FTextBlockStyle(NormalText)
+			.SetColorAndOpacity(FSlateColor::UseSubduedForeground())
+			.SetFont(FAppStyle::GetFontStyle(TEXT("PropertyWindow.NormalFont")))
+			.SetFontSize(10));
+	}
+
+	
 	// Command icons
 	{
 		// From generic engine
diff --git a/Engine/Plugins/Runtime/StateTree/Source/StateTreeEditorModule/Private/StateTreeState.cpp b/Engine/Plugins/Runtime/StateTree/Source/StateTreeEditorModule/Private/StateTreeState.cpp
--- a/Engine/Plugins/Runtime/StateTree/Source/StateTreeEditorModule/Private/StateTreeState.cpp
+++ b/Engine/Plugins/Runtime/StateTree/Source/StateTreeEditorModule/Private/StateTreeState.cpp
@@ -253,10 +253,6 @@ void UStateTreeState::PostEditChangeChainProperty(FPropertyChangedChainEvent& Pr
 			const int32 ArrayIndex = ChangePropertyPath.GetPropertyArrayIndex(StateTasksPath);
 			if (Tasks.IsValidIndex(ArrayIndex))
 			{
- 				if (FStateTreeTaskBase* Task = Tasks[ArrayIndex].Node.GetMutablePtr<FStateTreeTaskBase>())
-				{
-					Task->Name = FName(Task->Name.ToString() + TEXT(" Duplicate"));
-				}
 				const FGuid OldStructID = Tasks[ArrayIndex].ID; 
 				Tasks[ArrayIndex].ID = FGuid::NewGuid();
 				CopyBindings(OldStructID, Tasks[ArrayIndex].ID);
@@ -269,10 +265,6 @@ void UStateTreeState::PostEditChangeChainProperty(FPropertyChangedChainEvent& Pr
 			const int32 ArrayIndex = ChangePropertyPath.GetPropertyArrayIndex(StateEnterConditionsPath);
 			if (EnterConditions.IsValidIndex(ArrayIndex))
 			{
-				if (FStateTreeConditionBase* Condition = EnterConditions[ArrayIndex].Node.GetMutablePtr<FStateTreeConditionBase>())
-				{
-					Condition->Name = FName(Condition->Name.ToString() + TEXT(" Duplicate"));
-				}
 				const FGuid OldStructID = EnterConditions[ArrayIndex].ID; 
 				EnterConditions[ArrayIndex].ID = FGuid::NewGuid();
 				CopyBindings(OldStructID, EnterConditions[ArrayIndex].ID);
@@ -309,10 +301,6 @@ void UStateTreeState::PostEditChangeChainProperty(FPropertyChangedChainEvent& Pr
 				FStateTreeTransition& Transition = Transitions[TransitionsIndex];
 				if (Transition.Conditions.IsValidIndex(ConditionsIndex))
 				{
-					if (FStateTreeConditionBase* Condition = Transition.Conditions[ConditionsIndex].Node.GetMutablePtr<FStateTreeConditionBase>())
-					{
-						Condition->Name = FName(Condition->Name.ToString() + TEXT(" Duplicate"));
-					}
 					const FGuid OldStructID = Transition.Conditions[ConditionsIndex].ID;
 					Transition.Conditions[ConditionsIndex].ID = FGuid::NewGuid();
 					CopyBindings(OldStructID, Transition.Conditions[ConditionsIndex].ID);
diff --git a/Engine/Plugins/Runtime/StateTree/Source/StateTreeEditorModule/Private/TextStyleDecorator.cpp b/Engine/Plugins/Runtime/StateTree/Source/StateTreeEditorModule/Private/TextStyleDecorator.cpp
new file mode 100644
--- /dev/null
+++ b/Engine/Plugins/Runtime/StateTree/Source/StateTreeEditorModule/Private/TextStyleDecorator.cpp
@@ -0,0 +1,61 @@
+﻿// Copyright Epic Games, Inc. All Rights Reserved.
+
+#include "TextStyleDecorator.h"
+#include "Framework/Text/SlateTextRun.h"
+#include "Styling/SlateTypes.h"
+
+TSharedRef<FTextStyleDecorator> FTextStyleDecorator::Create(FString InName, const FTextBlockStyle& TextStyle)
+{
+	return MakeShareable(new FTextStyleDecorator(MoveTemp(InName), TextStyle));
+}
+
+bool FTextStyleDecorator::Supports(const FTextRunParseResults& RunInfo, const FString& Text) const
+{
+	return RunInfo.Name == DecoratorName;
+}
+
+TSharedRef<ISlateRun> FTextStyleDecorator::Create(const TSharedRef<class FTextLayout>& TextLayout, const FTextRunParseResults& RunParseResult, const FString& OriginalText, const TSharedRef<FString>& ModelText, const ISlateStyle* Style)
+{
+	FRunInfo RunInfo(RunParseResult.Name);
+	for (const TPair<FString, FTextRange>& Pair : RunParseResult.MetaData)
+	{
+		RunInfo.MetaData.Add(Pair.Key, OriginalText.Mid(Pair.Value.BeginIndex, Pair.Value.EndIndex - Pair.Value.BeginIndex));
+	}
+
+	FString Run;
+	if (!RunParseResult.ContentRange.IsEmpty())
+	{
+		Run = OriginalText.Mid(RunParseResult.ContentRange.BeginIndex, RunParseResult.ContentRange.Len());
+	}
+	else
+	{
+		// Handles the case when the decorator name is empty (matches the runs without any tags).
+		Run = OriginalText.Mid(RunParseResult.OriginalRange.BeginIndex, RunParseResult.OriginalRange.Len());
+	}
+
+	// Remove any formatting in between (e.g. when using nested tags).
+	int32 TagStartIndex = INDEX_NONE;
+	for (int32 Index = 0; Index < Run.Len(); Index++)
+	{
+		if (TagStartIndex == INDEX_NONE)
+		{
+			if (Run[Index] == TCHAR('<'))
+			{
+				TagStartIndex = Index;
+			}
+		}
+		else
+		{
+			if (Run[Index] == TCHAR('>'))
+			{
+				Run.RemoveAt(TagStartIndex, Index - TagStartIndex + 1);
+				Index = TagStartIndex;
+				TagStartIndex = INDEX_NONE;
+			}
+		}
+	}
+	
+	ModelText->Append(Run);
+
+	return FSlateTextRun::Create(RunInfo, ModelText, TextStyle);
+}
diff --git a/Engine/Plugins/Runtime/StateTree/Source/StateTreeEditorModule/Private/TextStyleDecorator.h b/Engine/Plugins/Runtime/StateTree/Source/StateTreeEditorModule/Private/TextStyleDecorator.h
new file mode 100644
--- /dev/null
+++ b/Engine/Plugins/Runtime/StateTree/Source/StateTreeEditorModule/Private/TextStyleDecorator.h
@@ -0,0 +1,42 @@
+﻿// Copyright Epic Games, Inc. All Rights Reserved.
+
+#pragma once
+
+#include "Framework/Text/ITextDecorator.h"
+#include "Styling/SlateTypes.h"
+
+/**
+ * Simple decorator for SRichTextBlock, which allows to map simple/shorthand tags to specific style names.
+ * The decorator also strips away any tags, e.g. in case the initial text came with nested tags.
+ */
+class FTextStyleDecorator : public ITextDecorator
+{
+public:
+
+	static TSharedRef<FTextStyleDecorator> Create(FString InName, const FTextBlockStyle& TextStyle);
+
+	virtual bool Supports(const FTextRunParseResults& RunInfo, const FString& Text) const override;
+
+	virtual TSharedRef<ISlateRun> Create(
+		const TSharedRef<class FTextLayout>& TextLayout,
+		const FTextRunParseResults& RunParseResult,
+		const FString& OriginalText,
+		const TSharedRef<FString>& ModelText,
+		const ISlateStyle* Style) override;
+
+private:
+
+	FTextStyleDecorator(FString&& InName, const FTextBlockStyle& InTextStyle)
+		: DecoratorName(MoveTemp(InName))
+		, TextStyle(InTextStyle)
+	{
+	}
+
+private:
+
+	/** Name of this decorator */
+	FString DecoratorName;
+
+	/** Color of this decorator */
+	FTextBlockStyle TextStyle;
+};
\ No newline at end of file
diff --git a/Engine/Plugins/Runtime/StateTree/Source/StateTreeEditorModule/Public/StateTreeEditorData.h b/Engine/Plugins/Runtime/StateTree/Source/StateTreeEditorModule/Public/StateTreeEditorData.h
--- a/Engine/Plugins/Runtime/StateTree/Source/StateTreeEditorModule/Public/StateTreeEditorData.h
+++ b/Engine/Plugins/Runtime/StateTree/Source/StateTreeEditorModule/Public/StateTreeEditorData.h
@@ -59,8 +59,18 @@ public:
 	virtual bool GetStructByID(const FGuid StructID, FStateTreeBindableStructDesc& OutStructDesc) const override;
 	virtual bool GetDataViewByID(const FGuid StructID, FStateTreeDataView& OutDataView) const override;
 	virtual FStateTreeEditorPropertyBindings* GetPropertyEditorBindings() override { return &EditorBindings; }
+	virtual const FStateTreeEditorPropertyBindings* GetPropertyEditorBindings() const override { return &EditorBindings; }
+	virtual FStateTreeBindableStructDesc FindContextData(const UStruct* ObjectType, const FString ObjectNameHint) const override;
+
 	// ~IStateTreeEditorPropertyBindingsOwner
 
+	/**
+	 * Returns the description for the node for UI.
+	 * Handles the name override logic, figures out required data for the GetDescription() call, and handles the fallbacks.
+	 * @return description for the node.
+	 */
+	FText GetNodeDescription(const FStateTreeEditorNode& Node, const EStateTreeNodeFormatting Formatting = EStateTreeNodeFormatting::Text) const;
+
 #if WITH_EDITOR
 	using FReplacementObjectMap = TMap<UObject*, UObject*>;
 	void OnObjectsReinstanced(const FReplacementObjectMap& ObjectMap);
@@ -121,13 +131,6 @@ public:
 	 */
 	void GetAccessibleStructs(const TConstArrayView<const UStateTreeState*> Path, const FGuid TargetStructID, TArray<FStateTreeBindableStructDesc>& OutStructDescs) const;
 
-	/**
-	 * Finds a bindable context struct based on name and type.
-	 * @param ObjectType Object type to match
-	 * @param ObjectNameHint Name to use if multiple context objects of same type are found. 
-	 */
-	FStateTreeBindableStructDesc FindContextData(const UStruct* ObjectType, const FString ObjectNameHint) const;
-
 	UE_DEPRECATED(5.3, "Use VisitHierarchyNodes with State, Desc, Value instead.")
 	void VisitHierarchyNodes(TFunctionRef<EStateTreeVisitor(const UStateTreeState* State, const FGuid& ID, const FName& Name, const EStateTreeNodeType NodeType, const UScriptStruct* NodeStruct, const UStruct* InstanceStruct)> InFunc) const;
 
diff --git a/Engine/Plugins/Runtime/StateTree/Source/StateTreeEditorModule/Public/StateTreeEditorNode.h b/Engine/Plugins/Runtime/StateTree/Source/StateTreeEditorModule/Public/StateTreeEditorNode.h
--- a/Engine/Plugins/Runtime/StateTree/Source/StateTreeEditorModule/Public/StateTreeEditorNode.h
+++ b/Engine/Plugins/Runtime/StateTree/Source/StateTreeEditorModule/Public/StateTreeEditorNode.h
@@ -31,16 +31,25 @@ struct STATETREEEDITORMODULE_API FStateTreeEditorNode
 		ID = FGuid();
 	}
 
+	/**
+	 * This is used to name nodes for runtime, as well as for error reporting.
+	 * If the node has a specified name, used that, or else of return the display name of the node.
+	 * @return name of the node.
+	 */
 	FName GetName() const
 	{
 		if (const FStateTreeNodeBase* NodePtr = Node.GetPtr<FStateTreeNodeBase>())
 		{
+			if (NodePtr->Name.IsNone())
+			{
+				return FName(Node.GetScriptStruct()->GetDisplayNameText().ToString());
+			}
 			return NodePtr->Name;
 		}
 		return FName();
 	}
 
-	const FStateTreeDataView GetInstance() const
+	FStateTreeDataView GetInstance() const
 	{
 		return InstanceObject ? FStateTreeDataView(InstanceObject) : FStateTreeDataView(const_cast<FInstancedStruct&>(Instance));
 	}
diff --git a/Engine/Plugins/Runtime/StateTree/Source/StateTreeEditorModule/Public/StateTreeEditorPropertyBindings.h b/Engine/Plugins/Runtime/StateTree/Source/StateTreeEditorModule/Public/StateTreeEditorPropertyBindings.h
--- a/Engine/Plugins/Runtime/StateTree/Source/StateTreeEditorModule/Public/StateTreeEditorPropertyBindings.h
+++ b/Engine/Plugins/Runtime/StateTree/Source/StateTreeEditorModule/Public/StateTreeEditorPropertyBindings.h
@@ -6,6 +6,8 @@
 #include "StateTreePropertyBindings.h"
 #include "StateTreeEditorPropertyBindings.generated.h"
 
+enum EStateTreeNodeFormatting : uint8;
+
 /**
  * Editor representation of a all property bindings in a StateTree
  */
@@ -116,6 +118,13 @@ class STATETREEEDITORMODULE_API IStateTreeEditorPropertyBindingsOwner
 	 */
 	virtual bool GetStructByID(const FGuid StructID, FStateTreeBindableStructDesc& OutStructDesc) const PURE_VIRTUAL(IStateTreeEditorPropertyBindingsOwner::GetStructByID, return false; );
 
+	/**
+	 * Finds a bindable context struct based on name and type.
+	 * @param ObjectType Object type to match
+	 * @param ObjectNameHint Name to use if multiple context objects of same type are found. 
+	 */
+	virtual FStateTreeBindableStructDesc FindContextData(const UStruct* ObjectType, const FString ObjectNameHint) const PURE_VIRTUAL(IStateTreeEditorPropertyBindingsOwner::FindContextData, return {}; );
+	
 	/**
 	 * Returns data view based on struct ID.
 	 * @param StructID Target struct ID
@@ -126,6 +135,9 @@ class STATETREEEDITORMODULE_API IStateTreeEditorPropertyBindingsOwner
 
 	/** @return Pointer to editor property bindings. */
 	virtual FStateTreeEditorPropertyBindings* GetPropertyEditorBindings() PURE_VIRTUAL(IStateTreeEditorPropertyBindingsOwner::GetPropertyEditorBindings, return nullptr; );
+
+	/** @return Pointer to editor property bindings. */
+	virtual const FStateTreeEditorPropertyBindings* GetPropertyEditorBindings() const PURE_VIRTUAL(IStateTreeEditorPropertyBindingsOwner::GetPropertyEditorBindings, return nullptr; );
 };
 
 // TODO: We should merge this with IStateTreeEditorPropertyBindingsOwner and FStateTreeEditorPropertyBindings.
@@ -134,13 +146,14 @@ class STATETREEEDITORMODULE_API IStateTreeEditorPropertyBindingsOwner
 // and IStateTreeBindingLookup is used in non-editor code and it cannot be in FStateTreeEditorPropertyBindings because bindings don't know about the owner.
 struct STATETREEEDITORMODULE_API FStateTreeBindingLookup : public IStateTreeBindingLookup
 {
-	FStateTreeBindingLookup(IStateTreeEditorPropertyBindingsOwner* InBindingOwner);
+	FStateTreeBindingLookup(const IStateTreeEditorPropertyBindingsOwner* InBindingOwner);
 
-	IStateTreeEditorPropertyBindingsOwner* BindingOwner = nullptr;
+	const IStateTreeEditorPropertyBindingsOwner* BindingOwner = nullptr;
 
 protected:
 	virtual const FStateTreePropertyPath* GetPropertyBindingSource(const FStateTreePropertyPath& InTargetPath) const override;
-	virtual FText GetPropertyPathDisplayName(const FStateTreePropertyPath& InTargetPath) const override;
+	virtual FText GetPropertyPathDisplayName(const FStateTreePropertyPath& InTargetPath, EStateTreeNodeFormatting Formatting) const override;
+	virtual FText GetBindingSourceDisplayName(const FStateTreePropertyPath& InTargetPath, EStateTreeNodeFormatting Formatting) const override;
 	virtual const FProperty* GetPropertyPathLeafProperty(const FStateTreePropertyPath& InPath) const override;
 
 };
diff --git a/Engine/Plugins/Runtime/StateTree/Source/StateTreeModule/Private/StateTreeNodeDescriptionHelpers.cpp b/Engine/Plugins/Runtime/StateTree/Source/StateTreeModule/Private/StateTreeNodeDescriptionHelpers.cpp
new file mode 100644
--- /dev/null
+++ b/Engine/Plugins/Runtime/StateTree/Source/StateTreeModule/Private/StateTreeNodeDescriptionHelpers.cpp
@@ -0,0 +1,145 @@
+﻿// Copyright Epic Games, Inc. All Rights Reserved.
+
+#include "StateTreeNodeDescriptionHelpers.h"
+#include "GameplayTagContainer.h"
+#include "StateTreeNodeBase.h"
+
+#if WITH_EDITOR
+#define LOCTEXT_NAMESPACE "StateTree"
+
+namespace UE::StateTree::DescHelpers
+{
+
+FText GetOperatorText(const EGenericAICheck Operator, EStateTreeNodeFormatting Formatting)
+{
+	if (Formatting == EStateTreeNodeFormatting::RichText)
+	{
+		switch (Operator)
+		{
+		case EGenericAICheck::Equal:
+			return FText::FromString(TEXT("<s>==</>"));
+		case EGenericAICheck::NotEqual:
+			return FText::FromString(TEXT("<s>!=</>"));
+		case EGenericAICheck::Less:
+			return FText::FromString(TEXT("<s>&lt;</>"));
+		case EGenericAICheck::LessOrEqual:
+			return FText::FromString(TEXT("<s>&lt;=</>"));
+		case EGenericAICheck::Greater:
+			return FText::FromString(TEXT("<s>&gt;</>"));
+		case EGenericAICheck::GreaterOrEqual:
+			return FText::FromString(TEXT("<s>&gt;=</>"));
+		default:
+			break;
+		}
+	}
+	else
+	{
+		switch (Operator)
+		{
+		case EGenericAICheck::Equal:
+			return FText::FromString(TEXT("=="));
+		case EGenericAICheck::NotEqual:
+			return FText::FromString(TEXT("!="));
+		case EGenericAICheck::Less:
+			return FText::FromString(TEXT("&lt;"));
+		case EGenericAICheck::LessOrEqual:
+			return FText::FromString(TEXT("&lt;="));
+		case EGenericAICheck::Greater:
+			return FText::FromString(TEXT("&gt;"));
+		case EGenericAICheck::GreaterOrEqual:
+			return FText::FromString(TEXT("&gt;="));
+		default:
+			break;
+		}
+	}
+
+	return FText::FromString(TEXT("??"));
+}
+
+FText GetInvertText(bool bInvert, EStateTreeNodeFormatting Formatting)
+{
+	FText InvertText = FText::GetEmpty();
+	if (bInvert)
+	{
+		if (Formatting == EStateTreeNodeFormatting::RichText)
+		{
+			InvertText = LOCTEXT("InvertRich", "<s>Not</>  ");
+		}
+		else
+		{
+			InvertText = LOCTEXT("Invert", "Not  ");
+		}
+	}
+	return InvertText;
+}
+
+FText GetBoolText(bool bValue, EStateTreeNodeFormatting Formatting)
+{
+	return bValue ? LOCTEXT("True", "True") : LOCTEXT("False", "False");
+}
+
+FText GetGameplayTagContainerAsText(const FGameplayTagContainer& TagContainer, const int ApproxMaxLength)
+{
+	if (TagContainer.IsEmpty())
+	{
+		return LOCTEXT("Empty", "Empty");
+	}
+		
+	FString Combined;
+	for (const FGameplayTag& Tag : TagContainer)
+	{
+		FString TagString = Tag.ToString();
+
+		if (Combined.Len() > 0)
+		{
+			Combined += TEXT(", ");
+		}
+			
+		if ((Combined.Len() + TagString.Len()) > ApproxMaxLength)
+		{
+			// Overflow
+			if (Combined.Len() == 0)
+			{
+				Combined += TagString.Left(ApproxMaxLength);
+			}
+			Combined += TEXT("...");
+			break;
+		}
+
+		Combined += TagString;
+	}
+
+	return FText::FromString(Combined);
+}
+
+FText GetGameplayTagQueryAsText(const FGameplayTagQuery& TagQuery, const int ApproxMaxLength)
+{
+	// Limit the presented query description.
+	FText QueryValue = LOCTEXT("Empty", "Empty");
+	constexpr int32 MaxDescLen = 120; 
+	FString QueryDesc = TagQuery.GetDescription();
+	if (!QueryDesc.IsEmpty())
+	{
+		if (QueryDesc.Len() > MaxDescLen)
+		{
+			QueryDesc = QueryDesc.Left(MaxDescLen);
+			QueryDesc += TEXT("...");
+		}
+		QueryValue = FText::FromString(QueryDesc);
+	}
+	return QueryValue;
+}
+
+FText GetExactMatchText(bool bExactMatch, EStateTreeNodeFormatting Formatting)
+{
+	if (Formatting == EStateTreeNodeFormatting::RichText)
+	{
+		return bExactMatch ? LOCTEXT("ExactlyRich", "<s>exactly</> ") : FText::GetEmpty();
+	}
+	return bExactMatch ? LOCTEXT("Exactly", "exactly ") : FText::GetEmpty();
+}
+
+} // UE::StateTree::Helpers
+
+#undef LOCTEXT_NAMESPACE
+#endif // WITH_EDITOR
\ No newline at end of file
diff --git a/Engine/Plugins/Runtime/StateTree/Source/StateTreeModule/Public/StateTreeNodeBase.h b/Engine/Plugins/Runtime/StateTree/Source/StateTreeModule/Public/StateTreeNodeBase.h
--- a/Engine/Plugins/Runtime/StateTree/Source/StateTreeModule/Public/StateTreeNodeBase.h
+++ b/Engine/Plugins/Runtime/StateTree/Source/StateTreeModule/Public/StateTreeNodeBase.h
@@ -10,6 +10,26 @@ struct FStateTreeEditorPropertyPath;
 struct FStateTreePropertyPath;
 struct IStateTreeBindingLookup;
 
+/**
+ * Enum describing in what format a text is expected to be returned.
+ *
+ * - Normal text should be used for values
+ * - Bold text should generally be used for actions, like name a of a task "<b>Play Animation</> {AnimName}".
+ * - Subdued should be generally used for secondary/structural information, like "{Left} <s>equals</> {Right}".
+ */
+UENUM()
+enum EStateTreeNodeFormatting : uint8
+{
+	/**
+	 * The returned text can contain following right text formatting (no nesting)
+	 *	- <b>Bold</> (bolder font is used)
+	 *	- <s>Subdued</> (normal font with lighter color) */
+	RichText,
+	
+	/** The text should be unformatted */
+	Text,
+};
+
 /**
  * Base struct of StateTree Conditions, Evaluators, and Tasks.
  */
@@ -53,15 +73,32 @@ PRAGMA_ENABLE_DEPRECATION_WARNINGS
 	UE_DEPRECATED(5.3, "Use version with FStateTreePropertyPath instead.")
 	virtual void OnBindingChanged(const FGuid& ID, FStateTreeDataView InstanceData, const FStateTreeEditorPropertyPath& SourcePath, const FStateTreeEditorPropertyPath& TargetPath, const IStateTreeBindingLookup& BindingLookup) final {}
 	PRAGMA_ENABLE_DEPRECATION_WARNINGS
+
+	/**
+	 * Returns description for the node, use in the UI.
+	 * The UI description is selected as follows: 
+	 * - Node Name, if not empty
+	 * - Description if not empty
+	 * - Display name of the node struct
+	 * @param ID ID of the item, can be used make property paths to this item.
+	 * @param InstanceDataView View to the instance data, can be struct or class.
+	 * @param BindingLookup Reference to binding lookup which can be used to reason about property paths.
+	 * @param Formatting Requested formatting (whether rich or plain text should be returned).
+	 */
+	virtual FText GetDescription(const FGuid& ID, FStateTreeDataView InstanceDataView, const IStateTreeBindingLookup& BindingLookup, EStateTreeNodeFormatting Formatting = EStateTreeNodeFormatting::Text) const
+	{
+		return FText::GetEmpty();
+	}
+	
 	/**
 	 * Called when binding of any of the properties in the node changes.
 	 * @param ID ID of the item, can be used make property paths to this item.
-	 * @param InstanceData view to the instance data, can be struct or class.
+	 * @param InstanceDataView view to the instance data, can be struct or class.
 	 * @param SourcePath Source path of the new binding.
 	 * @param TargetPath Target path of the new binding (the property in the condition).
 	 * @param BindingLookup Reference to binding lookup which can be used to reason about property paths.
 	 */
-	virtual void OnBindingChanged(const FGuid& ID, FStateTreeDataView InstanceData, const FStateTreePropertyPath& SourcePath, const FStateTreePropertyPath& TargetPath, const IStateTreeBindingLookup& BindingLookup) {}
+	virtual void OnBindingChanged(const FGuid& ID, FStateTreeDataView InstanceDataView, const FStateTreePropertyPath& SourcePath, const FStateTreePropertyPath& TargetPath, const IStateTreeBindingLookup& BindingLookup) {}
 
 	/**
 	 * Called when a property of the node has been modified externally
@@ -79,7 +116,7 @@ PRAGMA_ENABLE_DEPRECATION_WARNINGS
 
 	/**
 	* Called after the state tree asset that contains this node is loaded from disk.
-	* @param InstanceData view to the instance data, can be struct or class.
+	* @param InstanceDataView view to the instance data, can be struct or class.
 	*/
 	virtual void PostLoad(FStateTreeDataView InstanceDataView) {}
 #endif
diff --git a/Engine/Plugins/Runtime/StateTree/Source/StateTreeModule/Public/StateTreeNodeDescriptionHelpers.h b/Engine/Plugins/Runtime/StateTree/Source/StateTreeModule/Public/StateTreeNodeDescriptionHelpers.h
new file mode 100644
--- /dev/null
+++ b/Engine/Plugins/Runtime/StateTree/Source/StateTreeModule/Public/StateTreeNodeDescriptionHelpers.h
@@ -0,0 +1,35 @@
+﻿// Copyright Epic Games, Inc. All Rights Reserved.
+
+#pragma once
+
+#include "AITypes.h"
+#include "Internationalization/Text.h"
+
+enum EStateTreeNodeFormatting : uint8;
+struct FGameplayTagContainer;
+struct FGameplayTagQuery;
+
+namespace UE::StateTree::DescHelpers
+{
+#if WITH_EDITOR
+
+/** @return description for a EGenericAICheck. */
+extern STATETREEMODULE_API FText GetOperatorText(const EGenericAICheck Operator, EStateTreeNodeFormatting Formatting);
+
+/** @return description for condition inversion (returns "Not" plus a space). */
+extern STATETREEMODULE_API FText GetInvertText(bool bInvert, EStateTreeNodeFormatting Formatting);
+
+/** @return description of a boolean value. */
+extern STATETREEMODULE_API FText GetBoolText(bool bValue, EStateTreeNodeFormatting Formatting);
+
+/** @return description for a Gameplay Tag Container. If the length of container description is longer than ApproxMaxLength, the it truncated and ... as added to the end. */
+extern STATETREEMODULE_API FText GetGameplayTagContainerAsText(const FGameplayTagContainer& TagContainer, const int ApproxMaxLength = 60);
+
+/** @return description for a Gameplay Tag Query. If the query description is longer than ApproxMaxLength, the it truncated and ... as added to the end. */
+extern STATETREEMODULE_API FText GetGameplayTagQueryAsText(const FGameplayTagQuery& TagQuery, const int ApproxMaxLength = 120);
+
+/** @return description for exact match, used for Gameplay Tag matching functions (returns "Exactly" plus space). */
+extern STATETREEMODULE_API FText GetExactMatchText(bool bExactMatch, EStateTreeNodeFormatting Formatting);
+
+#endif // WITH_EDITOR 
+} // UE::StateTree::Helpers
\ No newline at end of file
diff --git a/Engine/Plugins/Runtime/StateTree/Source/StateTreeModule/Public/StateTreePropertyBindings.h b/Engine/Plugins/Runtime/StateTree/Source/StateTreeModule/Public/StateTreePropertyBindings.h
--- a/Engine/Plugins/Runtime/StateTree/Source/StateTreeModule/Public/StateTreePropertyBindings.h
+++ b/Engine/Plugins/Runtime/StateTree/Source/StateTreeModule/Public/StateTreePropertyBindings.h
@@ -2,6 +2,7 @@
 
 #pragma once
 
+#include "StateTreeNodeBase.h"
 #include "StateTreeTypes.h"
 #include "StateTreePropertyRefHelpers.h"
 #include "StateTreePropertyBindings.generated.h"
@@ -11,6 +12,7 @@ struct FStateTreePropertyPath;
 struct FStateTreePropertyBindingCompiler;
 struct FStateTreePropertyRef;
 class UStateTree;
+enum EStateTreeNodeFormatting : uint8;
 
 UENUM()
 enum class EStateTreeBindableStructSource : uint8
@@ -1017,11 +1019,16 @@ PRAGMA_ENABLE_DEPRECATION_WARNINGS
  */
 struct STATETREEMODULE_API IStateTreeBindingLookup
 {
+	virtual ~IStateTreeBindingLookup() {} 
+	
 	/** @return Source path for given target path, or null if binding does not exists. */
 	virtual const FStateTreePropertyPath* GetPropertyBindingSource(const FStateTreePropertyPath& InTargetPath) const PURE_VIRTUAL(IStateTreeBindingLookup::GetPropertyBindingSource, return nullptr; );
 
 	/** @return Display name given property path. */
-	virtual FText GetPropertyPathDisplayName(const FStateTreePropertyPath& InPath) const PURE_VIRTUAL(IStateTreeBindingLookup::GetPropertyPathDisplayName, return FText::GetEmpty(); );
+	virtual FText GetPropertyPathDisplayName(const FStateTreePropertyPath& InPath, EStateTreeNodeFormatting Formatting = EStateTreeNodeFormatting::Text) const PURE_VIRTUAL(IStateTreeBindingLookup::GetPropertyPathDisplayName, return FText::GetEmpty(); );
+
+	/** @return Display name of binding source, or empty if binding does not exists. */
+	virtual FText GetBindingSourceDisplayName(const FStateTreePropertyPath& InTargetPath, EStateTreeNodeFormatting Formatting = EStateTreeNodeFormatting::Text) const PURE_VIRTUAL(IStateTreeBindingLookup::GetPropertyPathDisplayName, return FText::GetEmpty(); );
 
 	/** @return Leaf property based on property path. */
 	virtual const FProperty* GetPropertyPathLeafProperty(const FStateTreePropertyPath& InPath) const PURE_VIRTUAL(IStateTreeBindingLookup::GetPropertyPathLeafProperty, return nullptr; );
