StateTreeDebugger: Fix "stuck timeline"
- Make sure SStateTreeDebuggerView::MaxTrackRecordingDuration is always kept in sync so that timeline clamping works
- Tick() is called only on visible debugger windows, in case of multiple debuggers some timeline might not get corrent legnth for time scrubbing clamping

diff --git a/Engine/Plugins/Runtime/StateTree/Source/StateTreeEditorModule/Private/Debugger/SStateTreeDebuggerView.cpp b/Engine/Plugins/Runtime/StateTree/Source/StateTreeEditorModule/Private/Debugger/SStateTreeDebuggerView.cpp
--- a/Engine/Plugins/Runtime/StateTree/Source/StateTreeEditorModule/Private/Debugger/SStateTreeDebuggerView.cpp
+++ b/Engine/Plugins/Runtime/StateTree/Source/StateTreeEditorModule/Private/Debugger/SStateTreeDebuggerView.cpp
@@ -683,12 +683,16 @@ void SStateTreeDebuggerView::Tick(const FGeometry& AllottedGeometry, const doubl
 	QUICK_SCOPE_CYCLE_COUNTER(STAT_StateTreeDebuggerView_TickView);
 
 	check(Debugger);
-	if (Debugger->IsAnalysisSessionActive() && !Debugger->IsAnalysisSessionPaused())
+	if (Debugger->IsAnalysisSessionActive())
 	{
 		MaxTrackRecordingDuration = FMath::Max(MaxTrackRecordingDuration, Debugger->GetRecordingDuration());
 
-		// Stick to most recent data if required
-		if (bAutoScroll)
+		// Stick to most recent data if auto scroll is enabled.
+		// Autoscroll is disabled when paused.
+		// This allows the user to pause the analysis, inspect the data, and continue and the autoscroll will catch up with latest.
+		// Complementary logic in OnTimeLineScrubPositionChanged().
+		if (!Debugger->IsAnalysisSessionPaused()
+			&& bAutoScroll)
 		{
 			Debugger->SetScrubTime(Debugger->GetRecordingDuration());
 		}
@@ -1018,6 +1022,10 @@ UStateTreeState* SStateTreeDebuggerView::FindStateAssociatedToBreakpoint(FStateT
 void SStateTreeDebuggerView::OnTimeLineScrubPositionChanged(double Time, bool bIsScrubbing)
 {
 	check(Debugger);
+	// Disable auto scroll when scrubbing.
+	// But, do not disable it if the analysis is in progress but paused.
+	// This allows the user to pause the analysis, inspect the data, and continue and the autoscroll will catch up with latest.
+	// Complementary logic in Tick().
 	if (Debugger->IsAnalysisSessionActive() && !Debugger->IsAnalysisSessionPaused())
 	{
 		bAutoScroll = false;
