StateTree: Improved transition visualization and fixed some logic to better match the execution.

diff --git a/Engine/Plugins/Runtime/StateTree/Source/StateTreeEditorModule/Private/Customizations/StateTreeStateLinkDetails.cpp b/Engine/Plugins/Runtime/StateTree/Source/StateTreeEditorModule/Private/Customizations/StateTreeStateLinkDetails.cpp
--- a/Engine/Plugins/Runtime/StateTree/Source/StateTreeEditorModule/Private/Customizations/StateTreeStateLinkDetails.cpp
+++ b/Engine/Plugins/Runtime/StateTree/Source/StateTreeEditorModule/Private/Customizations/StateTreeStateLinkDetails.cpp
@@ -12,7 +12,9 @@
 #include "ScopedTransaction.h"
 #include "StateTreeEditorStyle.h"
 #include "StateTreePropertyHelpers.h"
+#include "TextStyleDecorator.h"
 #include "Widgets/Input/SSearchBox.h"
+#include "Widgets/Text/SRichTextBlock.h"
 
 #define LOCTEXT_NAMESPACE "StateTreeEditor"
 
@@ -88,9 +90,14 @@ void FStateTreeStateLinkDetails::CustomizeHeader(TSharedRef<class IPropertyHandl
 				.AutoWidth()
 				.VAlign(VAlign_Center)
 				[
-					SNew(STextBlock)
+					SNew(SRichTextBlock)
 					.Text(this, &FStateTreeStateLinkDetails::GetCurrentStateDesc)
-					.Font(this, &FStateTreeStateLinkDetails::GetCurrentStateFont)
+					.TextStyle(&FStateTreeEditorStyle::Get().GetWidgetStyle<FTextBlockStyle>("Details.Normal"))
+					.OverflowPolicy(ETextOverflowPolicy::Ellipsis)
+					+SRichTextBlock::Decorator(FTextStyleDecorator::Create(TEXT(""), FStateTreeEditorStyle::Get().GetWidgetStyle<FTextBlockStyle>("Details.Normal")))
+					+SRichTextBlock::Decorator(FTextStyleDecorator::Create(TEXT("b"), FStateTreeEditorStyle::Get().GetWidgetStyle<FTextBlockStyle>("Details.Bold")))
+					+SRichTextBlock::Decorator(FTextStyleDecorator::Create(TEXT("i"), FStateTreeEditorStyle::Get().GetWidgetStyle<FTextBlockStyle>("Details.Italic")))
+					+SRichTextBlock::Decorator(FTextStyleDecorator::Create(TEXT("s"), FStateTreeEditorStyle::Get().GetWidgetStyle<FTextBlockStyle>("Details.Subdued")))
 				]
 			]
 		];
@@ -191,31 +198,31 @@ void FStateTreeStateLinkDetails::CacheStates()
 	if (!bDirectStatesOnly)
 	{
 		RootItem->Children.Add(MakeShared<FStateTreeStateItem>(
-			LOCTEXT("TransitionNone", "None"),
+			LOCTEXT("TransitionNone", "<i>None</>"),
 			LOCTEXT("TransitionNoneTooltip", "No transition."),
 			FStateTreeEditorStyle::Get().GetBrush("StateTreeEditor.Transition.None"),
 			EStateTreeTransitionType::None));
 
 		RootItem->Children.Add(MakeShared<FStateTreeStateItem>(
-			LOCTEXT("TransitionNextState", "Next State"),
+			LOCTEXT("TransitionNextState", "<i>Next State</>"),
 			LOCTEXT("TransitionNextTooltip", "Goto next sibling State."),
 			FStateTreeEditorStyle::Get().GetBrush("StateTreeEditor.Transition.Next"),
 			EStateTreeTransitionType::NextState));
 
 		RootItem->Children.Add(MakeShared<FStateTreeStateItem>(
-			LOCTEXT("TransitionNextSelectableState", "Next Selectable State"),
+			LOCTEXT("TransitionNextSelectableState", "<i>Next Selectable State</>"),
 			LOCTEXT("TransitionNextSelectableTooltip", "Goto next sibling state, whose enter conditions pass."),
 			FStateTreeEditorStyle::Get().GetBrush("StateTreeEditor.Transition.Next"),
 			EStateTreeTransitionType::NextSelectableState));
 
 		RootItem->Children.Add(MakeShared<FStateTreeStateItem>(
-			LOCTEXT("TransitionTreeSucceeded", "Tree Succeeded"),
+			LOCTEXT("TransitionTreeSucceeded", "<i>Tree Succeeded</>"),
 			LOCTEXT("TransitionTreeSuccessTooltip", "Complete tree with success."),
 			FStateTreeEditorStyle::Get().GetBrush("StateTreeEditor.Transition.Succeeded"),
 			EStateTreeTransitionType::Succeeded));
 
 		RootItem->Children.Add(MakeShared<FStateTreeStateItem>(
-			LOCTEXT("TransitionTreeFailed", "Tree Failed"),
+			LOCTEXT("TransitionTreeFailed", "<i>Tree Failed</>"),
 			LOCTEXT("TransitionTreeFailedTooltip", "Complete tree with failure."),
 			FStateTreeEditorStyle::Get().GetBrush("StateTreeEditor.Transition.Failed"),
 			EStateTreeTransitionType::Failed));
@@ -301,24 +308,8 @@ TSharedRef<SWidget> FStateTreeStateLinkDetails::GenerateStatePicker()
 	return MenuWidget;
 }
 
-FSlateFontInfo FStateTreeStateLinkDetails::GetItemFont(TSharedPtr<FStateTreeStateItem> Item)
-{
-	FSlateFontInfo Font = FAppStyle::Get().GetFontStyle("PropertyWindow.NormalFont");
-	if (Item->TransitionType != EStateTreeTransitionType::GotoState)
-	{
-		Font = FAppStyle::Get().GetFontStyle("PropertyWindow.ItalicFont");
-	}
-	else if (Item->bIsSubTree)
-	{
-		Font = FAppStyle::Get().GetFontStyle("PropertyWindow.BoldFont");
-	}
-	return Font;	
-}
-
 TSharedRef<ITableRow> FStateTreeStateLinkDetails::GenerateStateItemRow(TSharedPtr<FStateTreeStateItem> Item, const TSharedRef<STableViewBase>& OwnerTable)
 {
-	const FSlateFontInfo Font = GetItemFont(Item);
-
 	TSharedRef<SHorizontalBox> Container = SNew(SHorizontalBox);
 
 	// Icon
@@ -341,11 +332,16 @@ TSharedRef<ITableRow> FStateTreeStateLinkDetails::GenerateStateItemRow(TSharedPt
 		.VAlign(VAlign_Center)
 		.AutoWidth()
 		[
-			SNew(STextBlock)
-			.Font(Font)
+			SNew(SRichTextBlock)
 			.Text(Item->Desc)
+			.TextStyle(&FStateTreeEditorStyle::Get().GetWidgetStyle<FTextBlockStyle>("Normal.Normal"))
+			.OverflowPolicy(ETextOverflowPolicy::Ellipsis)
 			.ToolTipText(Item->TooltipText)
 			.HighlightText_Lambda([this]() { return SearchBox.IsValid() ? SearchBox->GetText() : FText::GetEmpty(); })
+			+SRichTextBlock::Decorator(FTextStyleDecorator::Create(TEXT(""), FStateTreeEditorStyle::Get().GetWidgetStyle<FTextBlockStyle>("Normal.Normal")))
+			+SRichTextBlock::Decorator(FTextStyleDecorator::Create(TEXT("b"), FStateTreeEditorStyle::Get().GetWidgetStyle<FTextBlockStyle>("Normal.Bold")))
+			+SRichTextBlock::Decorator(FTextStyleDecorator::Create(TEXT("i"), FStateTreeEditorStyle::Get().GetWidgetStyle<FTextBlockStyle>("Normal.Italic")))
+			+SRichTextBlock::Decorator(FTextStyleDecorator::Create(TEXT("s"), FStateTreeEditorStyle::Get().GetWidgetStyle<FTextBlockStyle>("Normal.Subdued")))
 		];
 
 	// Link
@@ -651,16 +647,6 @@ FSlateColor FStateTreeStateLinkDetails::GetCurrentStateColor() const
 	return {};
 }
 
-FSlateFontInfo FStateTreeStateLinkDetails::GetCurrentStateFont() const
-{
-	TArray<TSharedPtr<FStateTreeStateItem>> Path;
-	if (GetCurrentStateItem(Path))
-	{
-		return GetItemFont(Path.Last());
-	}
-	return FAppStyle::Get().GetFontStyle("NormalText");
-}
-
 bool FStateTreeStateLinkDetails::IsValidLink() const
 {
 	TArray<TSharedPtr<FStateTreeStateItem>> Path;
diff --git a/Engine/Plugins/Runtime/StateTree/Source/StateTreeEditorModule/Private/Customizations/StateTreeStateLinkDetails.h b/Engine/Plugins/Runtime/StateTree/Source/StateTreeEditorModule/Private/Customizations/StateTreeStateLinkDetails.h
--- a/Engine/Plugins/Runtime/StateTree/Source/StateTreeEditorModule/Private/Customizations/StateTreeStateLinkDetails.h
+++ b/Engine/Plugins/Runtime/StateTree/Source/StateTreeEditorModule/Private/Customizations/StateTreeStateLinkDetails.h
@@ -71,10 +71,8 @@ private:
 	FText GetCurrentStateDesc() const;
 	const FSlateBrush* GetCurrentStateIcon() const;
 	FSlateColor GetCurrentStateColor() const;
-	FSlateFontInfo GetCurrentStateFont() const;
 	bool IsValidLink() const;
 	TOptional<EStateTreeTransitionType> GetTransitionType() const;
-	static FSlateFontInfo GetItemFont(TSharedPtr<FStateTreeStateItem> Item);
 	TSharedRef<ITableRow> GenerateStateItemRow(TSharedPtr<FStateTreeStateItem> Item, const TSharedRef<STableViewBase>& OwnerTable);
 	void GetStateItemChildren(TSharedPtr<FStateTreeStateItem> Item, TArray<TSharedPtr<FStateTreeStateItem>>& OutItems) const;
 	void OnStateItemSelected(TSharedPtr<FStateTreeStateItem> SelectedItem, ESelectInfo::Type);
diff --git a/Engine/Plugins/Runtime/StateTree/Source/StateTreeEditorModule/Private/SStateTreeViewRow.cpp b/Engine/Plugins/Runtime/StateTree/Source/StateTreeEditorModule/Private/SStateTreeViewRow.cpp
--- a/Engine/Plugins/Runtime/StateTree/Source/StateTreeEditorModule/Private/SStateTreeViewRow.cpp
+++ b/Engine/Plugins/Runtime/StateTree/Source/StateTreeEditorModule/Private/SStateTreeViewRow.cpp
@@ -41,7 +41,7 @@ void SStateTreeViewRow::Construct(const FArguments& InArgs, const TSharedRef<STa
 	WeakEditorData = State != nullptr ? State->GetTypedOuter<UStateTreeEditorData>() : nullptr;
 
 	ConstructInternal(STableRow::FArguments()
-		.Padding(5.0f)
+		.Padding(5.f)
 		.OnDragDetected(this, &SStateTreeViewRow::HandleDragDetected)
 		.OnCanAcceptDrop(this, &SStateTreeViewRow::HandleCanAcceptDrop)
 		.OnAcceptDrop(this, &SStateTreeViewRow::HandleAcceptDrop)
@@ -75,18 +75,18 @@ void SStateTreeViewRow::Construct(const FArguments& InArgs, const TSharedRef<STa
 			[
 				SNew(SExpanderArrow, SharedThis(this))
 				.ShouldDrawWires(true)
-				.IndentAmount(32)
+				.IndentAmount(32.f)
 				.BaseIndentLevel(0)
 			]
 
 			+ SHorizontalBox::Slot()
 			.VAlign(VAlign_Fill)
 			.HAlign(HAlign_Left)
-			.Padding(FMargin(0.0f, 4.0f))
+			.Padding(FMargin(0.f, 4.f))
 			.AutoWidth()
 			[
 				SNew(SBox)
-				.HeightOverride(28.0f)
+				.HeightOverride(28.f)
 				.VAlign(VAlign_Fill)
 				[
 					SNew(SBorder)
@@ -101,8 +101,8 @@ void SStateTreeViewRow::Construct(const FArguments& InArgs, const TSharedRef<STa
 						.AutoWidth()
 						[
 							SNew(SBox)
-							.WidthOverride(4.0f)
-							.HeightOverride(28.0f)
+							.WidthOverride(4.f)
+							.HeightOverride(28.f)
 							[
 								SNew(SBorder)
 								.BorderImage(FAppStyle::GetBrush("WhiteBrush"))
@@ -116,13 +116,13 @@ void SStateTreeViewRow::Construct(const FArguments& InArgs, const TSharedRef<STa
 						.AutoWidth()
 						[
 							SNew(SBox)
-							.HeightOverride(28.0f)
+							.HeightOverride(28.f)
 							.VAlign(VAlign_Fill)
 							[
 								SNew(SBorder)
 								.BorderImage(FAppStyle::GetBrush("WhiteBrush"))
 								.BorderBackgroundColor(this, &SStateTreeViewRow::GetTitleColor)
-								.Padding(FMargin(4.0f, 0.0f, 12.0f, 0.0f))
+								.Padding(FMargin(4.f, 0.f, 12.f, 0.f))
 								.IsEnabled_Lambda([InState]
 									{
 										const UStateTreeState* State = InState.Get();
@@ -139,7 +139,7 @@ void SStateTreeViewRow::Construct(const FArguments& InArgs, const TSharedRef<STa
 										.AutoWidth()
 										[
 											SNew(SBox)
-											.Padding(FMargin(0.0f, 0.0f, 4.0f, 0.0f))
+											.Padding(FMargin(0.f, 0.f, 4.f, 0.f))
 											.Visibility(this, &SStateTreeViewRow::GetConditionVisibility)
 											[
 												SNew(SImage)
@@ -155,7 +155,7 @@ void SStateTreeViewRow::Construct(const FArguments& InArgs, const TSharedRef<STa
 										.AutoWidth()
 										[
 											SNew(SBox)
-											.Padding(FMargin(0.0f, 0.0f, 4.0f, 0.0f))
+											.Padding(FMargin(0.f, 0.f, 4.f, 0.f))
 											[
 												SNew(SImage)
 												.Image(this, &SStateTreeViewRow::GetSelectorIcon)
@@ -207,7 +207,7 @@ void SStateTreeViewRow::Construct(const FArguments& InArgs, const TSharedRef<STa
 										.AutoWidth()
 										[
 											SNew(SBox)
-											.Padding(FMargin(-12.0f, -6.0f, 0.0f, 0.0f))
+											.Padding(FMargin(-12.f, -6.f, 0.f, 0.f))
 											[
 												SNew(SImage)
 												.DesiredSizeOverride(FVector2D(12.f, 12.f))
@@ -227,21 +227,21 @@ void SStateTreeViewRow::Construct(const FArguments& InArgs, const TSharedRef<STa
 						.AutoWidth()
 						[
 							SNew(SBox)
-							.HeightOverride(28.0f)
+							.HeightOverride(28.f)
 							.VAlign(VAlign_Fill)
 							.Visibility(this, &SStateTreeViewRow::GetLinkedStateVisibility)
 							[
 								SNew(SBorder)
 								.BorderImage(FAppStyle::GetBrush("WhiteBrush"))
 								.BorderBackgroundColor(LinkBackground)
-								.Padding(FMargin(6.0f, 0.0f, 12.0f, 0.0f))
+								.Padding(FMargin(6.f, 0.f, 12.f, 0.f))
 								[
 									// Link icon
 									SNew(SHorizontalBox)
 									+ SHorizontalBox::Slot()
 									.VAlign(VAlign_Center)
 									.AutoWidth()
-									.Padding(FMargin(0.0f, 0.0f, 4.0f, 0.0f))
+									.Padding(FMargin(0.f, 0.f, 4.f, 0.f))
 									[
 										SNew(SImage)
 										.ColorAndOpacity(IconTint)
@@ -281,54 +281,61 @@ void SStateTreeViewRow::Construct(const FArguments& InArgs, const TSharedRef<STa
 			.VAlign(VAlign_Center)
 			.AutoWidth()
 			[
-				SNew(SHorizontalBox)
-				+ SHorizontalBox::Slot()
-				.VAlign(VAlign_Center)
-				.AutoWidth()
-				.Padding(FMargin(8.0f, 0.0f, 0, 0.0f))
+				SNew(SBox)
+				.Visibility(this, &SStateTreeViewRow::GetCompletedTransitionVisibility)
 				[
-					SNew(SOverlay)
-					+ SOverlay::Slot()
-					[
-						SNew(STextBlock)
-						.Text(this, &SStateTreeViewRow::GetCompletedTransitionsIcon)
-						.TextStyle(FStateTreeEditorStyle::Get(), "StateTree.Icon")
-						.Visibility(this, &SStateTreeViewRow::GetCompletedTransitionVisibility)
-					]
-					+ SOverlay::Slot()
+					SNew(SHorizontalBox)
+					+ SHorizontalBox::Slot()
+					.VAlign(VAlign_Center)
+					.AutoWidth()
+					.Padding(FMargin(8.f, 0.f, 0.f, 0.f))
 					[
-						// Breakpoint box
-						SNew(SHorizontalBox)
-						+ SHorizontalBox::Slot()
-						.VAlign(VAlign_Top)
-						.HAlign(HAlign_Left)
-						.AutoWidth()
+						SNew(SOverlay)
+						+ SOverlay::Slot()
 						[
-							SNew(SBox)
-							.Padding(FMargin(0.0f, -10.0f, 0.0f, 0.0f))
+							SNew(SImage)
+							.Image(this, &SStateTreeViewRow::GetCompletedTransitionsIcon)
+							.ColorAndOpacity(IconTint)
+						]
+						+ SOverlay::Slot()
+						[
+							// Breakpoint box
+							SNew(SHorizontalBox)
+							+ SHorizontalBox::Slot()
+							.VAlign(VAlign_Top)
+							.HAlign(HAlign_Left)
+							.AutoWidth()
 							[
-								SNew(SImage)
-								.DesiredSizeOverride(FVector2D(10.f, 10.f))
-								.Image(FStateTreeEditorStyle::Get().GetBrush(TEXT("StateTreeEditor.Debugger.Breakpoint.EnabledAndValid")))
-								.Visibility(this, &SStateTreeViewRow::GetCompletedTransitionBreakpointVisibility)
-								.ToolTipText_Lambda([this]
-									{
-										return FText::Format(LOCTEXT("TransitionBreakpointTooltip","Break when executing transition: {0}"),
-											GetCompletedTransitionWithBreakpointDesc());
-									})
+								SNew(SBox)
+								.Padding(FMargin(0.f, -10.f, 0.f, 0.f))
+								[
+									SNew(SImage)
+									.DesiredSizeOverride(FVector2D(10.f, 10.f))
+									.Image(FStateTreeEditorStyle::Get().GetBrush(TEXT("StateTreeEditor.Debugger.Breakpoint.EnabledAndValid")))
+									.Visibility(this, &SStateTreeViewRow::GetCompletedTransitionBreakpointVisibility)
+									.ToolTipText_Lambda([this]
+										{
+											return FText::Format(LOCTEXT("TransitionBreakpointTooltip","Break when executing transition: {0}"),
+												GetCompletedTransitionWithBreakpointDesc());
+										})
+								]
 							]
 						]
 					]
-				]
-				+ SHorizontalBox::Slot()
-				.VAlign(VAlign_Center)
-				.AutoWidth()
-				.Padding(FMargin(4.0f, 0, 0, 0))
-				[
-					SNew(STextBlock)
-					.Text(this, &SStateTreeViewRow::GetCompletedTransitionsDesc)
-					.TextStyle(FStateTreeEditorStyle::Get(), "StateTree.Details")
-					.Visibility(this, &SStateTreeViewRow::GetCompletedTransitionVisibility)
+					+ SHorizontalBox::Slot()
+					.VAlign(VAlign_Center)
+					.AutoWidth()
+					.Padding(FMargin(4.f, 0.f, 0.f, 0.f))
+					[
+						SNew(SRichTextBlock)
+						.Text(this, &SStateTreeViewRow::GetCompletedTransitionsDesc)
+						.TextStyle(&FStateTreeEditorStyle::Get().GetWidgetStyle<FTextBlockStyle>("Normal.Normal"))
+						.OverflowPolicy(ETextOverflowPolicy::Ellipsis)
+						+SRichTextBlock::Decorator(FTextStyleDecorator::Create(TEXT(""), FStateTreeEditorStyle::Get().GetWidgetStyle<FTextBlockStyle>("Normal.Normal")))
+						+SRichTextBlock::Decorator(FTextStyleDecorator::Create(TEXT("b"), FStateTreeEditorStyle::Get().GetWidgetStyle<FTextBlockStyle>("Normal.Bold")))
+						+SRichTextBlock::Decorator(FTextStyleDecorator::Create(TEXT("i"), FStateTreeEditorStyle::Get().GetWidgetStyle<FTextBlockStyle>("Normal.Italic")))
+						+SRichTextBlock::Decorator(FTextStyleDecorator::Create(TEXT("s"), FStateTreeEditorStyle::Get().GetWidgetStyle<FTextBlockStyle>("Normal.Subdued")))
+					]
 				]
 			]
 
@@ -337,65 +344,68 @@ void SStateTreeViewRow::Construct(const FArguments& InArgs, const TSharedRef<STa
 			.VAlign(VAlign_Center)
 			.AutoWidth()
 			[
-				SNew(SHorizontalBox)
-				+ SHorizontalBox::Slot()
-				.VAlign(VAlign_Center)
-				.Padding(FMargin(8.0f, 0.0f, 0, 0))
-				.AutoWidth()
-				[
-					SNew(STextBlock)
-					.Text(FEditorFontGlyphs::Check_Circle)
-					.ColorAndOpacity(FLinearColor(FColor(110,143,67)))
-					.TextStyle(FStateTreeEditorStyle::Get(), "StateTree.Icon")
-					.Visibility(this, &SStateTreeViewRow::GetSucceededTransitionVisibility)
-				]
-				+ SHorizontalBox::Slot()
-				.VAlign(VAlign_Center)
-				.Padding(FMargin(4.0f, 0.0f, 0, 0))
-				.AutoWidth()
+				SNew(SBox)
+				.Visibility(this, &SStateTreeViewRow::GetSucceededTransitionVisibility)
 				[
-					SNew(SOverlay)
-					+ SOverlay::Slot()
+					SNew(SHorizontalBox)
+					+ SHorizontalBox::Slot()
+					.VAlign(VAlign_Center)
+					.Padding(FMargin(8.f, 0.f, 0.f, 0.f))
+					.AutoWidth()
 					[
-						SNew(STextBlock)
-						.Text(this, &SStateTreeViewRow::GetSucceededTransitionIcon)
-						.TextStyle(FStateTreeEditorStyle::Get(), "StateTree.Icon")
-						.Visibility(this, &SStateTreeViewRow::GetSucceededTransitionVisibility)
+						SNew(SImage)
+						.Image(FStateTreeEditorStyle::Get().GetBrush("StateTreeEditor.Transition.Succeeded"))
 					]
-					+ SOverlay::Slot()
+					+ SHorizontalBox::Slot()
+					.VAlign(VAlign_Center)
+					.AutoWidth()
 					[
-						// Breakpoint box
-						SNew(SHorizontalBox)
-						+ SHorizontalBox::Slot()
-						.VAlign(VAlign_Top)
-						.HAlign(HAlign_Left)
-						.AutoWidth()
+						SNew(SOverlay)
+						+ SOverlay::Slot()
 						[
-							SNew(SBox)
-							.Padding(FMargin(0.0f, -10.0f, 0.0f, 0.0f))
+							SNew(SImage)
+							.Image(this, &SStateTreeViewRow::GetSucceededTransitionIcon)
+							.ColorAndOpacity(IconTint)
+						]
+						+ SOverlay::Slot()
+						[
+							// Breakpoint box
+							SNew(SHorizontalBox)
+							+ SHorizontalBox::Slot()
+							.VAlign(VAlign_Top)
+							.HAlign(HAlign_Left)
+							.AutoWidth()
 							[
-								SNew(SImage)
-								.DesiredSizeOverride(FVector2D(10.f, 10.f))
-								.Image(FStateTreeEditorStyle::Get().GetBrush(TEXT("StateTreeEditor.Debugger.Breakpoint.EnabledAndValid")))
-								.Visibility(this, &SStateTreeViewRow::GetSucceededTransitionBreakpointVisibility)
-								.ToolTipText_Lambda([this]
-									{
-										return FText::Format(LOCTEXT("TransitionBreakpointTooltip", "Break when executing transition: {0}"),
-											GetSucceededTransitionWithBreakpointDesc());
-									})
+								SNew(SBox)
+								.Padding(FMargin(0.f, -10.f, 0.f, 0.f))
+								[
+									SNew(SImage)
+									.DesiredSizeOverride(FVector2D(10.f, 10.f))
+									.Image(FStateTreeEditorStyle::Get().GetBrush(TEXT("StateTreeEditor.Debugger.Breakpoint.EnabledAndValid")))
+									.Visibility(this, &SStateTreeViewRow::GetSucceededTransitionBreakpointVisibility)
+									.ToolTipText_Lambda([this]
+										{
+											return FText::Format(LOCTEXT("TransitionBreakpointTooltip", "Break when executing transition: {0}"),
+												GetSucceededTransitionWithBreakpointDesc());
+										})
+								]
 							]
 						]
 					]
-				]
-				+ SHorizontalBox::Slot()
-				.VAlign(VAlign_Center)
-				.AutoWidth()
-				.Padding(FMargin(4.0f, 0, 0, 0))
-				[
-					SNew(STextBlock)
-					.Text(this, &SStateTreeViewRow::GetSucceededTransitionDesc)
-					.TextStyle(FStateTreeEditorStyle::Get(), "StateTree.Details")
-					.Visibility(this, &SStateTreeViewRow::GetSucceededTransitionVisibility)
+					+ SHorizontalBox::Slot()
+					.VAlign(VAlign_Center)
+					.AutoWidth()
+					.Padding(FMargin(4.f, 0.f, 0.f, 0.f))
+					[
+						SNew(SRichTextBlock)
+						.Text(this, &SStateTreeViewRow::GetSucceededTransitionDesc)
+						.TextStyle(&FStateTreeEditorStyle::Get().GetWidgetStyle<FTextBlockStyle>("Normal.Normal"))
+						.OverflowPolicy(ETextOverflowPolicy::Ellipsis)
+						+SRichTextBlock::Decorator(FTextStyleDecorator::Create(TEXT(""), FStateTreeEditorStyle::Get().GetWidgetStyle<FTextBlockStyle>("Normal.Normal")))
+						+SRichTextBlock::Decorator(FTextStyleDecorator::Create(TEXT("b"), FStateTreeEditorStyle::Get().GetWidgetStyle<FTextBlockStyle>("Normal.Bold")))
+						+SRichTextBlock::Decorator(FTextStyleDecorator::Create(TEXT("i"), FStateTreeEditorStyle::Get().GetWidgetStyle<FTextBlockStyle>("Normal.Italic")))
+						+SRichTextBlock::Decorator(FTextStyleDecorator::Create(TEXT("s"), FStateTreeEditorStyle::Get().GetWidgetStyle<FTextBlockStyle>("Normal.Subdued")))
+					]
 				]
 			]
 
@@ -404,65 +414,68 @@ void SStateTreeViewRow::Construct(const FArguments& InArgs, const TSharedRef<STa
 			.VAlign(VAlign_Center)
 			.AutoWidth()
 			[
-				SNew(SHorizontalBox)
-				+ SHorizontalBox::Slot()
-				.VAlign(VAlign_Center)
-				.Padding(FMargin(8.0f, 0.0f, 0, 0))
-				.AutoWidth()
-				[
-					SNew(STextBlock)
-					.Text(FEditorFontGlyphs::Times_Circle)
-					.ColorAndOpacity(FLinearColor(FColor(187,77,42)))
-					.TextStyle(FStateTreeEditorStyle::Get(), "StateTree.Icon")
-					.Visibility(this, &SStateTreeViewRow::GetFailedTransitionVisibility)
-				]
-				+ SHorizontalBox::Slot()
-				.VAlign(VAlign_Center)
-				.Padding(FMargin(4.0f, 0.0f, 0, 0))
-				.AutoWidth()
+				SNew(SBox)
+				.Visibility(this, &SStateTreeViewRow::GetFailedTransitionVisibility)
 				[
-					SNew(SOverlay)
-					+ SOverlay::Slot()
+					SNew(SHorizontalBox)
+					+ SHorizontalBox::Slot()
+					.VAlign(VAlign_Center)
+					.Padding(FMargin(8.f, 0.f, 0.f, 0.f))
+					.AutoWidth()
 					[
-						SNew(STextBlock)
-						.Text(this, &SStateTreeViewRow::GetFailedTransitionIcon)
-						.TextStyle(FStateTreeEditorStyle::Get(), "StateTree.Icon")
-						.Visibility(this, &SStateTreeViewRow::GetFailedTransitionVisibility)
+						SNew(SImage)
+						.Image(FStateTreeEditorStyle::Get().GetBrush("StateTreeEditor.Transition.Failed"))
 					]
-					+ SOverlay::Slot()
+					+ SHorizontalBox::Slot()
+					.VAlign(VAlign_Center)
+					.AutoWidth()
 					[
-						// Breakpoint box
-						SNew(SHorizontalBox)
-						+ SHorizontalBox::Slot()
-						.VAlign(VAlign_Top)
-						.HAlign(HAlign_Left)
-						.AutoWidth()
+						SNew(SOverlay)
+						+ SOverlay::Slot()
 						[
-							SNew(SBox)
-							.Padding(FMargin(0.0f, -10.0f, 0.0f, 0.0f))
+							SNew(SImage)
+							.Image(this, &SStateTreeViewRow::GetFailedTransitionIcon)
+							.ColorAndOpacity(IconTint)
+						]
+						+ SOverlay::Slot()
+						[
+							// Breakpoint box
+							SNew(SHorizontalBox)
+							+ SHorizontalBox::Slot()
+							.VAlign(VAlign_Top)
+							.HAlign(HAlign_Left)
+							.AutoWidth()
 							[
-								SNew(SImage)
-								.DesiredSizeOverride(FVector2D(10.f, 10.f))
-								.Image(FStateTreeEditorStyle::Get().GetBrush(TEXT("StateTreeEditor.Debugger.Breakpoint.EnabledAndValid")))
-								.Visibility(this, &SStateTreeViewRow::GetFailedTransitionBreakpointVisibility)
-								.ToolTipText_Lambda([this]
-									{
-										return FText::Format(LOCTEXT("TransitionBreakpointTooltip", "Break when executing transition: {0}"),
-											GetFailedTransitionWithBreakpointDesc());
-									})
+								SNew(SBox)
+								.Padding(FMargin(0.f, -10.f, 0.f, 0.f))
+								[
+									SNew(SImage)
+									.DesiredSizeOverride(FVector2D(10.f, 10.f))
+									.Image(FStateTreeEditorStyle::Get().GetBrush(TEXT("StateTreeEditor.Debugger.Breakpoint.EnabledAndValid")))
+									.Visibility(this, &SStateTreeViewRow::GetFailedTransitionBreakpointVisibility)
+									.ToolTipText_Lambda([this]
+										{
+											return FText::Format(LOCTEXT("TransitionBreakpointTooltip", "Break when executing transition: {0}"),
+												GetFailedTransitionWithBreakpointDesc());
+										})
+								]
 							]
 						]
 					]
-				]
-				+ SHorizontalBox::Slot()
-				.VAlign(VAlign_Center)
-				.AutoWidth()
-				.Padding(FMargin(4.0f, 0, 0, 0))
-				[
-					SNew(STextBlock)
-					.Text(this, &SStateTreeViewRow::GetFailedTransitionDesc)
-					.TextStyle(FStateTreeEditorStyle::Get(), "StateTree.Details")
-					.Visibility(this, &SStateTreeViewRow::GetFailedTransitionVisibility)
+					+ SHorizontalBox::Slot()
+					.VAlign(VAlign_Center)
+					.AutoWidth()
+					.Padding(FMargin(4.f, 0.f, 0.f, 0.f))
+					[
+						SNew(SRichTextBlock)
+						.Text(this, &SStateTreeViewRow::GetFailedTransitionDesc)
+						.TextStyle(&FStateTreeEditorStyle::Get().GetWidgetStyle<FTextBlockStyle>("Normal.Normal"))
+						.OverflowPolicy(ETextOverflowPolicy::Ellipsis)
+						+SRichTextBlock::Decorator(FTextStyleDecorator::Create(TEXT(""), FStateTreeEditorStyle::Get().GetWidgetStyle<FTextBlockStyle>("Normal.Normal")))
+						+SRichTextBlock::Decorator(FTextStyleDecorator::Create(TEXT("b"), FStateTreeEditorStyle::Get().GetWidgetStyle<FTextBlockStyle>("Normal.Bold")))
+						+SRichTextBlock::Decorator(FTextStyleDecorator::Create(TEXT("i"), FStateTreeEditorStyle::Get().GetWidgetStyle<FTextBlockStyle>("Normal.Italic")))
+						+SRichTextBlock::Decorator(FTextStyleDecorator::Create(TEXT("s"), FStateTreeEditorStyle::Get().GetWidgetStyle<FTextBlockStyle>("Normal.Subdued")))
+					]
 				]
 			]
 
@@ -471,64 +484,68 @@ void SStateTreeViewRow::Construct(const FArguments& InArgs, const TSharedRef<STa
 			.VAlign(VAlign_Center)
 			.AutoWidth()
 			[
-				SNew(SHorizontalBox)
-				+ SHorizontalBox::Slot()
-				.VAlign(VAlign_Center)
-				.Padding(FMargin(8.0f, 0.0f, 0, 0))
-				.AutoWidth()
-				[
-					SNew(SImage)
-					.Image(FAppStyle::Get().GetBrush("Icons.Help"))
-					.ColorAndOpacity(FLinearColor(FColor(31,151,167)))
-					.Visibility(this, &SStateTreeViewRow::GetConditionalTransitionsVisibility)
-				]
-				+ SHorizontalBox::Slot()
-				.VAlign(VAlign_Center)
-				.Padding(FMargin(4.0f, 0.0f, 0, 0))
-				.AutoWidth()
+				SNew(SBox)
+				.Visibility(this, &SStateTreeViewRow::GetConditionalTransitionsVisibility)
 				[
-					SNew(SOverlay)
-					+ SOverlay::Slot()
+					SNew(SHorizontalBox)
+					+ SHorizontalBox::Slot()
+					.VAlign(VAlign_Center)
+					.Padding(FMargin(8.f, 0.f, 0.f, 0.f))
+					.AutoWidth()
 					[
-						SNew(STextBlock)
-						.Text(FEditorFontGlyphs::Long_Arrow_Right)
-						.TextStyle(FStateTreeEditorStyle::Get(), "StateTree.Icon")
-						.Visibility(this, &SStateTreeViewRow::GetConditionalTransitionsVisibility)
+						SNew(SImage)
+						.Image(FStateTreeEditorStyle::Get().GetBrush("StateTreeEditor.Transition.Condition"))
 					]
-					+ SOverlay::Slot()
+					+ SHorizontalBox::Slot()
+					.VAlign(VAlign_Center)
+					.AutoWidth()
 					[
-						// Breakpoint box
-						SNew(SHorizontalBox)
-						+ SHorizontalBox::Slot()
-						.VAlign(VAlign_Top)
-						.HAlign(HAlign_Left)
-						.AutoWidth()
+						SNew(SOverlay)
+						+ SOverlay::Slot()
 						[
-							SNew(SBox)
-							.Padding(FMargin(0.0f, -10.0f, 0.0f, 0.0f))
+							SNew(SImage)
+							.Image(FStateTreeEditorStyle::Get().GetBrush("StateTreeEditor.Transition.Goto"))
+							.ColorAndOpacity(IconTint)
+						]
+						+ SOverlay::Slot()
+						[
+							// Breakpoint box
+							SNew(SHorizontalBox)
+							+ SHorizontalBox::Slot()
+							.VAlign(VAlign_Top)
+							.HAlign(HAlign_Left)
+							.AutoWidth()
 							[
-								SNew(SImage)
-								.DesiredSizeOverride(FVector2D(10.f, 10.f))
-								.Image(FStateTreeEditorStyle::Get().GetBrush(TEXT("StateTreeEditor.Debugger.Breakpoint.EnabledAndValid")))
-								.Visibility(this, &SStateTreeViewRow::GetConditionalTransitionsBreakpointVisibility)
-								.ToolTipText_Lambda([this]
-									{
-										return FText::Format(LOCTEXT("TransitionBreakpointTooltip", "Break when executing transition: {0}"),
-											GetConditionalTransitionsWithBreakpointDesc());
-									})
+								SNew(SBox)
+								.Padding(FMargin(0.f, -10.f, 0.f, 0.f))
+								[
+									SNew(SImage)
+									.DesiredSizeOverride(FVector2D(10.f, 10.f))
+									.Image(FStateTreeEditorStyle::Get().GetBrush(TEXT("StateTreeEditor.Debugger.Breakpoint.EnabledAndValid")))
+									.Visibility(this, &SStateTreeViewRow::GetConditionalTransitionsBreakpointVisibility)
+									.ToolTipText_Lambda([this]
+										{
+											return FText::Format(LOCTEXT("TransitionBreakpointTooltip", "Break when executing transition: {0}"),
+												GetConditionalTransitionsWithBreakpointDesc());
+										})
+								]
 							]
 						]
 					]
-				]
-				+ SHorizontalBox::Slot()
-				.VAlign(VAlign_Center)
-				.AutoWidth()
-				.Padding(FMargin(4.0f, 0, 0, 0))
-				[
-					SNew(STextBlock)
-					.Text(this, &SStateTreeViewRow::GetConditionalTransitionsDesc)
-					.TextStyle(FStateTreeEditorStyle::Get(), "StateTree.Details")
-					.Visibility(this, &SStateTreeViewRow::GetConditionalTransitionsVisibility)
+					+ SHorizontalBox::Slot()
+					.VAlign(VAlign_Center)
+					.AutoWidth()
+					.Padding(FMargin(4.f, 0.f, 0.f, 0.f))
+					[
+						SNew(SRichTextBlock)
+						.Text(this, &SStateTreeViewRow::GetConditionalTransitionsDesc)
+						.TextStyle(&FStateTreeEditorStyle::Get().GetWidgetStyle<FTextBlockStyle>("Normal.Normal"))
+						.OverflowPolicy(ETextOverflowPolicy::Ellipsis)
+						+SRichTextBlock::Decorator(FTextStyleDecorator::Create(TEXT(""), FStateTreeEditorStyle::Get().GetWidgetStyle<FTextBlockStyle>("Normal.Normal")))
+						+SRichTextBlock::Decorator(FTextStyleDecorator::Create(TEXT("b"), FStateTreeEditorStyle::Get().GetWidgetStyle<FTextBlockStyle>("Normal.Bold")))
+						+SRichTextBlock::Decorator(FTextStyleDecorator::Create(TEXT("i"), FStateTreeEditorStyle::Get().GetWidgetStyle<FTextBlockStyle>("Normal.Italic")))
+						+SRichTextBlock::Decorator(FTextStyleDecorator::Create(TEXT("s"), FStateTreeEditorStyle::Get().GetWidgetStyle<FTextBlockStyle>("Normal.Subdued")))
+					]
 				]
 			]
 		]
@@ -847,7 +864,9 @@ FText SStateTreeViewRow::GetSelectorTooltip() const
 		}
 		else if (State->SelectionBehavior == EStateTreeStateSelectionBehavior::TrySelectChildrenInOrder)
 		{
-			if (State->Children.IsEmpty())
+			if (State->Children.IsEmpty()
+				|| State->Type == EStateTreeStateType::Linked
+				|| State->Type == EStateTreeStateType::LinkedAsset)
 			{
 				const int32 EnterStateIndex = Enum->GetIndexByValue((int64)EStateTreeStateSelectionBehavior::TryEnterState);
 				return FText::Format(LOCTEXT("ConvertedToEnterState", "{0}\nAutomatically converted from '{1}' becase the State has no child States."),
@@ -972,19 +991,19 @@ FText SStateTreeViewRow::GetLinkDescription(const FStateTreeStateLink& Link)
 	switch (Link.LinkType)
 	{
 	case EStateTreeTransitionType::None:
-		return LOCTEXT("TransitionNoneStyled", "[None]");
+		return LOCTEXT("TransitionNoneStyled", "<i>None</>");
 		break;
 	case EStateTreeTransitionType::Succeeded:
-		return LOCTEXT("TransitionTreeSucceededStyled", "[Succeeded]");
+		return LOCTEXT("TransitionTreeSucceededStyled", "<i>Succeeded</>");
 		break;
 	case EStateTreeTransitionType::Failed:
-		return LOCTEXT("TransitionTreeFailedStyled", "[Failed]");
+		return LOCTEXT("TransitionTreeFailedStyled", "<i>Failed</>");
 		break;
 	case EStateTreeTransitionType::NextState:
-		return LOCTEXT("TransitionNextStateStyled", "[Next]");
+		return LOCTEXT("TransitionNextStateStyled", "<i>Next</>");
 		break;
 	case EStateTreeTransitionType::NextSelectableState:
-		return LOCTEXT("TransitionNextSelectableStateStyled", "[Next Selectable]");
+		return LOCTEXT("TransitionNextSelectableStateStyled", "<i>Next Selectable</s>");
 		break;
 	case EStateTreeTransitionType::GotoState:
 		return FText::FromName(Link.Name);
@@ -997,6 +1016,17 @@ FText SStateTreeViewRow::GetLinkDescription(const FStateTreeStateLink& Link)
 	return FText::GetEmpty();
 };
 
+bool SStateTreeViewRow::IsLeafState() const
+{
+	const UStateTreeState* State = WeakState.Get();
+	return State
+		&& State->Children.Num() == 0
+		&& !IsRootState()
+		&& (State->Type == EStateTreeStateType::State
+			|| State->Type == EStateTreeStateType::Linked
+			|| State->Type == EStateTreeStateType::LinkedAsset);
+}
+
 FText SStateTreeViewRow::GetTransitionsDesc(const UStateTreeState& State, const EStateTreeTransitionTrigger Trigger, const FTransitionDescFilterOptions FilterOptions) const
 {
 	TArray<FText> DescItems;
@@ -1061,32 +1091,31 @@ FText SStateTreeViewRow::GetTransitionsDesc(const UStateTreeState& State, const
 		AddLinksFromStruct(State.SingleTask.GetInstance());
 	}
 
-	if (State.Children.Num() == 0
-		&& State.Type == EStateTreeStateType::State
+	if (IsLeafState()
 		&& DescItems.Num() == 0
 		&& EnumHasAnyFlags(Trigger, EStateTreeTransitionTrigger::OnStateCompleted))
 	{
 		if (HasParentTransitionForTrigger(State, Trigger))
 		{
-			DescItems.Add(LOCTEXT("TransitionActionHandleInParentStyled", "[Parent]"));
+			DescItems.Add(LOCTEXT("TransitionActionHandleInParentStyled", "<i>Parent</>"));
 		}
 		else
 		{
-			DescItems.Add(LOCTEXT("TransitionActionRoot", "[Root]"));
+			DescItems.Add(LOCTEXT("TransitionActionRoot", "<i>Root</>"));
 		}
 	}
 	
 	return FText::Join(FText::FromString(TEXT(", ")), DescItems);
 }
 
-FText SStateTreeViewRow::GetTransitionsIcon(const UStateTreeState& State, const EStateTreeTransitionTrigger Trigger, const FTransitionDescFilterOptions FilterOptions) const
+const FSlateBrush* SStateTreeViewRow::GetTransitionsIcon(const UStateTreeState& State, const EStateTreeTransitionTrigger Trigger, const FTransitionDescFilterOptions FilterOptions) const
 {
 	enum EIconType
 	{
 		IconNone = 0,
-		IconRightArrow =	1 << 0,
-		IconDownArrow =		1 << 1,
-		IconLevelUp =		1 << 2,
+		IconGoto =	1 << 0,
+		IconNext =		1 << 1,
+		IconParent =		1 << 2,
 	};
 	uint8 IconType = IconNone;
 	
@@ -1118,20 +1147,20 @@ FText SStateTreeViewRow::GetTransitionsIcon(const UStateTreeState& State, const
 			switch (Transition.State.LinkType)
 			{
 			case EStateTreeTransitionType::None:
-				IconType |= IconRightArrow;
+				IconType |= IconGoto;
 				break;
 			case EStateTreeTransitionType::Succeeded:
-				IconType |= IconRightArrow;
+				IconType |= IconGoto;
 				break;
 			case EStateTreeTransitionType::Failed:
-				IconType |= IconRightArrow;
+				IconType |= IconGoto;
 				break;
 			case EStateTreeTransitionType::NextState:
 			case EStateTreeTransitionType::NextSelectableState:
-				IconType |= IconDownArrow;
+				IconType |= IconNext;
 				break;
 			case EStateTreeTransitionType::GotoState:
-				IconType |= IconRightArrow;
+				IconType |= IconGoto;
 				break;
 			default:
 				ensureMsgf(false, TEXT("Unhandled transition type."));
@@ -1143,29 +1172,30 @@ FText SStateTreeViewRow::GetTransitionsIcon(const UStateTreeState& State, const
 	if (FMath::CountBits(static_cast<uint64>(IconType)) > 1)
 	{
 		// Prune down to just one icon.
-		IconType = IconRightArrow;
+		IconType = IconGoto;
 	}
-	
-	if (State.Children.Num() == 0
-		&& State.Type == EStateTreeStateType::State
+
+	if (IsLeafState()
 		&& IconType == IconNone
 		&& EnumHasAnyFlags(Trigger, EStateTreeTransitionTrigger::OnStateCompleted))
 	{
 		// Transition is handled on parent state, or implicit Root.
-		IconType = IconLevelUp;
+		IconType = IconParent;
 	}
 
 	switch (IconType)
 	{
-		case IconRightArrow:
-			return FEditorFontGlyphs::Long_Arrow_Right;
-		case IconDownArrow:
-			return FEditorFontGlyphs::Long_Arrow_Down;
-		case IconLevelUp:
-			return FEditorFontGlyphs::Level_Up;
+		case IconGoto:
+			return FStateTreeEditorStyle::Get().GetBrush("StateTreeEditor.Transition.Goto");
+		case IconNext:
+			return FStateTreeEditorStyle::Get().GetBrush("StateTreeEditor.Transition.Next");
+		case IconParent:
+			return FStateTreeEditorStyle::Get().GetBrush("StateTreeEditor.Transition.Parent");
 		default:
-			return FText::GetEmpty();
+			break;
 	}
+	
+	return nullptr;
 }
 
 EVisibility SStateTreeViewRow::GetTransitionsVisibility(const UStateTreeState& State, const EStateTreeTransitionTrigger Trigger) const
@@ -1173,7 +1203,6 @@ EVisibility SStateTreeViewRow::GetTransitionsVisibility(const UStateTreeState& S
 	// Handle completed, succeeded and failed transitions.
 	if (EnumHasAnyFlags(Trigger, EStateTreeTransitionTrigger::OnStateCompleted))
 	{
-		const bool bIsLeafState = (State.Children.Num() == 0);
 		EStateTreeTransitionTrigger HandledTriggers = EStateTreeTransitionTrigger::None;
 		bool bExactMatch = false;
 
@@ -1195,7 +1224,7 @@ EVisibility SStateTreeViewRow::GetTransitionsVisibility(const UStateTreeState& S
 		}
 
 		// Assume that leaf states should have completion transitions.
-		if (!bExactMatch && bIsLeafState)
+		if (!bExactMatch && IsLeafState())
 		{
 			// Find the missing transition type, note: Completed = Succeeded|Failed.
 			const EStateTreeTransitionTrigger MissingTriggers = HandledTriggers ^ EStateTreeTransitionTrigger::OnStateCompleted;
@@ -1321,13 +1350,13 @@ FText SStateTreeViewRow::GetCompletedTransitionWithBreakpointDesc() const
 	return FText::GetEmpty();
 }
 
-FText SStateTreeViewRow::GetCompletedTransitionsIcon() const
+const FSlateBrush* SStateTreeViewRow::GetCompletedTransitionsIcon() const
 {
 	if (const UStateTreeState* State = WeakState.Get())
 	{
 		return GetTransitionsIcon(*State, EStateTreeTransitionTrigger::OnStateCompleted);
 	}
-	return FText::GetEmpty();
+	return nullptr;
 }
 
 EVisibility SStateTreeViewRow::GetSucceededTransitionVisibility() const
@@ -1368,13 +1397,13 @@ FText SStateTreeViewRow::GetSucceededTransitionWithBreakpointDesc() const
 	return FText::GetEmpty();
 }
 
-FText SStateTreeViewRow::GetSucceededTransitionIcon() const
+const FSlateBrush* SStateTreeViewRow::GetSucceededTransitionIcon() const
 {
 	if (const UStateTreeState* State = WeakState.Get())
 	{
 		return GetTransitionsIcon(*State, EStateTreeTransitionTrigger::OnStateSucceeded);
 	}
-	return FText::GetEmpty();
+	return nullptr;
 }
 
 EVisibility SStateTreeViewRow::GetFailedTransitionVisibility() const
@@ -1415,13 +1444,13 @@ FText SStateTreeViewRow::GetFailedTransitionWithBreakpointDesc() const
 	return FText::GetEmpty();
 }
 
-FText SStateTreeViewRow::GetFailedTransitionIcon() const
+const FSlateBrush* SStateTreeViewRow::GetFailedTransitionIcon() const
 {
 	if (const UStateTreeState* State = WeakState.Get())
 	{
 		return GetTransitionsIcon(*State, EStateTreeTransitionTrigger::OnStateFailed);
 	}
-	return FEditorFontGlyphs::Ban;
+	return nullptr;
 }
 
 EVisibility SStateTreeViewRow::GetConditionalTransitionsVisibility() const
diff --git a/Engine/Plugins/Runtime/StateTree/Source/StateTreeEditorModule/Private/SStateTreeViewRow.h b/Engine/Plugins/Runtime/StateTree/Source/StateTreeEditorModule/Private/SStateTreeViewRow.h
--- a/Engine/Plugins/Runtime/StateTree/Source/StateTreeEditorModule/Private/SStateTreeViewRow.h
+++ b/Engine/Plugins/Runtime/StateTree/Source/StateTreeEditorModule/Private/SStateTreeViewRow.h
@@ -57,19 +57,19 @@ private:
 	EVisibility GetCompletedTransitionBreakpointVisibility() const;
 	FText GetCompletedTransitionsDesc() const;
 	FText GetCompletedTransitionWithBreakpointDesc() const;
-	FText GetCompletedTransitionsIcon() const;
+	const FSlateBrush* GetCompletedTransitionsIcon() const;
 
 	EVisibility GetSucceededTransitionVisibility() const;
 	EVisibility GetSucceededTransitionBreakpointVisibility() const;
 	FText GetSucceededTransitionDesc() const;
 	FText GetSucceededTransitionWithBreakpointDesc() const;
-	FText GetSucceededTransitionIcon() const;
+	const FSlateBrush* GetSucceededTransitionIcon() const;
 
 	EVisibility GetFailedTransitionVisibility() const;
 	EVisibility GetFailedTransitionBreakpointVisibility() const;
 	FText GetFailedTransitionDesc() const;
 	FText GetFailedTransitionWithBreakpointDesc() const;
-	FText GetFailedTransitionIcon() const;
+	const FSlateBrush* GetFailedTransitionIcon() const;
 
 	EVisibility GetConditionalTransitionsVisibility() const;
 	EVisibility GetConditionalTransitionsBreakpointVisibility() const;
@@ -98,13 +98,14 @@ private:
 
 	static FText GetLinkDescription(const FStateTreeStateLink& Link);
 	FText GetTransitionsDesc(const UStateTreeState& State, const EStateTreeTransitionTrigger Trigger, const FTransitionDescFilterOptions FilterOptions = {}) const;
-	FText GetTransitionsIcon(const UStateTreeState& State, const EStateTreeTransitionTrigger Trigger, const FTransitionDescFilterOptions FilterOptions = {}) const;
+	const FSlateBrush* GetTransitionsIcon(const UStateTreeState& State, const EStateTreeTransitionTrigger Trigger, const FTransitionDescFilterOptions FilterOptions = {}) const;
 	EVisibility GetTransitionsVisibility(const UStateTreeState& State, const EStateTreeTransitionTrigger Trigger) const;
 	EVisibility GetTransitionsBreakpointVisibility(const UStateTreeState& State, const EStateTreeTransitionTrigger Trigger) const;
 
 	bool HasParentTransitionForTrigger(const UStateTreeState& State, const EStateTreeTransitionTrigger Trigger) const;
 
 	bool IsRootState() const;
+	bool IsLeafState() const;
 	bool IsStateSelected() const;
 
 	void HandleNodeLabelTextCommitted(const FText& NewLabel, ETextCommit::Type CommitType) const;
diff --git a/Engine/Plugins/Runtime/StateTree/Source/StateTreeEditorModule/Private/StateTreeEditorStyle.cpp b/Engine/Plugins/Runtime/StateTree/Source/StateTreeEditorModule/Private/StateTreeEditorStyle.cpp
--- a/Engine/Plugins/Runtime/StateTree/Source/StateTreeEditorModule/Private/StateTreeEditorStyle.cpp
+++ b/Engine/Plugins/Runtime/StateTree/Source/StateTreeEditorModule/Private/StateTreeEditorStyle.cpp
@@ -56,11 +56,6 @@ FStateTreeEditorStyle::FStateTreeEditorStyle()
 			.SetColorAndOpacity(FLinearColor(230.0f / 255.0f, 230.0f / 255.0f, 230.0f / 255.0f, 0.5f));
 		Set("StateTree.Icon", StateIcon);
 
-		const FTextBlockStyle StateDetailsIcon = FTextBlockStyle(NormalText)
-			.SetFont(FAppStyle::Get().GetFontStyle("FontAwesome.10"))
-		    .SetColorAndOpacity(FLinearColor(230.0f / 255.0f, 230.0f / 255.0f, 230.0f / 255.0f, 0.5f));
-		Set("StateTree.DetailsIcon", StateDetailsIcon);
-
 		const FTextBlockStyle StateTitle = FTextBlockStyle(NormalText)
 			.SetFont(DEFAULT_FONT("Bold", 12))
 			.SetColorAndOpacity(FLinearColor(230.0f / 255.0f, 230.0f / 255.0f, 230.0f / 255.0f, 0.9f));
@@ -128,11 +123,33 @@ FStateTreeEditorStyle::FStateTreeEditorStyle()
 		Set("Details.Bold", FTextBlockStyle(NormalText)
 			.SetFont(FAppStyle::GetFontStyle(TEXT("PropertyWindow.BoldFont"))));
 
+		Set("Details.Italic", FTextBlockStyle(NormalText)
+			.SetFont(FAppStyle::GetFontStyle(TEXT("PropertyWindow.ItalicFont"))));
+
 		Set("Details.Subdued", FTextBlockStyle(NormalText)
 			.SetColorAndOpacity(FSlateColor::UseSubduedForeground())
 			.SetFont(FAppStyle::GetFontStyle(TEXT("PropertyWindow.NormalFont"))));
 	}
 
+	// Normal rich text
+	{
+		Set("Normal.Normal", FTextBlockStyle(NormalText)
+			.SetColorAndOpacity(FSlateColor::UseForeground())
+			.SetFont(DEFAULT_FONT("Regular", 10)));
+
+		Set("Normal.Bold", FTextBlockStyle(NormalText)
+			.SetColorAndOpacity(FSlateColor::UseForeground())
+			.SetFont(DEFAULT_FONT("Bold", 10)));
+
+		Set("Normal.Italic", FTextBlockStyle(NormalText)
+			.SetColorAndOpacity(FSlateColor::UseForeground())
+			.SetFont(DEFAULT_FONT("Italic", 10)));
+
+		Set("Normal.Subdued", FTextBlockStyle(NormalText)
+			.SetColorAndOpacity(FSlateColor::UseSubduedForeground())
+			.SetFont(DEFAULT_FONT("Regular", 10)));
+	}
+
 	// Debugger
 	{
 		Set("StateTreeDebugger.Element.Normal",
@@ -247,15 +264,13 @@ FStateTreeEditorStyle::FStateTreeEditorStyle()
 
 		Set("StateTreeEditor.Debugger.Task.Enter", new CORE_IMAGE_BRUSH_SVG("Starship/Common/arrow-right", CoreStyleConstants::Icon16x16, FStyleColors::Foreground));
 		Set("StateTreeEditor.Debugger.Task.Exit", new CORE_IMAGE_BRUSH_SVG("Starship/Common/arrow-left", CoreStyleConstants::Icon16x16, FStyleColors::Foreground));
-		Set("StateTreeEditor.Debugger.Task.Failed", new CORE_IMAGE_BRUSH_SVG("Starship/Common/close", CoreStyleConstants::Icon16x16, FStyleColors::AccentRed));
+		Set("StateTreeEditor.Debugger.Task.Failed", new CORE_IMAGE_BRUSH_SVG("Starship/Common/close-small", CoreStyleConstants::Icon16x16, FStyleColors::AccentRed));
 		Set("StateTreeEditor.Debugger.Task.Succeeded", new CORE_IMAGE_BRUSH_SVG("Starship/Common/check", CoreStyleConstants::Icon16x16, FStyleColors::AccentGreen));
-		Set("StateTreeEditor.Debugger.Task.Stopped", new CORE_IMAGE_BRUSH_SVG("Starship/Common/close", CoreStyleConstants::Icon16x16, FStyleColors::AccentRed));
+		Set("StateTreeEditor.Debugger.Task.Stopped", new CORE_IMAGE_BRUSH_SVG("Starship/Common/close-small", CoreStyleConstants::Icon16x16, FStyleColors::AccentRed));
 
-		Set("StateTreeEditor.Debugger.Condition", new CORE_IMAGE_BRUSH_SVG("Starship/Common/help", CoreStyleConstants::Icon16x16, FStyleColors::Foreground));
 		Set("StateTreeEditor.Debugger.Condition.Passed", new CORE_IMAGE_BRUSH_SVG("Starship/Common/check", CoreStyleConstants::Icon16x16, FStyleColors::AccentGreen));
-		Set("StateTreeEditor.Debugger.Condition.Failed", new CORE_IMAGE_BRUSH_SVG("Starship/Common/close", CoreStyleConstants::Icon16x16, FStyleColors::AccentRed));
+		Set("StateTreeEditor.Debugger.Condition.Failed", new CORE_IMAGE_BRUSH_SVG("Starship/Common/close-small", CoreStyleConstants::Icon16x16, FStyleColors::AccentRed));
 		Set("StateTreeEditor.Debugger.Condition.OnEvaluating", new CORE_IMAGE_BRUSH_SVG("Starship/Common/Update", CoreStyleConstants::Icon16x16, FStyleColors::AccentYellow));
-		Set("StateTreeEditor.Debugger.Condition.OnTransition", new CORE_IMAGE_BRUSH_SVG("Starship/Common/help", CoreStyleConstants::Icon16x16, FStyleColors::AccentBlue));
 
 		Set("StateTreeEditor.Debugger.Unset", new CORE_IMAGE_BRUSH_SVG("Starship/Common/help", CoreStyleConstants::Icon16x16, FStyleColors::AccentBlack));
 	}
@@ -285,9 +300,11 @@ FStateTreeEditorStyle::FStateTreeEditorStyle()
 		Set("StateTreeEditor.Debugger.ResumeDebuggerAnalysis", new IMAGE_BRUSH_SVG("Starship/Common/Timeline", CoreStyleConstants::Icon16x16));
 
 		Set("StateTreeEditor.Transition.None", new CORE_IMAGE_BRUSH_SVG("Starship/Common/x-circle", CoreStyleConstants::Icon16x16, FSlateColor::UseSubduedForeground()));
-		Set("StateTreeEditor.Transition.Next", new CORE_IMAGE_BRUSH_SVG("Starship/Common/arrow-down", CoreStyleConstants::Icon16x16, FSlateColor::UseSubduedForeground()));
 		Set("StateTreeEditor.Transition.Succeeded", new CORE_IMAGE_BRUSH_SVG("Starship/Common/check", CoreStyleConstants::Icon16x16, FStyleColors::AccentGreen));
-		Set("StateTreeEditor.Transition.Failed", new CORE_IMAGE_BRUSH_SVG("Starship/Common/close", CoreStyleConstants::Icon16x16, FStyleColors::AccentRed));
+		Set("StateTreeEditor.Transition.Failed", new CORE_IMAGE_BRUSH_SVG("Starship/Common/close-small", CoreStyleConstants::Icon16x16, FStyleColors::AccentRed));
+
+		Set("StateTreeEditor.Transition.Succeeded", new CORE_IMAGE_BRUSH_SVG("Starship/Common/check", CoreStyleConstants::Icon16x16, FStyleColors::AccentGreen));
+		Set("StateTreeEditor.Transition.Failed", new CORE_IMAGE_BRUSH_SVG("Starship/Common/close-small", CoreStyleConstants::Icon16x16, FStyleColors::AccentRed));
 	}
 
 	{
@@ -305,8 +322,17 @@ FStateTreeEditorStyle::FStateTreeEditorStyle()
 
 		Set("StateTreeEditor.Conditions", new IMAGE_BRUSH_SVG("Icons/Conditions", CoreStyleConstants::Icon16x16));
 
+		Set("StateTreeEditor.Tasks", new IMAGE_BRUSH_SVG("Icons/Tasks", CoreStyleConstants::Icon16x16));
+
 		Set("StateTreeEditor.StateSubtree", new IMAGE_BRUSH_SVG("Icons/State_Subtree", CoreStyleConstants::Icon16x16));
 		Set("StateTreeEditor.StateLinked", new IMAGE_BRUSH_SVG("Icons/State_Linked", CoreStyleConstants::Icon16x16));
+
+		Set("StateTreeEditor.Transition.Goto", new IMAGE_BRUSH_SVG("Icons/Transition_Goto", CoreStyleConstants::Icon16x16, FStyleColors::Foreground));
+		Set("StateTreeEditor.Transition.Next", new IMAGE_BRUSH_SVG("Icons/Transition_Next", CoreStyleConstants::Icon16x16, FStyleColors::Foreground));
+		Set("StateTreeEditor.Transition.Parent", new IMAGE_BRUSH_SVG("Icons/Transition_Parent", CoreStyleConstants::Icon16x16, FStyleColors::Foreground));
+
+		Set("StateTreeEditor.Transition.Condition", new IMAGE_BRUSH_SVG("Icons/Conditions", CoreStyleConstants::Icon16x16, FStyleColors::AccentBlue));
+		Set("StateTreeEditor.Debugger.Condition.OnTransition", new IMAGE_BRUSH_SVG("Icons/Conditions", CoreStyleConstants::Icon16x16, FStyleColors::AccentBlue));
 	}
 }
 
