StateTree: Fix stack-use-after-scope

diff --git a/Engine/Plugins/Runtime/StateTree/Source/StateTreeModule/Private/StateTreeExecutionContext.cpp b/Engine/Plugins/Runtime/StateTree/Source/StateTreeModule/Private/StateTreeExecutionContext.cpp
--- a/Engine/Plugins/Runtime/StateTree/Source/StateTreeModule/Private/StateTreeExecutionContext.cpp
+++ b/Engine/Plugins/Runtime/StateTree/Source/StateTreeModule/Private/StateTreeExecutionContext.cpp
@@ -2743,14 +2743,14 @@ bool FStateTreeExecutionContext::TriggerTransitions()
 						}
 					}
 
-					TArray<FStateTreeSharedEvent*, TInlineAllocator<8>> TransitionEvents;
+					TArray<const FStateTreeSharedEvent*, TInlineAllocator<8>> TransitionEvents;
 
 					if (Transition.Trigger == EStateTreeTransitionTrigger::OnEvent)
 					{
 						check(Transition.RequiredEvent.IsValid());
 
 						TConstArrayView<FStateTreeSharedEvent> EventsQueue = GetEventsToProcessView();
-						for (FStateTreeSharedEvent Event : EventsQueue)
+						for (const FStateTreeSharedEvent& Event : EventsQueue)
 						{
 							check(Event.IsValid());
 							if (Transition.RequiredEvent.DoesEventMatchDesc(*Event))
@@ -2764,7 +2764,7 @@ bool FStateTreeExecutionContext::TriggerTransitions()
 						TransitionEvents.Emplace();
 					}
 					
-					for (FStateTreeSharedEvent* TransitionEvent : TransitionEvents)
+					for (const FStateTreeSharedEvent* TransitionEvent : TransitionEvents)
 					{
 						bool bPassed = false; 
 						{
