StateTree: Changed Task icon.

diff --git a/Engine/Plugins/Runtime/StateTree/Source/StateTreeEditorModule/Private/Customizations/StateTreeEditorNodeDetails.cpp b/Engine/Plugins/Runtime/StateTree/Source/StateTreeEditorModule/Private/Customizations/StateTreeEditorNodeDetails.cpp
--- a/Engine/Plugins/Runtime/StateTree/Source/StateTreeEditorModule/Private/Customizations/StateTreeEditorNodeDetails.cpp
+++ b/Engine/Plugins/Runtime/StateTree/Source/StateTreeEditorModule/Private/Customizations/StateTreeEditorNodeDetails.cpp
@@ -595,9 +595,9 @@ void FStateTreeEditorNodeDetails::CustomizeHeader(TSharedRef<class IPropertyHand
 					.VAlign(VAlign_Center)
 					.Padding(0, 0, 4, 0)
 					[
-						SNew(STextBlock)
-						.Text(FEditorFontGlyphs::Paper_Plane)
-						.TextStyle(FStateTreeEditorStyle::Get(), "StateTree.DetailsIcon")
+						SNew(SImage)
+						.Image(FStateTreeEditorStyle::Get().GetBrush("StateTreeEditor.Tasks"))
+						.ColorAndOpacity(FSlateColor::UseForeground())
 						.Visibility(this, &FStateTreeEditorNodeDetails::IsTaskVisible)
 					]
 					+ SHorizontalBox::Slot()
