StateTree: Make an error to fail ST compilation.

diff --git a/Engine/Plugins/Runtime/StateTree/Source/StateTreeEditorModule/Private/StateTreeCompiler.cpp b/Engine/Plugins/Runtime/StateTree/Source/StateTreeEditorModule/Private/StateTreeCompiler.cpp
--- a/Engine/Plugins/Runtime/StateTree/Source/StateTreeEditorModule/Private/StateTreeCompiler.cpp
+++ b/Engine/Plugins/Runtime/StateTree/Source/StateTreeEditorModule/Private/StateTreeCompiler.cpp
@@ -926,6 +926,7 @@ bool FStateTreeCompiler::CreateStateTransitions()
 					Log.Reportf(EMessageSeverity::Error,
 						TEXT("State completion transition to '%s' must have transition to valid state, 'None' not accepted."),
 						*Transition.State.Name.ToString());
+					return false;
 				}
 			}
 			
