StateTree: Improved event handling
- Added support for shared event buffer, so that multiple instance data can use the same
- Added API to consume events
- Removed double buffering from events
- Treat TriggerTransitions() as event handler (events flushed after each call to the method)
- Event handlers (including tasks) are executed in priority order
- Transitions and event capturing states can consume events on successful selection
- Added API to tick a StateTree in two passes (update tasks and trigger transitions)
- Change parallel tree tasks to do the task update in Tick(), and event handling in TriggerTransitions()
- Small improvements to the ST debugger to display events
- NOTE: this is breaking change for implementations that has relied events emitted during tick to be available on next EnterState()

#okfirgithub public

diff --git a/Engine/Plugins/Runtime/StateTree/Source/StateTreeEditorModule/Private/StateTreeCompiler.cpp b/Engine/Plugins/Runtime/StateTree/Source/StateTreeEditorModule/Private/StateTreeCompiler.cpp
--- a/Engine/Plugins/Runtime/StateTree/Source/StateTreeEditorModule/Private/StateTreeCompiler.cpp
+++ b/Engine/Plugins/Runtime/StateTree/Source/StateTreeEditorModule/Private/StateTreeCompiler.cpp
@@ -634,6 +634,7 @@ bool FStateTreeCompiler::CreateStateTasksAndParameters()
 			CompactState.EventDataIndex = FStateTreeIndex16(InstanceDataIndex++);
 			CompactState.RequiredEventToEnter.Tag = State->RequiredEventToEnter.Tag;
 			CompactState.RequiredEventToEnter.PayloadStruct = State->RequiredEventToEnter.PayloadStruct;
+			CompactState.bConsumeEventOnSelect = State->RequiredEventToEnter.bConsumeEventOnSelect;
 
 			FStateTreeBindableStructDesc Desc;
 			Desc.Struct = FStateTreeEvent::StaticStruct();
@@ -940,6 +941,7 @@ bool FStateTreeCompiler::CreateStateTransitions()
 			{
 				CompactTransition.RequiredEvent.Tag = Transition.RequiredEvent.Tag;
 				CompactTransition.RequiredEvent.PayloadStruct = Transition.RequiredEvent.PayloadStruct;
+				CompactTransition.bConsumeEventOnSelect = Transition.RequiredEvent.bConsumeEventOnSelect;
 
 				FStateTreeBindableStructDesc Desc;
 				Desc.Struct = FStateTreeEvent::StaticStruct();
diff --git a/Engine/Plugins/Runtime/StateTree/Source/StateTreeEditorModule/Public/StateTreeState.h b/Engine/Plugins/Runtime/StateTree/Source/StateTreeEditorModule/Public/StateTreeState.h
--- a/Engine/Plugins/Runtime/StateTree/Source/StateTreeEditorModule/Public/StateTreeState.h
+++ b/Engine/Plugins/Runtime/StateTree/Source/StateTreeEditorModule/Public/StateTreeState.h
@@ -3,6 +3,7 @@
 #pragma once
 
 #include "IStateTreeSchemaProvider.h"
+#include "StateTreeNodeBase.h"
 #include "StateTreeEditorNode.h"
 #include "StateTreeEditorTypes.h"
 #include "StateTreeEvents.h"
@@ -33,6 +34,10 @@ struct FStateTreeEventDesc
 	UPROPERTY(EditDefaultsOnly, Category = "Event")
 	TObjectPtr<const UScriptStruct> PayloadStruct;
 
+	/** If set to true, the event is consumed (later state selection cannot react to it) if state selection can be made. */
+	UPROPERTY(EditDefaultsOnly, Category = "Event")
+	bool bConsumeEventOnSelect = true;
+	
 	bool IsValid() const
 	{
 		return Tag.IsValid() || PayloadStruct;
@@ -82,8 +87,8 @@ PRAGMA_ENABLE_DEPRECATION_WARNINGS
 		FStateTreeEditorNode& CondNode = Conditions.AddDefaulted_GetRef();
 		CondNode.ID = FGuid::NewGuid();
 		CondNode.Node.InitializeAs<T>(Forward<TArgs>(InArgs)...);
-		T& Cond = CondNode.Node.GetMutable<T>();
-		if (const UScriptStruct* InstanceType = Cast<const UScriptStruct>(Cond.GetInstanceDataType()))
+		const FStateTreeNodeBase& Node = CondNode.Node.GetMutable<FStateTreeNodeBase>();
+		if (const UScriptStruct* InstanceType = Cast<const UScriptStruct>(Node.GetInstanceDataType()))
 		{
 			CondNode.Instance.InitializeAs(InstanceType);
 		}
@@ -101,6 +106,7 @@ PRAGMA_ENABLE_DEPRECATION_WARNINGS
 	UPROPERTY(EditDefaultsOnly, Category = "Transition")
 	EStateTreeTransitionTrigger Trigger = EStateTreeTransitionTrigger::OnStateCompleted;
 
+	/** Defines the event required to be present during state selection for the transition to trigger. */
 	UPROPERTY(EditDefaultsOnly, Category = "Transition", DisplayName = "Required Event")
 	FStateTreeEventDesc RequiredEvent; 
 
@@ -246,8 +252,8 @@ public:
 		FStateTreeEditorNode& CondNode = EnterConditions.AddDefaulted_GetRef();
 		CondNode.ID = FGuid::NewGuid();
 		CondNode.Node.InitializeAs<T>(Forward<TArgs>(InArgs)...);
-		T& Cond = CondNode.Node.GetMutable<T>();
-		if (const UScriptStruct* InstanceType = Cast<const UScriptStruct>(Cond.GetInstanceDataType()))
+		const FStateTreeNodeBase& Node = CondNode.Node.GetMutable<FStateTreeNodeBase>();
+		if (const UScriptStruct* InstanceType = Cast<const UScriptStruct>(Node.GetInstanceDataType()))
 		{
 			CondNode.Instance.InitializeAs(InstanceType);
 		}
@@ -264,8 +270,8 @@ public:
 		FStateTreeEditorNode& TaskItem = Tasks.AddDefaulted_GetRef();
 		TaskItem.ID = FGuid::NewGuid();
 		TaskItem.Node.InitializeAs<T>(Forward<TArgs>(InArgs)...);
-		T& Task = TaskItem.Node.GetMutable<T>();
-		if (const UScriptStruct* InstanceType = Cast<const UScriptStruct>(Task.GetInstanceDataType()))
+		const FStateTreeNodeBase& Node = TaskItem.Node.GetMutable<FStateTreeNodeBase>();
+		if (const UScriptStruct* InstanceType = Cast<const UScriptStruct>(Node.GetInstanceDataType()))
 		{
 			TaskItem.Instance.InitializeAs(InstanceType);
 		}
@@ -354,6 +360,7 @@ public:
 	UPROPERTY(EditDefaultsOnly, Category = "State", meta=(InlineEditConditionToggle))
 	bool bHasRequiredEventToEnter = false;
 
+	/** Defines the event required to be present during state selection for the state to be selected. */
 	UPROPERTY(EditDefaultsOnly, Category = "State", meta = (EditCondition = "bHasRequiredEventToEnter"))
 	FStateTreeEventDesc RequiredEventToEnter;
 
diff --git a/Engine/Plugins/Runtime/StateTree/Source/StateTreeModule/Private/Debugger/StateTreeTraceTypes.cpp b/Engine/Plugins/Runtime/StateTree/Source/StateTreeModule/Private/Debugger/StateTreeTraceTypes.cpp
--- a/Engine/Plugins/Runtime/StateTree/Source/StateTreeModule/Private/Debugger/StateTreeTraceTypes.cpp
+++ b/Engine/Plugins/Runtime/StateTree/Source/StateTreeModule/Private/Debugger/StateTreeTraceTypes.cpp
@@ -130,12 +130,6 @@ FString FStateTreeTraceTransitionEvent::GetValueString(const UStateTree& StateTr
 {
 	const FCompactStateTreeState* CompactState = StateTree.GetStateFromHandle(TransitionSource.TargetState);
 	FStringBuilderBase StrBuilder;
-	StrBuilder.Appendf(TEXT("go to State '%s'"), CompactState != nullptr ? *UE::StateTreeTrace::GetStateName(StateTree, CompactState) : *TransitionSource.TargetState.Describe());
-
-	if (TransitionSource.Priority != EStateTreeTransitionPriority::None)
-	{
-		StrBuilder.Appendf(TEXT(" (Priority: %s)"), *UEnum::GetDisplayValueAsText(TransitionSource.Priority).ToString()); 
-	}
 
 	if (TransitionSource.SourceType == EStateTreeTransitionSourceType::Asset)
 	{
@@ -143,21 +137,42 @@ FString FStateTreeTraceTransitionEvent::GetValueString(const UStateTree& StateTr
 		{
 			ensureAlways(Transition->Priority == TransitionSource.Priority);
 			ensureAlways(Transition->State == TransitionSource.TargetState);
-			if (Transition->RequiredEvent.Tag.IsValid())
-			{
-				StrBuilder.Appendf(TEXT("\n\t%s"), *Transition->RequiredEvent.Tag.ToString()); 
-			}
 
-			if (Transition->RequiredEvent.PayloadStruct)
+			const bool bHasTag = Transition->RequiredEvent.Tag.IsValid();
+			const bool bHasPayload = Transition->RequiredEvent.PayloadStruct != nullptr;
+			
+			if (bHasTag || bHasPayload)
 			{
-				StrBuilder.Appendf(TEXT("\n\t%s"), *Transition->RequiredEvent.PayloadStruct->GetName()); 
+				StrBuilder.Appendf(TEXT("("));
+				
+				if (bHasTag)
+				{
+					StrBuilder.Appendf(TEXT("Tag: '%s'"), *Transition->RequiredEvent.Tag.ToString()); 
+				}
+				if (bHasTag && bHasPayload)
+				{
+					StrBuilder.Appendf(TEXT(", "));
+				}
+				if (bHasPayload)
+				{
+					StrBuilder.Appendf(TEXT(" Payload: '%s'"), *Transition->RequiredEvent.PayloadStruct->GetName()); 
+				}
+				StrBuilder.Appendf(TEXT(") ")); 
 			}
 		}
 		else
 		{
-			StrBuilder.Appendf(TEXT("Invalid Transition Index %s for '%s'"), *LexToString(TransitionSource.TransitionIndex.Get()), *StateTree.GetFullName());
+			StrBuilder.Appendf(TEXT("[Invalid Transition Index %s for '%s']"), *LexToString(TransitionSource.TransitionIndex.Get()), *StateTree.GetFullName());
 		}
 	}
+	
+	StrBuilder.Appendf(TEXT("go to State '%s'"), CompactState != nullptr ? *UE::StateTreeTrace::GetStateName(StateTree, CompactState) : *TransitionSource.TargetState.Describe());
+
+	if (TransitionSource.Priority != EStateTreeTransitionPriority::None)
+	{
+		StrBuilder.Appendf(TEXT(" (Priority: %s)"), *UEnum::GetDisplayValueAsText(TransitionSource.Priority).ToString()); 
+	}
+
 
 	return StrBuilder.ToString();
 }
diff --git a/Engine/Plugins/Runtime/StateTree/Source/StateTreeModule/Private/StateTreeEvents.cpp b/Engine/Plugins/Runtime/StateTree/Source/StateTreeModule/Private/StateTreeEvents.cpp
--- a/Engine/Plugins/Runtime/StateTree/Source/StateTreeModule/Private/StateTreeEvents.cpp
+++ b/Engine/Plugins/Runtime/StateTree/Source/StateTreeModule/Private/StateTreeEvents.cpp
@@ -30,3 +30,11 @@ void FStateTreeEventQueue::SendEvent(const UObject* Owner, const FGameplayTag& T
 
 	SharedEvents.Emplace(Tag, Payload, Origin);
 }
+
+void FStateTreeEventQueue::ConsumeEvent(const FStateTreeSharedEvent& Event)
+{
+	SharedEvents.RemoveAll([&EventToRemove = Event](const FStateTreeSharedEvent& Event)
+	{
+		return Event == EventToRemove;
+	});
+}
diff --git a/Engine/Plugins/Runtime/StateTree/Source/StateTreeModule/Private/StateTreeExecutionContext.cpp b/Engine/Plugins/Runtime/StateTree/Source/StateTreeModule/Private/StateTreeExecutionContext.cpp
--- a/Engine/Plugins/Runtime/StateTree/Source/StateTreeModule/Private/StateTreeExecutionContext.cpp
+++ b/Engine/Plugins/Runtime/StateTree/Source/StateTreeModule/Private/StateTreeExecutionContext.cpp
@@ -13,6 +13,8 @@
 #include "ProfilingDebugging/CsvProfiler.h"
 #include "Logging/LogScopedVerbosityOverride.h"
 
+#include "Engine/World.h"
+
 #define STATETREE_LOG(Verbosity, Format, ...) UE_VLOG_UELOG(GetOwner(), LogStateTree, Verbosity, TEXT("%s: ") Format, *GetInstanceDescription(), ##__VA_ARGS__)
 #define STATETREE_CLOG(Condition, Verbosity, Format, ...) UE_CVLOG_UELOG((Condition), GetOwner(), LogStateTree, Verbosity, TEXT("%s: ") Format, *GetInstanceDescription(), ##__VA_ARGS__)
 
@@ -113,6 +115,8 @@ FStateTreeExecutionContext::FStateTreeExecutionContext(UObject& InOwner, const U
 
 		InstanceDataStorage = &InstanceData.GetMutableStorage();
 		check(InstanceDataStorage);
+		
+		EventQueue = InstanceData.GetSharedMutableEventQueue();
 	}
 	else
 	{
@@ -133,6 +137,8 @@ FStateTreeExecutionContext::FStateTreeExecutionContext(const FStateTreeExecution
 			const int32 TargetIndex = TargetDataDesc.Handle.DataHandle.GetIndex();
 			ContextAndExternalDataViews[TargetIndex] = InContextToCopy.ContextAndExternalDataViews[TargetIndex];
 		}
+		
+		EventQueue = InstanceData.GetSharedMutableEventQueue();
 	}
 	else
 	{
@@ -403,11 +409,6 @@ EStateTreeRunStatus FStateTreeExecutionContext::Stop(EStateTreeRunStatus Complet
 
 	EStateTreeRunStatus Result = Exec.TreeRunStatus;
 	
-	// Capture events added between ticks.
-	FStateTreeEventQueue& EventQueue = InstanceData.GetMutableEventQueue();
-	EventsToProcess = EventQueue.GetEventsView();
-	EventQueue.Reset();
-
 	// Exit states if still in some valid state.
 	if (Exec.TreeRunStatus == EStateTreeRunStatus::Running)
 	{
@@ -438,14 +439,14 @@ EStateTreeRunStatus FStateTreeExecutionContext::Stop(EStateTreeRunStatus Complet
 	// Destruct all allocated instance data (does not shrink the buffer). This will invalidate Exec too.
 	InstanceData.Reset();
 
+	// External data needs to be recollected if this exec context is reused.
+	bActiveExternalDataCollected = false;
+
 	return Result;
 }
 
-EStateTreeRunStatus FStateTreeExecutionContext::Tick(const float DeltaTime)
+EStateTreeRunStatus FStateTreeExecutionContext::TickPrelude()
 {
-	CSV_SCOPED_TIMING_STAT_EXCLUSIVE(StateTree_Tick);
-	STATETREE_TRACE_SCOPED_PHASE(EStateTreeUpdatePhase::TickStateTree);
-
 	if (!IsValid())
 	{
 		STATETREE_LOG(Warning, TEXT("%hs: StateTree context is not initialized properly ('%s' using StateTree '%s')"),
@@ -460,7 +461,6 @@ EStateTreeRunStatus FStateTreeExecutionContext::Tick(const float DeltaTime)
 		return EStateTreeRunStatus::Failed;
 	}
 
-	FStateTreeEventQueue& EventQueue = InstanceData.GetMutableEventQueue();
 	FStateTreeExecutionState& Exec = GetExecState();
 
 	// No ticking if the tree is done or stopped.
@@ -477,10 +477,88 @@ EStateTreeRunStatus FStateTreeExecutionContext::Tick(const float DeltaTime)
 
 	// From this point any calls to Stop should be deferred.
 	Exec.CurrentPhase = EStateTreeUpdatePhase::TickStateTree;
+
+	return EStateTreeRunStatus::Running;
+}
+
+
+EStateTreeRunStatus FStateTreeExecutionContext::TickPostlude()
+{
+	FStateTreeExecutionState& Exec = GetExecState();
+
+	// Reset phase since we are now safe to stop.
+	Exec.CurrentPhase = EStateTreeUpdatePhase::Unset;
+
+	// Use local for resulting run state since Stop will reset the instance data.
+	EStateTreeRunStatus Result = Exec.TreeRunStatus;
+	
+	if (Exec.RequestedStop != EStateTreeRunStatus::Unset)
+	{
+		STATETREE_LOG_AND_TRACE(VeryVerbose, TEXT("Processing Deferred Stop"));
+		Result = Stop(Exec.RequestedStop);
+	}
+
+	return Result;
+}
+
+EStateTreeRunStatus FStateTreeExecutionContext::Tick(const float DeltaTime)
+{
+	CSV_SCOPED_TIMING_STAT_EXCLUSIVE(StateTree_Tick);
+	STATETREE_TRACE_SCOPED_PHASE(EStateTreeUpdatePhase::TickStateTree);
+
+//	STATETREE_LOG_AND_TRACE(Warning, TEXT("Tick %f"), GetWorld()->GetTimeSeconds());
+
+	const EStateTreeRunStatus PreludeResult = TickPrelude();
+	if (PreludeResult != EStateTreeRunStatus::Running)
+	{
+		return PreludeResult;
+	}
+
+	TickUpdateTasksInternal(DeltaTime);
+	TickTriggerTransitionsInternal();
+
+	return TickPostlude();
+}
+
+EStateTreeRunStatus FStateTreeExecutionContext::TickUpdateTasks(const float DeltaTime)
+{
+	CSV_SCOPED_TIMING_STAT_EXCLUSIVE(StateTree_Tick);
+	STATETREE_TRACE_SCOPED_PHASE(EStateTreeUpdatePhase::TickStateTree);
+
+//	STATETREE_LOG_AND_TRACE(Warning, TEXT("TickUpdateTasks %f"), GetWorld()->GetTimeSeconds());
+
+	const EStateTreeRunStatus PreludeResult = TickPrelude();
+	if (PreludeResult != EStateTreeRunStatus::Running)
+	{
+		return PreludeResult;
+	}
+
+	TickUpdateTasksInternal(DeltaTime);
 	
-	// Capture events added between ticks.
-	EventsToProcess = EventQueue.GetEventsView();
-	EventQueue.Reset();
+	return TickPostlude();
+}
+	
+EStateTreeRunStatus FStateTreeExecutionContext::TickTriggerTransitions()
+{
+	CSV_SCOPED_TIMING_STAT_EXCLUSIVE(StateTree_Tick);
+	STATETREE_TRACE_SCOPED_PHASE(EStateTreeUpdatePhase::TickStateTree);
+
+//	STATETREE_LOG_AND_TRACE(Warning, TEXT("TickTriggerTransitions %f"), GetWorld()->GetTimeSeconds());
+
+	const EStateTreeRunStatus PreludeResult = TickPrelude();
+	if (PreludeResult != EStateTreeRunStatus::Running)
+	{
+		return PreludeResult;
+	}
+
+	TickTriggerTransitionsInternal();
+
+	return TickPostlude();
+}
+
+void FStateTreeExecutionContext::TickUpdateTasksInternal(const float DeltaTime)
+{
+	FStateTreeExecutionState& Exec = GetExecState();
 	
 	// Update the delayed transitions.
 	for (FStateTreeTransitionDelayedState& DelayedState : Exec.DelayedTransitions)
@@ -503,93 +581,73 @@ EStateTreeRunStatus FStateTreeExecutionContext::Tick(const float DeltaTime)
 				StateCompleted();
 			}
 		}
+	}
+	else
+	{
+		STATETREE_TRACE_LOG_EVENT(TEXT("Global tasks completed (%s), stopping the tree"), *UEnum::GetDisplayValueAsText(EvalAndGlobalTaskStatus).ToString());
+		Exec.RequestedStop = EvalAndGlobalTaskStatus;
+	}
+}
 
-		// The state selection is repeated up to MaxIteration time. This allows failed EnterState() to potentially find a new state immediately.
-		// This helps event driven StateTrees to not require another event/tick to find a suitable state.
-		static constexpr int32 MaxIterations = 5;
-		for (int32 Iter = 0; Iter < MaxIterations; Iter++)
-		{
-			// Append events accumulated during the tick, so that transitions can immediately act on them.
-			// We'll consume the events only if they lead to state change below (EnterState is treated the same as Tick),
-			// or let them be processed next frame if no transition.
-			EventsToProcess.Append(EventQueue.GetEventsView());
-
-			// Trigger conditional transitions or state succeed/failed transitions. First tick transition is handled here too.
-			if (TriggerTransitions())
-			{
-				STATETREE_TRACE_SCOPED_PHASE(EStateTreeUpdatePhase::ApplyTransitions);
-				STATETREE_TRACE_TRANSITION_EVENT(NextTransitionSource, EStateTreeTraceEventType::OnTransition);
-				NextTransitionSource.Reset();
-
-				// We have committed to state change, consume events that were accumulated during the tick above.
-				EventQueue.Reset();
-
-				ExitState(NextTransition);
+void FStateTreeExecutionContext::TickTriggerTransitionsInternal()
+{
+	FStateTreeExecutionState& Exec = GetExecState();
 
-				// Tree succeeded or failed.
-				if (NextTransition.TargetState.IsCompletionState())
-				{
-					// Transition to a terminal state (succeeded/failed), or default transition failed.
-					Exec.TreeRunStatus = NextTransition.TargetState.ToCompletionStatus();
+	// If stop is requested, do not try to trigger transitions.
+	if (Exec.RequestedStop != EStateTreeRunStatus::Unset)
+	{
+		return;
+	}
 
-					// Stop evaluators and global tasks.
-					StopEvaluatorsAndGlobalTasks(Exec.TreeRunStatus);
+	// The state selection is repeated up to MaxIteration time. This allows failed EnterState() to potentially find a new state immediately.
+	// This helps event driven StateTrees to not require another event/tick to find a suitable state.
+	static constexpr int32 MaxIterations = 5;
+	for (int32 Iter = 0; Iter < MaxIterations; Iter++)
+	{
+		// Trigger conditional transitions or state succeed/failed transitions. First tick transition is handled here too.
+		if (TriggerTransitions())
+		{
+			STATETREE_TRACE_SCOPED_PHASE(EStateTreeUpdatePhase::ApplyTransitions);
+			STATETREE_TRACE_TRANSITION_EVENT(NextTransitionSource, EStateTreeTraceEventType::OnTransition);
+			NextTransitionSource.Reset();
 
-					// No active states or global tasks anymore, reset frames.
-					Exec.ActiveFrames.Reset();
+			ExitState(NextTransition);
 
-					break;
-				}
+			// Tree succeeded or failed.
+			if (NextTransition.TargetState.IsCompletionState())
+			{
+				// Transition to a terminal state (succeeded/failed), or default transition failed.
+				Exec.TreeRunStatus = NextTransition.TargetState.ToCompletionStatus();
 
-				// Append and consume the events accumulated during the state exit.
-				EventsToProcess.Append(EventQueue.GetEventsView());
-				EventQueue.Reset();
+				// Stop evaluators and global tasks.
+				StopEvaluatorsAndGlobalTasks(Exec.TreeRunStatus);
 
-				// Enter state tasks can fail/succeed, treat it same as tick.
-				const EStateTreeRunStatus LastTickStatus = EnterState(NextTransition);
+				// No active states or global tasks anymore, reset frames.
+				Exec.ActiveFrames.Reset();
 
-				NextTransition.Reset();
+				break;
+			}
 
-				Exec.LastTickStatus = LastTickStatus;
+			// Enter state tasks can fail/succeed, treat it same as tick.
+			const EStateTreeRunStatus LastTickStatus = EnterState(NextTransition);
 
-				// Consider events so far processed. Events sent during EnterState went into EventQueue, and are processed in next iteration.
-				EventsToProcess.Reset();
+			NextTransition.Reset();
 
-				// Report state completed immediately.
-				if (Exec.LastTickStatus != EStateTreeRunStatus::Running)
-				{
-					StateCompleted();
-				}
-			}
+			Exec.LastTickStatus = LastTickStatus;
 
-			// Stop as soon as have found a running state.
-			if (Exec.LastTickStatus == EStateTreeRunStatus::Running)
+			// Report state completed immediately.
+			if (Exec.LastTickStatus != EStateTreeRunStatus::Running)
 			{
-				break;
+				StateCompleted();
 			}
 		}
-	}
-	else
-	{
-		STATETREE_TRACE_LOG_EVENT(TEXT("Global tasks completed (%s), stopping the tree"), *UEnum::GetDisplayValueAsText(EvalAndGlobalTaskStatus).ToString());
-		Exec.RequestedStop = EvalAndGlobalTaskStatus;
-	}
-
-	EventsToProcess.Reset();
-
-	// Reset phase since we are now safe to stop.
-	Exec.CurrentPhase = EStateTreeUpdatePhase::Unset;
 
-	// Use local for resulting run state since Stop will reset the instance data.
-	EStateTreeRunStatus Result = Exec.TreeRunStatus;
-	
-	if (Exec.RequestedStop != EStateTreeRunStatus::Unset)
-	{
-		STATETREE_LOG_AND_TRACE(VeryVerbose, TEXT("Processing Deferred Stop"));
-		Result = Stop(Exec.RequestedStop);
+		// Stop as soon as have found a running state.
+		if (Exec.LastTickStatus == EStateTreeRunStatus::Running)
+		{
+			break;
+		}
 	}
-
-	return Result;
 }
 
 EStateTreeRunStatus FStateTreeExecutionContext::GetStateTreeRunStatus() const
@@ -609,11 +667,6 @@ EStateTreeRunStatus FStateTreeExecutionContext::GetStateTreeRunStatus() const
 	return EStateTreeRunStatus::Failed;
 }
 
-void FStateTreeExecutionContext::SendEvent(const FStateTreeEvent& Event) const
-{
-	SendEvent(Event.Tag, Event.Payload, Event.Origin);
-}
-
 void FStateTreeExecutionContext::SendEvent(const FGameplayTag Tag, const FConstStructView Payload, const FName Origin) const
 {
 	CSV_SCOPED_TIMING_STAT_EXCLUSIVE(StateTree_SendEvent);
@@ -627,8 +680,8 @@ void FStateTreeExecutionContext::SendEvent(const FGameplayTag Tag, const FConstS
 
 	STATETREE_LOG_AND_TRACE(Verbose, TEXT("Send Event '%s'"), *Tag.ToString());
 
-	FStateTreeEventQueue& EventQueue = InstanceData.GetMutableEventQueue();
-	EventQueue.SendEvent(&Owner, Tag, Payload, Origin);
+	FStateTreeEventQueue& LocalEventQueue = InstanceData.GetMutableEventQueue();
+	LocalEventQueue.SendEvent(&Owner, Tag, Payload, Origin);
 }
 
 void FStateTreeExecutionContext::RequestTransition(const FStateTreeTransitionRequest& Request)
@@ -1890,7 +1943,7 @@ EStateTreeRunStatus FStateTreeExecutionContext::TickEvaluatorsAndGlobalTasks(con
 			{
 				// Used to stop ticking tasks after one fails, but we still want to keep updating the data views so that property binding works properly.
 				bool bShouldTickTasks = true;
-				const bool bHasEvents = !EventsToProcess.IsEmpty();
+				const bool bHasEvents = EventQueue && EventQueue->HasEvents();
 
 				for (int32 TaskIndex = CurrentStateTree->GlobalTasksBegin; TaskIndex < (CurrentStateTree->GlobalTasksBegin + CurrentStateTree->GlobalTasksNum); TaskIndex++)
 				{
@@ -2316,7 +2369,7 @@ EStateTreeRunStatus FStateTreeExecutionContext::TickTasks(const float DeltaTime)
 	EStateTreeRunStatus Result = EStateTreeRunStatus::Running;
 	int32 NumTotalTasks = 0;
 
-	const bool bHasEvents = !EventsToProcess.IsEmpty();
+	const bool bHasEvents = EventQueue && EventQueue->HasEvents();
 
 	Exec.CompletedFrameIndex = FStateTreeIndex16::Invalid;
 	Exec.CompletedStateHandle = FStateTreeStateHandle::Invalid;
@@ -2522,18 +2575,45 @@ bool FStateTreeExecutionContext::TestAllConditions(const FStateTreeExecutionFram
 
 FString FStateTreeExecutionContext::DebugGetEventsAsString() const
 {
-	FString Result;
-	for (const FStateTreeSharedEvent& Event : EventsToProcess)
+	FStringBuilderBase StrBuilder;
+	
+	if (EventQueue)
 	{
-		if (!Result.IsEmpty())
+		for (const FStateTreeSharedEvent& Event : EventQueue->GetEventsView())
 		{
-			Result += TEXT(", ");
-		}
+			if (Event.IsValid())
+			{
+				if (StrBuilder.Len() > 0)
+				{
+					StrBuilder += TEXT(", ");
+				}
 
-		check(Event.IsValid());
-		Result += Event->Tag.ToString();
+				const bool bHasTag = Event->Tag.IsValid();
+				const bool bHasPayload = Event->Payload.GetScriptStruct() != nullptr;
+			
+				if (bHasTag || bHasPayload)
+				{
+					StrBuilder.Appendf(TEXT("("));
+				
+					if (bHasTag)
+					{
+						StrBuilder.Appendf(TEXT("Tag: '%s'"), *Event->Tag.ToString()); 
+					}
+					if (bHasTag && bHasPayload)
+					{
+						StrBuilder.Appendf(TEXT(", "));
+					}
+					if (bHasPayload)
+					{
+						StrBuilder.Appendf(TEXT(" Payload: '%s'"), *Event->Payload.GetScriptStruct()->GetName()); 
+					}
+					StrBuilder.Appendf(TEXT(") "));
+				}
+			}
+		}
 	}
-	return Result;
+	
+	return StrBuilder.ToString();
 }
 
 bool FStateTreeExecutionContext::RequestTransition(
@@ -2570,6 +2650,25 @@ bool FStateTreeExecutionContext::RequestTransition(
 		NextTransition.NextActiveFrames = StateSelectionResult.GetSelectedFrames();
 		NextTransition.NextActiveFrameEvents = StateSelectionResult.GetFramesStateSelectionEvents();
 
+		// Consume events from states, if required. 
+		for (int32 FrameIndex = 0; FrameIndex < NextTransition.NextActiveFrames.Num(); FrameIndex++)
+		{
+			const FStateTreeExecutionFrame& Frame = NextTransition.NextActiveFrames[FrameIndex]; 
+			const FStateTreeFrameStateSelectionEvents& FrameEvents = NextTransition.NextActiveFrameEvents[FrameIndex];
+
+			for (int32 StateIndex = 0; StateIndex < Frame.ActiveStates.Num(); StateIndex++)
+			{
+				if (FrameEvents.Events[StateIndex].IsValid())
+				{
+					const FCompactStateTreeState& State = Frame.StateTree->States[StateIndex];
+					if (State.bConsumeEventOnSelect)
+					{
+						 ConsumeEvent(FrameEvents.Events[StateIndex]);
+					}
+				}
+			}
+		}
+		
 		STATETREE_LOG(Verbose, TEXT("Transition on state '%s' -[%s]-> state '%s'"),
 			*GetSafeStateName(CurrentFrame, CurrentFrame.ActiveStates.Last()),
 			*GetSafeStateName(CurrentFrame, NextState),
@@ -2614,9 +2713,9 @@ bool FStateTreeExecutionContext::TriggerTransitions()
 	FAllowDirectTransitionsScope AllowDirectTransitionsScope(*this); // Set flag for the scope of this function to allow direct transitions without buffering.
 	FStateTreeExecutionState& Exec = GetExecState();
 
-	if (EventsToProcess.Num() > 0)
+	if (EventQueue && EventQueue->HasEvents())
 	{
-		STATETREE_LOG_AND_TRACE(Verbose, TEXT("Trigger transitions with events [%s]"), *DebugGetEventsAsString());
+		STATETREE_LOG_AND_TRACE(Verbose, TEXT("Trigger transitions with events: %s"), *DebugGetEventsAsString());
 	}
 
 	NextTransition.Reset();
@@ -2641,18 +2740,49 @@ bool FStateTreeExecutionContext::TriggerTransitions()
 	InstanceData.ResetTransitionRequests();
 
 	//
-	// Check tick, event, and task based transitions first.
+	// Collect tick, event, and task based transitions.
 	//
+	struct FTransitionHandler
+	{
+		FTransitionHandler() = default;
+		
+		FTransitionHandler(const uint8 InFrameIndex, const FStateTreeStateHandle InStateHandle, const EStateTreeTransitionPriority InPriority)
+			: StateHandle(InStateHandle)
+			, TaskIndex(FStateTreeIndex16::Invalid)
+			, FrameIndex(InFrameIndex)
+			, Priority(InPriority)
+		{
+		}
+
+		FTransitionHandler(const uint8 InFrameIndex, const FStateTreeStateHandle InStateHandle, const FStateTreeIndex16 InTaskIndex, const EStateTreeTransitionPriority InPriority)
+			: StateHandle(InStateHandle)
+			, TaskIndex(InTaskIndex)
+			, FrameIndex(InFrameIndex)
+			, Priority(InPriority)
+		{
+		}
+
+		FStateTreeStateHandle StateHandle;
+		FStateTreeIndex16 TaskIndex = FStateTreeIndex16::Invalid;
+		uint8 FrameIndex = 0;
+		EStateTreeTransitionPriority Priority = EStateTreeTransitionPriority::Normal;
+
+		bool operator<(const FTransitionHandler& Other) const
+		{
+			// Highest priority first.
+			return Priority > Other.Priority;
+		}
+	};
+
+	TArray<FTransitionHandler, TInlineAllocator<16>> TransitionHandlers;
+
 	if (Exec.ActiveFrames.Num() > 0)
 	{
-		for (int32 FrameIndex = 0; FrameIndex < Exec.ActiveFrames.Num(); FrameIndex++)
+		for (int32 FrameIndex = Exec.ActiveFrames.Num() - 1; FrameIndex >= 0; FrameIndex--)
 		{
-			FStateTreeExecutionFrame* CurrentParentFrame = FrameIndex > 0 ? &Exec.ActiveFrames[FrameIndex - 1] : nullptr;
 			FStateTreeExecutionFrame& CurrentFrame = Exec.ActiveFrames[FrameIndex];
 			const UStateTree* CurrentStateTree = CurrentFrame.StateTree;
 
-			FCurrentlyProcessedFrameScope FrameScope(*this, CurrentParentFrame, CurrentFrame);
-
 			for (int32 StateIndex = CurrentFrame.ActiveStates.Num() - 1; StateIndex >= 0; StateIndex--)
 			{
 				const FStateTreeStateHandle StateHandle = CurrentFrame.ActiveStates[StateIndex];
@@ -2663,205 +2793,239 @@ bool FStateTreeExecutionContext::TriggerTransitions()
 				{
 					continue;
 				}
-				
-				FCurrentlyProcessedStateScope StateScope(*this, StateHandle);
-				STATETREE_TRACE_SCOPED_STATE(StateHandle);
 
+				// Transition tasks.
 				if (State.bHasTransitionTasks)
 				{
-					STATETREE_CLOG(State.TasksNum > 0, VeryVerbose, TEXT("%*sTrigger task transitions in state '%s'"), StateIndex*UE::StateTree::DebugIndentSize, TEXT(""), *DebugGetStatePath(Exec.ActiveFrames, &CurrentFrame, StateIndex));
-
 					for (int32 TaskIndex = (State.TasksBegin + State.TasksNum) - 1; TaskIndex >= State.TasksBegin; TaskIndex--)
 					{
 						const FStateTreeTaskBase& Task = CurrentStateTree->Nodes[TaskIndex].Get<const FStateTreeTaskBase>();
-						const FStateTreeDataView TaskInstanceView = GetDataView(CurrentParentFrame, CurrentFrame, Task.InstanceDataHandle);
-						FNodeInstanceDataScope DataScope(*this, Task.InstanceDataHandle, TaskInstanceView);
-
-						// Ignore disabled task
-						if (Task.bTaskEnabled == false)
+						if (Task.bShouldAffectTransitions)
 						{
-							STATETREE_LOG(VeryVerbose, TEXT("%*sSkipped 'TriggerTransitions' for disabled Task: '%s'"), UE::StateTree::DebugIndentSize, TEXT(""), *Task.Name.ToString());
-							continue;
+							TransitionHandlers.Emplace((uint8)FrameIndex, StateHandle, FStateTreeIndex16(TaskIndex), Task.TransitionHandlingPriority);
 						}
+					}
+				}
+				
+				// Regular transitions on state
+				if (State.TransitionsNum > 0)
+				{
+					TransitionHandlers.Emplace((uint8)FrameIndex, StateHandle, EStateTreeTransitionPriority::Normal);
+				}
+			}
 
+			if (CurrentFrame.bIsGlobalFrame)
+			{
+				// Global transition tasks.
+				if (CurrentFrame.StateTree->bHasGlobalTransitionTasks)
+				{
+					for (int32 TaskIndex = (CurrentStateTree->GlobalTasksBegin + CurrentStateTree->GlobalTasksNum) - 1; TaskIndex >= CurrentFrame.StateTree->GlobalTasksBegin; TaskIndex--)
+					{
+						const FStateTreeTaskBase& Task = CurrentStateTree->Nodes[TaskIndex].Get<const FStateTreeTaskBase>();
 						if (Task.bShouldAffectTransitions)
 						{
-							STATETREE_LOG(VeryVerbose, TEXT("%*sTriggerTransitions: '%s'"), UE::StateTree::DebugIndentSize, TEXT(""), *Task.Name.ToString());
-							STATETREE_TRACE_TASK_EVENT(TaskIndex, TaskInstanceView, EStateTreeTraceEventType::OnEvaluating, EStateTreeRunStatus::Running);
-							check(TaskInstanceView.IsValid());
-							Task.TriggerTransitions(*this);
+							TransitionHandlers.Emplace((uint8)FrameIndex, FStateTreeStateHandle(), FStateTreeIndex16(TaskIndex), Task.TransitionHandlingPriority);
 						}
 					}
+				}				
+			}
+		}
+
+		// Sort by priority and adding order.
+		TransitionHandlers.StableSort();
+	}
+
+	//
+	// Process task and state transitions in priority order. 
+	//
+	for (const FTransitionHandler& Handler : TransitionHandlers)
+	{
+		const int32 FrameIndex = Handler.FrameIndex;
+		FStateTreeExecutionFrame* CurrentParentFrame = FrameIndex > 0 ? &Exec.ActiveFrames[FrameIndex - 1] : nullptr;
+		FStateTreeExecutionFrame& CurrentFrame = Exec.ActiveFrames[FrameIndex];
+		const UStateTree* CurrentStateTree = CurrentFrame.StateTree;
+
+		FCurrentlyProcessedFrameScope FrameScope(*this, CurrentParentFrame, CurrentFrame);
+		FCurrentlyProcessedStateScope StateScope(*this, Handler.StateHandle);
+		STATETREE_TRACE_SCOPED_STATE(Handler.StateHandle);
+
+		if (Handler.TaskIndex.IsValid())
+		{
+			const FStateTreeTaskBase& Task = CurrentStateTree->Nodes[Handler.TaskIndex.Get()].Get<const FStateTreeTaskBase>();
+
+			// Ignore disabled task
+			if (Task.bTaskEnabled == false)
+			{
+				STATETREE_LOG(VeryVerbose, TEXT("%*sSkipped 'TriggerTransitions' for disabled Task: '%s'"), UE::StateTree::DebugIndentSize, TEXT(""), *Task.Name.ToString());
+				continue;
+			}
+
+			const FStateTreeDataView TaskInstanceView = GetDataView(CurrentParentFrame, CurrentFrame, Task.InstanceDataHandle);
+			FNodeInstanceDataScope DataScope(*this, Task.InstanceDataHandle, TaskInstanceView);
+
+			STATETREE_LOG(VeryVerbose, TEXT("%*sTriggerTransitions: '%s'"), UE::StateTree::DebugIndentSize, TEXT(""), *Task.Name.ToString());
+			STATETREE_TRACE_TASK_EVENT(Handler.TaskIndex.Get(), TaskInstanceView, EStateTreeTraceEventType::OnEvaluating, EStateTreeRunStatus::Running);
+			check(TaskInstanceView.IsValid());
+			Task.TriggerTransitions(*this);
+		}
+		else if (Handler.StateHandle.IsValid())
+		{
+			const FCompactStateTreeState& State = CurrentStateTree->States[Handler.StateHandle.Index];
+
+			// Transitions
+			for (uint8 i = 0; i < State.TransitionsNum; i++)
+			{
+				// All transition conditions must pass
+				const int16 TransitionIndex = State.TransitionsBegin + i;
+				const FCompactStateTransition& Transition = CurrentStateTree->Transitions[TransitionIndex];
+
+				// Skip disabled transitions
+				if (Transition.bTransitionEnabled == false)
+				{
+					continue;
 				}
 				
-				for (uint8 i = 0; i < State.TransitionsNum; i++)
+				// No need to test the transition if same or higher priority transition has already been processed.
+				if (Transition.Priority <= NextTransition.Priority)
 				{
-					// All transition conditions must pass
-					const int16 TransitionIndex = State.TransitionsBegin + i;
-					const FCompactStateTransition& Transition = CurrentStateTree->Transitions[TransitionIndex];
+					continue;
+				}
 
-					// Skip disabled transitions
-					if (Transition.bTransitionEnabled == false)
-					{
-						continue;
-					}
-					
-					// No need to test the transition if same or higher priority transition has already been processed.
-					if (Transition.Priority <= NextTransition.Priority)
-					{
-						continue;
-					}
+				// Skip completion transitions
+				if (EnumHasAnyFlags(Transition.Trigger, EStateTreeTransitionTrigger::OnStateCompleted))
+				{
+					continue;
+				}
 
-					// Skip completion transitions
-					if (EnumHasAnyFlags(Transition.Trigger, EStateTreeTransitionTrigger::OnStateCompleted))
+				// If a delayed transition has passed the delay, and remove it from the queue, and try trigger it.
+				if (Transition.HasDelay())
+				{
+					bool bTriggeredDelayedTransition = false;
+					const TArray<FStateTreeTransitionDelayedState, TInlineAllocator<8>> ExpiredDelayedStates = Exec.FindAndRemoveExpiredDelayedTransitions(CurrentFrame.StateTree, FStateTreeIndex16(TransitionIndex));
+					for (const FStateTreeTransitionDelayedState& DelayedState : ExpiredDelayedStates)
 					{
-						continue;
-					}
+						STATETREE_LOG(Verbose, TEXT("Passed delayed transition from '%s' (%s) -> '%s'"),
+							*GetSafeStateName(CurrentFrame, CurrentFrame.ActiveStates.Last()), *State.Name.ToString(), *GetSafeStateName(CurrentFrame, Transition.State));
 
-					// If a delayed transition has passed the delay, and remove it from the queue, and try trigger it.
-					if (Transition.HasDelay())
-					{
-						bool bTriggeredDelayedTransition = false;
-						const TArray<FStateTreeTransitionDelayedState, TInlineAllocator<8>> ExpiredDelayedStates = Exec.FindAndRemoveExpiredDelayedTransitions(CurrentFrame.StateTree, FStateTreeIndex16(TransitionIndex));
-						for (const FStateTreeTransitionDelayedState& DelayedState : ExpiredDelayedStates)
+						// Trigger Delayed Transition when the delay has passed.
+						if (RequestTransition(CurrentFrame, Transition.State, Transition.Priority, &DelayedState.CapturedEvent, Transition.Fallback))
 						{
-							STATETREE_LOG(Verbose, TEXT("Passed delayed transition from '%s' (%s) -> '%s'"),
-								*GetSafeStateName(CurrentFrame, CurrentFrame.ActiveStates.Last()), *State.Name.ToString(), *GetSafeStateName(CurrentFrame, Transition.State));
-
-							// Trigger Delayed Transition when the delay has passed.
-							if (RequestTransition(CurrentFrame, Transition.State, Transition.Priority, &DelayedState.CapturedEvent, Transition.Fallback))
+							// If the transition was successfully requested with a specific event, consume and remove the event, it's been used.  
+							if (DelayedState.CapturedEvent.IsValid() && Transition.bConsumeEventOnSelect)
 							{
-								NextTransitionSource = FStateTreeTransitionSource(FStateTreeIndex16(TransitionIndex), Transition.State, Transition.Priority);
-								bTriggeredDelayedTransition = true;
-								break;
+								ConsumeEvent(DelayedState.CapturedEvent);
 							}
-						}
 
-						if (bTriggeredDelayedTransition)
-						{
-							continue;
+							NextTransitionSource = FStateTreeTransitionSource(FStateTreeIndex16(TransitionIndex), Transition.State, Transition.Priority);
+							bTriggeredDelayedTransition = true;
+							break;
 						}
 					}
 
-					TArray<const FStateTreeSharedEvent*, TInlineAllocator<8>> TransitionEvents;
-
-					if (Transition.Trigger == EStateTreeTransitionTrigger::OnEvent)
+					if (bTriggeredDelayedTransition)
 					{
-						check(Transition.RequiredEvent.IsValid());
+						continue;
+					}
+				}
 
-						TConstArrayView<FStateTreeSharedEvent> EventsQueue = GetEventsToProcessView();
-						for (const FStateTreeSharedEvent& Event : EventsQueue)
+				TArray<const FStateTreeSharedEvent*, TInlineAllocator<8>> TransitionEvents;
+
+				if (Transition.Trigger == EStateTreeTransitionTrigger::OnEvent)
+				{
+					check(Transition.RequiredEvent.IsValid());
+
+					TConstArrayView<FStateTreeSharedEvent> EventsQueue = GetEventsToProcessView();
+					for (const FStateTreeSharedEvent& Event : EventsQueue)
+					{
+						check(Event.IsValid());
+						if (Transition.RequiredEvent.DoesEventMatchDesc(*Event))
 						{
-							check(Event.IsValid());
-							if (Transition.RequiredEvent.DoesEventMatchDesc(*Event))
-							{
-								TransitionEvents.Emplace(&Event);
-							}
+							TransitionEvents.Emplace(&Event);
 						}
 					}
-					else if (Transition.Trigger == EStateTreeTransitionTrigger::OnTick)
+				}
+				else if (Transition.Trigger == EStateTreeTransitionTrigger::OnTick)
+				{
+					TransitionEvents.Emplace();
+				}
+				
+				for (const FStateTreeSharedEvent* TransitionEvent : TransitionEvents)
+				{
+					bool bPassed = false; 
 					{
-						TransitionEvents.Emplace();
+						FCurrentlyProcessedTransitionEventScope TransitionEventScope(*this, TransitionEvent ? TransitionEvent->Get() : nullptr);
+						STATETREE_TRACE_TRANSITION_EVENT(FStateTreeTransitionSource(FStateTreeIndex16(TransitionIndex), Transition.State, Transition.Priority), EStateTreeTraceEventType::OnEvaluating);
+						STATETREE_TRACE_SCOPED_PHASE(EStateTreeUpdatePhase::TransitionConditions);
+						bPassed = TestAllConditions(CurrentParentFrame, CurrentFrame, Transition.ConditionsBegin, Transition.ConditionsNum);
 					}
-					
-					for (const FStateTreeSharedEvent* TransitionEvent : TransitionEvents)
+
+					if (bPassed)
 					{
-						bool bPassed = false; 
+						// If the transitions is delayed, set up the delay. 
+						if (Transition.HasDelay())
 						{
-							FCurrentlyProcessedTransitionEventScope TransitionEventScope(*this, TransitionEvent ? TransitionEvent->Get() : nullptr);
-							STATETREE_TRACE_TRANSITION_EVENT(FStateTreeTransitionSource(FStateTreeIndex16(TransitionIndex), Transition.State, Transition.Priority), EStateTreeTraceEventType::OnEvaluating);
-							STATETREE_TRACE_SCOPED_PHASE(EStateTreeUpdatePhase::TransitionConditions);
-							bPassed = TestAllConditions(CurrentParentFrame, CurrentFrame, Transition.ConditionsBegin, Transition.ConditionsNum);
-						}
+							uint32 TransitionEventHash = 0u;
+							if (TransitionEvent && TransitionEvent->IsValid())
+							{		
+								TransitionEventHash = GetTypeHash(*TransitionEvent->Get());
+							}
 
-						if (bPassed)
-						{
-							// If the transitions is delayed, set up the delay. 
-							if (Transition.HasDelay())
+							const bool bIsDelayedTransitionExisting = Exec.DelayedTransitions.ContainsByPredicate([StateTree = CurrentFrame.StateTree, TransitionIndex, TransitionEventHash](const FStateTreeTransitionDelayedState& DelayedState)
 							{
-								uint32 TransitionEventHash = 0u;
-								if (TransitionEvent && TransitionEvent->IsValid())
-								{		
-									TransitionEventHash = GetTypeHash(*TransitionEvent->Get());
-								}
+								return DelayedState.StateTree == StateTree && DelayedState.TransitionIndex.Get() == TransitionIndex && DelayedState.CapturedEventHash == TransitionEventHash;
+							});
 
-								const bool bIsDelayedTransitionExisting = Exec.DelayedTransitions.ContainsByPredicate([StateTree = CurrentFrame.StateTree, TransitionIndex, TransitionEventHash](const FStateTreeTransitionDelayedState& DelayedState)
-								{
-									return DelayedState.StateTree == StateTree && DelayedState.TransitionIndex.Get() == TransitionIndex && DelayedState.CapturedEventHash == TransitionEventHash;
-								});
-
-								if (!bIsDelayedTransitionExisting)
-								{
-									// Initialize new delayed transition.
-									const float DelayDuration = Transition.Delay.GetRandomDuration();
-									if (DelayDuration > 0.0f)
-									{
-										FStateTreeTransitionDelayedState& DelayedState = Exec.DelayedTransitions.AddDefaulted_GetRef();
-										DelayedState.StateTree = CurrentFrame.StateTree;
-										DelayedState.TransitionIndex = FStateTreeIndex16(TransitionIndex);
-										DelayedState.TimeLeft = DelayDuration;
-										DelayedState.CapturedEvent = *TransitionEvent;
-										DelayedState.CapturedEventHash = TransitionEventHash;
-
-										BeginDelayedTransition(DelayedState);
-										STATETREE_LOG(Verbose, TEXT("Delayed transition triggered from '%s' (%s) -> '%s' %.1fs"),
-											*GetSafeStateName(CurrentFrame, CurrentFrame.ActiveStates.Last()), *State.Name.ToString(), *GetSafeStateName(CurrentFrame, Transition.State), DelayedState.TimeLeft);
-									
-										// Delay state added, skip requesting the transition.
-										continue;
-									}
-									// Fallthrough to request transition if duration was zero. 
-								}
-								else
+							if (!bIsDelayedTransitionExisting)
+							{
+								// Initialize new delayed transition.
+								const float DelayDuration = Transition.Delay.GetRandomDuration();
+								if (DelayDuration > 0.0f)
 								{
-									// We get here if the transitions re-triggers during the delay, on which case we'll just ignore it.
+									FStateTreeTransitionDelayedState& DelayedState = Exec.DelayedTransitions.AddDefaulted_GetRef();
+									DelayedState.StateTree = CurrentFrame.StateTree;
+									DelayedState.TransitionIndex = FStateTreeIndex16(TransitionIndex);
+									DelayedState.TimeLeft = DelayDuration;
+									DelayedState.CapturedEvent = *TransitionEvent;
+									DelayedState.CapturedEventHash = TransitionEventHash;
+
+									BeginDelayedTransition(DelayedState);
+									STATETREE_LOG(Verbose, TEXT("Delayed transition triggered from '%s' (%s) -> '%s' %.1fs"),
+										*GetSafeStateName(CurrentFrame, CurrentFrame.ActiveStates.Last()), *State.Name.ToString(), *GetSafeStateName(CurrentFrame, Transition.State), DelayedState.TimeLeft);
+								
+									// Delay state added, skip requesting the transition.
 									continue;
 								}
+								// Fallthrough to request transition if duration was zero. 
 							}
-
-							if (RequestTransition(CurrentFrame, Transition.State, Transition.Priority, TransitionEvent, Transition.Fallback))
+							else
 							{
-								NextTransitionSource = FStateTreeTransitionSource(FStateTreeIndex16(TransitionIndex), Transition.State, Transition.Priority);
-								break;
+								// We get here if the transitions re-triggers during the delay, on which case we'll just ignore it.
+								continue;
 							}
 						}
-					}
-				}
-			}
 
-			if (CurrentFrame.bIsGlobalFrame)
-			{
-				// Global frame
-				if (CurrentFrame.StateTree->bHasGlobalTransitionTasks)
-				{
-					STATETREE_LOG(VeryVerbose, TEXT("Trigger global task transitions"));
-					for (int32 TaskIndex = (CurrentStateTree->GlobalTasksBegin + CurrentStateTree->GlobalTasksNum) - 1; TaskIndex >= CurrentFrame.StateTree->GlobalTasksBegin; TaskIndex--)
-					{
-						const FStateTreeTaskBase& Task =  CurrentStateTree->Nodes[TaskIndex].Get<const FStateTreeTaskBase>();
-						const FStateTreeDataView TaskInstanceView = GetDataView(CurrentParentFrame, CurrentFrame, Task.InstanceDataHandle);
-						FNodeInstanceDataScope DataScope(*this, Task.InstanceDataHandle, TaskInstanceView);
-
-						// Ignore disabled task
-						if (Task.bTaskEnabled == false)
-						{
-							STATETREE_LOG(VeryVerbose, TEXT("%*sSkipped 'TriggerTransitions' for disabled Task: '%s'"), UE::StateTree::DebugIndentSize, TEXT(""), *Task.Name.ToString());
-							continue;
-						}
-
-						if (Task.bShouldAffectTransitions)
+						if (RequestTransition(CurrentFrame, Transition.State, Transition.Priority, TransitionEvent, Transition.Fallback))
 						{
-							STATETREE_LOG(VeryVerbose, TEXT("%*sTriggerTransitions: '%s'"), UE::StateTree::DebugIndentSize, TEXT(""), *Task.Name.ToString());
-							STATETREE_TRACE_TASK_EVENT(TaskIndex, TaskInstanceView, EStateTreeTraceEventType::OnEvaluating, EStateTreeRunStatus::Running);
-							check(TaskInstanceView.IsValid());
-							Task.TriggerTransitions(*this);
+							// If the transition was successfully requested with a specific event, consume and remove the event, it's been used.  
+							if (TransitionEvent && Transition.bConsumeEventOnSelect)
+							{
+								ConsumeEvent(*TransitionEvent);
+							}
+							
+							NextTransitionSource = FStateTreeTransitionSource(FStateTreeIndex16(TransitionIndex), Transition.State, Transition.Priority);
+							break;
 						}
 					}
 				}
-				
 			}
 		}
 	}
 
+	// All events have had the change to be reacted to, clear the event queue (if this instance owns it).
+	if (InstanceData.IsOwningEventQueue())
+	{
+		EventQueue->Reset();
+	}
 
 	//
 	// Check state completion transitions.
diff --git a/Engine/Plugins/Runtime/StateTree/Source/StateTreeModule/Private/StateTreeInstanceData.cpp b/Engine/Plugins/Runtime/StateTree/Source/StateTreeModule/Private/StateTreeInstanceData.cpp
--- a/Engine/Plugins/Runtime/StateTree/Source/StateTreeModule/Private/StateTreeInstanceData.cpp
+++ b/Engine/Plugins/Runtime/StateTree/Source/StateTreeModule/Private/StateTreeInstanceData.cpp
@@ -83,6 +83,60 @@ namespace UE::StateTree
 // FStateTreeInstanceStorage
 //----------------------------------------------------------------//
 
+FStateTreeInstanceStorage::FStateTreeInstanceStorage() = default;
+
+FStateTreeInstanceStorage::FStateTreeInstanceStorage(const FStateTreeInstanceStorage& Other)
+	: InstanceStructs(Other.InstanceStructs)
+	, ExecutionState(Other.ExecutionState)
+	, TemporaryInstances(Other.TemporaryInstances)
+	, EventQueue(MakeShared<FStateTreeEventQueue>(*Other.EventQueue))
+	, TransitionRequests(Other.TransitionRequests)
+	, GlobalParameters(Other.GlobalParameters)
+{
+}
+
+FStateTreeInstanceStorage::FStateTreeInstanceStorage(FStateTreeInstanceStorage&& Other) noexcept
+	: InstanceStructs(MoveTemp(Other.InstanceStructs))
+	, ExecutionState(MoveTemp(Other.ExecutionState))
+	, TemporaryInstances(MoveTemp(Other.TemporaryInstances))
+	, EventQueue(Other.EventQueue)
+	, TransitionRequests(MoveTemp(Other.TransitionRequests))
+	, GlobalParameters(MoveTemp(Other.GlobalParameters))
+{
+	Other.EventQueue = MakeShared<FStateTreeEventQueue>();
+}
+
+FStateTreeInstanceStorage& FStateTreeInstanceStorage::operator=(const FStateTreeInstanceStorage& Other)
+{
+	InstanceStructs = Other.InstanceStructs;
+	ExecutionState = Other.ExecutionState;
+	TemporaryInstances = Other.TemporaryInstances;
+	EventQueue = MakeShared<FStateTreeEventQueue>(*Other.EventQueue);
+	TransitionRequests = Other.TransitionRequests;
+	GlobalParameters = Other.GlobalParameters;
+
+	return *this;
+}
+
+FStateTreeInstanceStorage& FStateTreeInstanceStorage::operator=(FStateTreeInstanceStorage&& Other) noexcept
+{
+	InstanceStructs = MoveTemp(Other.InstanceStructs);
+	ExecutionState = MoveTemp(Other.ExecutionState);
+	TemporaryInstances = MoveTemp(Other.TemporaryInstances);
+	EventQueue = Other.EventQueue;
+	Other.EventQueue = MakeShared<FStateTreeEventQueue>();
+	TransitionRequests = MoveTemp(Other.TransitionRequests);
+	GlobalParameters = MoveTemp(Other.GlobalParameters);
+
+	return *this;
+}
+
+void FStateTreeInstanceStorage::SetSharedEventQueue(const TSharedRef<FStateTreeEventQueue>& InSharedEventQueue)
+{
+	EventQueue = InSharedEventQueue;
+	bIsOwningEventQueue = false;
+}
+
 void FStateTreeInstanceStorage::AddTransitionRequest(const UObject* Owner, const FStateTreeTransitionRequest& Request)
 {
 	constexpr int32 MaxPendingTransitionRequests = 32;
@@ -197,6 +251,24 @@ void FStateTreeInstanceStorage::SetGlobalParameters(const FInstancedPropertyBag&
 	GlobalParameters = Parameters;
 }
 
+void FStateTreeInstanceStorage::AddStructReferencedObjects(FReferenceCollector& Collector)
+{
+	Collector.AddPropertyReferencesWithStructARO(TBaseStructure<FStateTreeInstanceStorage>::Get(), this);
+	Collector.AddPropertyReferencesWithStructARO(TBaseStructure<FStateTreeEventQueue>::Get(), &EventQueue.Get());
+}
+
+void FStateTreeInstanceStorage::Reset()
+{
+	InstanceStructs.Reset();
+	if (bIsOwningEventQueue)
+	{
+		EventQueue->Reset();
+	}
+	ExecutionState.Reset();
+	TemporaryInstances.Reset();
+	GlobalParameters.Reset();
+}
+
 //----------------------------------------------------------------//
 // FStateTreeInstanceData
 //----------------------------------------------------------------//
@@ -208,7 +280,7 @@ FStateTreeInstanceData::FStateTreeInstanceData(const FStateTreeInstanceData& Oth
 	InstanceStorage = MakeShared<FStateTreeInstanceStorage>(*Other.InstanceStorage);
 }
 
-FStateTreeInstanceData::FStateTreeInstanceData(FStateTreeInstanceData&& Other)
+FStateTreeInstanceData::FStateTreeInstanceData(FStateTreeInstanceData&& Other) noexcept
 {
 	InstanceStorage = Other.InstanceStorage;
 	Other.InstanceStorage = MakeShared<FStateTreeInstanceStorage>();
@@ -220,7 +292,7 @@ FStateTreeInstanceData& FStateTreeInstanceData::operator=(const FStateTreeInstan
 	return *this;
 }
 
-FStateTreeInstanceData& FStateTreeInstanceData::operator=(FStateTreeInstanceData&& Other)
+FStateTreeInstanceData& FStateTreeInstanceData::operator=(FStateTreeInstanceData&& Other) noexcept
 {
 	InstanceStorage = Other.InstanceStorage;
 	Other.InstanceStorage = MakeShared<FStateTreeInstanceStorage>();
@@ -254,12 +326,27 @@ FStateTreeInstanceStorage& FStateTreeInstanceData::GetMutableStorage()
 
 FStateTreeEventQueue& FStateTreeInstanceData::GetMutableEventQueue()
 {
-	return GetMutableStorage().EventQueue;	
+	return GetMutableStorage().GetMutableEventQueue();
+}
+
+const TSharedRef<FStateTreeEventQueue>& FStateTreeInstanceData::GetSharedMutableEventQueue()
+{
+	return GetMutableStorage().GetSharedMutableEventQueue();
 }
 
 const FStateTreeEventQueue& FStateTreeInstanceData::GetEventQueue() const
 {
-	return GetStorage().EventQueue;
+	return GetStorage().GetEventQueue();
+}
+
+bool FStateTreeInstanceData::IsOwningEventQueue() const
+{
+	return GetStorage().IsOwningEventQueue();
+}
+
+void FStateTreeInstanceData::SetSharedEventQueue(const TSharedRef<FStateTreeEventQueue>& InSharedEventQueue)
+{
+	return GetMutableStorage().SetSharedEventQueue(InSharedEventQueue);
 }
 
 void FStateTreeInstanceData::AddTransitionRequest(const UObject* Owner, const FStateTreeTransitionRequest& Request)
@@ -378,7 +465,7 @@ bool FStateTreeInstanceData::Identical(const FStateTreeInstanceData* Other, uint
 
 void FStateTreeInstanceData::AddStructReferencedObjects(FReferenceCollector& Collector)
 {
-	Collector.AddPropertyReferencesWithStructARO(FStateTreeInstanceStorage::StaticStruct(), &GetMutableStorage());
+	GetMutableStorage().AddStructReferencedObjects(Collector);
 }
 
 bool FStateTreeInstanceData::Serialize(FArchive& Ar)
@@ -553,10 +640,6 @@ void FStateTreeInstanceData::ShrinkTo(const int32 NumStructs)
 
 void FStateTreeInstanceData::Reset()
 {
-	FStateTreeInstanceStorage& Storage = GetMutableStorage();
-	Storage.InstanceStructs.Reset();
-	Storage.EventQueue.Reset();
-	Storage.ExecutionState.Reset();
-	Storage.TemporaryInstances.Reset();
-	Storage.GlobalParameters.Reset();
+	GetMutableStorage().Reset();
+
 }
diff --git a/Engine/Plugins/Runtime/StateTree/Source/StateTreeModule/Private/Tasks/StateTreeRunParallelStateTreeTask.cpp b/Engine/Plugins/Runtime/StateTree/Source/StateTreeModule/Private/Tasks/StateTreeRunParallelStateTreeTask.cpp
--- a/Engine/Plugins/Runtime/StateTree/Source/StateTreeModule/Private/Tasks/StateTreeRunParallelStateTreeTask.cpp
+++ b/Engine/Plugins/Runtime/StateTree/Source/StateTreeModule/Private/Tasks/StateTreeRunParallelStateTreeTask.cpp
@@ -10,6 +10,7 @@ FStateTreeRunParallelStateTreeTask::FStateTreeRunParallelStateTreeTask()
 {
 	bShouldCopyBoundPropertiesOnTick = false;
 	bShouldCopyBoundPropertiesOnExitState = false;
+	bShouldAffectTransitions = true;
 }
 
 EStateTreeRunStatus FStateTreeRunParallelStateTreeTask::EnterState(FStateTreeExecutionContext& Context, const FStateTreeTransitionResult& Transitions) const
@@ -21,6 +22,12 @@ EStateTreeRunStatus FStateTreeRunParallelStateTreeTask::EnterState(FStateTreeExe
 		return EStateTreeRunStatus::Failed;
 	}
 
+	// Share event queue with parent tree.
+	if (FStateTreeInstanceData* OuterInstanceData = Context.GetMutableInstanceData())
+	{
+		InstanceData.TreeInstanceData.SetSharedEventQueue(OuterInstanceData->GetSharedMutableEventQueue());
+	}
+	
 	InstanceData.RunningStateTree = StateTreeToRun.GetStateTree();
 	FStateTreeExecutionContext ParallelTreeContext(Context, *InstanceData.RunningStateTree, InstanceData.TreeInstanceData);
 	if (!ParallelTreeContext.IsValid())
@@ -45,7 +52,24 @@ EStateTreeRunStatus FStateTreeRunParallelStateTreeTask::Tick(FStateTreeExecution
 		return EStateTreeRunStatus::Failed;
 	}
 
-	return ParallelTreeContext.Tick(DeltaTime);
+	return ParallelTreeContext.TickUpdateTasks(DeltaTime);
+}
+
+void FStateTreeRunParallelStateTreeTask::TriggerTransitions(FStateTreeExecutionContext& Context) const
+{
+	FInstanceDataType& InstanceData = Context.GetInstanceData(*this);
+	if (!InstanceData.RunningStateTree)
+	{
+		return;
+	}
+
+	FStateTreeExecutionContext ParallelTreeContext(Context, *InstanceData.RunningStateTree, InstanceData.TreeInstanceData);
+	if (!ParallelTreeContext.IsValid())
+	{
+		return;
+	}
+
+	ParallelTreeContext.TickTriggerTransitions();
 }
 
 void FStateTreeRunParallelStateTreeTask::ExitState(FStateTreeExecutionContext& Context, const FStateTreeTransitionResult& Transition) const
@@ -79,6 +103,13 @@ const FStateTreeReference& FStateTreeRunParallelStateTreeTask::GetStateTreeToRun
 }
 
 #if WITH_EDITOR
+EDataValidationResult FStateTreeRunParallelStateTreeTask::Compile(FStateTreeDataView InstanceDataView, TArray<FText>& ValidationMessages)
+{
+	TransitionHandlingPriority = EventHandlingPriority;
+
+	return EDataValidationResult::Valid;
+}
+
 void FStateTreeRunParallelStateTreeTask::PostEditInstanceDataChangeChainProperty(const FPropertyChangedChainEvent& PropertyChangedEvent, FStateTreeDataView InstanceDataView)
 {
 	if (PropertyChangedEvent.Property && PropertyChangedEvent.Property->GetFName() == GET_MEMBER_NAME_CHECKED(FStateTreeRunParallelStateTreeTaskInstanceData, StateTree))
diff --git a/Engine/Plugins/Runtime/StateTree/Source/StateTreeModule/Private/Tasks/StateTreeRunParallelStateTreeTask.h b/Engine/Plugins/Runtime/StateTree/Source/StateTreeModule/Private/Tasks/StateTreeRunParallelStateTreeTask.h
deleted file mode 100644
diff --git a/Engine/Plugins/Runtime/StateTree/Source/StateTreeModule/Public/StateTreeEvents.h b/Engine/Plugins/Runtime/StateTree/Source/StateTreeModule/Public/StateTreeEvents.h
--- a/Engine/Plugins/Runtime/StateTree/Source/StateTreeModule/Public/StateTreeEvents.h
+++ b/Engine/Plugins/Runtime/StateTree/Source/StateTreeModule/Public/StateTreeEvents.h
@@ -15,6 +15,8 @@ enum class EStateTreeLoopEvents : uint8
 	Next,
 	/** Stops the event handling loop. */
 	Break,
+	/** Consumes and removes the current event. */
+	Consume,
 };
 
 /**
@@ -156,8 +158,11 @@ struct STATETREEMODULE_API FStateTreeEventQueue
 		return SharedEvents;
 	}
 
-	UE_DEPRECATED(5.5, "Use GetEventsView() instead.")
-	TConstArrayView<FStateTreeEvent> GetEvents() const { return {};	}
+	/** @return view to all the events in the buffer. */
+	TArrayView<FStateTreeSharedEvent> GetMutableEventsView()
+	{
+		return SharedEvents;
+	}
 
 	/** Resets the events in the event queue */
 	void Reset()
@@ -165,6 +170,12 @@ struct STATETREEMODULE_API FStateTreeEventQueue
 		SharedEvents.Reset();
 	}
 
+	/** @return true if the queue has any events. */
+	bool HasEvents() const
+	{
+		return !SharedEvents.IsEmpty();
+	}
+	
 	/**
 	 * Buffers and event to be sent to the State Tree.
 	 * @param Owner Optional pointer to an owner UObject that is used for logging errors.
@@ -174,6 +185,33 @@ struct STATETREEMODULE_API FStateTreeEventQueue
 	 */
 	void SendEvent(const UObject* Owner, const FGameplayTag& Tag, const FConstStructView Payload = FConstStructView(), const FName Origin = FName());
 
+	/** Consumes and removes the specified event from the event queue. */
+	void ConsumeEvent(const FStateTreeSharedEvent& Event);
+
+	/**
+	 * Iterates over all events.
+	 * @param Function a lambda which takes const FStateTreeSharedEvent& Event, and returns EStateTreeLoopEvents.
+	 */
+	template<typename TFunc>
+	void ForEachEvent(TFunc&& Function) const
+	{
+		for (TArray<FStateTreeSharedEvent>::TConstIterator It(SharedEvents); It; ++It)
+		{
+			const EStateTreeLoopEvents Result = Function(*It);
+			if (Result == EStateTreeLoopEvents::Break)
+			{
+				break;
+			}
+			if (Result == EStateTreeLoopEvents::Consume)
+			{
+				It.RemoveCurrent();
+			}
+		}
+	}
+
+	UE_DEPRECATED(5.5, "Use GetEventsView() instead.")
+	TConstArrayView<FStateTreeEvent> GetEvents() const { return {};	}
+
 protected:
 	// Used by FStateTreeExecutionState to implement deprecated functionality.
 	TArray<FStateTreeSharedEvent>& GetEventsArray() { return SharedEvents; };
diff --git a/Engine/Plugins/Runtime/StateTree/Source/StateTreeModule/Public/StateTreeExecutionContext.h b/Engine/Plugins/Runtime/StateTree/Source/StateTreeModule/Public/StateTreeExecutionContext.h
--- a/Engine/Plugins/Runtime/StateTree/Source/StateTreeModule/Public/StateTreeExecutionContext.h
+++ b/Engine/Plugins/Runtime/StateTree/Source/StateTreeModule/Public/StateTreeExecutionContext.h
@@ -154,12 +154,27 @@ public:
 	EStateTreeRunStatus Stop(const EStateTreeRunStatus CompletionStatus = EStateTreeRunStatus::Stopped);
 
 	/**
-	 * Tick the state tree logic.
+	 * Tick the state tree logic, updates the tasks and triggers transitions.
 	 * @param DeltaTime time to advance the logic.
 	 * @returns tree run status after the tick.
 	 */
 	EStateTreeRunStatus Tick(const float DeltaTime);
 
+	/**
+	 * Tick the state tree logic partially, updates the tasks.
+	 * For full update TickTriggerTransitions() should be called after.
+	 * @param DeltaTime time to advance the logic.
+	 * @returns tree run status after the partial tick.
+	 */
+	EStateTreeRunStatus TickUpdateTasks(const float DeltaTime);
+	
+	/**
+	 * Tick the state tree logic partially, triggers the transitions.
+	 * For full update TickUpdateTasks() should be called before.
+	 * @returns tree run status after the partial tick.
+	 */
+	EStateTreeRunStatus TickTriggerTransitions();
+	
 	/** @return the tree run status. */
 	EStateTreeRunStatus GetStateTreeRunStatus() const;
 
@@ -194,44 +209,61 @@ public:
 	/** @return the names of all the active state. */
 	TArray<FName> GetActiveStateNames() const;
 
-	/** Sends event for the StateTree. */
-	UE_DEPRECATED(5.2, "Use AddEvent() with individual parameters instead.")
-	void SendEvent(const FStateTreeEvent& Event) const;
-
 	/** Sends event for the StateTree. */
 	void SendEvent(const FGameplayTag Tag, const FConstStructView Payload = FConstStructView(), const FName Origin = FName()) const;
 
-	/** Iterates over all events. Can only be used during StateTree tick. Expects a lambda which takes const FStateTreeSharedEvent& Event, and returns EStateTreeLoopEvents. */
+	/**
+	 * Iterates over all events.
+	 * @param Function a lambda which takes const FStateTreeSharedEvent& Event, and returns EStateTreeLoopEvents.
+	 */
 	template<typename TFunc>
 	typename TEnableIf<TIsInvocable<TFunc, FStateTreeSharedEvent>::Value, void>::Type ForEachEvent(TFunc&& Function) const
 	{
-		for (const FStateTreeSharedEvent& Event : EventsToProcess)
+		if (!EventQueue)
 		{
-			if (Function(Event) == EStateTreeLoopEvents::Break)
-			{
-				break;
-			}
+			return;
 		}
+		EventQueue->template ForEachEvent(Function);
 	}
 
-	/** Iterates over all events. Can only be used during StateTree tick. Expects a lambda which takes const FStateTreeEvent& Event, and returns EStateTreeLoopEvents. Less preferable than FStateTreeSharedEvent version. */
+	/**
+	 * Iterates over all events.
+	 * @param Function a lambda which takes const FStateTreeSharedEvent& Event, and returns EStateTreeLoopEvents.
+	 * Less preferable than FStateTreeSharedEvent version.
+	 */
 	template<typename TFunc>
 	typename TEnableIf<TIsInvocable<TFunc, FStateTreeEvent>::Value, void>::Type ForEachEvent(TFunc&& Function) const
 	{
-		for (const FStateTreeSharedEvent& Event : EventsToProcess)
+		if (!EventQueue)
 		{
-			if (Function(*Event) == EStateTreeLoopEvents::Break)
-			{
-				break;
-			}
+			return;
 		}
+		EventQueue->template ForEachEvent([Function](const FStateTreeSharedEvent& Event)
+		{
+			return Function(*Event);
+		});
 	}
 
 	/** @return events to process this tick. */
-	TArrayView<FStateTreeSharedEvent> GetMutableEventsToProcessView() { return EventsToProcess; }
+	TArrayView<FStateTreeSharedEvent> GetMutableEventsToProcessView()
+	{
+		return EventQueue ? EventQueue->GetMutableEventsView() : TArrayView<FStateTreeSharedEvent>();
+	}
 
 	/** @return events to process this tick. */
-	TConstArrayView<FStateTreeSharedEvent> GetEventsToProcessView() const { return EventsToProcess; }
+	TConstArrayView<FStateTreeSharedEvent> GetEventsToProcessView() const
+	{
+		return EventQueue ? EventQueue->GetMutableEventsView() : TArrayView<FStateTreeSharedEvent>();
+	}
+
+	/** Consumes and removes the specified event from the event queue. */
+	void ConsumeEvent(const FStateTreeSharedEvent& Event)
+	{
+		if (EventQueue)
+		{
+			EventQueue->ConsumeEvent(Event);
+		}
+	}
 
 	UE_DEPRECATED(5.5, "Use GetEventsToProcessView() instead.")
 	TConstArrayView<FStateTreeEvent> GetEventsToProcess() const { return {}; }
@@ -239,12 +271,7 @@ public:
 	/** @return true if there is a pending event with specified tag. */
 	bool HasEventToProcess(const FGameplayTag Tag) const
 	{
-		if (EventsToProcess.IsEmpty())
-		{
-			return false;
-		}
-		
-		return EventsToProcess.ContainsByPredicate([Tag](const FStateTreeSharedEvent& Event)
+		return GetEventsToProcessView().ContainsByPredicate([Tag](const FStateTreeSharedEvent& Event)
 		{
 			check(Event.IsValid());
 			return Event->Tag.MatchesTag(Tag);
@@ -531,6 +558,16 @@ protected:
 	 */
 	EStateTreeRunStatus TickTasks(const float DeltaTime);
 
+	/** Common functionality shared by the tick methods. */
+	EStateTreeRunStatus TickPrelude();
+	EStateTreeRunStatus TickPostlude();
+
+	/** Handles task ticking part of the tick. */
+	void TickUpdateTasksInternal(const float DeltaTime);
+	
+	/** Handles transition triggering part of the tick. */
+	void TickTriggerTransitionsInternal();
+
 	/**
 	 * Checks all conditions at given range
 	 * @return True if all conditions pass.
@@ -594,13 +631,15 @@ protected:
 	/** @return StateTree execution state from the instance storage. */
 	FStateTreeExecutionState& GetExecState()
 	{
-		return *InstanceData.GetMutableExecutionState();
+		check(InstanceDataStorage);
+		return InstanceDataStorage->GetMutableExecutionState();
 	}
 
 	/** @return const StateTree execution state from the instance storage. */
 	const FStateTreeExecutionState& GetExecState() const
 	{
-		return *InstanceData.GetExecutionState();
+		check(InstanceDataStorage);
+		return InstanceDataStorage->GetExecutionState();
 	}
 
 	/** @return String describing state status for logging and debug. */
@@ -671,15 +710,15 @@ protected:
 	/** Data storage of the instance data, cached for less indirections. */
 	FStateTreeInstanceStorage* InstanceDataStorage = nullptr;
 
+	/** Events queue to use, cached for less indirections. */
+	TSharedPtr<FStateTreeEventQueue> EventQueue;
+
 	/** Pointer to linked state tree overrides. */
 	const FStateTreeReferenceOverrides* LinkedStateTreeOverrides = nullptr;
 	
 	/** Data view of the context data. */
 	TArray<FStateTreeDataView, TConcurrentLinearArrayAllocator<FDefaultBlockAllocationTag>> ContextAndExternalDataViews;
 
-	/** Events to process in current tick. */
-	TArray<FStateTreeSharedEvent, TConcurrentLinearArrayAllocator<FDefaultBlockAllocationTag>> EventsToProcess;
-
 	FOnCollectStateTreeExternalData CollectExternalDataDelegate;
 
 	struct FCollectedExternalDataCache
diff --git a/Engine/Plugins/Runtime/StateTree/Source/StateTreeModule/Public/StateTreeInstanceData.h b/Engine/Plugins/Runtime/StateTree/Source/StateTreeModule/Public/StateTreeInstanceData.h
--- a/Engine/Plugins/Runtime/StateTree/Source/StateTreeModule/Public/StateTreeInstanceData.h
+++ b/Engine/Plugins/Runtime/StateTree/Source/StateTreeModule/Public/StateTreeInstanceData.h
@@ -88,18 +88,40 @@ struct STATETREEMODULE_API FStateTreeInstanceStorage
 {
 	GENERATED_BODY()
 
+	FStateTreeInstanceStorage();
+	FStateTreeInstanceStorage(const FStateTreeInstanceStorage& Other);
+	FStateTreeInstanceStorage(FStateTreeInstanceStorage&& Other) noexcept;
+
+	FStateTreeInstanceStorage& operator=(const FStateTreeInstanceStorage& Other);
+	FStateTreeInstanceStorage& operator=(FStateTreeInstanceStorage&& Other) noexcept;
+
 	/** @return reference to the event queue. */
 	FStateTreeEventQueue& GetMutableEventQueue()
 	{
-		return EventQueue;
+		return *EventQueue;
 	}
 
 	/** @return reference to the event queue. */
 	const FStateTreeEventQueue& GetEventQueue() const
+	{
+		return *EventQueue;
+	}
+
+	/** @return true if the storage owns the event queue. */
+	bool IsOwningEventQueue() const
+	{
+		return bIsOwningEventQueue;
+	}
+
+	/** @return shared pointer to the event queue. */
+	const TSharedRef<FStateTreeEventQueue>& GetSharedMutableEventQueue()
 	{
 		return EventQueue;
 	}
 
+	/** Sets event queue from another storage. Marks the event queue not owned. */
+	void SetSharedEventQueue(const TSharedRef<FStateTreeEventQueue>& InSharedEventQueue);
+	
 	/** 
 	 * Buffers a transition request to be sent to the State Tree.
 	 * @param Owner Optional pointer to an owner UObject that is used for logging errors.
@@ -210,6 +232,12 @@ struct STATETREEMODULE_API FStateTreeInstanceStorage
 	{
 		return GlobalParameters.GetMutableValue();
 	}
+
+	/** Note, called by FStateTreeInstanceData. */
+	void AddStructReferencedObjects(FReferenceCollector& Collector);
+
+	/** Resets the storage to initial state. */
+	void Reset();
 	
 	UE_DEPRECATED(5.4, "Use Num() instead.")
 	int32 NumStructs() const { return 0; }
@@ -224,28 +252,30 @@ struct STATETREEMODULE_API FStateTreeInstanceStorage
 	bool IsValid() const { return false; }
 
 protected:
-	/** Execution state of the state tree instance. */
-	UPROPERTY()
-	FStateTreeExecutionState ExecutionState;
-
-	/** Struct instances */
+	/** Struct instances (Not transient, as we use FStateTreeInstanceData to store default values for instance data) */
 	UPROPERTY()
 	FInstancedStructContainer InstanceStructs;
 
+	/** Execution state of the state tree instance. */
+	UPROPERTY(Transient)
+	FStateTreeExecutionState ExecutionState;
+
 	/** Temporary instances */
-	UPROPERTY()
+	UPROPERTY(Transient)
 	TArray<FStateTreeTemporaryInstanceData> TemporaryInstances;
 
-	/** Events */
-	UPROPERTY()
-	FStateTreeEventQueue EventQueue;
+	/** Events (Transient) */
+	TSharedRef<FStateTreeEventQueue> EventQueue = MakeShared<FStateTreeEventQueue>();
 
+	/* True if the storage owns the event queue. */
+	bool bIsOwningEventQueue = true;
+	
 	/** Requested transitions */
-	UPROPERTY()
+	UPROPERTY(Transient)
 	TArray<FStateTreeTransitionRequest> TransitionRequests;
 
 	/** Global parameters */
-	UPROPERTY()
+	UPROPERTY(Transient)
 	FInstancedPropertyBag GlobalParameters;
 
 	friend struct FStateTreeInstanceData;
@@ -266,10 +296,10 @@ struct STATETREEMODULE_API FStateTreeInstanceData
 
 	FStateTreeInstanceData();
 	FStateTreeInstanceData(const FStateTreeInstanceData& Other);
-	FStateTreeInstanceData(FStateTreeInstanceData&& Other);
+	FStateTreeInstanceData(FStateTreeInstanceData&& Other) noexcept;
 
 	FStateTreeInstanceData& operator=(const FStateTreeInstanceData& Other);
-	FStateTreeInstanceData& operator=(FStateTreeInstanceData&& Other);
+	FStateTreeInstanceData& operator=(FStateTreeInstanceData&& Other) noexcept;
 
 	~FStateTreeInstanceData();
 	
@@ -350,6 +380,13 @@ struct STATETREEMODULE_API FStateTreeInstanceData
 	/** @return reference to the event queue. */
 	FStateTreeEventQueue& GetMutableEventQueue();
 	const FStateTreeEventQueue& GetEventQueue() const;
+	const TSharedRef<FStateTreeEventQueue>& GetSharedMutableEventQueue();
+
+	/** @return true if the instance data owns its' event queue. */
+	bool IsOwningEventQueue() const;
+
+	/** Sets event queue from another instance data. Marks the event queue not owned. */
+	void SetSharedEventQueue(const TSharedRef<FStateTreeEventQueue>& InSharedEventQueue);
 
 	/** 
 	 * Buffers a transition request to be sent to the State Tree.
diff --git a/Engine/Plugins/Runtime/StateTree/Source/StateTreeModule/Public/StateTreeTaskBase.h b/Engine/Plugins/Runtime/StateTree/Source/StateTreeModule/Public/StateTreeTaskBase.h
--- a/Engine/Plugins/Runtime/StateTree/Source/StateTreeModule/Public/StateTreeTaskBase.h
+++ b/Engine/Plugins/Runtime/StateTree/Source/StateTreeModule/Public/StateTreeTaskBase.h
@@ -103,6 +103,9 @@ struct STATETREEMODULE_API FStateTreeTaskBase : public FStateTreeNodeBase
 	/** True if the node is Enabled (i.e. not explicitly disabled in the asset). */
 	UPROPERTY()
 	uint8 bTaskEnabled : 1;
+
+	UPROPERTY()
+	EStateTreeTransitionPriority TransitionHandlingPriority = EStateTreeTransitionPriority::Normal;
 };
 
 /**
diff --git a/Engine/Plugins/Runtime/StateTree/Source/StateTreeModule/Public/StateTreeTypes.h b/Engine/Plugins/Runtime/StateTree/Source/StateTreeModule/Public/StateTreeTypes.h
--- a/Engine/Plugins/Runtime/StateTree/Source/StateTreeModule/Public/StateTreeTypes.h
+++ b/Engine/Plugins/Runtime/StateTree/Source/StateTreeModule/Public/StateTreeTypes.h
@@ -133,7 +133,10 @@ UENUM(BlueprintType)
 enum class EStateTreeTransitionPriority : uint8
 {
 	None UMETA(Hidden),
-	
+
+	/** Low priority. */
+	Low,
+
 	/** Normal priority. */
 	Normal,
 	
@@ -543,7 +546,6 @@ struct STATETREEMODULE_API FCompactStateTransition
 	GENERATED_BODY()
 
 	explicit FCompactStateTransition()
-		: bTransitionEnabled(true)
 	{
 	}
 
@@ -587,7 +589,11 @@ struct STATETREEMODULE_API FCompactStateTransition
 
 	/** Indicates if the transition is enabled and should be considered. */
 	UPROPERTY()
-	uint8 bTransitionEnabled : 1;
+	uint8 bTransitionEnabled : 1 = true;
+
+	/** If set to true, the required event is consumed (later state selection cannot react to it) if state selection can be made. */
+	UPROPERTY()
+	uint8 bConsumeEventOnSelect : 1 = true;
 };
 
 /**
@@ -699,6 +705,10 @@ struct STATETREEMODULE_API FCompactStateTreeState
 	/** True if the state is Enabled (i.e. not explicitly marked as disabled). */
 	UPROPERTY()
 	uint8 bEnabled : 1 = true;
+
+	/** If set to true, the required event is consumed (later state selection cannot react to it) if state selection can be made. */
+	UPROPERTY()
+	uint8 bConsumeEventOnSelect : 1 = true;
 };
 
 USTRUCT()
diff --git a/Engine/Plugins/Runtime/StateTree/Source/StateTreeModule/Public/Tasks/StateTreeRunParallelStateTreeTask.h b/Engine/Plugins/Runtime/StateTree/Source/StateTreeModule/Public/Tasks/StateTreeRunParallelStateTreeTask.h
new file mode 100644
--- /dev/null
+++ b/Engine/Plugins/Runtime/StateTree/Source/StateTreeModule/Public/Tasks/StateTreeRunParallelStateTreeTask.h
@@ -0,0 +1,81 @@
+// Copyright Epic Games, Inc. All Rights Reserved.
+
+#pragma once
+
+#include "StateTreeInstanceData.h"
+#include "StateTreeReference.h"
+#include "StateTreeTaskBase.h"
+
+#include "StateTreeRunParallelStateTreeTask.generated.h"
+
+USTRUCT()
+struct FStateTreeRunParallelStateTreeTaskInstanceData
+{
+	GENERATED_BODY()
+
+	/** State tree and parameters that will be run when this task is started. */
+	UPROPERTY(EditAnywhere, Category = "Parameter", meta=(SchemaCanBeOverriden))
+	FStateTreeReference StateTree;
+
+	UPROPERTY(Transient)
+	FStateTreeInstanceData TreeInstanceData;
+
+	UPROPERTY(Transient)
+	TObjectPtr<const UStateTree> RunningStateTree = nullptr;
+};
+
+/**
+* Task that will run another state tree in the current state while allowing the current tree to continue selection and process of child state.
+* It will succeed, fail or run depending on the result of the parallel tree.
+* Less efficient then Linked Asset state, it has the advantage of allowing multiple trees to run in parallel.
+*/
+USTRUCT(meta = (DisplayName = "Run Parallel Tree", Category = "Common"))
+struct STATETREEMODULE_API FStateTreeRunParallelStateTreeTask : public FStateTreeTaskCommonBase
+{
+	GENERATED_BODY()
+	using FInstanceDataType = FStateTreeRunParallelStateTreeTaskInstanceData;
+	
+	FStateTreeRunParallelStateTreeTask();
+
+#if WITH_EDITORONLY_DATA
+	// Sets event handling priority
+	void SetEventHandlingPriority(const EStateTreeTransitionPriority NewPriority)
+	{
+		EventHandlingPriority = NewPriority;
+	}
+#endif	
+	
+protected:
+	virtual const UStruct* GetInstanceDataType() const override { return FInstanceDataType::StaticStruct(); }
+
+	virtual EStateTreeRunStatus EnterState(FStateTreeExecutionContext& Context, const FStateTreeTransitionResult& Transitions) const override;
+	virtual EStateTreeRunStatus Tick(FStateTreeExecutionContext& Context, const float DeltaTime) const override;
+	virtual void TriggerTransitions(FStateTreeExecutionContext& Context) const override;
+	virtual void ExitState(FStateTreeExecutionContext& Context, const FStateTreeTransitionResult& Transition) const override;
+
+#if WITH_EDITOR
+	virtual EDataValidationResult Compile(FStateTreeDataView InstanceDataView, TArray<FText>& ValidationMessages) override;
+	virtual void PostEditInstanceDataChangeChainProperty(const FPropertyChangedChainEvent& PropertyChangedEvent, FStateTreeDataView InstanceDataView) override;
+	virtual void PostLoad(FStateTreeDataView InstanceDataView) override;
+	virtual FText GetDescription(const FGuid& ID, FStateTreeDataView InstanceDataView, const IStateTreeBindingLookup& BindingLookup, EStateTreeNodeFormatting Formatting = EStateTreeNodeFormatting::Text) const override;
+#endif // WITH_EDITOR
+
+	const FStateTreeReference& GetStateTreeToRun(FStateTreeExecutionContext& Context, FInstanceDataType& InstanceData) const;
+
+	/** If set the task will look at the linked state tree override to replace the state tree it's running. */
+	UPROPERTY(EditAnywhere, Category = Parameter)
+	FGameplayTag StateTreeOverrideTag;
+
+#if WITH_EDITORONLY_DATA
+	/**
+	 * At what priority the events should be handled in the parallel State Tree.
+	 * If set to 'Normal' the order of the States in the State Tree will define the handling order.
+	 * If the priority is set to Low, the main tree is let to handle the transitions first.
+	 * If set to High or above, the parallel tree has change to handle events first.
+	 * If multiple tasks has same priority, the State order of the States defines the handling order.
+	 * The tree handling order is: States and handle from leaf to root, tasks before and handled before transitions per State.
+	 */
+	UPROPERTY(EditAnywhere, Category = Parameter)
+	EStateTreeTransitionPriority EventHandlingPriority = EStateTreeTransitionPriority::Normal;
+#endif // WITH_EDITORONLY_DATA	
+};
\ No newline at end of file
diff --git a/Engine/Plugins/Runtime/StateTree/Source/StateTreeTestSuite/Private/StateTreeTest.cpp b/Engine/Plugins/Runtime/StateTree/Source/StateTreeTestSuite/Private/StateTreeTest.cpp
--- a/Engine/Plugins/Runtime/StateTree/Source/StateTreeTestSuite/Private/StateTreeTest.cpp
+++ b/Engine/Plugins/Runtime/StateTree/Source/StateTreeTestSuite/Private/StateTreeTest.cpp
@@ -6,6 +6,7 @@
 #include "StateTreeEditorData.h"
 #include "StateTreeCompiler.h"
 #include "Conditions/StateTreeCommonConditions.h"
+#include "Tasks/StateTreeRunParallelStateTreeTask.h"
 #include "StateTreeTestTypes.h"
 #include "Engine/World.h"
 #include "Async/ParallelFor.h"
@@ -2976,6 +2977,149 @@ struct FStateTreeTest_NestedOverride : FAITestBase
 };
 IMPLEMENT_AI_INSTANT_TEST(FStateTreeTest_NestedOverride, "System.StateTree.NestedOverride");
 
+// Test parallel tree event priority handling.
+struct FStateTreeTest_ParallelEventPriority : FAITestBase
+{
+	EStateTreeTransitionPriority ParallelTreePriority = EStateTreeTransitionPriority::Normal;
+	
+	virtual bool InstantTest() override
+	{
+		FStateTreeCompilerLog Log;
+		
+		const FGameplayTag EventTag = UE::StateTree::Tests::FNativeGameplayTags::Get().TestTag;
+
+		// Parallel tree
+		// - Root
+		//   - State1 ?-> State2
+		//   - State2
+		UStateTree& StateTreePar = UE::StateTree::Tests::NewStateTree(&GetWorld());
+		UStateTreeEditorData& EditorDataPar = *Cast<UStateTreeEditorData>(StateTreePar.EditorData);
+
+		UStateTreeState& RootPar = EditorDataPar.AddSubTree(FName(TEXT("Root")));
+		UStateTreeState& State1 = RootPar.AddChildState(FName(TEXT("State1")));
+		UStateTreeState& State2 = RootPar.AddChildState(FName(TEXT("State2")));
+
+		TStateTreeEditorNode<FTestTask_Stand>& Task1 = State1.AddTask<FTestTask_Stand>(FName(TEXT("Task1")));
+		Task1.GetNode().TicksToCompletion = 100;
+		State1.AddTransition(EStateTreeTransitionTrigger::OnEvent, EventTag, EStateTreeTransitionType::NextState);
+
+		TStateTreeEditorNode<FTestTask_Stand>& Task2 = State2.AddTask<FTestTask_Stand>(FName(TEXT("Task2")));
+		Task2.GetNode().TicksToCompletion = 100;
+
+		{
+			FStateTreeCompiler Compiler(Log);
+			const bool bResult = Compiler.Compile(StateTreePar);
+			AITEST_TRUE("StateTreePar should get compiled", bResult);
+		}
+
+		// Main asset
+		// - Root [StateTreePar]
+		//   - State3 ?-> State4
+		//   - State4
+		UStateTree& StateTree = UE::StateTree::Tests::NewStateTree(&GetWorld());
+		UStateTreeEditorData& EditorData = *Cast<UStateTreeEditorData>(StateTree.EditorData);
+
+		UStateTreeState& Root = EditorData.AddSubTree(FName(TEXT("Root")));
+		UStateTreeState& State3 = Root.AddChildState(FName(TEXT("State3")));
+		UStateTreeState& State4 = Root.AddChildState(FName(TEXT("State4")));
+
+		TStateTreeEditorNode<FStateTreeRunParallelStateTreeTask>& TaskPar = Root.AddTask<FStateTreeRunParallelStateTreeTask>();
+		TaskPar.GetNode().SetEventHandlingPriority(ParallelTreePriority);
+		
+		TaskPar.GetInstanceData().StateTree.SetStateTree(&StateTreePar);
+
+		TStateTreeEditorNode<FTestTask_Stand>& Task3 = State3.AddTask<FTestTask_Stand>(FName(TEXT("Task3")));
+		Task3.GetNode().TicksToCompletion = 100;
+		State3.AddTransition(EStateTreeTransitionTrigger::OnEvent, EventTag, EStateTreeTransitionType::NextState);
+
+		TStateTreeEditorNode<FTestTask_Stand>& Task4 = State4.AddTask<FTestTask_Stand>(FName(TEXT("Task4")));
+		Task4.GetNode().TicksToCompletion = 100;
+		
+		{
+			FStateTreeCompiler Compiler(Log);
+			const bool bResult = Compiler.Compile(StateTree);
+			AITEST_TRUE("StateTree should get compiled", bResult);
+		}
+
+		const FString TickStr(TEXT("Tick"));
+		const FString EnterStateStr(TEXT("EnterState"));
+		const FString ExitStateStr(TEXT("ExitState"));
+
+		// Run StateTreePar in parallel with the main tree.
+		// Both trees have a transition on same event.
+		// Setting the priority to Low, should make the main tree to take the transition.
+		EStateTreeRunStatus Status = EStateTreeRunStatus::Unset;
+		FStateTreeInstanceData InstanceData;
+		FTestStateTreeExecutionContext Exec(StateTree, StateTree, InstanceData);
+		const bool bInitSucceeded = Exec.IsValid();
+		AITEST_TRUE("StateTree should init", bInitSucceeded);
+
+		Status = Exec.Start();
+		AITEST_EQUAL("Start should complete with Running", Status, EStateTreeRunStatus::Running);
+		AITEST_TRUE("StateTree should enter Task1, Task3", Exec.Expect(Task1.GetName(), EnterStateStr).Then(Task3.GetName(), EnterStateStr));
+		Exec.LogClear();
+
+		Status = Exec.Tick(0.1f);
+		AITEST_EQUAL("Tick should complete with Running", Status, EStateTreeRunStatus::Running);
+		AITEST_TRUE("StateTree should tick Task1, Task3", Exec.Expect(Task1.GetName(), TickStr).Then(Task3.GetName(), TickStr));
+		Exec.LogClear();
+
+		Exec.SendEvent(EventTag);
+
+		// If the parallel tree priority is < Normal, then it should always be handled after the main tree.
+		// If the parallel tree priority is Normal, then the state order decides (leaf to root)
+		// If the parallel tree priority is > Normal, then it should always be handled before the main tree.
+		if (ParallelTreePriority <= EStateTreeTransitionPriority::Normal)
+		{
+			// Main tree should do the transition.
+			Status = Exec.Tick(0.1f);
+			AITEST_EQUAL("Tick should complete with Running", Status, EStateTreeRunStatus::Running);
+			AITEST_TRUE("StateTree should enter Task4", Exec.Expect(Task4.GetName(), EnterStateStr));
+			Exec.LogClear();
+
+			Status = Exec.Tick(0.1f);
+			AITEST_EQUAL("Tick should complete with Running", Status, EStateTreeRunStatus::Running);
+			AITEST_TRUE("StateTree should tick Task1, Task4", Exec.Expect(Task1.GetName(), TickStr).Then(Task4.GetName(), TickStr));
+			Exec.LogClear();
+		}
+		else
+		{
+			// Parallel tree should do the transition.
+			Status = Exec.Tick(0.1f);
+			AITEST_EQUAL("Tick should complete with Running", Status, EStateTreeRunStatus::Running);
+			AITEST_TRUE("StateTree should enter Task2", Exec.Expect(Task2.GetName(), EnterStateStr));
+			Exec.LogClear();
+
+			Status = Exec.Tick(0.1f);
+			AITEST_EQUAL("Tick should complete with Running", Status, EStateTreeRunStatus::Running);
+			AITEST_TRUE("StateTree should tick Task2, Task3", Exec.Expect(Task2.GetName(), TickStr).Then(Task3.GetName(), TickStr));
+			Exec.LogClear();
+		}
+		
+		return true;
+	}
+};
+IMPLEMENT_AI_INSTANT_TEST(FStateTreeTest_ParallelEventPriority, "System.StateTree.ParallelEventPriority");
+
+
+struct FStateTreeTest_ParallelEventPriority_Low : FStateTreeTest_ParallelEventPriority
+{
+	FStateTreeTest_ParallelEventPriority_Low()
+	{
+		ParallelTreePriority = EStateTreeTransitionPriority::Low;
+	}
+};
+IMPLEMENT_AI_INSTANT_TEST(FStateTreeTest_ParallelEventPriority_Low, "System.StateTree.ParallelEventPriority.Low");
+
+struct FStateTreeTest_ParallelEventPriority_High : FStateTreeTest_ParallelEventPriority
+{
+	FStateTreeTest_ParallelEventPriority_High()
+	{
+		ParallelTreePriority = EStateTreeTransitionPriority::High;
+	}
+};
+IMPLEMENT_AI_INSTANT_TEST(FStateTreeTest_ParallelEventPriority_High, "System.StateTree.ParallelEventPriority.High");
+
 UE_ENABLE_OPTIMIZATION_SHIP
 
 #undef LOCTEXT_NAMESPACE
diff --git a/Engine/Plugins/Runtime/StateTree/Source/StateTreeTestSuite/Private/StateTreeTest.h b/Engine/Plugins/Runtime/StateTree/Source/StateTreeTestSuite/Private/StateTreeTest.h
--- a/Engine/Plugins/Runtime/StateTree/Source/StateTreeTestSuite/Private/StateTreeTest.h
+++ b/Engine/Plugins/Runtime/StateTree/Source/StateTreeTestSuite/Private/StateTreeTest.h
@@ -11,5 +11,6 @@ class UStateTreeTestSchema : public UStateTreeSchema
 	GENERATED_BODY()
 
 	virtual bool IsStructAllowed(const UScriptStruct* InScriptStruct) const override { return true; }
+	virtual bool IsExternalItemAllowed(const UStruct& InStruct) const override { return true; }
 };
 
diff --git a/Engine/Plugins/Runtime/StateTree/Source/StateTreeTestSuite/Private/StateTreeTestTypes.h b/Engine/Plugins/Runtime/StateTree/Source/StateTreeTestSuite/Private/StateTreeTestTypes.h
--- a/Engine/Plugins/Runtime/StateTree/Source/StateTreeTestSuite/Private/StateTreeTestTypes.h
+++ b/Engine/Plugins/Runtime/StateTree/Source/StateTreeTestSuite/Private/StateTreeTestTypes.h
@@ -6,20 +6,19 @@
 #include "StateTreeEvaluatorBase.h"
 #include "StateTreeConditionBase.h"
 #include "StateTreeExecutionContext.h"
+#include "StateTreeLinker.h"
 #include "StateTreePropertyRef.h"
 #include "StateTreeTestTypes.generated.h"
 
 class UStateTree;
 struct FStateTreeInstanceData;
 
-
-struct FTestStateTreeExecutionContext : public FStateTreeExecutionContext
+// Test log that can be passed as external data.
+USTRUCT()
+struct FStateTreeTestLog
 {
-	FTestStateTreeExecutionContext(UObject& InOwner, const UStateTree& InStateTree, FStateTreeInstanceData& InInstanceData)
-		: FStateTreeExecutionContext(InOwner, InStateTree, InInstanceData)
-	{
-	}
-	
+	GENERATED_BODY()
+
 	struct FLogItem
 	{
 		FLogItem() = default;
@@ -27,16 +26,44 @@ struct FTestStateTreeExecutionContext : public FStateTreeExecutionContext
 		FName Name;
 		FString Message; 
 	};
-	TArray<FLogItem> LogItems;
+	TArray<FLogItem> Items;
 	
 	void Log(const FName& Name, const FString& Message)
 	{
-		LogItems.Emplace(Name, Message);
+		Items.Emplace(Name, Message);
+	}
+};
+
+struct FTestStateTreeExecutionContext : public FStateTreeExecutionContext
+{
+	
+	FTestStateTreeExecutionContext(UObject& InOwner, const UStateTree& InStateTree, FStateTreeInstanceData& InInstanceData)
+		: FStateTreeExecutionContext(InOwner, InStateTree, InInstanceData)
+	{
+		// Handle getting pointer to the test log.
+		FStateTreeDataView TestLogView(TBaseStructure<FStateTreeTestLog>::Get(), (uint8*)&Log);
+		
+		CollectExternalDataDelegate = FOnCollectStateTreeExternalData::CreateLambda(
+			[TestLogView](const FStateTreeExecutionContext& Context, const UStateTree* StateTree, TArrayView<const FStateTreeExternalDataDesc> ExternalDataDescs, TArrayView<FStateTreeDataView> OutDataViews)
+			{
+				for (int32 Index = 0; Index < ExternalDataDescs.Num(); Index++)
+				{
+					const FStateTreeExternalDataDesc& ItemDesc = ExternalDataDescs[Index];
+					if (ItemDesc.Struct == TBaseStructure<FStateTreeTestLog>::Get())
+					{
+						OutDataViews[Index] = TestLogView;
+						break;
+					}
+				}
+				return true;
+			});
 	}
 
+	FStateTreeTestLog Log;
+
 	void LogClear()
 	{
-		LogItems.Empty();
+		Log.Items.Empty();
 	}
 
 	struct FLogOrder
@@ -46,9 +73,9 @@ struct FTestStateTreeExecutionContext : public FStateTreeExecutionContext
 		FLogOrder Then(const FName& Name, const FString& Message) const
 		{
 			int32 NextIndex = Index;
-			while (NextIndex < Context.LogItems.Num())
+			while (NextIndex < Context.Log.Items.Num())
 			{
-				const FLogItem& Item = Context.LogItems[NextIndex];
+				const FStateTreeTestLog::FLogItem& Item = Context.Log.Items[NextIndex];
 				if (Item.Name == Name && Item.Message == Message)
 				{
 					break;
@@ -58,7 +85,7 @@ struct FTestStateTreeExecutionContext : public FStateTreeExecutionContext
 			return FLogOrder(Context, NextIndex);
 		}
 
-		operator bool() const { return Index < Context.LogItems.Num(); }
+		operator bool() const { return Index < Context.Log.Items.Num(); }
 		
 		const FTestStateTreeExecutionContext& Context;
 		int32 Index = 0;
@@ -154,10 +181,18 @@ struct FTestTask_B : public FStateTreeTaskBase
 
 	virtual EStateTreeRunStatus EnterState(FStateTreeExecutionContext& Context, const FStateTreeTransitionResult& Transition) const override
 	{
-		FTestStateTreeExecutionContext& TestContext = static_cast<FTestStateTreeExecutionContext&>(Context);
-		TestContext.Log(Name,  TEXT("EnterState"));
+		FStateTreeTestLog& TestLog = Context.GetExternalData(LogHandle);
+		TestLog.Log(Name,  TEXT("EnterState"));
 		return EStateTreeRunStatus::Running;
 	}
+
+	virtual bool Link(FStateTreeLinker& Linker) override
+	{
+		Linker.LinkExternalData(LogHandle);
+		return true;
+	}
+	
+	TStateTreeExternalDataHandle<FStateTreeTestLog> LogHandle;
 };
 
 USTRUCT()
@@ -184,27 +219,35 @@ struct FTestTask_PrintValue : public FStateTreeTaskBase
 
 	virtual EStateTreeRunStatus EnterState(FStateTreeExecutionContext& Context, const FStateTreeTransitionResult& Transition) const override
 	{
-		FTestStateTreeExecutionContext& TestContext = static_cast<FTestStateTreeExecutionContext&>(Context);
+		FStateTreeTestLog& TestLog = Context.GetExternalData(LogHandle);
 		const FInstanceDataType& InstanceData = Context.GetInstanceData(*this);
-		TestContext.Log(Name,  FString::Printf(TEXT("EnterState%d"), InstanceData.Value));
+		TestLog.Log(Name,  FString::Printf(TEXT("EnterState%d"), InstanceData.Value));
 		return EStateTreeRunStatus::Running;
 	}
 
 	virtual void ExitState(FStateTreeExecutionContext& Context, const FStateTreeTransitionResult& Transition) const override
 	{
-		FTestStateTreeExecutionContext& TestContext = static_cast<FTestStateTreeExecutionContext&>(Context);
+		FStateTreeTestLog& TestLog = Context.GetExternalData(LogHandle);
 		const FInstanceDataType& InstanceData = Context.GetInstanceData(*this);
-		TestContext.Log(Name,  FString::Printf(TEXT("ExitState%d"), InstanceData.Value));
+		TestLog.Log(Name,  FString::Printf(TEXT("ExitState%d"), InstanceData.Value));
 	}
 
 	virtual EStateTreeRunStatus Tick(FStateTreeExecutionContext& Context, const float DeltaTime) const override
 	{
-		FTestStateTreeExecutionContext& TestContext = static_cast<FTestStateTreeExecutionContext&>(Context);
+		FStateTreeTestLog& TestLog = Context.GetExternalData(LogHandle);
 		const FInstanceDataType& InstanceData = Context.GetInstanceData(*this);
-		TestContext.Log(Name,  FString::Printf(TEXT("Tick%d"), InstanceData.Value));
+		TestLog.Log(Name,  FString::Printf(TEXT("Tick%d"), InstanceData.Value));
 		
 		return EStateTreeRunStatus::Running;
 	};
+
+	virtual bool Link(FStateTreeLinker& Linker) override
+	{
+		Linker.LinkExternalData(LogHandle);
+		return true;
+	}
+	
+	TStateTreeExternalDataHandle<FStateTreeTestLog> LogHandle;
 };
 
 
@@ -253,6 +296,14 @@ struct FTestTask_StopTree : public FStateTreeTaskBase
 		return EStateTreeRunStatus::Running;
 	};
 
+	virtual bool Link(FStateTreeLinker& Linker) override
+	{
+		Linker.LinkExternalData(LogHandle);
+		return true;
+	}
+	
+	TStateTreeExternalDataHandle<FStateTreeTestLog> LogHandle;
+	
 	/** Indicates in which phase the call to Stop should be performed. Possible values are EnterStates, ExitStats and TickStateTree */
 	UPROPERTY()
 	EStateTreeUpdatePhase Phase = EStateTreeUpdatePhase::Unset;
@@ -282,8 +333,8 @@ struct FTestTask_Stand : public FStateTreeTaskBase
 
 	virtual EStateTreeRunStatus EnterState(FStateTreeExecutionContext& Context, const FStateTreeTransitionResult& Transition) const override
 	{
-		FTestStateTreeExecutionContext& TestContext = static_cast<FTestStateTreeExecutionContext&>(Context);
-		TestContext.Log(Name, TEXT("EnterState"));
+		FStateTreeTestLog& TestLog = Context.GetExternalData(LogHandle);
+		TestLog.Log(Name, TEXT("EnterState"));
 
 		FInstanceDataType& InstanceData = Context.GetInstanceData(*this);
 		
@@ -296,34 +347,34 @@ struct FTestTask_Stand : public FStateTreeTaskBase
 
 	virtual void ExitState(FStateTreeExecutionContext& Context, const FStateTreeTransitionResult& Transition) const override
 	{
-		FTestStateTreeExecutionContext& TestContext = static_cast<FTestStateTreeExecutionContext&>(Context);
+		FStateTreeTestLog& TestLog = Context.GetExternalData(LogHandle);
 
 		if (Transition.CurrentRunStatus == EStateTreeRunStatus::Succeeded)
 		{
-			TestContext.Log(Name, TEXT("ExitSucceeded"));
+			TestLog.Log(Name, TEXT("ExitSucceeded"));
 		}
 		else if (Transition.CurrentRunStatus == EStateTreeRunStatus::Failed)
 		{
-			TestContext.Log(Name, TEXT("ExitFailed"));
+			TestLog.Log(Name, TEXT("ExitFailed"));
 		}
 		else if (Transition.CurrentRunStatus == EStateTreeRunStatus::Stopped)
 		{
-			TestContext.Log(Name, TEXT("ExitStopped"));
+			TestLog.Log(Name, TEXT("ExitStopped"));
 		}
 		
-		TestContext.Log(Name, TEXT("ExitState"));
+		TestLog.Log(Name, TEXT("ExitState"));
 	}
 
 	virtual void StateCompleted(FStateTreeExecutionContext& Context, const EStateTreeRunStatus CompletionStatus, const FStateTreeActiveStates& CompletedActiveStates) const override
 	{
-		FTestStateTreeExecutionContext& TestContext = static_cast<FTestStateTreeExecutionContext&>(Context);
-		TestContext.Log(Name, TEXT("StateCompleted"));
+		FStateTreeTestLog& TestLog = Context.GetExternalData(LogHandle);
+		TestLog.Log(Name, TEXT("StateCompleted"));
 	}
 	
 	virtual EStateTreeRunStatus Tick(FStateTreeExecutionContext& Context, const float DeltaTime) const override
 	{
-		FTestStateTreeExecutionContext& TestContext = static_cast<FTestStateTreeExecutionContext&>(Context);
-		TestContext.Log(Name, TEXT("Tick"));
+		FStateTreeTestLog& TestLog = Context.GetExternalData(LogHandle);
+		TestLog.Log(Name, TEXT("Tick"));
 
 		FInstanceDataType& InstanceData = Context.GetInstanceData(*this);
 		
@@ -332,6 +383,14 @@ struct FTestTask_Stand : public FStateTreeTaskBase
 		return (InstanceData.CurrentTick >= TicksToCompletion) ? TickCompletionResult : EStateTreeRunStatus::Running;
 	};
 
+	virtual bool Link(FStateTreeLinker& Linker) override
+	{
+		Linker.LinkExternalData(LogHandle);
+		return true;
+	}
+	
+	TStateTreeExternalDataHandle<FStateTreeTestLog> LogHandle;
+	
 	UPROPERTY(EditAnywhere, Category = Parameter)
 	int32 TicksToCompletion = 1;
 
