StateTree: Added warning about linked state's child states (which cannot be selected)

diff --git a/Engine/Plugins/Runtime/StateTree/Source/StateTreeEditorModule/Private/SStateTreeViewRow.cpp b/Engine/Plugins/Runtime/StateTree/Source/StateTreeEditorModule/Private/SStateTreeViewRow.cpp
--- a/Engine/Plugins/Runtime/StateTree/Source/StateTreeEditorModule/Private/SStateTreeViewRow.cpp
+++ b/Engine/Plugins/Runtime/StateTree/Source/StateTreeEditorModule/Private/SStateTreeViewRow.cpp
@@ -133,6 +133,22 @@ void SStateTreeViewRow::Construct(const FArguments& InArgs, const TSharedRef<STa
 									+ SOverlay::Slot()
 									[
 										SNew(SHorizontalBox)
+
+										// Warnings
+										+SHorizontalBox::Slot()
+										.VAlign(VAlign_Center)
+										.AutoWidth()
+										[
+											SNew(SBox)
+											.Padding(FMargin(0.f, 0.f, 4.f, 0.f))
+											.Visibility(this, &SStateTreeViewRow::GetWarningsVisibility)
+											[
+												SNew(SImage)
+												.Image(FAppStyle::Get().GetBrush("Icons.Warning"))
+												.ToolTipText(this, &SStateTreeViewRow::GetWarningsTooltipText)
+											]
+										]
+
 										// Conditions icon
 										+SHorizontalBox::Slot()
 										.VAlign(VAlign_Center)
@@ -952,6 +968,31 @@ EVisibility SStateTreeViewRow::GetLinkedStateVisibility() const
 	return EVisibility::Collapsed;
 }
 
+bool SStateTreeViewRow::GetStateWarnings(FText* OutText) const
+{
+	bool bHasWarnings = false;
+	
+	const UStateTreeState* State = WeakState.Get();
+	if (!State)
+	{
+		return bHasWarnings;
+	}
+
+	// Linked States cannot have children.
+	if ((State->Type == EStateTreeStateType::Linked
+		|| State->Type == EStateTreeStateType::LinkedAsset)
+		&& State->Children.Num() > 0)
+	{
+		if (OutText)
+		{
+			*OutText = LOCTEXT("LinkedStateChildWarning", "Linked State cannot have child states, because the state selection will enter to the linked state on activation.");
+		}
+		bHasWarnings = true;
+	}
+
+	return bHasWarnings;
+}
+
 FText SStateTreeViewRow::GetLinkedStateDesc() const
 {
 	const UStateTreeState* State = WeakState.Get();
@@ -972,6 +1013,18 @@ FText SStateTreeViewRow::GetLinkedStateDesc() const
 	return FText::GetEmpty();
 }
 
+EVisibility SStateTreeViewRow::GetWarningsVisibility() const
+{
+	return GetStateWarnings(nullptr) ? EVisibility::Visible : EVisibility::Collapsed;
+}
+
+FText SStateTreeViewRow::GetWarningsTooltipText() const
+{
+	FText Warnings = FText::GetEmpty();
+	GetStateWarnings(&Warnings);
+	return Warnings;
+}
+
 bool SStateTreeViewRow::HasParentTransitionForTrigger(const UStateTreeState& State, const EStateTreeTransitionTrigger Trigger) const
 {
 	EStateTreeTransitionTrigger CombinedTrigger = EStateTreeTransitionTrigger::None;
diff --git a/Engine/Plugins/Runtime/StateTree/Source/StateTreeEditorModule/Private/SStateTreeViewRow.h b/Engine/Plugins/Runtime/StateTree/Source/StateTreeEditorModule/Private/SStateTreeViewRow.h
--- a/Engine/Plugins/Runtime/StateTree/Source/StateTreeEditorModule/Private/SStateTreeViewRow.h
+++ b/Engine/Plugins/Runtime/StateTree/Source/StateTreeEditorModule/Private/SStateTreeViewRow.h
@@ -53,6 +53,10 @@ private:
 	EVisibility GetLinkedStateVisibility() const;
 	FText GetLinkedStateDesc() const;
 
+	bool GetStateWarnings(FText* OutText) const;
+	EVisibility GetWarningsVisibility() const;
+	FText GetWarningsTooltipText() const;
+
 	EVisibility GetCompletedTransitionVisibility() const;
 	EVisibility GetCompletedTransitionBreakpointVisibility() const;
 	FText GetCompletedTransitionsDesc() const;
diff --git a/Engine/Plugins/Runtime/StateTree/Source/StateTreeEditorModule/Private/StateTreeCompiler.cpp b/Engine/Plugins/Runtime/StateTree/Source/StateTreeEditorModule/Private/StateTreeCompiler.cpp
--- a/Engine/Plugins/Runtime/StateTree/Source/StateTreeEditorModule/Private/StateTreeCompiler.cpp
+++ b/Engine/Plugins/Runtime/StateTree/Source/StateTreeEditorModule/Private/StateTreeCompiler.cpp
@@ -409,6 +409,13 @@ bool FStateTreeCompiler::CreateStateRecursive(UStateTreeState& State, const FSta
 
 	FStateTreeCompilerLogStateScope LogStateScope(&State, Log);
 
+	if (State.Type == EStateTreeStateType::LinkedAsset
+		|| State.Type == EStateTreeStateType::Linked)
+	{
+		Log.Reportf(EMessageSeverity::Warning,
+			TEXT("Linked State cannot have child states, because the state selection will enter to the linked state on activation."));
+	}
+
 	const int32 StateIdx = StateTree->States.AddDefaulted();
 	FCompactStateTreeState& CompactState = StateTree->States[StateIdx];
 	CompactState.Name = State.Name;
