StateTree: Fixed UI when undoing
- When redoing add node + set node type the property customization for the node children is reused, but we cache the node ID from the initially empty node.
- Changed FBindableNodeInstanceDetails to always get the up to date ID

diff --git a/Engine/Plugins/Runtime/StateTree/Source/StateTreeEditorModule/Private/Customizations/StateTreeEditorNodeDetails.cpp b/Engine/Plugins/Runtime/StateTree/Source/StateTreeEditorModule/Private/Customizations/StateTreeEditorNodeDetails.cpp
--- a/Engine/Plugins/Runtime/StateTree/Source/StateTreeEditorModule/Private/Customizations/StateTreeEditorNodeDetails.cpp
+++ b/Engine/Plugins/Runtime/StateTree/Source/StateTreeEditorModule/Private/Customizations/StateTreeEditorNodeDetails.cpp
@@ -443,20 +443,23 @@ class FBindableNodeInstanceDetails : public FInstancedStructDataDetails
 {
 public:
 
-	FBindableNodeInstanceDetails(TSharedPtr<IPropertyHandle> InStructProperty, FGuid InID, UStateTreeEditorData* InEditorData)
+	FBindableNodeInstanceDetails(TSharedPtr<IPropertyHandle> InStructProperty, TSharedPtr<IPropertyHandle> InIDProperty, UStateTreeEditorData* InEditorData)
 		: FInstancedStructDataDetails(InStructProperty)
 		, EditorData(InEditorData)
+		, IDProperty(InIDProperty)
 	{
-		ID = InID;
 	}
 
 	virtual void OnChildRowAdded(IDetailPropertyRow& ChildRow)
 	{
+		FGuid ID;
+		UE::StateTree::PropertyHelpers::GetStructValue<FGuid>(IDProperty, ID);
+		
 		UE::StateTreeEditor::Internal::ModifyRow(ChildRow, ID, EditorData);
 	}
 
 	UStateTreeEditorData* EditorData;
-	FGuid ID;
+	TSharedPtr<IPropertyHandle> IDProperty;
 };
 
 ////////////////////////////////////
@@ -895,11 +898,11 @@ void FStateTreeEditorNodeDetails::CustomizeChildren(TSharedRef<class IPropertyHa
 	}
 	
 	// Node
-	TSharedRef<FBindableNodeInstanceDetails> NodeDetails = MakeShareable(new FBindableNodeInstanceDetails(NodeProperty, FGuid(), EditorData));
+	TSharedRef<FBindableNodeInstanceDetails> NodeDetails = MakeShareable(new FBindableNodeInstanceDetails(NodeProperty, {}, EditorData));
 	StructBuilder.AddCustomBuilder(NodeDetails);
 
 	// Instance
-	TSharedRef<FBindableNodeInstanceDetails> InstanceDetails = MakeShareable(new FBindableNodeInstanceDetails(InstanceProperty, ID, EditorData));
+	TSharedRef<FBindableNodeInstanceDetails> InstanceDetails = MakeShareable(new FBindableNodeInstanceDetails(InstanceProperty, IDProperty, EditorData));
 	StructBuilder.AddCustomBuilder(InstanceDetails);
 
 	// InstanceObject
