StateTree: Fixed StateTree state UI to refresh properly after undo

diff --git a/Engine/Plugins/Runtime/StateTree/Source/StateTreeEditorModule/Private/Customizations/StateTreeStateDetails.cpp b/Engine/Plugins/Runtime/StateTree/Source/StateTreeEditorModule/Private/Customizations/StateTreeStateDetails.cpp
--- a/Engine/Plugins/Runtime/StateTree/Source/StateTreeEditorModule/Private/Customizations/StateTreeStateDetails.cpp
+++ b/Engine/Plugins/Runtime/StateTree/Source/StateTreeEditorModule/Private/Customizations/StateTreeStateDetails.cpp
@@ -180,8 +180,8 @@ void FStateTreeStateDetails::CustomizeDetails(IDetailLayoutBuilder& DetailBuilde
 	MakeArrayCategory(DetailBuilder, "Transitions", LOCTEXT("StateDetailsTransitions", "Transitions"), 4, TransitionsProperty);
 
 	// Refresh the UI when the type changes.	
-	TSharedPtr<IPropertyUtilities> PropUtils = DetailBuilder.GetPropertyUtilities();
-	TypeProperty->SetOnPropertyValueChanged(FSimpleDelegate::CreateLambda([PropUtils] ()
+	PropUtils = DetailBuilder.GetPropertyUtilities();
+	TypeProperty->SetOnPropertyValueChanged(FSimpleDelegate::CreateLambda([PropUtils = PropUtils] ()
 	{
 		if (PropUtils.IsValid())
 		{
@@ -190,6 +190,25 @@ void FStateTreeStateDetails::CustomizeDetails(IDetailLayoutBuilder& DetailBuilde
 	}));
 }
 
+void FStateTreeStateDetails::PostUndo(bool bSuccess)
+{
+	// Refresh view on undo or redo so that the customization based on e.g. State type will be reflected correctly.
+	if (PropUtils.IsValid())
+	{
+		PropUtils->ForceRefresh();
+	}
+	
+}
+
+void FStateTreeStateDetails::PostRedo(bool bSuccess)
+{
+	// Refresh view on undo or redo so that the customization based on e.g. State type will be reflected correctly.
+	if (PropUtils.IsValid())
+	{
+		PropUtils->ForceRefresh();
+	}
+}
+
 void FStateTreeStateDetails::MakeArrayCategory(IDetailLayoutBuilder& DetailBuilder, FName CategoryName, const FText& DisplayName, int32 SortOrder, TSharedPtr<IPropertyHandle> PropertyHandle)
 {
 	IDetailCategoryBuilder& Category = DetailBuilder.EditCategory(CategoryName, DisplayName);
diff --git a/Engine/Plugins/Runtime/StateTree/Source/StateTreeEditorModule/Private/Customizations/StateTreeStateDetails.h b/Engine/Plugins/Runtime/StateTree/Source/StateTreeEditorModule/Private/Customizations/StateTreeStateDetails.h
--- a/Engine/Plugins/Runtime/StateTree/Source/StateTreeEditorModule/Private/Customizations/StateTreeStateDetails.h
+++ b/Engine/Plugins/Runtime/StateTree/Source/StateTreeEditorModule/Private/Customizations/StateTreeStateDetails.h
@@ -3,12 +3,14 @@
 #pragma once
 
 #include "IDetailCustomization.h"
+#include "EditorUndoClient.h"
 
 class IDetailLayoutBuilder;
 class IDetailChildrenBuilder;
 class IPropertyHandle;
+class IPropertyUtilities;
 
-class FStateTreeStateDetails : public IDetailCustomization
+class FStateTreeStateDetails : public IDetailCustomization, FSelfRegisteringEditorUndoClient
 {
 public:
 	/** Makes a new instance of this detail layout class for a specific detail view requesting it */
@@ -18,5 +20,12 @@ public:
 	virtual void CustomizeDetails(IDetailLayoutBuilder& DetailBuilder) override;
 
 protected:
+
+	//~ FEditorUndoClient
+	virtual void PostUndo(bool bSuccess) override;
+	virtual void PostRedo(bool bSuccess) override;
+
 	void MakeArrayCategory(IDetailLayoutBuilder& DetailBuilder, FName CategoryName, const FText& DisplayName, int32 SortOrder, TSharedPtr<IPropertyHandle> PropertyHandle);
+
+	TSharedPtr<IPropertyUtilities> PropUtils;
 };
