StateTree: Fix parameter handling entering to a subtree directly.
- When transitioning to a subtree directly, self host the parameters
- Made transitioning to a subtree directly a warning, may make it an error in future versions

diff --git a/Engine/Plugins/Runtime/StateTree/Source/StateTreeEditorModule/Private/StateTreeCompiler.cpp b/Engine/Plugins/Runtime/StateTree/Source/StateTreeEditorModule/Private/StateTreeCompiler.cpp
--- a/Engine/Plugins/Runtime/StateTree/Source/StateTreeEditorModule/Private/StateTreeCompiler.cpp
+++ b/Engine/Plugins/Runtime/StateTree/Source/StateTreeEditorModule/Private/StateTreeCompiler.cpp
@@ -944,6 +944,18 @@ bool FStateTreeCompiler::CreateStateTransitions()
 				return false;
 			}
 
+			if (CompactTransition.State.IsValid()
+				&& !CompactTransition.State.IsCompletionState())
+			{
+				FCompactStateTreeState& TransitionTargetState = StateTree->States[CompactTransition.State.Index];
+				if (TransitionTargetState.Type == EStateTreeStateType::Subtree)
+				{
+					Log.Reportf(EMessageSeverity::Warning,
+						TEXT("Transitioning directly to a Subtree State '%s' is not recommended, as it may have unexpected results. Subtree States should be used with Linked States instead."),
+						*TransitionTargetState.Name.ToString());
+				}
+			}
+			
 			if (Transition.Trigger == EStateTreeTransitionTrigger::OnEvent)
 			{
 				CompactTransition.RequiredEvent.Tag = Transition.RequiredEvent.Tag;
@@ -965,13 +977,14 @@ bool FStateTreeCompiler::CreateStateTransitions()
 					return false;
 				}
 
-				if (!CompactTransition.State.IsCompletionState())
+				if (CompactTransition.State.IsValid()
+					&& !CompactTransition.State.IsCompletionState())
 				{
 					FCompactStateTreeState& TransitionTargetState = StateTree->States[CompactTransition.State.Index];
 					if (TransitionTargetState.RequiredEventToEnter.IsValid() && !TransitionTargetState.RequiredEventToEnter.IsSubsetOfAnotherDesc(CompactTransition.RequiredEvent))
 					{
 						Log.Reportf(EMessageSeverity::Error, Desc,
-							TEXT("On Event transition to %s will never succeed as transition and state required events are incompatible."),
+							TEXT("On Event transition to '%s' will never succeed as transition and state required events are incompatible."),
 							*TransitionTargetState.Name.ToString());
 						return false;
 					}
diff --git a/Engine/Plugins/Runtime/StateTree/Source/StateTreeModule/Private/StateTreeExecutionContext.cpp b/Engine/Plugins/Runtime/StateTree/Source/StateTreeModule/Private/StateTreeExecutionContext.cpp
--- a/Engine/Plugins/Runtime/StateTree/Source/StateTreeModule/Private/StateTreeExecutionContext.cpp
+++ b/Engine/Plugins/Runtime/StateTree/Source/StateTreeModule/Private/StateTreeExecutionContext.cpp
@@ -1061,27 +1061,30 @@ FStateTreeDataView FStateTreeExecutionContext::GetDataViewFromInstanceStorage(FS
 		return CurrentlyProcessedSharedInstanceStorage->GetMutableObject(Handle.GetIndex());
 
 	case EStateTreeDataSourceType::GlobalParameterData:
+		// Defined in parent frame or is root state tree parameters
+		if (ParentFrame)
 		{
-			// Defined in parent frame or is root state tree parameters
-			if (ParentFrame)
-			{
-				return GetDataViewFromInstanceStorage(InstanceDataStorage, CurrentlyProcessedSharedInstanceStorage, nullptr, *ParentFrame, CurrentFrame.GlobalParameterDataHandle);
-			}
-
-			return InstanceDataStorage.GetMutableGlobalParameters();
+			return GetDataViewFromInstanceStorage(InstanceDataStorage, CurrentlyProcessedSharedInstanceStorage, nullptr, *ParentFrame, CurrentFrame.GlobalParameterDataHandle);
 		}
+		return InstanceDataStorage.GetMutableGlobalParameters();
 
 	case EStateTreeDataSourceType::SubtreeParameterData:
 		{
 			// Defined in parent frame.
-			check(ParentFrame);
-			return GetDataViewFromInstanceStorage(InstanceDataStorage, CurrentlyProcessedSharedInstanceStorage, nullptr, *ParentFrame, CurrentFrame.StateParameterDataHandle);
+			if (ParentFrame)
+			{
+				// Linked subtree, params defined in parent scope.
+				return GetDataViewFromInstanceStorage(InstanceDataStorage, CurrentlyProcessedSharedInstanceStorage, nullptr, *ParentFrame, CurrentFrame.StateParameterDataHandle);
+			}
+			// Standalone subtree, params define as state params.
+			FCompactStateTreeParameters& SubtreeParams = InstanceDataStorage.GetMutableStruct(CurrentFrame.ActiveInstanceIndexBase.Get() + Handle.GetIndex()).Get<FCompactStateTreeParameters>();
+			return SubtreeParams.Parameters.GetMutableValue();
 		}
 
 	case EStateTreeDataSourceType::StateParameterData:
 		{
-			FCompactStateTreeParameters& Params = InstanceDataStorage.GetMutableStruct(CurrentFrame.ActiveInstanceIndexBase.Get() + Handle.GetIndex()).Get<FCompactStateTreeParameters>();
-			return Params.Parameters.GetMutableValue();
+			FCompactStateTreeParameters& StateParams = InstanceDataStorage.GetMutableStruct(CurrentFrame.ActiveInstanceIndexBase.Get() + Handle.GetIndex()).Get<FCompactStateTreeParameters>();
+			return StateParams.Parameters.GetMutableValue();
 		}
 
 	case EStateTreeDataSourceType::StateEvent:
@@ -1169,8 +1172,15 @@ bool FStateTreeExecutionContext::IsHandleSourceValid(const FStateTreeExecutionFr
 			: CurrentFrame.GlobalParameterDataHandle.IsValid();
 
 	case EStateTreeDataSourceType::SubtreeParameterData:
-		return ParentFrame
-			&& IsHandleSourceValid(nullptr, *ParentFrame, CurrentFrame.StateParameterDataHandle);
+		if (ParentFrame)
+		{
+			// Linked subtree, params defined in parent scope.
+			return IsHandleSourceValid(nullptr, *ParentFrame, CurrentFrame.StateParameterDataHandle);
+		}
+		// Standalone subtree, params define as state params.
+		return CurrentFrame.ActiveInstanceIndexBase.IsValid()
+			&& CurrentFrame.ActiveStates.Contains(Handle.GetState(), CurrentFrame.NumCurrentlyActiveStates)
+			&& InstanceDataStorage->IsValidIndex(CurrentFrame.ActiveInstanceIndexBase.Get() + Handle.GetIndex());
 
 	case EStateTreeDataSourceType::StateParameterData:
 		return CurrentFrame.ActiveInstanceIndexBase.IsValid()
@@ -1225,19 +1235,23 @@ FStateTreeDataView FStateTreeExecutionContext::GetDataViewOrTemporary(const FSta
 	case EStateTreeDataSourceType::SubtreeParameterData:
 		if (ParentFrame)
 		{
+			// Linked subtree, params defined in parent scope.
 			if (FCompactStateTreeParameters* Params = InstanceDataStorage->GetMutableTemporaryStruct(*ParentFrame, CurrentFrame.StateParameterDataHandle).GetPtr<FCompactStateTreeParameters>())
 			{
 				return Params->Parameters.GetMutableValue();
 			}
 		}
+		// Standalone subtree, params define as state params.
+		if (FCompactStateTreeParameters* Params = InstanceDataStorage->GetMutableTemporaryStruct(CurrentFrame, Handle).GetPtr<FCompactStateTreeParameters>())
+		{
+			return Params->Parameters.GetMutableValue();
+		}
 		break;
 
 	case EStateTreeDataSourceType::StateParameterData:
+		if (FCompactStateTreeParameters* Params = InstanceDataStorage->GetMutableTemporaryStruct(CurrentFrame, Handle).GetPtr<FCompactStateTreeParameters>())
 		{
-			if (FCompactStateTreeParameters* Params = InstanceDataStorage->GetMutableTemporaryStruct(CurrentFrame, Handle).GetPtr<FCompactStateTreeParameters>())
-			{
-				return Params->Parameters.GetMutableValue();
-			}
+			return Params->Parameters.GetMutableValue();
 		}
 		break;
 
