StateTree: Fixed duplicate localizations strings.

diff --git a/Engine/Plugins/Runtime/StateTree/Source/StateTreeEditorModule/Private/Customizations/StateTreeStateLinkDetails.cpp b/Engine/Plugins/Runtime/StateTree/Source/StateTreeEditorModule/Private/Customizations/StateTreeStateLinkDetails.cpp
--- a/Engine/Plugins/Runtime/StateTree/Source/StateTreeEditorModule/Private/Customizations/StateTreeStateLinkDetails.cpp
+++ b/Engine/Plugins/Runtime/StateTree/Source/StateTreeEditorModule/Private/Customizations/StateTreeStateLinkDetails.cpp
@@ -198,31 +198,31 @@ void FStateTreeStateLinkDetails::CacheStates()
 	if (!bDirectStatesOnly)
 	{
 		RootItem->Children.Add(MakeShared<FStateTreeStateItem>(
-			LOCTEXT("TransitionNone", "<i>None</>"),
+			LOCTEXT("TransitionNoneRich", "<i>None</>"),
 			LOCTEXT("TransitionNoneTooltip", "No transition."),
 			FStateTreeEditorStyle::Get().GetBrush("StateTreeEditor.Transition.None"),
 			EStateTreeTransitionType::None));
 
 		RootItem->Children.Add(MakeShared<FStateTreeStateItem>(
-			LOCTEXT("TransitionNextState", "<i>Next State</>"),
+			LOCTEXT("TransitionNextStateRich", "<i>Next State</>"),
 			LOCTEXT("TransitionNextTooltip", "Goto next sibling State."),
 			FStateTreeEditorStyle::Get().GetBrush("StateTreeEditor.Transition.Next"),
 			EStateTreeTransitionType::NextState));
 
 		RootItem->Children.Add(MakeShared<FStateTreeStateItem>(
-			LOCTEXT("TransitionNextSelectableState", "<i>Next Selectable State</>"),
+			LOCTEXT("TransitionNextSelectableStateRich", "<i>Next Selectable State</>"),
 			LOCTEXT("TransitionNextSelectableTooltip", "Goto next sibling state, whose enter conditions pass."),
 			FStateTreeEditorStyle::Get().GetBrush("StateTreeEditor.Transition.Next"),
 			EStateTreeTransitionType::NextSelectableState));
 
 		RootItem->Children.Add(MakeShared<FStateTreeStateItem>(
-			LOCTEXT("TransitionTreeSucceeded", "<i>Tree Succeeded</>"),
+			LOCTEXT("TransitionTreeSucceededRich", "<i>Tree Succeeded</>"),
 			LOCTEXT("TransitionTreeSuccessTooltip", "Complete tree with success."),
 			FStateTreeEditorStyle::Get().GetBrush("StateTreeEditor.Transition.Succeeded"),
 			EStateTreeTransitionType::Succeeded));
 
 		RootItem->Children.Add(MakeShared<FStateTreeStateItem>(
-			LOCTEXT("TransitionTreeFailed", "<i>Tree Failed</>"),
+			LOCTEXT("TransitionTreeFailedRich", "<i>Tree Failed</>"),
 			LOCTEXT("TransitionTreeFailedTooltip", "Complete tree with failure."),
 			FStateTreeEditorStyle::Get().GetBrush("StateTreeEditor.Transition.Failed"),
 			EStateTreeTransitionType::Failed));
