Mark FStateTreeReference BlueprintType
#rnx

diff --git a/Engine/Plugins/Runtime/StateTree/Source/StateTreeModule/Public/StateTreeReference.h b/Engine/Plugins/Runtime/StateTree/Source/StateTreeModule/Public/StateTreeReference.h
--- a/Engine/Plugins/Runtime/StateTree/Source/StateTreeModule/Public/StateTreeReference.h
+++ b/Engine/Plugins/Runtime/StateTree/Source/StateTreeModule/Public/StateTreeReference.h
@@ -11,7 +11,7 @@ class UStateTree;
 /**
  * Struct to hold reference to a StateTree asset along with values to parameterized it.
  */
-USTRUCT()
+USTRUCT(BlueprintType)
 struct STATETREEMODULE_API FStateTreeReference
 {
 	GENERATED_BODY()
