Add compile in StateTree PostLoad and compile on PIE begin to state tree so that the trees are never out of sync with the editor state in game
Add validation warning if the state tree is not compiled to prevent uncompiled state tree to be submitted

diff --git a/Engine/Plugins/Runtime/StateTree/Source/StateTreeEditorModule/Private/StateTreeEditorModule.cpp b/Engine/Plugins/Runtime/StateTree/Source/StateTreeEditorModule/Private/StateTreeEditorModule.cpp
--- a/Engine/Plugins/Runtime/StateTree/Source/StateTreeEditorModule/Private/StateTreeEditorModule.cpp
+++ b/Engine/Plugins/Runtime/StateTree/Source/StateTreeEditorModule/Private/StateTreeEditorModule.cpp
@@ -73,6 +73,7 @@ namespace UE::StateTree::Editor
 void FStateTreeEditorModule::StartupModule()
 {
 	UE::StateTree::Delegates::OnRequestCompile.BindStatic(&UE::StateTree::Editor::CompileStateTree);
+	UE::StateTree::Delegates::OnRequestEditorHash.BindStatic(&UE::StateTree::Editor::CalcAssetHash);
 
 #if WITH_STATETREE_DEBUGGER
 	FStateTreeDebuggerCommands::Register();
diff --git a/Engine/Plugins/Runtime/StateTree/Source/StateTreeModule/Private/StateTree.cpp b/Engine/Plugins/Runtime/StateTree/Source/StateTreeModule/Private/StateTree.cpp
--- a/Engine/Plugins/Runtime/StateTree/Source/StateTreeModule/Private/StateTree.cpp
+++ b/Engine/Plugins/Runtime/StateTree/Source/StateTreeModule/Private/StateTree.cpp
@@ -15,6 +15,7 @@
 #include "Misc/EnumerateRange.h"
 #include "UObject/AssetRegistryTagsContext.h"
 #if WITH_EDITOR
+#include "Editor.h"
 #include "Engine/UserDefinedStruct.h"
 #endif
 #include UE_INLINE_GENERATED_CPP_BY_NAME(StateTree)
@@ -247,6 +248,7 @@ void UStateTree::PostInitProperties()
 	
 	OnObjectsReinstancedHandle = FCoreUObjectDelegates::OnObjectsReinstanced.AddUObject(this, &UStateTree::OnObjectsReinstanced);
 	OnUserDefinedStructReinstancedHandle = UE::StructUtils::Delegates::OnUserDefinedStructReinstanced.AddUObject(this, &UStateTree::OnUserDefinedStructReinstanced);
+	FEditorDelegates::PreBeginPIE.AddUObject(this, &UStateTree::OnPreBeginPIE);
 }
 
 void UStateTree::BeginDestroy()
@@ -262,6 +264,10 @@ void UStateTree::BeginDestroy()
 		OnUserDefinedStructReinstancedHandle.Reset();
 	}
 	
+#if WITH_EDITOR
+	FEditorDelegates::PreBeginPIE.RemoveAll(this);
+#endif // WITH_EDITOR
+
 	Super::BeginDestroy();
 }
 
@@ -298,6 +304,20 @@ void UStateTree::ThreadedPostLoadAssetRegistryTagsOverride(FPostLoadAssetRegistr
 
 EDataValidationResult UStateTree::IsDataValid(FDataValidationContext& Context) const
 {
+	// Don't warn user that the tree they just saved is not compiled. Only for submit or manual validation
+	if (Context.GetValidationUsecase() != EDataValidationUsecase::Save)
+	{
+		if (UE::StateTree::Delegates::OnRequestEditorHash.IsBound())
+		{
+			const uint32 CurrentHash = UE::StateTree::Delegates::OnRequestEditorHash.Execute(*this);
+			if (CurrentHash != LastCompiledEditorDataHash)
+			{
+				Context.AddWarning(FText::FromString(FString::Printf(TEXT("%s is not compiled. Please recompile the State Tree."), *GetFullName())));
+				return EDataValidationResult::Invalid;
+			}
+		}
+	}
+
 	if (!const_cast<UStateTree*>(this)->Link())
 	{
 		Context.AddError(FText::FromString(FString::Printf(TEXT("%s failed to link. Please recompile the State Tree for more details errors."), *GetFullName())));
@@ -355,16 +375,15 @@ void UStateTree::PostLoad()
 #endif // WITH_EDITOR
 
 	const int32 CurrentVersion = GetLinkerCustomVersion(FStateTreeCustomVersion::GUID);
+#if WITH_EDITOR
+	if (EditorData)
+	{
+		// Make sure all the fix up logic in the editor data has had chance to happen.
+		EditorData->ConditionalPostLoad();
+	}
 
 	if (CurrentVersion < FStateTreeCustomVersion::LatestVersion)
 	{
-#if WITH_EDITOR
-		if (EditorData)
-		{
-			// Make sure all the fix up logic in the editor data has had chance to happen.
-			EditorData->ConditionalPostLoad();
-		}
-		
 		// Compiled data is in older format, try to compile the StateTree.
 		if (UE::StateTree::Delegates::OnRequestCompile.IsBound())
 		{
@@ -377,12 +396,19 @@ void UStateTree::PostLoad()
 			ResetCompiled();
 			UE_LOG(LogStateTree, Warning, TEXT("%s: compiled data is in older format. Please resave the StateTree asset."), *GetFullName());
 		}
+	}
+	else
+	{
+		CompileIfChanged();
+	}
 #else
+	if (CurrentVersion < FStateTreeCustomVersion::LatestVersion)
+	{		
 		UE_LOG(LogStateTree, Error, TEXT("%s: compiled data is in older format. Please recompile the StateTree asset."), *GetFullName());
-#endif
 		return;
 	}
-	
+#endif
+
 	if (!Link())
 	{
 		UE_LOG(LogStateTree, Log, TEXT("%s failed to link. Asset will not be usable at runtime."), *GetFullName());	
@@ -977,5 +1003,24 @@ TArray<FStateTreeMemoryUsage> UStateTree::CalculateEstimatedMemoryUsage() const
 
 	return MemoryUsages;
 }
+
+void UStateTree::OnPreBeginPIE(const bool bIsSimulating)
+{
+	CompileIfChanged();
+}
+
+void UStateTree::CompileIfChanged()
+{
+	if (UE::StateTree::Delegates::OnRequestCompile.IsBound() && UE::StateTree::Delegates::OnRequestEditorHash.IsBound())
+	{
+		const uint32 CurrentHash = UE::StateTree::Delegates::OnRequestEditorHash.Execute(*this);
+		if (LastCompiledEditorDataHash != CurrentHash)
+		{
+			UE_LOG(LogStateTree, Log, TEXT("%s: Editor data has changed. Recompiling state tree."), *GetFullName());
+			UE::StateTree::Delegates::OnRequestCompile.Execute(*this);
+		}
+	}
+}
+
 #endif // WITH_EDITOR
 
diff --git a/Engine/Plugins/Runtime/StateTree/Source/StateTreeModule/Private/StateTreeDelegates.cpp b/Engine/Plugins/Runtime/StateTree/Source/StateTreeModule/Private/StateTreeDelegates.cpp
--- a/Engine/Plugins/Runtime/StateTree/Source/StateTreeModule/Private/StateTreeDelegates.cpp
+++ b/Engine/Plugins/Runtime/StateTree/Source/StateTreeModule/Private/StateTreeDelegates.cpp
@@ -14,6 +14,7 @@ FOnStateParametersChanged OnStateParametersChanged;
 FOnBreakpointsChanged OnBreakpointsChanged;
 FOnPostCompile OnPostCompile;
 FOnRequestCompile OnRequestCompile;
+FOnRequestEditorHash OnRequestEditorHash;
 #endif // WITH_EDITOR
 
 #if WITH_STATETREE_DEBUGGER
diff --git a/Engine/Plugins/Runtime/StateTree/Source/StateTreeModule/Public/StateTree.h b/Engine/Plugins/Runtime/StateTree/Source/StateTreeModule/Public/StateTree.h
--- a/Engine/Plugins/Runtime/StateTree/Source/StateTreeModule/Public/StateTree.h
+++ b/Engine/Plugins/Runtime/StateTree/Source/StateTreeModule/Public/StateTree.h
@@ -172,6 +172,12 @@ public:
 
 	/** Calculates runtime memory usage for different sections of the tree. */
 	TArray<FStateTreeMemoryUsage> CalculateEstimatedMemoryUsage() const;
+
+	/** Called when the editor is preparing to start a pie session. */
+	void OnPreBeginPIE(const bool bIsSimulating);
+
+	/** Compile the state trees if the editor hash data as changed since the last compilation. */
+	void CompileIfChanged();
 #endif
 
 #if WITH_EDITORONLY_DATA
diff --git a/Engine/Plugins/Runtime/StateTree/Source/StateTreeModule/Public/StateTreeDelegates.h b/Engine/Plugins/Runtime/StateTree/Source/StateTreeModule/Public/StateTreeDelegates.h
--- a/Engine/Plugins/Runtime/StateTree/Source/StateTreeModule/Public/StateTreeDelegates.h
+++ b/Engine/Plugins/Runtime/StateTree/Source/StateTreeModule/Public/StateTreeDelegates.h
@@ -60,6 +60,10 @@ extern STATETREEMODULE_API FOnPostCompile OnPostCompile;
 DECLARE_DELEGATE_RetVal_OneParam(bool, FOnRequestCompile, UStateTree& /*StateTreeToCompile*/);
 extern STATETREEMODULE_API FOnRequestCompile OnRequestCompile;
 
+/** Request the editor hash of the given state tree. */
+DECLARE_DELEGATE_RetVal_OneParam(uint32, FOnRequestEditorHash, const UStateTree& /*StateTreeToHash*/);
+extern STATETREEMODULE_API FOnRequestEditorHash OnRequestEditorHash;
+
 #endif // WITH_EDITOR
 
 #if WITH_STATETREE_DEBUGGER
