[State Tree] Uniformly Random State Selection
- Added a new option for selecint child at uniform random in StateTreeStateSelectionBehavior
- Added RandomStream in StateTreeInstanceData and relevant APIs to get stream and set seed

diff --git a/Engine/Plugins/Runtime/StateTree/Source/StateTreeModule/Private/StateTreeExecutionContext.cpp b/Engine/Plugins/Runtime/StateTree/Source/StateTreeModule/Private/StateTreeExecutionContext.cpp
--- a/Engine/Plugins/Runtime/StateTree/Source/StateTreeModule/Private/StateTreeExecutionContext.cpp
+++ b/Engine/Plugins/Runtime/StateTree/Source/StateTreeModule/Private/StateTreeExecutionContext.cpp
@@ -3835,7 +3835,42 @@ bool FStateTreeExecutionContext::SelectStateInternal(
 				return true;
 			}
 		}
+		else if (NextState.SelectionBehavior == EStateTreeStateSelectionBehavior::TrySelectChildrenAtUniformRandom)
+		{
+			if (NextState.HasChildren())
+			{
+				STATETREE_TRACE_SCOPED_STATE_PHASE(NextStateHandle, EStateTreeUpdatePhase::TrySelectBehavior);
 
+				TArray<uint16, TInlineAllocator<8>> RandomChildStates;
+				for (uint16 ChildState = NextState.ChildrenBegin; ChildState < NextState.ChildrenEnd; ChildState = CurrentStateTree->States[ChildState].GetNextSibling())
+				{
+					RandomChildStates.Push(ChildState);
+				}
+
+				const uint32 LastIndex = RandomChildStates.Num() - 1;
+				for (uint32 Index = 0; Index < LastIndex; ++Index)
+				{
+					// Get a random integer in [Index, Num)
+					const uint32 IndexToSwap = InstanceDataStorage->GetRandomStream().RandRange(Index, LastIndex);
+					RandomChildStates.Swap(Index, IndexToSwap);
+				}
+
+				for (uint16 ChildState : RandomChildStates)
+				{
+					if (SelectStateInternal(CurrentParentFrame, CurrentFrame, CurrentFrameInActiveFrames, { FStateTreeStateHandle(ChildState) }, OutSelectionResult))
+					{
+						// Selection succeeded
+						return true;
+					}
+				}
+			}
+			else
+			{
+				// Select this state (For backwards compatibility)
+				STATETREE_TRACE_STATE_EVENT(NextStateHandle, EStateTreeTraceEventType::OnStateSelected);
+				return true;
+			}
+		}
 		// State could not be selected, restore.
 		CurrentFrame.NumCurrentlyActiveStates = PrevNumCurrentlyActiveStates;
 		CurrentFrame.ActiveStates.Pop();
diff --git a/Engine/Plugins/Runtime/StateTree/Source/StateTreeModule/Private/StateTreeInstanceData.cpp b/Engine/Plugins/Runtime/StateTree/Source/StateTreeModule/Private/StateTreeInstanceData.cpp
--- a/Engine/Plugins/Runtime/StateTree/Source/StateTreeModule/Private/StateTreeInstanceData.cpp
+++ b/Engine/Plugins/Runtime/StateTree/Source/StateTreeModule/Private/StateTreeInstanceData.cpp
@@ -83,7 +83,10 @@ namespace UE::StateTree
 // FStateTreeInstanceStorage
 //----------------------------------------------------------------//
 
-FStateTreeInstanceStorage::FStateTreeInstanceStorage() = default;
+FStateTreeInstanceStorage::FStateTreeInstanceStorage()
+	: RandomStream(FPlatformTime::Cycles())
+{
+}
 
 FStateTreeInstanceStorage::FStateTreeInstanceStorage(const FStateTreeInstanceStorage& Other)
 	: InstanceStructs(Other.InstanceStructs)
@@ -92,6 +95,7 @@ FStateTreeInstanceStorage::FStateTreeInstanceStorage(const FStateTreeInstanceSto
 	, EventQueue(MakeShared<FStateTreeEventQueue>(*Other.EventQueue))
 	, TransitionRequests(Other.TransitionRequests)
 	, GlobalParameters(Other.GlobalParameters)
+	, RandomStream(Other.RandomStream)
 {
 }
 
@@ -102,6 +106,7 @@ FStateTreeInstanceStorage::FStateTreeInstanceStorage(FStateTreeInstanceStorage&&
 	, EventQueue(Other.EventQueue)
 	, TransitionRequests(MoveTemp(Other.TransitionRequests))
 	, GlobalParameters(MoveTemp(Other.GlobalParameters))
+	, RandomStream(MoveTemp(Other.RandomStream))
 {
 	Other.EventQueue = MakeShared<FStateTreeEventQueue>();
 }
@@ -114,6 +119,7 @@ FStateTreeInstanceStorage& FStateTreeInstanceStorage::operator=(const FStateTree
 	EventQueue = MakeShared<FStateTreeEventQueue>(*Other.EventQueue);
 	TransitionRequests = Other.TransitionRequests;
 	GlobalParameters = Other.GlobalParameters;
+	RandomStream = Other.RandomStream;
 
 	return *this;
 }
@@ -127,6 +133,7 @@ FStateTreeInstanceStorage& FStateTreeInstanceStorage::operator=(FStateTreeInstan
 	Other.EventQueue = MakeShared<FStateTreeEventQueue>();
 	TransitionRequests = MoveTemp(Other.TransitionRequests);
 	GlobalParameters = MoveTemp(Other.GlobalParameters);
+	RandomStream = MoveTemp(Other.RandomStream);
 
 	return *this;
 }
@@ -251,6 +258,11 @@ void FStateTreeInstanceStorage::SetGlobalParameters(const FInstancedPropertyBag&
 	GlobalParameters = Parameters;
 }
 
+void FStateTreeInstanceStorage::SetRandomSeed(int32 Seed)
+{
+	RandomStream.Initialize(Seed);
+}
+
 void FStateTreeInstanceStorage::AddStructReferencedObjects(FReferenceCollector& Collector)
 {
 	Collector.AddPropertyReferencesWithStructARO(TBaseStructure<FStateTreeInstanceStorage>::Get(), this);
@@ -267,6 +279,7 @@ void FStateTreeInstanceStorage::Reset()
 	ExecutionState.Reset();
 	TemporaryInstances.Reset();
 	GlobalParameters.Reset();
+	RandomStream.Reset();
 }
 
 //----------------------------------------------------------------//
@@ -390,6 +403,11 @@ int32 FStateTreeInstanceData::GetEstimatedMemoryUsage() const
 	return Size;
 }
 
+void FStateTreeInstanceData::SetRandomSeed(int32 Seed)
+{
+	GetMutableStorage().SetRandomSeed(Seed);
+}
+
 bool FStateTreeInstanceData::Identical(const FStateTreeInstanceData* Other, uint32 PortFlags) const
 {
 	if (Other == nullptr)
diff --git a/Engine/Plugins/Runtime/StateTree/Source/StateTreeModule/Public/StateTreeInstanceData.h b/Engine/Plugins/Runtime/StateTree/Source/StateTreeModule/Public/StateTreeInstanceData.h
--- a/Engine/Plugins/Runtime/StateTree/Source/StateTreeModule/Public/StateTreeInstanceData.h
+++ b/Engine/Plugins/Runtime/StateTree/Source/StateTreeModule/Public/StateTreeInstanceData.h
@@ -239,6 +239,15 @@ struct STATETREEMODULE_API FStateTreeInstanceStorage
 	/** Resets the storage to initial state. */
 	void Reset();
 	
+	/** Update seed for the RandomStream used by state tree random-based operations. */
+	void SetRandomSeed(int32 Seed);
+
+	/** @return RandomStream used by state tree random-based operations. */
+	const FRandomStream& GetRandomStream() const
+	{
+		return RandomStream;
+	}
+
 	UE_DEPRECATED(5.4, "Use Num() instead.")
 	int32 NumStructs() const { return 0; }
 
@@ -278,6 +287,10 @@ protected:
 	UPROPERTY(Transient)
 	FInstancedPropertyBag GlobalParameters;
 
+	/** Used by state tree random-based operations. By Default it is initialized with seed 0 */
+	UPROPERTY()
+	FRandomStream RandomStream;
+
 	friend struct FStateTreeInstanceData;
 };
 
@@ -412,6 +425,15 @@ struct STATETREEMODULE_API FStateTreeInstanceData
 
 	int32 GetEstimatedMemoryUsage() const;
 	
+	/** @return RandomStream used by state tree random-based operations. */
+	const FRandomStream& GetRandomStream() const
+	{
+		return GetStorage().GetRandomStream();
+	}
+
+	/** Update seed for the RandomStream used by state tree random-based operations. */
+	void SetRandomSeed(int32 Seed);
+
 	/** Type traits */
 	bool Identical(const FStateTreeInstanceData* Other, uint32 PortFlags) const;
 	void AddStructReferencedObjects(FReferenceCollector& Collector);
diff --git a/Engine/Plugins/Runtime/StateTree/Source/StateTreeModule/Public/StateTreeTypes.h b/Engine/Plugins/Runtime/StateTree/Source/StateTreeModule/Public/StateTreeTypes.h
--- a/Engine/Plugins/Runtime/StateTree/Source/StateTreeModule/Public/StateTreeTypes.h
+++ b/Engine/Plugins/Runtime/StateTree/Source/StateTreeModule/Public/StateTreeTypes.h
@@ -94,11 +94,15 @@ enum class EStateTreeStateSelectionBehavior : uint8
 	/** When state is considered for selection, it is selected even if it has child states. */
 	TryEnterState UMETA(DisplayName = "Try Enter"),
 
-	/** When state is considered for selection, try to selects the first child state (in order they appear in the child list). If no child states are present, behaves like SelectState. */
+	/** When state is considered for selection, try to select the first child state (in order they appear in the child list). If no child states are present, behaves like SelectState. */
 	TrySelectChildrenInOrder UMETA(DisplayName = "Try Select Children In Order"),
 	
+	/** When state is considered for selection, shuffle the order of child states and try to select the first one. If no child states are present, behaves like SelectState. */
+	TrySelectChildrenAtUniformRandom UMETA(DisplayName = "Try Select Children At Uniform Random"),
+
 	/** When state is considered for selection, try to trigger the transitions instead. */
 	TryFollowTransitions UMETA(DisplayName = "Try Follow Transitions"),
+
 };
 
 
