[State Tree] Added icon for random selection

diff --git a/Engine/Plugins/Runtime/StateTree/Source/StateTreeEditorModule/Private/SStateTreeViewRow.cpp b/Engine/Plugins/Runtime/StateTree/Source/StateTreeEditorModule/Private/SStateTreeViewRow.cpp
--- a/Engine/Plugins/Runtime/StateTree/Source/StateTreeEditorModule/Private/SStateTreeViewRow.cpp
+++ b/Engine/Plugins/Runtime/StateTree/Source/StateTreeEditorModule/Private/SStateTreeViewRow.cpp
@@ -853,6 +853,15 @@ const FSlateBrush* SStateTreeViewRow::GetSelectorIcon() const
 				return FStateTreeEditorStyle::Get().GetBrush("StateTreeEditor.TrySelectChildrenInOrder");
 			}
 		}
+		else if (State->SelectionBehavior == EStateTreeStateSelectionBehavior::TrySelectChildrenAtUniformRandom)
+		{
+			if (State->Children.IsEmpty())
+			{
+				return FStateTreeEditorStyle::Get().GetBrush("StateTreeEditor.TryEnterState");			
+			}
+
+			return FStateTreeEditorStyle::Get().GetBrush("StateTreeEditor.TrySelectCHildrenAtRandom");
+		}
 		else if (State->SelectionBehavior == EStateTreeStateSelectionBehavior::TryFollowTransitions)
 		{
 			return FStateTreeEditorStyle::Get().GetBrush("StateTreeEditor.TryFollowTransitions");
diff --git a/Engine/Plugins/Runtime/StateTree/Source/StateTreeEditorModule/Private/StateTreeEditorStyle.cpp b/Engine/Plugins/Runtime/StateTree/Source/StateTreeEditorModule/Private/StateTreeEditorStyle.cpp
--- a/Engine/Plugins/Runtime/StateTree/Source/StateTreeEditorModule/Private/StateTreeEditorStyle.cpp
+++ b/Engine/Plugins/Runtime/StateTree/Source/StateTreeEditorModule/Private/StateTreeEditorStyle.cpp
@@ -318,6 +318,7 @@ FStateTreeEditorStyle::FStateTreeEditorStyle()
 		Set("StateTreeEditor.SelectNone", new IMAGE_BRUSH_SVG("Icons/Select_None", CoreStyleConstants::Icon16x16));
 		Set("StateTreeEditor.TryEnterState", new IMAGE_BRUSH_SVG("Icons/Try_Enter_State", CoreStyleConstants::Icon16x16));
 		Set("StateTreeEditor.TrySelectChildrenInOrder", new IMAGE_BRUSH_SVG("Icons/Try_Select_Children_In_Order", CoreStyleConstants::Icon16x16));
+		Set("StateTreeEditor.TrySelectChildrenAtRandom", new IMAGE_BRUSH_SVG("Icons/Try_Select_Children_At_Random", CoreStyleConstants::Icon16x16));
 		Set("StateTreeEditor.TryFollowTransitions", new IMAGE_BRUSH_SVG("Icons/Try_Follow_Transitions", CoreStyleConstants::Icon16x16));
 
 		Set("StateTreeEditor.Conditions", new IMAGE_BRUSH_SVG("Icons/Conditions", CoreStyleConstants::Icon16x16));
