StateTree: Debugger fixes and improvements
- Another fix for stuck timeline
- Removed exec context to emit  EStateTreeUpdatePhase::TickStateTree scope, to allow cleaner logs for parallel STs

diff --git a/Engine/Plugins/Runtime/StateTree/Source/StateTreeEditorModule/Private/Debugger/SStateTreeDebuggerView.cpp b/Engine/Plugins/Runtime/StateTree/Source/StateTreeEditorModule/Private/Debugger/SStateTreeDebuggerView.cpp
--- a/Engine/Plugins/Runtime/StateTree/Source/StateTreeEditorModule/Private/Debugger/SStateTreeDebuggerView.cpp
+++ b/Engine/Plugins/Runtime/StateTree/Source/StateTreeEditorModule/Private/Debugger/SStateTreeDebuggerView.cpp
@@ -683,19 +683,28 @@ void SStateTreeDebuggerView::Tick(const FGeometry& AllottedGeometry, const doubl
 	QUICK_SCOPE_CYCLE_COUNTER(STAT_StateTreeDebuggerView_TickView);
 
 	check(Debugger);
-	if (Debugger->IsAnalysisSessionActive())
-	{
-		MaxTrackRecordingDuration = FMath::Max(MaxTrackRecordingDuration, Debugger->GetRecordingDuration());
 
-		// Stick to most recent data if auto scroll is enabled.
-		// Autoscroll is disabled when paused.
-		// This allows the user to pause the analysis, inspect the data, and continue and the autoscroll will catch up with latest.
-		// Complementary logic in OnTimeLineScrubPositionChanged().
-		if (!Debugger->IsAnalysisSessionPaused()
-			&& bAutoScroll)
+	const double RecordingDuration = Debugger->GetRecordingDuration();
+	const bool bHasMoreRecentData = LastUpdatedTrackRecordingDuration != RecordingDuration;
+	LastUpdatedTrackRecordingDuration = RecordingDuration;
+	MaxTrackRecordingDuration = FMath::Max(MaxTrackRecordingDuration, RecordingDuration);
+	
+	if ((Debugger->IsAnalysisSessionActive() && !Debugger->IsAnalysisSessionPaused())
+		|| bHasMoreRecentData)
+	{
+		if (bAutoScroll)
 		{
+			// Stick to most recent data if auto scroll is enabled.
+			// Autoscroll is disabled when paused.
+			// This allows the user to pause the analysis, inspect the data, and continue and the autoscroll will catch up with latest.
+			// Complementary logic in OnTimeLineScrubPositionChanged().
 			Debugger->SetScrubTime(Debugger->GetRecordingDuration());
 		}
+		else
+		{
+			// Set scrub time to self to request update the UI.
+			Debugger->SetScrubTime(Debugger->GetScrubTime());
+		}
 	}
 	
 	RefreshTracks();
diff --git a/Engine/Plugins/Runtime/StateTree/Source/StateTreeEditorModule/Private/Debugger/SStateTreeDebuggerView.h b/Engine/Plugins/Runtime/StateTree/Source/StateTreeEditorModule/Private/Debugger/SStateTreeDebuggerView.h
--- a/Engine/Plugins/Runtime/StateTree/Source/StateTreeEditorModule/Private/Debugger/SStateTreeDebuggerView.h
+++ b/Engine/Plugins/Runtime/StateTree/Source/StateTreeEditorModule/Private/Debugger/SStateTreeDebuggerView.h
@@ -161,6 +161,9 @@ private:
 	/** In case tracks are not reset when a new analysis session is started we keep track of the longest duration to adjust our clamp range. */
 	double MaxTrackRecordingDuration = 0;
 
+	/** The recording duration time the UI was last updated from the debugger. Used to detect if new data has been collect while the UI was inactive. */
+	double LastUpdatedTrackRecordingDuration = 0;
+
 	/** Indicates that a live session was started (record button or auto record in PIE) to generate StateTree traces. */
 	bool bRecording = false;
 
diff --git a/Engine/Plugins/Runtime/StateTree/Source/StateTreeModule/Private/StateTreeExecutionContext.cpp b/Engine/Plugins/Runtime/StateTree/Source/StateTreeModule/Private/StateTreeExecutionContext.cpp
--- a/Engine/Plugins/Runtime/StateTree/Source/StateTreeModule/Private/StateTreeExecutionContext.cpp
+++ b/Engine/Plugins/Runtime/StateTree/Source/StateTreeModule/Private/StateTreeExecutionContext.cpp
@@ -504,9 +504,6 @@ EStateTreeRunStatus FStateTreeExecutionContext::TickPostlude()
 EStateTreeRunStatus FStateTreeExecutionContext::Tick(const float DeltaTime)
 {
 	CSV_SCOPED_TIMING_STAT_EXCLUSIVE(StateTree_Tick);
-	STATETREE_TRACE_SCOPED_PHASE(EStateTreeUpdatePhase::TickStateTree);
-
-//	STATETREE_LOG_AND_TRACE(Warning, TEXT("Tick %f"), GetWorld()->GetTimeSeconds());
 
 	const EStateTreeRunStatus PreludeResult = TickPrelude();
 	if (PreludeResult != EStateTreeRunStatus::Running)
@@ -523,9 +520,6 @@ EStateTreeRunStatus FStateTreeExecutionContext::Tick(const float DeltaTime)
 EStateTreeRunStatus FStateTreeExecutionContext::TickUpdateTasks(const float DeltaTime)
 {
 	CSV_SCOPED_TIMING_STAT_EXCLUSIVE(StateTree_Tick);
-	STATETREE_TRACE_SCOPED_PHASE(EStateTreeUpdatePhase::TickStateTree);
-
-//	STATETREE_LOG_AND_TRACE(Warning, TEXT("TickUpdateTasks %f"), GetWorld()->GetTimeSeconds());
 
 	const EStateTreeRunStatus PreludeResult = TickPrelude();
 	if (PreludeResult != EStateTreeRunStatus::Running)
@@ -541,9 +535,6 @@ EStateTreeRunStatus FStateTreeExecutionContext::TickUpdateTasks(const float Delt
 EStateTreeRunStatus FStateTreeExecutionContext::TickTriggerTransitions()
 {
 	CSV_SCOPED_TIMING_STAT_EXCLUSIVE(StateTree_Tick);
-	STATETREE_TRACE_SCOPED_PHASE(EStateTreeUpdatePhase::TickStateTree);
-
-//	STATETREE_LOG_AND_TRACE(Warning, TEXT("TickTriggerTransitions %f"), GetWorld()->GetTimeSeconds());
 
 	const EStateTreeRunStatus PreludeResult = TickPrelude();
 	if (PreludeResult != EStateTreeRunStatus::Running)
