Make FStateTreeExecutionContext::AreContextDataViewsValid properly account for optional external data.
#rnx

diff --git a/Engine/Plugins/Runtime/StateTree/Source/StateTreeModule/Private/StateTreeExecutionContext.cpp b/Engine/Plugins/Runtime/StateTree/Source/StateTreeModule/Private/StateTreeExecutionContext.cpp
--- a/Engine/Plugins/Runtime/StateTree/Source/StateTreeModule/Private/StateTreeExecutionContext.cpp
+++ b/Engine/Plugins/Runtime/StateTree/Source/StateTreeModule/Private/StateTreeExecutionContext.cpp
@@ -199,11 +199,22 @@ bool FStateTreeExecutionContext::AreContextDataViewsValid() const
 	{
 		const FStateTreeDataView& DataView = ContextAndExternalDataViews[DataDesc.Handle.DataHandle.GetIndex()];
 
-		// Items must have valid pointer of the expected type.  
-		if (!DataView.IsValid() || !DataView.GetStruct()->IsChildOf(DataDesc.Struct))
+		// Required items must have valid pointer of the expected type.  
+		if (DataDesc.Requirement == EStateTreeExternalDataRequirement::Required)
 		{
-			bResult = false;
-			break;
+			if (!DataView.IsValid() || !DataDesc.IsCompatibleWith(DataView))
+			{
+				bResult = false;
+				break;
+			}
+		}
+		else // Optional items must have the expected type if they are set.
+		{
+			if (DataView.IsValid() && !DataDesc.IsCompatibleWith(DataView))
+			{
+				bResult = false;
+				break;
+			}
 		}
 	}
 	return bResult;
