[StateTreeDebugger]
- new recording started from the debugger will be selected as the new active analysis
- allow multiple recordings during the same PIE session to reactive existing tracks
- do not apply trace selection if user reselect the current trace to avoid clearing data

diff --git a/Engine/Plugins/Runtime/StateTree/Source/StateTreeEditorModule/Private/Debugger/SStateTreeDebuggerView.cpp b/Engine/Plugins/Runtime/StateTree/Source/StateTreeEditorModule/Private/Debugger/SStateTreeDebuggerView.cpp
--- a/Engine/Plugins/Runtime/StateTree/Source/StateTreeEditorModule/Private/Debugger/SStateTreeDebuggerView.cpp
+++ b/Engine/Plugins/Runtime/StateTree/Source/StateTreeEditorModule/Private/Debugger/SStateTreeDebuggerView.cpp
@@ -229,6 +229,10 @@ void SStateTreeDebuggerView::StartRecording()
 		check(Debugger);
 		Debugger->ClearSelection();
 
+		// Stop the current analysis so we have a chance to reconnect to an existing instance
+		// if there is existing data and it is still active.
+		Debugger->StopSessionAnalysis();
+
 		// We give priority to the Editor actions even if an analysis was active (remote process)
 		// This will stop current analysis and connect to the new live trace.
 		bRecording = Debugger->RequestAnalysisOfEditorSession();
@@ -1309,6 +1313,15 @@ void SStateTreeDebuggerView::OnNewSession()
 	{
 		ResetTracks();
 	}
+	else if (Debugger->GetSelectedInstanceDescriptor() == nullptr)
+	{
+		// In PIE it is possible to stop/start the recording multiple times during the same game session,
+		// in this case we try to reselect the currently selected instance in case it can be reactivated.
+		if (FStateTreeDebuggerBaseTrack* DebuggerBaseTrack = static_cast<FStateTreeDebuggerBaseTrack*>(InstancesTreeView->GetSelection().Get()))
+		{
+			DebuggerBaseTrack->OnSelected();
+		}
+	}
 
 	// Restore automatic scroll to most recent data.
 	bAutoScroll = true;
@@ -1384,7 +1397,11 @@ TSharedRef<SWidget> SStateTreeDebuggerView::OnGetDebuggerTracesMenu() const
 
 		FUIAction ItemAction(FExecuteAction::CreateSPLambda(Debugger.ToSharedRef(), [Debugger = Debugger, TraceDescriptor]()
 			{
-				Debugger->RequestSessionAnalysis(TraceDescriptor);
+				// Request new analysis only if user picked a different trace (we don't want to clear the tracks)
+				if (Debugger->GetSelectedTraceDescriptor() != TraceDescriptor)
+				{
+					Debugger->RequestSessionAnalysis(TraceDescriptor);
+				}
 			}));
 		MenuBuilder.AddMenuEntry(Desc, TAttribute<FText>(), FSlateIcon(), ItemAction);
 	}
diff --git a/Engine/Plugins/Runtime/StateTree/Source/StateTreeEditorModule/Private/Debugger/StateTreeDebuggerTrack.cpp b/Engine/Plugins/Runtime/StateTree/Source/StateTreeEditorModule/Private/Debugger/StateTreeDebuggerTrack.cpp
--- a/Engine/Plugins/Runtime/StateTree/Source/StateTreeEditorModule/Private/Debugger/StateTreeDebuggerTrack.cpp
+++ b/Engine/Plugins/Runtime/StateTree/Source/StateTreeEditorModule/Private/Debugger/StateTreeDebuggerTrack.cpp
@@ -72,11 +72,14 @@ bool FStateTreeDebuggerInstanceTrack::UpdateInternal()
 			Window.Description = FText::FromString(StatePath);
 			Window.TimeStart = Span.GetWorldTimeStart();
 
-			// When there is another state change after the current one in the list we use its start time to close the window.
+			// When there is another state change after the current one in the list we use
+			// the end time of the previous frame span to close the window. Not using the start time
+			// since we might have a gap in case of multiple recordings during the same game session.
 			if (StateChangeIndex < NumStateChanges-1)
 			{
-				const uint32 NextSpanIndex = EventCollection.ActiveStatesChanges[StateChangeIndex+1].SpanIndex;
-				Window.TimeEnd = EventCollection.FrameSpans[NextSpanIndex].GetWorldTimeStart();
+				const uint32 EndSpanIndex = EventCollection.ActiveStatesChanges[StateChangeIndex+1].SpanIndex-1;
+				check(EventCollection.FrameSpans.IsValidIndex(EndSpanIndex));
+				Window.TimeEnd = EventCollection.FrameSpans[EndSpanIndex].GetWorldTimeEnd();
 			}
 			else
 			{
@@ -104,6 +107,13 @@ bool FStateTreeDebuggerInstanceTrack::UpdateInternal()
 	}
 
 	const bool bChanged = (PrevNumPoints != EventData->Points.Num() || PrevNumWindows != EventData->Windows.Num());
+
+	// Tracks can be reactivated when multiple recordings are made in a single PIE session.
+	if (bChanged && IsStale())
+	{
+		bIsStale = false;
+	}
+
 	return bChanged;
 }
 
diff --git a/Engine/Plugins/Runtime/StateTree/Source/StateTreeModule/Private/Debugger/StateTreeDebugger.cpp b/Engine/Plugins/Runtime/StateTree/Source/StateTreeModule/Private/Debugger/StateTreeDebugger.cpp
--- a/Engine/Plugins/Runtime/StateTree/Source/StateTreeModule/Private/Debugger/StateTreeDebugger.cpp
+++ b/Engine/Plugins/Runtime/StateTree/Source/StateTreeModule/Private/Debugger/StateTreeDebugger.cpp
@@ -309,7 +309,6 @@ bool FStateTreeDebugger::StartSessionAnalysis(const FTraceDescriptor& TraceDescr
 		return false;
 	}
 
-	RecordingDuration = 0;
 	AnalysisDuration = 0;
 	LastTraceReadTime = 0;
 
@@ -854,19 +853,47 @@ bool FStateTreeDebugger::ProcessEvent(const FStateTreeInstanceDebugId InstanceId
 	check(ExistingCollection);
 	TArray<FStateTreeTraceEventVariantType>& Events = ExistingCollection->Events;
 
-	// Add new frame span if none added yet or new frame
-	if (ExistingCollection->FrameSpans.IsEmpty() || ExistingCollection->FrameSpans.Last().Frame.Index < Frame.Index)
+	TraceServices::FFrame FrameToAddInSpans = Frame;
+	bool bShouldAddFrameToSpans = false;
+
+	double RecordingWorldTime = 0;
+	Visit([&RecordingWorldTime](auto& TypedEvent)
+		{
+			RecordingWorldTime = TypedEvent.RecordingWorldTime;
+		}, Event);
+
+	// Add new frame span if none added yet
+	if (ExistingCollection->FrameSpans.IsEmpty())
 	{
-		double RecordingWorldTime = 0;
-		Visit([&RecordingWorldTime](auto& TypedEvent)
-			{
-				RecordingWorldTime = TypedEvent.RecordingWorldTime;
-			}, Event);
+		bShouldAddFrameToSpans = true;
+	}
+	else
+	{
+		const TraceServices::FFrame& LastFrame = ExistingCollection->FrameSpans.Last().Frame;
+		// Add new frame span for new larger frame index
+		if (Frame.Index > LastFrame.Index)
+		{
+			bShouldAddFrameToSpans = true;
+		}
+		else if (Frame.Index < LastFrame.Index && Frame.StartTime > LastFrame.StartTime)
+		{
+			// Some events are buffered and can be sent from an older world recording
+			// time in case of late recording (e.g., ActiveStatesEvent) so if we want to aggregate the events
+			// with existing data we'll snap it to the most recent time to not break the timelines
+			RecordingWorldTime = FMath::Max(RecordingWorldTime, RecordingDuration);
+
+			// Frame index will restart at 0 if a new session is started,
+			// in that case we offset the frame we store to append to existing data
+			FrameToAddInSpans.Index += LastFrame.Index + 1;
+			bShouldAddFrameToSpans = true;
+		}
+	}
 
+	if (bShouldAddFrameToSpans)
+	{
 		// Update global recording duration
 		RecordingDuration = RecordingWorldTime;
-
-		ExistingCollection->FrameSpans.Add(UE::StateTreeDebugger::FFrameSpan(Frame, RecordingWorldTime, Events.Num()));
+		ExistingCollection->FrameSpans.Add(UE::StateTreeDebugger::FFrameSpan(FrameToAddInSpans, RecordingWorldTime, Events.Num()));
 	}
 
 	// Add activate states change info
@@ -912,6 +939,7 @@ void FStateTreeDebugger::ResetEventCollections()
 {
 	EventCollections.Reset();
 	SetScrubStateCollectionIndex(INDEX_NONE);
+	RecordingDuration = 0;
 }
 
 void FStateTreeDebugger::AddEvents(const double StartTime, const double EndTime, const TraceServices::IFrameProvider& FrameProvider, const IStateTreeTraceProvider& StateTreeTraceProvider)
