[StateTreeDebugger] Moved Breakpoint commands to ViewModel and main TreeView

diff --git a/Engine/Plugins/Runtime/StateTree/Source/StateTreeEditorModule/Private/Debugger/SStateTreeDebuggerView.cpp b/Engine/Plugins/Runtime/StateTree/Source/StateTreeEditorModule/Private/Debugger/SStateTreeDebuggerView.cpp
--- a/Engine/Plugins/Runtime/StateTree/Source/StateTreeEditorModule/Private/Debugger/SStateTreeDebuggerView.cpp
+++ b/Engine/Plugins/Runtime/StateTree/Source/StateTreeEditorModule/Private/Debugger/SStateTreeDebuggerView.cpp
@@ -377,21 +377,6 @@ void SStateTreeDebuggerView::Construct(const FArguments& InArgs, const UStateTre
 	// Register the play world commands
 	InCommandList->Append(FPlayWorldCommands::GlobalPlayWorldActions.ToSharedRef());
 
-	// Debug commands
-	InCommandList->MapAction(
-		FStateTreeDebuggerCommands::Get().EnableOnEnterStateBreakpoint,
-		FExecuteAction::CreateLambda([this] { HandleEnableStateBreakpoint(EStateTreeBreakpointType::OnEnter); }),
-		FCanExecuteAction(),
-		FGetActionCheckState::CreateLambda([this]{ return GetStateBreakpointCheckState(EStateTreeBreakpointType::OnEnter); }),
-		FIsActionButtonVisible::CreateLambda([this] { return CanAddStateBreakpoint(EStateTreeBreakpointType::OnEnter) || CanRemoveStateBreakpoint(EStateTreeBreakpointType::OnEnter); }));
-
-	InCommandList->MapAction(
-		FStateTreeDebuggerCommands::Get().EnableOnExitStateBreakpoint,
-		FExecuteAction::CreateLambda([this] { HandleEnableStateBreakpoint(EStateTreeBreakpointType::OnExit); }),
-		FCanExecuteAction(),
-		FGetActionCheckState::CreateLambda([this] { return GetStateBreakpointCheckState(EStateTreeBreakpointType::OnExit); }),
-		FIsActionButtonVisible::CreateLambda([this] { return CanAddStateBreakpoint(EStateTreeBreakpointType::OnExit) || CanRemoveStateBreakpoint(EStateTreeBreakpointType::OnExit); }));
-
 	// Toolbars
 	FSlimHorizontalToolBarBuilder LeftToolbar(InCommandList, FMultiBoxCustomization::None, /*InExtender*/ nullptr, /*InForceSmallIcons*/ true);
 	LeftToolbar.BeginSection(TEXT("Debugging"));
@@ -837,201 +822,6 @@ void SStateTreeDebuggerView::StepForwardToNextStateChange()
 	bAutoScroll = false;
 }
 
-bool SStateTreeDebuggerView::CanAddStateBreakpoint(const EStateTreeBreakpointType Type) const
-{
-	check(StateTreeViewModel);
-
-	TArray<UStateTreeState*> SelectedStates;
-	StateTreeViewModel->GetSelectedStates(SelectedStates);
-	if (SelectedStates.IsEmpty())
-	{
-		return false;
-	}
-
-	const UStateTreeEditorData* EditorData = StateTreeEditorData.Get();
-	if (!ensure(EditorData != nullptr))
-	{
-		return false;
-	}
-
-	for (const UStateTreeState* SelectedState : SelectedStates)
-	{
-		if (SelectedState != nullptr)
-		{
-			if (EditorData->HasBreakpoint(SelectedState->ID, Type) == false)
-			{
-				return true;
-			}
-		}
-	}
-
-	return false;
-}
-
-bool SStateTreeDebuggerView::CanRemoveStateBreakpoint(const EStateTreeBreakpointType Type) const
-{
-	check(StateTreeViewModel);
-
-	TArray<UStateTreeState*> SelectedStates;
-	StateTreeViewModel->GetSelectedStates(SelectedStates);
-	if (SelectedStates.IsEmpty())
-	{
-		return false;
-	}
-
-	const UStateTreeEditorData* EditorData = StateTreeEditorData.Get();
-	if (!ensure(EditorData != nullptr))
-	{
-		return false;
-	}
-
-	for (const UStateTreeState* SelectedState : SelectedStates)
-	{
-		if (SelectedState != nullptr)
-		{
-			if (EditorData->HasBreakpoint(SelectedState->ID, Type))
-			{
-				return true;
-			}
-		}
-	}
-
-	return false;
-}
-
-
-ECheckBoxState SStateTreeDebuggerView::GetStateBreakpointCheckState(const EStateTreeBreakpointType Type) const
-{
-	const bool bCanAdd = CanAddStateBreakpoint(Type);
-	const bool bCanRemove = CanRemoveStateBreakpoint(Type);
-	if (bCanAdd && bCanRemove)
-	{
-		return ECheckBoxState::Undetermined;
-	}
-
-	if (bCanRemove)
-	{
-		return ECheckBoxState::Checked;
-	}
-
-	if (bCanAdd)
-	{
-		return ECheckBoxState::Unchecked;
-	}
-
-	// Should not happen since action is not visible in this case
-	return ECheckBoxState::Undetermined;
-}
-
-void SStateTreeDebuggerView::HandleEnableStateBreakpoint(EStateTreeBreakpointType Type)
-{
-	check(StateTreeViewModel);
-
-	TArray<UStateTreeState*> SelectedStates;
-	StateTreeViewModel->GetSelectedStates(SelectedStates);
-	if (SelectedStates.IsEmpty())
-	{
-		return;
-	}
-
-	UStateTreeEditorData* EditorData = StateTreeEditorData.Get();
-	if (!ensure(EditorData != nullptr))
-	{
-		return;
-	}
-
-	TBitArray<> HasBreakpoint;
-	HasBreakpoint.Reserve(SelectedStates.Num());
-	for (const UStateTreeState* SelectedState : SelectedStates)
-	{
-		HasBreakpoint.Add(SelectedState != nullptr && EditorData->HasBreakpoint(SelectedState->ID, Type));
-	}
-
-	check(HasBreakpoint.Num() == SelectedStates.Num());
-
-	// Process CanAdd first so in case of undetermined state (mixed selection) we add by default. 
-	if (CanAddStateBreakpoint(Type))
-	{
-		const FScopedTransaction Transaction(LOCTEXT("AddStateBreakpoint", "Add State Breakpoint(s)"));
-		EditorData->Modify();
-		for (int Index = 0; Index < SelectedStates.Num(); ++Index)
-		{
-			const UStateTreeState* SelectedState = SelectedStates[Index];
-			if (HasBreakpoint[Index] == false && SelectedState != nullptr)
-			{
-				EditorData->AddBreakpoint(SelectedState->ID, Type);	
-			}
-		}
-	}
-	else if (CanRemoveStateBreakpoint(Type))
-	{
-		const FScopedTransaction Transaction(LOCTEXT("RemoveStateBreakpoint", "Remove State Breakpoint(s)"));
-		EditorData->Modify();
-		for (int Index = 0; Index < SelectedStates.Num(); ++Index)
-		{
-			const UStateTreeState* SelectedState = SelectedStates[Index];
-			if (HasBreakpoint[Index] && SelectedState != nullptr)
-			{
-				EditorData->RemoveBreakpoint(SelectedState->ID, Type);	
-			}
-		}
-	}
-}
-
-UStateTreeState* SStateTreeDebuggerView::FindStateAssociatedToBreakpoint(FStateTreeDebuggerBreakpoint Breakpoint) const
-{
-	const UStateTree* Tree = StateTree.Get();
-	UStateTreeEditorData* TreeEditorData = StateTreeEditorData.Get();
-	if (Tree == nullptr || TreeEditorData == nullptr)
-	{
-		return nullptr;
-	}
-
-	UStateTreeState* StateTreeState = nullptr;
-
-	if (const FStateTreeStateHandle* StateHandle = Breakpoint.ElementIdentifier.TryGet<FStateTreeStateHandle>())
-	{
-		const FGuid StateId = StateTree->GetStateIdFromHandle(*StateHandle);
-		StateTreeState = TreeEditorData->GetMutableStateByID(StateId);
-	}
-	else if (const FStateTreeDebuggerBreakpoint::FStateTreeTaskIndex* TaskIndex = Breakpoint.ElementIdentifier.TryGet<FStateTreeDebuggerBreakpoint::FStateTreeTaskIndex>())
-	{
-		const FGuid TaskId = StateTree->GetNodeIdFromIndex(TaskIndex->Index);
-
-		TreeEditorData->VisitHierarchy([&TaskId, &StateTreeState](UStateTreeState& State, UStateTreeState* /*ParentState*/)
-			{
-				for (const FStateTreeEditorNode& EditorNode : State.Tasks)
-				{
-					if (EditorNode.ID == TaskId)
-					{
-						StateTreeState = &State;
-						return EStateTreeVisitor::Break;
-					}
-				}
-				return EStateTreeVisitor::Continue;
-			});
-	}
-	else if (const FStateTreeDebuggerBreakpoint::FStateTreeTransitionIndex* TransitionIndex = Breakpoint.ElementIdentifier.TryGet<FStateTreeDebuggerBreakpoint::FStateTreeTransitionIndex>())
-	{
-		const FGuid TransitionId = StateTree->GetTransitionIdFromIndex(TransitionIndex->Index);
-
-		TreeEditorData->VisitHierarchy([&TransitionId, &StateTreeState](UStateTreeState& State, UStateTreeState* /*ParentState*/)
-			{
-				for (const FStateTreeTransition& StateTransition : State.Transitions)
-				{
-					if (StateTransition.ID == TransitionId)
-					{
-						StateTreeState = &State;
-						return EStateTreeVisitor::Break;
-					}
-				}
-				return EStateTreeVisitor::Continue;
-			});
-	}
-
-	return StateTreeState;
-}
-
 void SStateTreeDebuggerView::OnTimeLineScrubPositionChanged(double Time, bool bIsScrubbing)
 {
 	check(Debugger);
@@ -1272,9 +1062,9 @@ void SStateTreeDebuggerView::OnBreakpointHit(const FStateTreeInstanceDebugId Ins
 	}
 
 	// Extract associated UStateTreeState to focus on it.
-	if (UStateTreeState* AssociatedState = FindStateAssociatedToBreakpoint(Breakpoint))
+	check(StateTreeViewModel);
+	if (UStateTreeState* AssociatedState = StateTreeViewModel->FindStateAssociatedToBreakpoint(Breakpoint))
 	{
-		check(StateTreeViewModel);
 		StateTreeViewModel->SetSelection(AssociatedState);
 	}
 
diff --git a/Engine/Plugins/Runtime/StateTree/Source/StateTreeEditorModule/Private/Debugger/SStateTreeDebuggerView.h b/Engine/Plugins/Runtime/StateTree/Source/StateTreeEditorModule/Private/Debugger/SStateTreeDebuggerView.h
--- a/Engine/Plugins/Runtime/StateTree/Source/StateTreeEditorModule/Private/Debugger/SStateTreeDebuggerView.h
+++ b/Engine/Plugins/Runtime/StateTree/Source/StateTreeEditorModule/Private/Debugger/SStateTreeDebuggerView.h
@@ -99,13 +99,6 @@ private:
 	bool CanStepForwardToNextStateChange() const;
 	void StepForwardToNextStateChange();
 
-	bool CanAddStateBreakpoint(EStateTreeBreakpointType Type) const;
-	bool CanRemoveStateBreakpoint(EStateTreeBreakpointType Type) const;
-	ECheckBoxState GetStateBreakpointCheckState(EStateTreeBreakpointType Type) const;
-	void HandleEnableStateBreakpoint(EStateTreeBreakpointType Type);
-
-	UStateTreeState* FindStateAssociatedToBreakpoint(FStateTreeDebuggerBreakpoint Breakpoint) const;
-
 	/** Callback from timeline widgets to update the debugger scrub state. */
 	void OnTimeLineScrubPositionChanged(double Time, bool bIsScrubbing);
 
diff --git a/Engine/Plugins/Runtime/StateTree/Source/StateTreeEditorModule/Private/Debugger/StateTreeDebuggerCommands.cpp b/Engine/Plugins/Runtime/StateTree/Source/StateTreeEditorModule/Private/Debugger/StateTreeDebuggerCommands.cpp
--- a/Engine/Plugins/Runtime/StateTree/Source/StateTreeEditorModule/Private/Debugger/StateTreeDebuggerCommands.cpp
+++ b/Engine/Plugins/Runtime/StateTree/Source/StateTreeEditorModule/Private/Debugger/StateTreeDebuggerCommands.cpp
@@ -22,9 +22,6 @@ FStateTreeDebuggerCommands::FStateTreeDebuggerCommands()
 
 void FStateTreeDebuggerCommands::RegisterCommands()
 {
-	UI_COMMAND(EnableOnEnterStateBreakpoint, "Break on Enter", "Adds or removes a breakpoint when entering the selected state(s)", EUserInterfaceActionType::Check, FInputChord(EKeys::F9));
-	UI_COMMAND(EnableOnExitStateBreakpoint, "Break on Exit", "Adds or removes a breakpoint when exiting the selected state(s)", EUserInterfaceActionType::Check, FInputChord(EModifierKey::Shift, EKeys::F9));
-
 	UI_COMMAND(StartRecording, "Start Recording", "Start a new trace session.", EUserInterfaceActionType::Button, FInputChord(EModifierKey::Control | EModifierKey::Shift, EKeys::R));
 	UI_COMMAND(StopRecording, "Stop Recording", "Stop the current trace session.", EUserInterfaceActionType::Button, FInputChord(EModifierKey::Control | EModifierKey::Shift, EKeys::S));
 	
diff --git a/Engine/Plugins/Runtime/StateTree/Source/StateTreeEditorModule/Private/Debugger/StateTreeDebuggerCommands.h b/Engine/Plugins/Runtime/StateTree/Source/StateTreeEditorModule/Private/Debugger/StateTreeDebuggerCommands.h
--- a/Engine/Plugins/Runtime/StateTree/Source/StateTreeEditorModule/Private/Debugger/StateTreeDebuggerCommands.h
+++ b/Engine/Plugins/Runtime/StateTree/Source/StateTreeEditorModule/Private/Debugger/StateTreeDebuggerCommands.h
@@ -17,8 +17,6 @@ public:
 	// TCommands<> overrides
 	virtual void RegisterCommands() override;
 
-	TSharedPtr<FUICommandInfo> EnableOnEnterStateBreakpoint;
-	TSharedPtr<FUICommandInfo> EnableOnExitStateBreakpoint;
 	TSharedPtr<FUICommandInfo> StartRecording;
 	TSharedPtr<FUICommandInfo> StopRecording;
 	TSharedPtr<FUICommandInfo> PreviousFrameWithStateChange;
diff --git a/Engine/Plugins/Runtime/StateTree/Source/StateTreeEditorModule/Private/SStateTreeView.cpp b/Engine/Plugins/Runtime/StateTree/Source/StateTreeEditorModule/Private/SStateTreeView.cpp
--- a/Engine/Plugins/Runtime/StateTree/Source/StateTreeEditorModule/Private/SStateTreeView.cpp
+++ b/Engine/Plugins/Runtime/StateTree/Source/StateTreeEditorModule/Private/SStateTreeView.cpp
@@ -1,16 +1,16 @@
 // Copyright Epic Games, Inc. All Rights Reserved.
 
 #include "SStateTreeView.h"
+#include "Debugger/StateTreeDebuggerTypes.h"
+#include "Framework/Commands/UICommandList.h"
 #include "Framework/MultiBox/MultiBoxBuilder.h"
-#include "Widgets/Layout/SScrollBox.h"
 #include "SPositiveActionButton.h"
 #include "SStateTreeViewRow.h"
-#include "Debugger/StateTreeDebuggerCommands.h"
 #include "StateTreeViewModel.h"
 #include "StateTreeState.h"
 #include "StateTreeEditorCommands.h"
 #include "StateTreeSettings.h"
-#include "Framework/Commands/UICommandList.h"
+#include "Widgets/Layout/SScrollBox.h"
 
 #define LOCTEXT_NAMESPACE "StateTreeEditor"
 
@@ -187,7 +187,7 @@ void SStateTreeView::BindCommands()
 		Commands.EnableStates,
 		FExecuteAction::CreateSP(this, &SStateTreeView::HandleEnableSelectedStates),
 		FCanExecuteAction(),
-		FGetActionCheckState::CreateLambda([this]
+		FGetActionCheckState::CreateSPLambda(this, [this]
 			{
 				const bool bCanEnable = CanEnableStates();
 				const bool bCanDisable = CanDisableStates();
@@ -209,7 +209,54 @@ void SStateTreeView::BindCommands()
 				// Should not happen since action is not visible in this case
 				return ECheckBoxState::Undetermined;
 			}),
-		FIsActionButtonVisible::CreateLambda([this] { return CanEnableStates() || CanDisableStates(); }));
+		FIsActionButtonVisible::CreateSPLambda(this, [this]
+		{
+			return CanEnableStates() || CanDisableStates();
+		}));
+
+#if WITH_STATETREE_DEBUGGER
+	CommandList->MapAction(
+		Commands.EnableOnEnterStateBreakpoint,
+		FExecuteAction::CreateSPLambda(this, [this]
+		{
+			if (StateTreeViewModel)
+			{
+				StateTreeViewModel->HandleEnableStateBreakpoint(EStateTreeBreakpointType::OnEnter);
+			}
+		}),
+		FCanExecuteAction(),
+		FGetActionCheckState::CreateSPLambda(this, [this]
+		{
+			return StateTreeViewModel ? StateTreeViewModel->GetStateBreakpointCheckState(EStateTreeBreakpointType::OnEnter) : ECheckBoxState::Unchecked;
+		}),
+		FIsActionButtonVisible::CreateSPLambda(this, [this]
+		{
+			return (StateTreeViewModel)
+				&& (StateTreeViewModel->CanAddStateBreakpoint(EStateTreeBreakpointType::OnEnter)
+					|| StateTreeViewModel->CanRemoveStateBreakpoint(EStateTreeBreakpointType::OnEnter));
+		}));
+
+	CommandList->MapAction(
+		Commands.EnableOnExitStateBreakpoint,
+		FExecuteAction::CreateSPLambda(this, [this]
+		{
+			if (StateTreeViewModel)
+			{
+				StateTreeViewModel->HandleEnableStateBreakpoint(EStateTreeBreakpointType::OnExit);
+			}
+		}),
+		FCanExecuteAction(),
+		FGetActionCheckState::CreateSPLambda(this, [this]
+		{
+			return StateTreeViewModel ? StateTreeViewModel->GetStateBreakpointCheckState(EStateTreeBreakpointType::OnExit) : ECheckBoxState::Unchecked;
+		}),
+		FIsActionButtonVisible::CreateSPLambda(this, [this]
+		{
+			return (StateTreeViewModel)
+				&& (StateTreeViewModel->CanAddStateBreakpoint(EStateTreeBreakpointType::OnExit)
+					|| StateTreeViewModel->CanRemoveStateBreakpoint(EStateTreeBreakpointType::OnExit));
+		}));
+#endif // WITH_STATETREE_DEBUGGER
 }
 
 bool SStateTreeView::HasSelection() const
@@ -478,8 +525,8 @@ TSharedPtr<SWidget> SStateTreeView::HandleContextMenuOpening()
 
 #if WITH_STATETREE_DEBUGGER
 	MenuBuilder.AddSeparator();
-	MenuBuilder.AddMenuEntry(FStateTreeDebuggerCommands::Get().EnableOnEnterStateBreakpoint);
-	MenuBuilder.AddMenuEntry(FStateTreeDebuggerCommands::Get().EnableOnExitStateBreakpoint);
+	MenuBuilder.AddMenuEntry(FStateTreeEditorCommands::Get().EnableOnEnterStateBreakpoint);
+	MenuBuilder.AddMenuEntry(FStateTreeEditorCommands::Get().EnableOnExitStateBreakpoint);
 #endif // WITH_STATETREE_DEBUGGER
 	
 	return MenuBuilder.MakeWidget();
diff --git a/Engine/Plugins/Runtime/StateTree/Source/StateTreeEditorModule/Private/StateTreeEditorCommands.cpp b/Engine/Plugins/Runtime/StateTree/Source/StateTreeEditorModule/Private/StateTreeEditorCommands.cpp
--- a/Engine/Plugins/Runtime/StateTree/Source/StateTreeEditorModule/Private/StateTreeEditorCommands.cpp
+++ b/Engine/Plugins/Runtime/StateTree/Source/StateTreeEditorModule/Private/StateTreeEditorCommands.cpp
@@ -32,6 +32,11 @@ void FStateTreeEditorCommands::RegisterCommands()
 	UI_COMMAND(DeleteStates, "Delete", "Delete Selected States", EUserInterfaceActionType::Button, FInputChord(EKeys::Delete));
 	UI_COMMAND(RenameState, "Rename", "Rename Selected State", EUserInterfaceActionType::Button, FInputChord(EKeys::F2));
 	UI_COMMAND(EnableStates, "State Enabled", "Enables Selected States", EUserInterfaceActionType::Check, FInputChord());
+
+#if WITH_STATETREE_DEBUGGER
+	UI_COMMAND(EnableOnEnterStateBreakpoint, "Break on Enter", "Adds or removes a breakpoint when entering the selected state(s)", EUserInterfaceActionType::Check, FInputChord(EKeys::F9));
+	UI_COMMAND(EnableOnExitStateBreakpoint, "Break on Exit", "Adds or removes a breakpoint when exiting the selected state(s)", EUserInterfaceActionType::Check, FInputChord(EModifierKey::Shift, EKeys::F9));
+#endif //WITH_STATETREE_DEBUGGER
 }
 
 #undef LOCTEXT_NAMESPACE
diff --git a/Engine/Plugins/Runtime/StateTree/Source/StateTreeEditorModule/Private/StateTreeEditorCommands.h b/Engine/Plugins/Runtime/StateTree/Source/StateTreeEditorModule/Private/StateTreeEditorCommands.h
--- a/Engine/Plugins/Runtime/StateTree/Source/StateTreeEditorModule/Private/StateTreeEditorCommands.h
+++ b/Engine/Plugins/Runtime/StateTree/Source/StateTreeEditorModule/Private/StateTreeEditorCommands.h
@@ -31,4 +31,9 @@ public:
 	TSharedPtr<FUICommandInfo> DuplicateStates;
 	TSharedPtr<FUICommandInfo> DeleteStates;
 	TSharedPtr<FUICommandInfo> EnableStates;
+
+#if WITH_STATETREE_DEBUGGER
+	TSharedPtr<FUICommandInfo> EnableOnEnterStateBreakpoint;
+	TSharedPtr<FUICommandInfo> EnableOnExitStateBreakpoint;
+#endif // WITH_STATETREE_DEBUGGER
 };
\ No newline at end of file
diff --git a/Engine/Plugins/Runtime/StateTree/Source/StateTreeEditorModule/Private/StateTreeViewModel.cpp b/Engine/Plugins/Runtime/StateTree/Source/StateTreeEditorModule/Private/StateTreeViewModel.cpp
--- a/Engine/Plugins/Runtime/StateTree/Source/StateTreeEditorModule/Private/StateTreeViewModel.cpp
+++ b/Engine/Plugins/Runtime/StateTree/Source/StateTreeEditorModule/Private/StateTreeViewModel.cpp
@@ -283,6 +283,185 @@ void FStateTreeViewModel::HandleIdentifierChanged(const UStateTree& StateTree) c
 }
 
 #if WITH_STATETREE_DEBUGGER
+bool FStateTreeViewModel::CanAddStateBreakpoint(const EStateTreeBreakpointType Type) const
+{
+	const UStateTreeEditorData* EditorData = TreeDataWeak.Get();
+	if (!ensure(EditorData != nullptr))
+	{
+		return false;
+	}
+
+	for (const TWeakObjectPtr<UStateTreeState>& WeakState : SelectedStates)
+	{
+		if (const UStateTreeState* State = WeakState.Get())
+		{
+			if (EditorData->HasBreakpoint(State->ID, Type) == false)
+			{
+				return true;
+			}
+		}
+	}
+
+	return false;
+}
+
+bool FStateTreeViewModel::CanRemoveStateBreakpoint(const EStateTreeBreakpointType Type) const
+{
+	const UStateTreeEditorData* EditorData = TreeDataWeak.Get();
+	if (!ensure(EditorData != nullptr))
+	{
+		return false;
+	}
+
+	for (const TWeakObjectPtr<UStateTreeState>& WeakState : SelectedStates)
+	{
+		if (const UStateTreeState* State = WeakState.Get())
+		{
+			if (EditorData->HasBreakpoint(State->ID, Type))
+			{
+				return true;
+			}
+		}
+	}
+
+	return false;
+}
+
+
+ECheckBoxState FStateTreeViewModel::GetStateBreakpointCheckState(const EStateTreeBreakpointType Type) const
+{
+	const bool bCanAdd = CanAddStateBreakpoint(Type);
+	const bool bCanRemove = CanRemoveStateBreakpoint(Type);
+	if (bCanAdd && bCanRemove)
+	{
+		return ECheckBoxState::Undetermined;
+	}
+
+	if (bCanRemove)
+	{
+		return ECheckBoxState::Checked;
+	}
+
+	if (bCanAdd)
+	{
+		return ECheckBoxState::Unchecked;
+	}
+
+	// Should not happen since action is not visible in this case
+	return ECheckBoxState::Undetermined;
+}
+
+void FStateTreeViewModel::HandleEnableStateBreakpoint(EStateTreeBreakpointType Type)
+{
+	TArray<UStateTreeState*> ValidatedSelectedStates;
+	GetSelectedStates(ValidatedSelectedStates);
+	if (ValidatedSelectedStates.IsEmpty())
+	{
+		return;
+	}
+
+	UStateTreeEditorData* EditorData = TreeDataWeak.Get();
+	if (!ensure(EditorData != nullptr))
+	{
+		return;
+	}
+
+	TBitArray<> HasBreakpoint;
+	HasBreakpoint.Reserve(ValidatedSelectedStates.Num());
+	for (const UStateTreeState* SelectedState : ValidatedSelectedStates)
+	{
+		HasBreakpoint.Add(SelectedState != nullptr && EditorData->HasBreakpoint(SelectedState->ID, Type));
+	}
+
+	check(HasBreakpoint.Num() == ValidatedSelectedStates.Num());
+
+	// Process CanAdd first so in case of undetermined state (mixed selection) we add by default. 
+	if (CanAddStateBreakpoint(Type))
+	{
+		const FScopedTransaction Transaction(LOCTEXT("AddStateBreakpoint", "Add State Breakpoint(s)"));
+		EditorData->Modify();
+		for (int Index = 0; Index < ValidatedSelectedStates.Num(); ++Index)
+		{
+			const UStateTreeState* SelectedState = ValidatedSelectedStates[Index];
+			if (HasBreakpoint[Index] == false && SelectedState != nullptr)
+			{
+				EditorData->AddBreakpoint(SelectedState->ID, Type);	
+			}
+		}
+	}
+	else if (CanRemoveStateBreakpoint(Type))
+	{
+		const FScopedTransaction Transaction(LOCTEXT("RemoveStateBreakpoint", "Remove State Breakpoint(s)"));
+		EditorData->Modify();
+		for (int Index = 0; Index < ValidatedSelectedStates.Num(); ++Index)
+		{
+			const UStateTreeState* SelectedState = ValidatedSelectedStates[Index];
+			if (HasBreakpoint[Index] && SelectedState != nullptr)
+			{
+				EditorData->RemoveBreakpoint(SelectedState->ID, Type);	
+			}
+		}
+	}
+}
+
+UStateTreeState* FStateTreeViewModel::FindStateAssociatedToBreakpoint(FStateTreeDebuggerBreakpoint Breakpoint) const
+{
+	UStateTreeEditorData* EditorData = TreeDataWeak.Get();
+	if (EditorData == nullptr)
+	{
+		return nullptr;
+	}
+	const UStateTree* StateTree = GetStateTree();
+	if (StateTree == nullptr)
+	{
+		return nullptr;
+	}
+
+	UStateTreeState* StateTreeState = nullptr;
+
+	if (const FStateTreeStateHandle* StateHandle = Breakpoint.ElementIdentifier.TryGet<FStateTreeStateHandle>())
+	{
+		const FGuid StateId = StateTree->GetStateIdFromHandle(*StateHandle);
+		StateTreeState = EditorData->GetMutableStateByID(StateId);
+	}
+	else if (const FStateTreeDebuggerBreakpoint::FStateTreeTaskIndex* TaskIndex = Breakpoint.ElementIdentifier.TryGet<FStateTreeDebuggerBreakpoint::FStateTreeTaskIndex>())
+	{
+		const FGuid TaskId = StateTree->GetNodeIdFromIndex(TaskIndex->Index);
+
+		EditorData->VisitHierarchy([&TaskId, &StateTreeState](UStateTreeState& State, UStateTreeState* /*ParentState*/)
+			{
+				for (const FStateTreeEditorNode& EditorNode : State.Tasks)
+				{
+					if (EditorNode.ID == TaskId)
+					{
+						StateTreeState = &State;
+						return EStateTreeVisitor::Break;
+					}
+				}
+				return EStateTreeVisitor::Continue;
+			});
+	}
+	else if (const FStateTreeDebuggerBreakpoint::FStateTreeTransitionIndex* TransitionIndex = Breakpoint.ElementIdentifier.TryGet<FStateTreeDebuggerBreakpoint::FStateTreeTransitionIndex>())
+	{
+		const FGuid TransitionId = StateTree->GetTransitionIdFromIndex(TransitionIndex->Index);
+
+		EditorData->VisitHierarchy([&TransitionId, &StateTreeState](UStateTreeState& State, UStateTreeState* /*ParentState*/)
+			{
+				for (const FStateTreeTransition& StateTransition : State.Transitions)
+				{
+					if (StateTransition.ID == TransitionId)
+					{
+						StateTreeState = &State;
+						return EStateTreeVisitor::Break;
+					}
+				}
+				return EStateTreeVisitor::Continue;
+			});
+	}
+
+	return StateTreeState;
+}
+
 void FStateTreeViewModel::HandleBreakpointsChanged(const UStateTree& StateTree)
 {
 	if (GetStateTree() == &StateTree)
diff --git a/Engine/Plugins/Runtime/StateTree/Source/StateTreeEditorModule/Private/StateTreeViewModel.h b/Engine/Plugins/Runtime/StateTree/Source/StateTreeEditorModule/Private/StateTreeViewModel.h
--- a/Engine/Plugins/Runtime/StateTree/Source/StateTreeEditorModule/Private/StateTreeViewModel.h
+++ b/Engine/Plugins/Runtime/StateTree/Source/StateTreeEditorModule/Private/StateTreeViewModel.h
@@ -6,14 +6,18 @@
 #include "EditorUndoClient.h"
 #include "StateTreeViewModel.generated.h"
 
-struct FStateTreeEditorBreakpoint;
-struct FStateTreeDebugger;
-struct FPropertyChangedEvent;
-
+class FMenuBuilder;
+class UStateTree;
 class UStateTreeEditorData;
 class UStateTreeState;
-class UStateTree;
-class FMenuBuilder;
+
+enum class ECheckBoxState : uint8;
+enum class EStateTreeBreakpointType : uint8;
+
+struct FPropertyChangedEvent;
+struct FStateTreeDebugger;
+struct FStateTreeDebuggerBreakpoint;
+struct FStateTreeEditorBreakpoint;
 struct FStateTreePropertyPathBinding;
 
 enum class FStateTreeViewModelInsert : uint8
@@ -91,6 +95,13 @@ public:
 
 	// Debugging
 #if WITH_STATETREE_DEBUGGER
+	bool CanAddStateBreakpoint(EStateTreeBreakpointType Type) const;
+	bool CanRemoveStateBreakpoint(EStateTreeBreakpointType Type) const;
+	ECheckBoxState GetStateBreakpointCheckState(EStateTreeBreakpointType Type) const;
+	void HandleEnableStateBreakpoint(EStateTreeBreakpointType Type);
+
+	UStateTreeState* FindStateAssociatedToBreakpoint(FStateTreeDebuggerBreakpoint Breakpoint) const;
+
 	TSharedRef<FStateTreeDebugger> GetDebugger() const { return Debugger; }
 	void RefreshDebuggerBreakpoints();
 #endif // WITH_STATETREE_DEBUGGER
