StateTree: Fix crash when transitioning into linked tree.

diff --git a/Engine/Plugins/Runtime/StateTree/Source/StateTreeModule/Private/StateTreeExecutionContext.cpp b/Engine/Plugins/Runtime/StateTree/Source/StateTreeModule/Private/StateTreeExecutionContext.cpp
--- a/Engine/Plugins/Runtime/StateTree/Source/StateTreeModule/Private/StateTreeExecutionContext.cpp
+++ b/Engine/Plugins/Runtime/StateTree/Source/StateTreeModule/Private/StateTreeExecutionContext.cpp
@@ -872,6 +872,8 @@ void FStateTreeExecutionContext::UpdateInstanceData(TConstArrayView<FStateTreeEx
 			InstanceStructs.AddDefaulted(State.InstanceDataNum);
 			TempInstanceStructs.AddZeroed(State.InstanceDataNum);
 
+			bool bCanHaveTempData = false;
+			
 			if (State.Type == EStateTreeStateType::Subtree)
 			{
 				check(State.ParameterDataHandle.IsValid());
@@ -882,6 +884,7 @@ void FStateTreeExecutionContext::UpdateInstanceData(TConstArrayView<FStateTreeEx
 					// Parameters are not set by a linked state, create instance data.
 					InstanceStructs[BaseIndex + State.ParameterDataHandle.GetIndex()] = ParamsInstanceData;
 					NextFrame.StateParameterDataHandle = State.ParameterDataHandle;
+					bCanHaveTempData = true;
 				}
 				else
 				{
@@ -892,10 +895,18 @@ void FStateTreeExecutionContext::UpdateInstanceData(TConstArrayView<FStateTreeEx
 					
 					NextFrame.StateParameterDataHandle = NextStateParameterDataHandle;
 					NextStateParameterDataHandle = FStateTreeDataHandle::Invalid; // Mark as used.
+
+					// This state will not instantiate parameter data, so we don't care about the temp data either.
+					bCanHaveTempData = false;
 				}
 			}
 			else
 			{
+				if (NextStateParameterDataHandle.IsValid())
+				{
+					UE_DEBUG_BREAK();
+				}
+				
 				if (State.ParameterTemplateIndex.IsValid())
 				{
 					// Linked state's instance data is the parameters.
@@ -908,6 +919,7 @@ void FStateTreeExecutionContext::UpdateInstanceData(TConstArrayView<FStateTreeEx
 						// We expect overridden linked assets to hit this code path. 
 						InstanceStructs[BaseIndex + State.ParameterDataHandle.GetIndex()] = FConstStructView(TempParamsInstanceData->GetScriptStruct());
 						Params = TempParamsInstanceData->GetPtr<const FCompactStateTreeParameters>();
+						bCanHaveTempData = true;
 					}
 					else
 					{
@@ -915,6 +927,7 @@ void FStateTreeExecutionContext::UpdateInstanceData(TConstArrayView<FStateTreeEx
 						const FConstStructView ParamsInstanceData = NextFrame.StateTree->DefaultInstanceData.GetStruct(State.ParameterTemplateIndex.Get());
 						InstanceStructs[BaseIndex + State.ParameterDataHandle.GetIndex()] = ParamsInstanceData;
 						Params = ParamsInstanceData.GetPtr<const FCompactStateTreeParameters>();
+						bCanHaveTempData = true;
 					}
 
 					if (State.Type == EStateTreeStateType::Linked
@@ -929,7 +942,7 @@ void FStateTreeExecutionContext::UpdateInstanceData(TConstArrayView<FStateTreeEx
 				}
 			}
 			
-			if (!bAreCommon && State.ParameterDataHandle.IsValid())
+			if (!bAreCommon && bCanHaveTempData)
 			{
 				TempInstanceStructs[BaseIndex + State.ParameterDataHandle.GetIndex()] = FindInstanceTempData(NextFrame, State.ParameterDataHandle);
 			}
@@ -2962,7 +2975,8 @@ bool FStateTreeExecutionContext::TriggerTransitions()
 				}
 				else if (Transition.Trigger == EStateTreeTransitionTrigger::OnTick)
 				{
-					TransitionEvents.Emplace();
+					// Dummy event to make sure we iterate to loop below once.
+					TransitionEvents.Emplace(nullptr);
 				}
 				
 				for (const FStateTreeSharedEvent* TransitionEvent : TransitionEvents)
@@ -3001,8 +3015,11 @@ bool FStateTreeExecutionContext::TriggerTransitions()
 									DelayedState.StateTree = CurrentFrame.StateTree;
 									DelayedState.TransitionIndex = FStateTreeIndex16(TransitionIndex);
 									DelayedState.TimeLeft = DelayDuration;
-									DelayedState.CapturedEvent = *TransitionEvent;
-									DelayedState.CapturedEventHash = TransitionEventHash;
+									if (TransitionEvent && TransitionEvent->IsValid())
+									{
+										DelayedState.CapturedEvent = *TransitionEvent;
+										DelayedState.CapturedEventHash = TransitionEventHash;
+									}
 
 									BeginDelayedTransition(DelayedState);
 									STATETREE_LOG(Verbose, TEXT("Delayed transition triggered from '%s' (%s) -> '%s' %.1fs"),
diff --git a/Engine/Plugins/Runtime/StateTree/Source/StateTreeModule/Private/StateTreeInstanceData.cpp b/Engine/Plugins/Runtime/StateTree/Source/StateTreeModule/Private/StateTreeInstanceData.cpp
--- a/Engine/Plugins/Runtime/StateTree/Source/StateTreeModule/Private/StateTreeInstanceData.cpp
+++ b/Engine/Plugins/Runtime/StateTree/Source/StateTreeModule/Private/StateTreeInstanceData.cpp
@@ -631,15 +631,18 @@ void FStateTreeInstanceData::Append(UObject& InOwner, TConstArrayView<FConstStru
 		FInstancedStruct* Source = InInstancesToMove[Index - StartIndex];
 
 		// The source is used to move temporary instance data into instance data. Not all entries may have it.
-		// If the source is specified, move it to the instance data. We assume that if the source is object wrapper, it is already the instance we want.
-		if (Source && Source->IsValid())
+		// The instance struct can be empty, in which case the temporary instance is ignored.
+		// If the source is specified, move it to the instance data.
+		// We assume that if the source is object wrapper, it is already the instance we want.
+		if (Struct.IsValid()
+			&& (Source && Source->IsValid()))
 		{
 			check(Struct.GetScriptStruct() == Source->GetScriptStruct());
 				
 			FMemory::Memswap(Struct.GetMemory(), Source->GetMutableMemory(), Struct.GetScriptStruct()->GetStructureSize());
 			Source->Reset();
 		}
-		else if (FStateTreeInstanceObjectWrapper* Wrapper = Storage.InstanceStructs[Index].GetPtr<FStateTreeInstanceObjectWrapper>())
+		else if (FStateTreeInstanceObjectWrapper* Wrapper = Struct.GetPtr<FStateTreeInstanceObjectWrapper>())
 		{
 			if (Wrapper->InstanceObject)
 			{
