StateTrees: Fixed skipping last property in edit chain during PostEditProperty dispatch.

diff --git a/Engine/Plugins/Runtime/StateTree/Source/StateTreeEditorModule/Private/StateTreePropertyHelpers.cpp b/Engine/Plugins/Runtime/StateTree/Source/StateTreeEditorModule/Private/StateTreePropertyHelpers.cpp
--- a/Engine/Plugins/Runtime/StateTree/Source/StateTreeEditorModule/Private/StateTreePropertyHelpers.cpp
+++ b/Engine/Plugins/Runtime/StateTree/Source/StateTreeEditorModule/Private/StateTreePropertyHelpers.cpp
@@ -65,7 +65,7 @@ void DispatchPostEditToNodes(UObject& Owner, FPropertyChangedChainEvent& InPrope
 
 							// Take copy of the event, we'll modify it.
 							FEditPropertyChain PropertyChainCopy;
-							for (const TDoubleLinkedList<FProperty*>::TDoubleLinkedListNode* Node = InPropertyChangedEvent.PropertyChain.GetHead(); Node->GetNextNode(); Node = Node->GetNextNode())
+							for (const TDoubleLinkedList<FProperty*>::TDoubleLinkedListNode* Node = InPropertyChangedEvent.PropertyChain.GetHead(); Node; Node = Node->GetNextNode())
 							{
 								PropertyChainCopy.AddTail(Node->GetValue());
 							}
