StateTrees: Made TStateTreeInstanceDataStructRef usable when const.

diff --git a/Engine/Plugins/Runtime/StateTree/Source/StateTreeModule/Public/StateTreeInstanceData.h b/Engine/Plugins/Runtime/StateTree/Source/StateTreeModule/Public/StateTreeInstanceData.h
--- a/Engine/Plugins/Runtime/StateTree/Source/StateTreeModule/Public/StateTreeInstanceData.h
+++ b/Engine/Plugins/Runtime/StateTree/Source/StateTreeModule/Public/StateTreeInstanceData.h
@@ -575,7 +575,7 @@ struct TStateTreeInstanceDataStructRef
 
 	bool IsValid() const { return WeakStateTree.IsValid() && RootState.IsValid() && DataHandle.IsValid(); }
 
-	T* GetPtr()
+	T* GetPtr() const
 	{
 		if (!WeakStorage.IsValid())
 		{
@@ -619,7 +619,7 @@ struct TStateTreeInstanceDataStructRef
 
 protected:
 
-	FStructView GetDataView(FStateTreeInstanceStorage& Storage, const FStateTreeExecutionFrame& CurrentFrame, const FStateTreeDataHandle Handle)
+	FStructView GetDataView(FStateTreeInstanceStorage& Storage, const FStateTreeExecutionFrame& CurrentFrame, const FStateTreeDataHandle Handle) const
 	{
 		switch (DataHandle.GetSource())
 		{
