StateTrees: AddOverride replaces now already existing override.
Also, added removing overrides

diff --git a/Engine/Plugins/Runtime/StateTree/Source/StateTreeModule/Public/StateTreeReference.h b/Engine/Plugins/Runtime/StateTree/Source/StateTreeModule/Public/StateTreeReference.h
--- a/Engine/Plugins/Runtime/StateTree/Source/StateTreeModule/Public/StateTreeReference.h
+++ b/Engine/Plugins/Runtime/StateTree/Source/StateTreeModule/Public/StateTreeReference.h
@@ -159,14 +159,45 @@ struct STATETREEMODULE_API FStateTreeReferenceOverrides
 {
 	GENERATED_BODY()
 
+	/** Removes all overrides. */
 	void Reset()
 	{
 		OverrideItems.Reset();	
 	}
 	
-	void AddOverride(const FGameplayTag StateTag, FStateTreeReference& StateTreeReference)
+	/** Adds or replaces override for a selected tag. */
+	void AddOverride(const FGameplayTag StateTag, const FStateTreeReference& StateTreeReference)
 	{
-		OverrideItems.Emplace(StateTag, StateTreeReference);
+		FStateTreeReferenceOverrideItem* FoundOverride = OverrideItems.FindByPredicate([StateTag](const FStateTreeReferenceOverrideItem& Override)
+		{
+			return Override.GetStateTag() == StateTag;
+		});
+
+		if (FoundOverride)
+		{
+			*FoundOverride = FStateTreeReferenceOverrideItem(StateTag, StateTreeReference);
+		}
+		else
+		{
+			OverrideItems.Emplace(StateTag, StateTreeReference);
+		}
+	}
+
+	/** Returns true if removing an override succeeded. */
+	bool RemoveOverride(const FGameplayTag StateTag)
+	{
+		const int32 Index = OverrideItems.IndexOfByPredicate([StateTag](const FStateTreeReferenceOverrideItem& Override)
+		{
+			return Override.GetStateTag() == StateTag;
+		});
+
+		if (Index != INDEX_NONE)
+		{
+			OverrideItems.RemoveAtSwap(Index);
+			return true;
+		}
+
+		return false;
 	}
 
 	TConstArrayView<FStateTreeReferenceOverrideItem> GetOverrideItems() const
