StateTrees: Made empty linked asset a valid setup.

diff --git a/Engine/Plugins/Runtime/StateTree/Source/StateTreeEditorModule/Private/StateTreeCompiler.cpp b/Engine/Plugins/Runtime/StateTree/Source/StateTreeEditorModule/Private/StateTreeCompiler.cpp
--- a/Engine/Plugins/Runtime/StateTree/Source/StateTreeEditorModule/Private/StateTreeCompiler.cpp
+++ b/Engine/Plugins/Runtime/StateTree/Source/StateTreeEditorModule/Private/StateTreeCompiler.cpp
@@ -857,26 +857,29 @@ bool FStateTreeCompiler::CreateStateTransitions()
 				return false;
 			}
 
-			// Linked asset must have same schema.
-			const UStateTreeSchema* LinkedAssetSchema = SourceState->LinkedAsset ? SourceState->LinkedAsset->GetSchema() : nullptr;
-
-			if (!LinkedAssetSchema)
+			if (SourceState->LinkedAsset)
 			{
-				Log.Reportf(EMessageSeverity::Error,
-					TEXT("Linked State Tree asset must have valid schema."));
-				return false;
-			}
+				// Linked asset must have same schema.
+				const UStateTreeSchema* LinkedAssetSchema = SourceState->LinkedAsset->GetSchema();
+
+				if (!LinkedAssetSchema)
+				{
+					Log.Reportf(EMessageSeverity::Error,
+						TEXT("Linked State Tree asset must have valid schema."));
+					return false;
+				}
 			
-			check(Schema);
-			if (LinkedAssetSchema->GetClass() != Schema->GetClass())
-			{
-				Log.Reportf(EMessageSeverity::Error,
-					TEXT("Linked State Tree asset '%s' must have same schema class as this asset. Linked asset has '%s', expected '%s'."),
-					*GetFullNameSafe(SourceState->LinkedAsset),
-					*LinkedAssetSchema->GetClass()->GetDisplayNameText().ToString(),
-					*Schema->GetClass()->GetDisplayNameText().ToString()
-				);
-				return false;
+				check(Schema);
+				if (LinkedAssetSchema->GetClass() != Schema->GetClass())
+				{
+					Log.Reportf(EMessageSeverity::Error,
+						TEXT("Linked State Tree asset '%s' must have same schema class as this asset. Linked asset has '%s', expected '%s'."),
+						*GetFullNameSafe(SourceState->LinkedAsset),
+						*LinkedAssetSchema->GetClass()->GetDisplayNameText().ToString(),
+						*Schema->GetClass()->GetDisplayNameText().ToString()
+					);
+					return false;
+				}
 			}
 			
 			CompactState.LinkedAsset = SourceState->LinkedAsset;
diff --git a/Engine/Plugins/Runtime/StateTree/Source/StateTreeModule/Private/StateTreeExecutionContext.cpp b/Engine/Plugins/Runtime/StateTree/Source/StateTreeModule/Private/StateTreeExecutionContext.cpp
--- a/Engine/Plugins/Runtime/StateTree/Source/StateTreeModule/Private/StateTreeExecutionContext.cpp
+++ b/Engine/Plugins/Runtime/StateTree/Source/StateTreeModule/Private/StateTreeExecutionContext.cpp
@@ -3408,15 +3408,21 @@ bool FStateTreeExecutionContext::SelectStateInternal(
 		return false;
 	}
 
+	const UStateTree* NextLinkedStateAsset = NextState.LinkedAsset;
+
 	// Look up linked state overrides
-	const UStateTree* NextLinkedStateAssetOverride = nullptr;
 	const FInstancedPropertyBag* NextLinkedStateParameterOverride = nullptr;
 	if (NextState.Type == EStateTreeStateType::LinkedAsset)
 	{
 		if (const FStateTreeReference* Override = GetLinkedStateTreeOverrideForTag(NextState.Tag))
 		{
-			NextLinkedStateAssetOverride = Override->GetStateTree();
+			NextLinkedStateAsset = Override->GetStateTree();
 			NextLinkedStateParameterOverride = &Override->GetParameters();
+
+			STATETREE_LOG(VeryVerbose, TEXT("%hs: In state '%s', overriding linked asset '%s' with '%s'. '%s' using StateTree '%s'."),
+					__FUNCTION__, *GetSafeStateName(CurrentFrame, NextStateHandle),
+					*GetFullNameSafe(NextState.LinkedAsset), *GetFullNameSafe(NextLinkedStateAsset),
+					*GetNameSafe(&Owner), *GetFullNameSafe(CurrentFrame.StateTree));
 		}
 	}
 
@@ -3620,106 +3626,93 @@ bool FStateTreeExecutionContext::SelectStateInternal(
 		}
 		else if (NextState.Type == EStateTreeStateType::LinkedAsset)
 		{
-			const UStateTree* LinkedAsset = NextState.LinkedAsset;
-			if (NextLinkedStateAssetOverride)
+			if (NextLinkedStateAsset == nullptr)
 			{
-				STATETREE_LOG(VeryVerbose, TEXT("%hs: In state '%s', overriding linked asset '%s' with '%s'. '%s' using StateTree '%s'."),
-					__FUNCTION__, *GetSafeStateName(CurrentFrame, NextStateHandle),
-					*GetFullNameSafe(NextState.LinkedAsset), *GetFullNameSafe(NextLinkedStateAssetOverride),
-					*GetNameSafe(&Owner), *GetFullNameSafe(CurrentFrame.StateTree));
-				LinkedAsset = NextLinkedStateAssetOverride;
+				return false;
 			}
-			
-			if (LinkedAsset)
+
+			if (OutSelectionResult.IsFull())
 			{
-				if (OutSelectionResult.IsFull())
-				{
-					STATETREE_LOG(Error, TEXT("%hs: Reached max execution depth when trying to select state %s from '%s'.  '%s' using StateTree '%s'."),
-						__FUNCTION__, *GetSafeStateName(CurrentFrame, NextStateHandle), *GetStateStatusString(Exec), *GetNameSafe(&Owner), *GetFullNameSafe(CurrentFrame.StateTree));
-					return false;
-				}
+				STATETREE_LOG(Error, TEXT("%hs: Reached max execution depth when trying to select state %s from '%s'.  '%s' using StateTree '%s'."),
+					__FUNCTION__, *GetSafeStateName(CurrentFrame, NextStateHandle), *GetStateStatusString(Exec), *GetNameSafe(&Owner), *GetFullNameSafe(CurrentFrame.StateTree));
+				return false;
+			}
 
-				// The linked state tree should have compatible context requirements.
-				if (!LinkedAsset->HasCompatibleContextData(RootStateTree))
-				{
-					STATETREE_LOG(Error, TEXT("%hs: The linked State Tree '%s' does not have compatible schema, trying to select state %s from '%s'.  '%s' using StateTree '%s'."),
-						__FUNCTION__, *GetFullNameSafe(LinkedAsset), *GetSafeStateName(CurrentFrame, NextStateHandle), *GetStateStatusString(Exec), *GetNameSafe(&Owner), *GetFullNameSafe(CurrentFrame.StateTree));
-					return false;
-				}
+			// The linked state tree should have compatible context requirements.
+			if (!NextLinkedStateAsset->HasCompatibleContextData(RootStateTree))
+			{
+				STATETREE_LOG(Error, TEXT("%hs: The linked State Tree '%s' does not have compatible schema, trying to select state %s from '%s'.  '%s' using StateTree '%s'."),
+					__FUNCTION__, *GetFullNameSafe(NextLinkedStateAsset), *GetSafeStateName(CurrentFrame, NextStateHandle), *GetStateStatusString(Exec), *GetNameSafe(&Owner), *GetFullNameSafe(CurrentFrame.StateTree));
+				return false;
+			}
 				
-				FStateTreeExecutionFrame NewFrame;
-				NewFrame.StateTree = LinkedAsset;
-				NewFrame.RootState = FStateTreeStateHandle::Root;
-				NewFrame.bIsGlobalFrame = true;
-
-				// Check and prevent recursion.
-				const bool bNewFrameAlreadySelected = OutSelectionResult.GetSelectedFrames().ContainsByPredicate([&NewFrame](const FStateTreeExecutionFrame& Frame) {
-					return Frame.IsSameFrame(NewFrame);
-				});
+			FStateTreeExecutionFrame NewFrame;
+			NewFrame.StateTree = NextLinkedStateAsset;
+			NewFrame.RootState = FStateTreeStateHandle::Root;
+			NewFrame.bIsGlobalFrame = true;
+
+			// Check and prevent recursion.
+			const bool bNewFrameAlreadySelected = OutSelectionResult.GetSelectedFrames().ContainsByPredicate([&NewFrame](const FStateTreeExecutionFrame& Frame) {
+				return Frame.IsSameFrame(NewFrame);
+			});
 				
-				if (bNewFrameAlreadySelected)
+			if (bNewFrameAlreadySelected)
+			{
+				STATETREE_LOG(Error, TEXT("%hs: Trying to recursively enter subtree '%s' from '%s'.  '%s' using StateTree '%s'."),
+					__FUNCTION__, *GetSafeStateName(NewFrame, NewFrame.RootState), *GetStateStatusString(Exec), *GetNameSafe(&Owner), *GetFullNameSafe(CurrentFrame.StateTree));
+				return false;
+			}
+
+			// If the Frame already exists, copy instance indices so that conditions that rely on active states work correctly.
+			const FStateTreeExecutionFrame* ExistingFrame = Exec.ActiveFrames.FindByPredicate(
+				[StateTree = NewFrame.StateTree, RootState = NewFrame.RootState](const FStateTreeExecutionFrame& Frame)
 				{
-					STATETREE_LOG(Error, TEXT("%hs: Trying to recursively enter subtree '%s' from '%s'.  '%s' using StateTree '%s'."),
-						__FUNCTION__, *GetSafeStateName(NewFrame, NewFrame.RootState), *GetStateStatusString(Exec), *GetNameSafe(&Owner), *GetFullNameSafe(CurrentFrame.StateTree));
-					return false;
-				}
+					return Frame.StateTree == StateTree && Frame.RootState == RootState;
+				});
+			if (ExistingFrame)
+			{
+				NewFrame.ActiveInstanceIndexBase = ExistingFrame->ActiveInstanceIndexBase;
+				NewFrame.GlobalInstanceIndexBase = ExistingFrame->GlobalInstanceIndexBase;
+				NewFrame.StateParameterDataHandle = ExistingFrame->StateParameterDataHandle;
+				NewFrame.GlobalParameterDataHandle = ExistingFrame->GlobalParameterDataHandle;
+				NewFrame.ExternalDataBaseIndex = ExistingFrame->ExternalDataBaseIndex;
+			}
+			else
+			{
+				// Pass the linked state's parameters as global parameters to the linked asset.
+				NewFrame.GlobalParameterDataHandle = NextState.ParameterDataHandle;
 
-				// If the Frame already exists, copy instance indices so that conditions that rely on active states work correctly.
-				const FStateTreeExecutionFrame* ExistingFrame = Exec.ActiveFrames.FindByPredicate(
-					[StateTree = NewFrame.StateTree, RootState = NewFrame.RootState](const FStateTreeExecutionFrame& Frame)
-					{
-						return Frame.StateTree == StateTree && Frame.RootState == RootState;
-					});
-				if (ExistingFrame)
+				// Collect external data if needed
+				NewFrame.ExternalDataBaseIndex = CollectExternalData(NewFrame.StateTree);
+				if (!NewFrame.ExternalDataBaseIndex.IsValid())
 				{
-					NewFrame.ActiveInstanceIndexBase = ExistingFrame->ActiveInstanceIndexBase;
-					NewFrame.GlobalInstanceIndexBase = ExistingFrame->GlobalInstanceIndexBase;
-					NewFrame.StateParameterDataHandle = ExistingFrame->StateParameterDataHandle;
-					NewFrame.GlobalParameterDataHandle = ExistingFrame->GlobalParameterDataHandle;
-					NewFrame.ExternalDataBaseIndex = ExistingFrame->ExternalDataBaseIndex;
+					STATETREE_LOG(VeryVerbose, TEXT("%hs: Cannot select state '%s' because failed to collect external data for nested tree '%s'.  '%s' using StateTree '%s'."),
+						__FUNCTION__, *GetSafeStateName(CurrentFrame, NextStateHandle), *GetFullNameSafe(NewFrame.StateTree), *GetNameSafe(&Owner), *GetFullNameSafe(CurrentFrame.StateTree));
+					return false;
 				}
-				else
-				{
-					// Pass the linked state's parameters as global parameters to the linked asset.
-					NewFrame.GlobalParameterDataHandle = NextState.ParameterDataHandle;
-
-					// Collect external data if needed
-					NewFrame.ExternalDataBaseIndex = CollectExternalData(NewFrame.StateTree);
-					if (!NewFrame.ExternalDataBaseIndex.IsValid())
-					{
-						STATETREE_LOG(VeryVerbose, TEXT("%hs: Cannot select state '%s' because failed to collect external data for nested tree '%s'.  '%s' using StateTree '%s'."),
-							__FUNCTION__, *GetSafeStateName(CurrentFrame, NextStateHandle), *GetFullNameSafe(NewFrame.StateTree), *GetNameSafe(&Owner), *GetFullNameSafe(CurrentFrame.StateTree));
-						return false;
-					}
 					
-					// The state parameters will be from the root state.
-					const FCompactStateTreeState& RootState = NewFrame.StateTree->States[NewFrame.RootState.Index];
-					NewFrame.StateParameterDataHandle = RootState.ParameterDataHandle;
-
-					// Start global tasks and evaluators temporarily, so that their data is available already during select.
-					if (StartTemporaryEvaluatorsAndGlobalTasks(nullptr, NewFrame) != EStateTreeRunStatus::Running)
-					{
-						STATETREE_LOG(VeryVerbose, TEXT("%hs: Cannot select state '%s' because cannot start nested tree's '%s' global tasks and evaluators.  '%s' using StateTree '%s'."),
-							__FUNCTION__, *GetSafeStateName(CurrentFrame, NextStateHandle), *GetFullNameSafe(NewFrame.StateTree), *GetNameSafe(&Owner), *GetFullNameSafe(CurrentFrame.StateTree));
-						return false;
-					}
-				}
-				
-				OutSelectionResult.PushFrame(NewFrame);
+				// The state parameters will be from the root state.
+				const FCompactStateTreeState& RootState = NewFrame.StateTree->States[NewFrame.RootState.Index];
+				NewFrame.StateParameterDataHandle = RootState.ParameterDataHandle;
 
-				// If State is linked, proceed to the linked state.
-				if (SelectStateInternal(&CurrentFrame, OutSelectionResult.GetSelectedFrames().Last(), ExistingFrame, {NewFrame.RootState}, OutSelectionResult))
+				// Start global tasks and evaluators temporarily, so that their data is available already during select.
+				if (StartTemporaryEvaluatorsAndGlobalTasks(nullptr, NewFrame) != EStateTreeRunStatus::Running)
 				{
-					return true;
+					STATETREE_LOG(VeryVerbose, TEXT("%hs: Cannot select state '%s' because cannot start nested tree's '%s' global tasks and evaluators.  '%s' using StateTree '%s'."),
+						__FUNCTION__, *GetSafeStateName(CurrentFrame, NextStateHandle), *GetFullNameSafe(NewFrame.StateTree), *GetNameSafe(&Owner), *GetFullNameSafe(CurrentFrame.StateTree));
+					return false;
 				}
-				
-				OutSelectionResult.PopFrame();
 			}
-			else
+				
+			OutSelectionResult.PushFrame(NewFrame);
+
+			// If State is linked, proceed to the linked state.
+			if (SelectStateInternal(&CurrentFrame, OutSelectionResult.GetSelectedFrames().Last(), ExistingFrame, {NewFrame.RootState}, OutSelectionResult))
 			{
-				STATETREE_LOG(Warning, TEXT("%hs: Trying to enter invalid linked asset from '%s'.  '%s' using StateTree '%s'."),
-					__FUNCTION__, *GetStateStatusString(Exec), *GetNameSafe(&Owner), *GetFullNameSafe(CurrentFrame.StateTree));
+				return true;
 			}
+				
+			OutSelectionResult.PopFrame();
 		}
 		else if (NextState.SelectionBehavior == EStateTreeStateSelectionBehavior::TryEnterState)
 		{
