[StateTreeDebugger] some layout tweaks to allow users to reduce all panels completely.
Also replaced the scrollbars for the Events TreeView by external ones to make sure both are visible when the window is too small in both dimensions.
Fixed edge case in SSimpleTimeSlider when receiving negative 'PixelsPerInput'

diff --git a/Engine/Plugins/Runtime/StateTree/Source/StateTreeEditorModule/Private/Debugger/SStateTreeDebuggerView.cpp b/Engine/Plugins/Runtime/StateTree/Source/StateTreeEditorModule/Private/Debugger/SStateTreeDebuggerView.cpp
--- a/Engine/Plugins/Runtime/StateTree/Source/StateTreeEditorModule/Private/Debugger/SStateTreeDebuggerView.cpp
+++ b/Engine/Plugins/Runtime/StateTree/Source/StateTreeEditorModule/Private/Debugger/SStateTreeDebuggerView.cpp
@@ -12,7 +12,6 @@
 #include "Kismet2/DebuggerCommands.h"
 #include "Modules/ModuleManager.h"
 #include "PropertyEditorModule.h"
-#include "ScopedTransaction.h"
 #include "SStateTreeDebuggerInstanceTree.h"
 #include "SStateTreeDebuggerTimelines.h"
 #include "StateTree.h"
@@ -29,7 +28,6 @@
 #include "Widgets/Input/SComboButton.h"
 #include "Widgets/Layout/SBorder.h"
 #include "Widgets/Layout/SScrollBox.h"
-#include "Widgets/Layout/SSpacer.h"
 
 #define LOCTEXT_NAMESPACE "StateTreeEditor"
 
@@ -491,6 +489,15 @@ void SStateTreeDebuggerView::Construct(const FArguments& InArgs, const UStateTre
 		.ScrubPosition(ScrubTimeAttribute)
 		.OnScrubPositionChanged_Lambda([this](double NewScrubTime, bool bIsScrubbing) { OnTimeLineScrubPositionChanged(NewScrubTime, bIsScrubbing); });
 
+	// EventsTreeView scrollbars
+	TSharedRef<SScrollBar> HorizontalScrollBar = SNew(SScrollBar)
+		.Orientation(Orient_Horizontal)
+		.Thickness(FVector2D(12.0f, 12.0f));
+
+	TSharedRef<SScrollBar> VerticalScrollBar = SNew(SScrollBar)
+		.Orientation(Orient_Vertical)
+		.Thickness(FVector2D(12.0f, 12.0f));
+
 	// EventsTreeView
 	EventsTreeView = SNew(STreeView<TSharedPtr<FStateTreeDebuggerEventTreeElement>>)
 			.OnGenerateRow_Lambda([this](const TSharedPtr<FStateTreeDebuggerEventTreeElement>& InElement, const TSharedRef<STableViewBase>& InOwnerTableView)
@@ -506,7 +513,8 @@ void SStateTreeDebuggerView::Construct(const FArguments& InArgs, const UStateTre
 			})
 		.TreeItemsSource(&EventsTreeElements)
 		.ItemHeight(32)
-		.AllowOverscroll(EAllowOverscroll::No);
+		.AllowOverscroll(EAllowOverscroll::No)
+		.ExternalScrollbar(VerticalScrollBar);
 
 	ChildSlot
 	[
@@ -520,7 +528,6 @@ void SStateTreeDebuggerView::Construct(const FArguments& InArgs, const UStateTre
 				SNew(SSplitter)
 				.Orientation(Orient_Horizontal)
 				+ SSplitter::Slot()
-				.MinSize(600)
 				[
 						SNew(SVerticalBox)
 						+ SVerticalBox::Slot()
@@ -530,13 +537,13 @@ void SStateTreeDebuggerView::Construct(const FArguments& InArgs, const UStateTre
 							SNew(SHorizontalBox)
 							+ SHorizontalBox::Slot()
 							.HAlign(HAlign_Left)
-							.FillWidth(1.0f)
+							.FillContentWidth(1.f)
 							[
 								LeftToolbar.MakeWidget()
 							]
 							+ SHorizontalBox::Slot()
 							.HAlign(HAlign_Right)
-							.AutoWidth()
+							.FillContentWidth(1.f)
 							[
 								RightToolbar.MakeWidget()
 							]
@@ -548,7 +555,6 @@ void SStateTreeDebuggerView::Construct(const FArguments& InArgs, const UStateTre
 							.Orientation(Orient_Horizontal)
 							+ SSplitter::Slot()
 							.Value(0.2f)
-							.MinSize(350)
 							.Resizable(false)
 							[
 								TraceSelectionBox
@@ -574,7 +580,6 @@ void SStateTreeDebuggerView::Construct(const FArguments& InArgs, const UStateTre
 							.Orientation(Orient_Horizontal)
 							+ SSplitter::Slot()
 							.Value(0.2f)
-							.MinSize(350)
 							.OnSlotResized_Lambda([this](float Size)
 								{
 									// Sync both header and content
@@ -610,7 +615,6 @@ void SStateTreeDebuggerView::Construct(const FArguments& InArgs, const UStateTre
 						]
 				]
 				+ SSplitter::Slot()
-				.MinSize(400)
 				[
 					SNew(SVerticalBox)
 					+ SVerticalBox::Slot()
@@ -622,13 +626,31 @@ void SStateTreeDebuggerView::Construct(const FArguments& InArgs, const UStateTre
 					+ SVerticalBox::Slot()
 					.FillHeight(1.0f)
 					[
-						SNew(SScrollBox)
-						.Orientation(Orient_Horizontal)
-						+ SScrollBox::Slot()
-						.FillSize(1.0f)
+						SNew(SHorizontalBox)
+						+ SHorizontalBox::Slot()
+						.FillWidth(1.0f)
+						.Padding(0.0f)
 						[
-							EventsTreeView.ToSharedRef()
+							SNew(SScrollBox)
+							.Orientation(Orient_Horizontal)
+							.ExternalScrollbar(HorizontalScrollBar)
+							+ SScrollBox::Slot()
+							.FillSize(1.0f)
+							[
+								EventsTreeView.ToSharedRef()
+							]
 						]
+
+						+ SHorizontalBox::Slot()
+						.AutoWidth()
+						[
+							VerticalScrollBar
+						]
+					]
+					+SVerticalBox::Slot()
+					.AutoHeight()
+					[
+						HorizontalScrollBar
 					]
 				]
 			]
diff --git a/Engine/Source/Developer/ToolWidgets/Private/SSimpleTimeSlider.cpp b/Engine/Source/Developer/ToolWidgets/Private/SSimpleTimeSlider.cpp
--- a/Engine/Source/Developer/ToolWidgets/Private/SSimpleTimeSlider.cpp
+++ b/Engine/Source/Developer/ToolWidgets/Private/SSimpleTimeSlider.cpp
@@ -87,7 +87,7 @@ static float GetNextSpacing( uint32 CurrentStep )
 */
 static double DetermineOptimalSpacing(float InPixelsPerInput, uint32 InMinTick, double InMinTickSpacing)
 {
-	if (InPixelsPerInput == 0.0f)
+	if (InPixelsPerInput <= 0.0f)
 		return InMinTickSpacing;
 
 	uint32 CurStep = 0;
