[StateTreeDebugger]
allowed to provide verbosity to STATETREE_TRACE_LOG_EVENT and STATETREE_LOG_AND_TRACE so the debugger can add custom icons accordingly.

diff --git a/Engine/Plugins/Runtime/StateTree/Source/StateTreeEditorModule/Private/Debugger/SStateTreeDebuggerViewRow.cpp b/Engine/Plugins/Runtime/StateTree/Source/StateTreeEditorModule/Private/Debugger/SStateTreeDebuggerViewRow.cpp
--- a/Engine/Plugins/Runtime/StateTree/Source/StateTreeEditorModule/Private/Debugger/SStateTreeDebuggerViewRow.cpp
+++ b/Engine/Plugins/Runtime/StateTree/Source/StateTreeEditorModule/Private/Debugger/SStateTreeDebuggerViewRow.cpp
@@ -30,15 +30,15 @@ void SStateTreeDebuggerViewRow::Construct(const FArguments& InArgs,
 			.BaseIndentLevel(0)
 		];
 
-	const TSharedPtr<SWidget> EventWidget = GenerateEventWidget();
-	if (EventWidget.IsValid())
+	const TSharedPtr<SWidget> EventImage = CreateImageForEvent();
+	if (EventImage.IsValid())
 	{
 		HorizontalBox->AddSlot()
 		.VAlign(VAlign_Center)
 		.HAlign(HAlign_Left)
 		.AutoWidth()
 		[
-			EventWidget.ToSharedRef()
+			EventImage.ToSharedRef()
 		];
 	}
 
@@ -60,9 +60,11 @@ void SStateTreeDebuggerViewRow::Construct(const FArguments& InArgs,
 		];
 }
 
-TSharedPtr<SWidget> SStateTreeDebuggerViewRow::GenerateEventWidget() const
+TSharedPtr<SWidget> SStateTreeDebuggerViewRow::CreateImageForEvent() const
 {
 	const FStateTreeEditorStyle& StyleSet = FStateTreeEditorStyle::Get();
+
+	// Phase events
 	if (const FStateTreeTracePhaseEvent* PhaseEvent = Item->Event.TryGet<FStateTreeTracePhaseEvent>())
 	{
 		const FSlateBrush* Image = nullptr;
@@ -78,6 +80,23 @@ TSharedPtr<SWidget> SStateTreeDebuggerViewRow::GenerateEventWidget() const
 		return SNew(SImage).Image(Image);
 	}
 
+	// Log events
+	if (const FStateTreeTraceLogEvent* LogEvent = Item->Event.TryGet<FStateTreeTraceLogEvent>())
+	{
+		const FSlateBrush* Image = nullptr;
+		switch (LogEvent->Verbosity)
+		{
+		case ELogVerbosity::Fatal:
+		case ELogVerbosity::Error:		Image = StyleSet.GetBrush("StateTreeEditor.Debugger.Log.Error");	break;
+		case ELogVerbosity::Warning:	Image = StyleSet.GetBrush("StateTreeEditor.Debugger.Log.Warning");	break;
+		default:
+			return nullptr;
+		}
+
+		return SNew(SImage).Image(Image);
+	}
+
+	// State events
 	if (const FStateTreeTraceStateEvent* StateEvent = Item->Event.TryGet<FStateTreeTraceStateEvent>())
 	{
 		const FSlateBrush* Image = nullptr;
@@ -92,6 +111,7 @@ TSharedPtr<SWidget> SStateTreeDebuggerViewRow::GenerateEventWidget() const
 		return SNew(SImage).Image(Image);
 	}
 
+	// Task events
 	if (const FStateTreeTraceTaskEvent* TaskEvent = Item->Event.TryGet<FStateTreeTraceTaskEvent>())
 	{
 		const FSlateBrush* Image = nullptr;
@@ -117,6 +137,7 @@ TSharedPtr<SWidget> SStateTreeDebuggerViewRow::GenerateEventWidget() const
 		return SNew(SImage).Image(Image);
 	}
 
+	// Condition events
 	if (Item->Event.IsType<FStateTreeTraceConditionEvent>())
 	{
 		const FStateTreeTraceConditionEvent& ConditionEvent = Item->Event.Get<FStateTreeTraceConditionEvent>();
diff --git a/Engine/Plugins/Runtime/StateTree/Source/StateTreeEditorModule/Private/Debugger/SStateTreeDebuggerViewRow.h b/Engine/Plugins/Runtime/StateTree/Source/StateTreeEditorModule/Private/Debugger/SStateTreeDebuggerViewRow.h
--- a/Engine/Plugins/Runtime/StateTree/Source/StateTreeEditorModule/Private/Debugger/SStateTreeDebuggerViewRow.h
+++ b/Engine/Plugins/Runtime/StateTree/Source/StateTreeEditorModule/Private/Debugger/SStateTreeDebuggerViewRow.h
@@ -40,7 +40,7 @@ public:
 				   const TSharedPtr<FStateTreeDebuggerEventTreeElement>& InElement);
 
 private:
-	TSharedPtr<SWidget> GenerateEventWidget() const;
+	TSharedPtr<SWidget> CreateImageForEvent() const;
 	const FTextBlockStyle& GetEventTextStyle() const;
 	FText GetEventDescription() const;
 	FText GetEventTooltip() const;
diff --git a/Engine/Plugins/Runtime/StateTree/Source/StateTreeEditorModule/Private/StateTreeEditorStyle.cpp b/Engine/Plugins/Runtime/StateTree/Source/StateTreeEditorModule/Private/StateTreeEditorStyle.cpp
--- a/Engine/Plugins/Runtime/StateTree/Source/StateTreeEditorModule/Private/StateTreeEditorStyle.cpp
+++ b/Engine/Plugins/Runtime/StateTree/Source/StateTreeEditorModule/Private/StateTreeEditorStyle.cpp
@@ -350,6 +350,9 @@ FStateTreeEditorStyle::FStateTreeEditorStyle()
 		Set("StateTreeEditor.Transition.Condition", new IMAGE_BRUSH_SVG("Icons/State_Conditions", CoreStyleConstants::Icon16x16, FStyleColors::AccentGray));
 		Set("StateTreeEditor.Debugger.Condition.OnTransition", new IMAGE_BRUSH_SVG("Icons/State_Conditions", CoreStyleConstants::Icon16x16, FStyleColors::AccentGray));
 
+		Set("StateTreeEditor.Debugger.Log.Warning", new IMAGE_BRUSH_SVG("Icons/Alert", CoreStyleConstants::Icon16x16, FStyleColors::AccentYellow));
+		Set("StateTreeEditor.Debugger.Log.Error", new IMAGE_BRUSH_SVG("Icons/Failure", CoreStyleConstants::Icon16x16, FStyleColors::AccentRed));
+
 		// Common Node Icons
 		Set("Node.Movement", new IMAGE_BRUSH_SVG("Icons/Movement", CoreStyleConstants::Icon16x16));
 		Set("Node.Tag", new IMAGE_BRUSH_SVG("Icons/Tag", CoreStyleConstants::Icon16x16));
diff --git a/Engine/Plugins/Runtime/StateTree/Source/StateTreeModule/Private/Debugger/StateTreeTrace.cpp b/Engine/Plugins/Runtime/StateTree/Source/StateTreeModule/Private/Debugger/StateTreeTrace.cpp
--- a/Engine/Plugins/Runtime/StateTree/Source/StateTreeModule/Private/Debugger/StateTreeTrace.cpp
+++ b/Engine/Plugins/Runtime/StateTree/Source/StateTreeModule/Private/Debugger/StateTreeTrace.cpp
@@ -60,6 +60,7 @@ UE_TRACE_EVENT_BEGIN(StateTreeDebugger, LogEvent)
 	UE_TRACE_EVENT_FIELD(uint64, Cycle)
 	UE_TRACE_EVENT_FIELD(uint32, InstanceId)
 	UE_TRACE_EVENT_FIELD(uint32, InstanceSerial)
+	UE_TRACE_EVENT_FIELD(std::underlying_type_t<ELogVerbosity::Type>, Verbosity)
 	UE_TRACE_EVENT_FIELD(UE::Trace::WideString, Message)
 UE_TRACE_EVENT_END()
 
@@ -621,6 +622,7 @@ void OutputInstanceFrameEvent(
 
 void OutputLogEventTrace(
 	const FStateTreeInstanceDebugId InstanceId,
+	ELogVerbosity::Type Verbosity,
 	const TCHAR* Fmt, ...
 	)
 {
@@ -633,6 +635,7 @@ void OutputLogEventTrace(
 		<< LogEvent.Cycle(FPlatformTime::Cycles64())
 		<< LogEvent.InstanceId(InstanceId.Id)
 		<< LogEvent.InstanceSerial(InstanceId.SerialNumber)
+		<< LogEvent.Verbosity(Verbosity)
 		<< LogEvent.Message(TraceStaticBuffer);
 }
 
diff --git a/Engine/Plugins/Runtime/StateTree/Source/StateTreeModule/Private/Debugger/StateTreeTraceAnalyzer.cpp b/Engine/Plugins/Runtime/StateTree/Source/StateTreeModule/Private/Debugger/StateTreeTraceAnalyzer.cpp
--- a/Engine/Plugins/Runtime/StateTree/Source/StateTreeModule/Private/Debugger/StateTreeTraceAnalyzer.cpp
+++ b/Engine/Plugins/Runtime/StateTree/Source/StateTreeModule/Private/Debugger/StateTreeTraceAnalyzer.cpp
@@ -127,7 +127,7 @@ bool FStateTreeTraceAnalyzer::OnEvent(const uint16 RouteId, EStyle Style, const
 		{
 			FString Message;
         	EventData.GetString("Message", Message);
-			const FStateTreeTraceLogEvent Event(WorldTime, Message);
+			const FStateTreeTraceLogEvent Event(WorldTime, EventData.GetValue<ELogVerbosity::Type>("Verbosity"), Message);
 
 			Provider.AppendEvent(FStateTreeInstanceDebugId(EventData.GetValue<uint32>("InstanceId"), EventData.GetValue<uint32>("InstanceSerial")),
 				Context.EventTime.AsSeconds(EventData.GetValue<uint64>("Cycle")),
diff --git a/Engine/Plugins/Runtime/StateTree/Source/StateTreeModule/Private/StateTreeExecutionContext.cpp b/Engine/Plugins/Runtime/StateTree/Source/StateTreeModule/Private/StateTreeExecutionContext.cpp
--- a/Engine/Plugins/Runtime/StateTree/Source/StateTreeModule/Private/StateTreeExecutionContext.cpp
+++ b/Engine/Plugins/Runtime/StateTree/Source/StateTreeModule/Private/StateTreeExecutionContext.cpp
@@ -18,9 +18,9 @@
 #define STATETREE_LOG(Verbosity, Format, ...) UE_VLOG_UELOG(GetOwner(), LogStateTree, Verbosity, TEXT("%s: ") Format, *GetInstanceDescription(), ##__VA_ARGS__)
 #define STATETREE_CLOG(Condition, Verbosity, Format, ...) UE_CVLOG_UELOG((Condition), GetOwner(), LogStateTree, Verbosity, TEXT("%s: ") Format, *GetInstanceDescription(), ##__VA_ARGS__)
 
-#define STATETREE_LOG_AND_TRACE(Verbosity, Format, ...) \
-	UE_VLOG_UELOG(GetOwner(), LogStateTree, Verbosity, TEXT("%s: ") Format, *GetInstanceDescription(), ##__VA_ARGS__); \
-	STATETREE_TRACE_LOG_EVENT(Format, ##__VA_ARGS__)
+#define STATETREE_LOG_AND_TRACE(LogVerbosity, TraceVerbosity, Format, ...) \
+	UE_VLOG_UELOG(GetOwner(), LogStateTree, LogVerbosity, TEXT("%s: ") Format, *GetInstanceDescription(), ##__VA_ARGS__); \
+	STATETREE_TRACE_LOG_EVENT(TraceVerbosity, Format, ##__VA_ARGS__)
 
 #if WITH_STATETREE_DEBUGGER
 	#define ID_NAME PREPROCESSOR_JOIN(InstanceId,__LINE__) \
@@ -45,7 +45,7 @@
 	#define STATETREE_TRACE_PHASE_BEGIN(Phase)								TRACE_STATETREE_PHASE_EVENT(GetInstanceDebugId(), Phase, EStateTreeTraceEventType::Push, FStateTreeStateHandle::Invalid)
 	#define STATETREE_TRACE_PHASE_END(Phase)								TRACE_STATETREE_PHASE_EVENT(GetInstanceDebugId(), Phase, EStateTreeTraceEventType::Pop, FStateTreeStateHandle::Invalid)
 	#define STATETREE_TRACE_ACTIVE_STATES_EVENT(ActiveFrames)				TRACE_STATETREE_ACTIVE_STATES_EVENT(GetInstanceDebugId(), ActiveFrames);
-	#define STATETREE_TRACE_LOG_EVENT(Format, ...)							TRACE_STATETREE_LOG_EVENT(GetInstanceDebugId(), Format, ##__VA_ARGS__)
+	#define STATETREE_TRACE_LOG_EVENT(TraceVerbosity, Format, ...)			TRACE_STATETREE_LOG_EVENT(GetInstanceDebugId(), TraceVerbosity, Format, ##__VA_ARGS__)
 	#define STATETREE_TRACE_STATE_EVENT(StateHandle, EventType)				TRACE_STATETREE_STATE_EVENT(GetInstanceDebugId(), StateHandle, EventType);
 	#define STATETREE_TRACE_TASK_EVENT(Index, DataView, EventType, Status)	TRACE_STATETREE_TASK_EVENT(GetInstanceDebugId(), FStateTreeIndex16(Index), DataView, EventType, Status);
 	#define STATETREE_TRACE_EVALUATOR_EVENT(Index, DataView, EventType)		TRACE_STATETREE_EVALUATOR_EVENT(GetInstanceDebugId(), FStateTreeIndex16(Index), DataView, EventType);
@@ -60,7 +60,7 @@
 	#define STATETREE_TRACE_PHASE_BEGIN(Phase)
 	#define STATETREE_TRACE_PHASE_END(Phase)
 	#define STATETREE_TRACE_ACTIVE_STATES_EVENT(ActiveFrames)
-	#define STATETREE_TRACE_LOG_EVENT(Format, ...)
+	#define STATETREE_TRACE_LOG_EVENT(TraceVerbosity, Format, ...)
 	#define STATETREE_TRACE_STATE_EVENT(StateHandle, EventType)
 	#define STATETREE_TRACE_TASK_EVENT(Index, DataView, EventType, Status)
 	#define STATETREE_TRACE_EVALUATOR_EVENT(Index, DataView, EventType)
@@ -368,7 +368,7 @@ EStateTreeRunStatus FStateTreeExecutionContext::Start(const FInstancedPropertyBa
 	
 	if (Exec.RequestedStop != EStateTreeRunStatus::Unset)
 	{
-		STATETREE_LOG_AND_TRACE(VeryVerbose, TEXT("Processing Deferred Stop"));
+		STATETREE_LOG_AND_TRACE(VeryVerbose, Log, TEXT("Processing Deferred Stop"));
 		Result = Stop(Exec.RequestedStop);
 	}
 	
@@ -409,7 +409,7 @@ EStateTreeRunStatus FStateTreeExecutionContext::Stop(EStateTreeRunStatus Complet
 	// A reentrant call to Stop or a call from Start or Tick must be deferred.
 	if (Exec.CurrentPhase != EStateTreeUpdatePhase::Unset)
 	{
-		STATETREE_LOG_AND_TRACE(VeryVerbose, TEXT("Deferring Stop at end of %s"), *UEnum::GetDisplayValueAsText(Exec.CurrentPhase).ToString());
+		STATETREE_LOG_AND_TRACE(VeryVerbose, Log, TEXT("Deferring Stop at end of %s"), *UEnum::GetDisplayValueAsText(Exec.CurrentPhase).ToString());
 
 		Exec.RequestedStop = CompletionStatus;
 		return EStateTreeRunStatus::Running;
@@ -505,7 +505,7 @@ EStateTreeRunStatus FStateTreeExecutionContext::TickPostlude()
 	
 	if (Exec.RequestedStop != EStateTreeRunStatus::Unset)
 	{
-		STATETREE_LOG_AND_TRACE(VeryVerbose, TEXT("Processing Deferred Stop"));
+		STATETREE_LOG_AND_TRACE(VeryVerbose, Log, TEXT("Processing Deferred Stop"));
 		Result = Stop(Exec.RequestedStop);
 	}
 
@@ -586,7 +586,7 @@ void FStateTreeExecutionContext::TickUpdateTasksInternal(const float DeltaTime)
 	}
 	else
 	{
-		STATETREE_TRACE_LOG_EVENT(TEXT("Global tasks completed (%s), stopping the tree"), *UEnum::GetDisplayValueAsText(EvalAndGlobalTaskStatus).ToString());
+		STATETREE_TRACE_LOG_EVENT(Log, TEXT("Global tasks completed (%s), stopping the tree"), *UEnum::GetDisplayValueAsText(EvalAndGlobalTaskStatus).ToString());
 		Exec.RequestedStop = EvalAndGlobalTaskStatus;
 	}
 }
@@ -680,7 +680,7 @@ void FStateTreeExecutionContext::SendEvent(const FGameplayTag Tag, const FConstS
 		return;
 	}
 
-	STATETREE_LOG_AND_TRACE(Verbose, TEXT("Send Event '%s'"), *Tag.ToString());
+	STATETREE_LOG_AND_TRACE(Verbose, Log, TEXT("Send Event '%s'"), *Tag.ToString());
 
 	FStateTreeEventQueue& LocalEventQueue = InstanceData.GetMutableEventQueue();
 	LocalEventQueue.SendEvent(&Owner, Tag, Payload, Origin);
@@ -2542,7 +2542,7 @@ bool FStateTreeExecutionContext::TestAllConditions(const FStateTreeExecutionFram
 				{
 					// If the source data cannot be accessed, the whole expression evaluates to false.
 					STATETREE_TRACE_CONDITION_EVENT(ConditionIndex, ConditionInstanceView, EStateTreeTraceEventType::InternalForcedFailure);
-					STATETREE_TRACE_LOG_EVENT(TEXT("Evaluation forced to false: source data cannot be accessed (e.g. enter conditions trying to access inactive parent state)"));
+					STATETREE_TRACE_LOG_EVENT(Warning, TEXT("Evaluation forced to false: source data cannot be accessed (e.g. enter conditions trying to access inactive parent state)"));
 					Values[0] = false;
 					break;
 				}
@@ -2744,7 +2744,7 @@ bool FStateTreeExecutionContext::TriggerTransitions()
 
 	if (EventQueue && EventQueue->HasEvents())
 	{
-		STATETREE_LOG_AND_TRACE(Verbose, TEXT("Trigger transitions with events: %s"), *DebugGetEventsAsString());
+		STATETREE_LOG_AND_TRACE(Verbose, Log, TEXT("Trigger transitions with events: %s"), *DebugGetEventsAsString());
 	}
 
 	NextTransition.Reset();
@@ -3144,7 +3144,7 @@ bool FStateTreeExecutionContext::TriggerTransitions()
 		// Handle the case where no transition was found.
 		if (NextTransition.Priority == EStateTreeTransitionPriority::None) //-V547
 		{
-			STATETREE_LOG_AND_TRACE(Verbose, TEXT("Could not trigger completion transition, jump back to root state."));
+			STATETREE_LOG_AND_TRACE(Verbose, Warning, TEXT("Could not trigger completion transition, jump back to root state."));
 
 			check(!Exec.ActiveFrames.IsEmpty());
 			FStateTreeExecutionFrame& RootFrame = Exec.ActiveFrames[0];
@@ -3157,7 +3157,7 @@ bool FStateTreeExecutionContext::TriggerTransitions()
 			}
 			else
 			{
-				STATETREE_LOG_AND_TRACE(Warning, TEXT("Failed to select root state. Stopping the tree with failure."));
+				STATETREE_LOG_AND_TRACE(Warning, Error, TEXT("Failed to select root state. Stopping the tree with failure."));
 
 				SetupNextTransition(RootFrame, FStateTreeStateHandle::Failed, EStateTreeTransitionPriority::Critical);
 
diff --git a/Engine/Plugins/Runtime/StateTree/Source/StateTreeModule/Public/Debugger/StateTreeTrace.h b/Engine/Plugins/Runtime/StateTree/Source/StateTreeModule/Public/Debugger/StateTreeTrace.h
--- a/Engine/Plugins/Runtime/StateTree/Source/StateTreeModule/Public/Debugger/StateTreeTrace.h
+++ b/Engine/Plugins/Runtime/StateTree/Source/StateTreeModule/Public/Debugger/StateTreeTrace.h
@@ -19,6 +19,7 @@ enum class EStateTreeStateSelectionBehavior : uint8;
 enum class EStateTreeRunStatus : uint8;
 enum class EStateTreeTraceEventType : uint8;
 enum class EStateTreeUpdatePhase : uint8;
+namespace ELogVerbosity { enum Type : uint8; }
 
 UE_TRACE_CHANNEL_EXTERN(StateTreeDebugChannel, STATETREEMODULE_API)
 
@@ -31,7 +32,7 @@ namespace UE::StateTreeTrace
 	void OutputAssetDebugIdEvent(const UStateTree* StateTree, FStateTreeIndex16 AssetDebugId);
 	void OutputInstanceLifetimeEvent(FStateTreeInstanceDebugId InstanceId, const UStateTree* StateTree, const TCHAR* InstanceName, EStateTreeTraceEventType EventType);
 	void OutputInstanceFrameEvent(FStateTreeInstanceDebugId InstanceId, const FStateTreeExecutionFrame* Frame);
-	void OutputLogEventTrace(FStateTreeInstanceDebugId InstanceId, const TCHAR* Fmt, ...);
+	void OutputLogEventTrace(FStateTreeInstanceDebugId InstanceId, ELogVerbosity::Type Verbosity, const TCHAR* Fmt, ...);
 	void OutputStateEventTrace(FStateTreeInstanceDebugId InstanceId, FStateTreeStateHandle StateHandle, EStateTreeTraceEventType EventType);
 	void OutputTaskEventTrace(FStateTreeInstanceDebugId InstanceId, FStateTreeIndex16 TaskIdx, FStateTreeDataView DataView, EStateTreeTraceEventType EventType, EStateTreeRunStatus Status);
 	void OutputEvaluatorEventTrace(FStateTreeInstanceDebugId InstanceId, FStateTreeIndex16 EvaluatorIdx, FStateTreeDataView DataView, EStateTreeTraceEventType EventType);
@@ -49,10 +50,10 @@ namespace UE::StateTreeTrace
 #define TRACE_STATETREE_PHASE_EVENT(InstanceID, Phase, EventType, StateHandle) \
 	UE::StateTreeTrace::OutputPhaseScopeEvent(InstanceID, Phase, EventType, StateHandle); \
 
-#define TRACE_STATETREE_LOG_EVENT(InstanceId, Format, ...) \
+#define TRACE_STATETREE_LOG_EVENT(InstanceId, TraceVerbosity, Format, ...) \
 	if (UE_TRACE_CHANNELEXPR_IS_ENABLED(StateTreeDebugChannel)) \
 	{ \
-		UE::StateTreeTrace::OutputLogEventTrace(InstanceId, Format, ##__VA_ARGS__); \
+		UE::StateTreeTrace::OutputLogEventTrace(InstanceId, ELogVerbosity::TraceVerbosity, Format, ##__VA_ARGS__); \
 	}
 
 #define TRACE_STATETREE_STATE_EVENT(InstanceId, StateHandle, EventType) \
@@ -93,7 +94,7 @@ namespace UE::StateTreeTrace
 #define TRACE_STATETREE_INSTANCE_EVENT(InstanceID, StateTree, InstanceName, EventType)
 #define TRACE_STATETREE_INSTANCE_FRAME_EVENT(InstanceID, Frame)
 #define TRACE_STATETREE_PHASE_EVENT(InstanceID, Phase, EventType, StateHandle)
-#define TRACE_STATETREE_LOG_EVENT(InstanceId, Format, ...)
+#define TRACE_STATETREE_LOG_EVENT(InstanceId, TraceVerbosity, Format, ...)
 #define TRACE_STATETREE_STATE_EVENT(InstanceId, StateHandle, EventType)
 #define TRACE_STATETREE_TASK_EVENT(InstanceId, TaskIdx, DataView, EventType, Status)
 #define TRACE_STATETREE_EVALUATOR_EVENT(InstanceId, EvaluatorIdx, DataView, EventType)
diff --git a/Engine/Plugins/Runtime/StateTree/Source/StateTreeModule/Public/Debugger/StateTreeTraceTypes.h b/Engine/Plugins/Runtime/StateTree/Source/StateTreeModule/Public/Debugger/StateTreeTraceTypes.h
--- a/Engine/Plugins/Runtime/StateTree/Source/StateTreeModule/Public/Debugger/StateTreeTraceTypes.h
+++ b/Engine/Plugins/Runtime/StateTree/Source/StateTreeModule/Public/Debugger/StateTreeTraceTypes.h
@@ -76,8 +76,9 @@ struct FStateTreeTracePhaseEvent : FStateTreeTraceBaseEvent
 
 struct FStateTreeTraceLogEvent : FStateTreeTraceBaseEvent
 {
-	explicit FStateTreeTraceLogEvent(const double RecordingWorldTime, const FString& Message)
+	explicit FStateTreeTraceLogEvent(const double RecordingWorldTime, ELogVerbosity::Type Verbosity, const FString& Message)
 		: FStateTreeTraceBaseEvent(RecordingWorldTime, EStateTreeTraceEventType::Unset)
+		, Verbosity(Verbosity)
 		, Message(Message)
 	{
 	}
@@ -86,13 +87,14 @@ struct FStateTreeTraceLogEvent : FStateTreeTraceBaseEvent
 	STATETREEMODULE_API FString GetValueString(const UStateTree& StateTree) const;
 	STATETREEMODULE_API FString GetTypeString(const UStateTree& StateTree) const;
 
+	ELogVerbosity::Type Verbosity;
 	FString Message;
 };
 
 struct FStateTreeTracePropertyEvent : FStateTreeTraceLogEvent
 {
 	explicit FStateTreeTracePropertyEvent(const double RecordingWorldTime, const FString& Message)
-		: FStateTreeTraceLogEvent(RecordingWorldTime, Message)
+		: FStateTreeTraceLogEvent(RecordingWorldTime, ELogVerbosity::Verbose, Message)
 	{
 	}
 
