StateTree: Remove debug code left in by mistake.

diff --git a/Engine/Plugins/Runtime/StateTree/Source/StateTreeModule/Private/StateTreeExecutionContext.cpp b/Engine/Plugins/Runtime/StateTree/Source/StateTreeModule/Private/StateTreeExecutionContext.cpp
--- a/Engine/Plugins/Runtime/StateTree/Source/StateTreeModule/Private/StateTreeExecutionContext.cpp
+++ b/Engine/Plugins/Runtime/StateTree/Source/StateTreeModule/Private/StateTreeExecutionContext.cpp
@@ -902,11 +902,6 @@ void FStateTreeExecutionContext::UpdateInstanceData(TConstArrayView<FStateTreeEx
 			}
 			else
 			{
-				if (NextStateParameterDataHandle.IsValid())
-				{
-					UE_DEBUG_BREAK();
-				}
-				
 				if (State.ParameterTemplateIndex.IsValid())
 				{
 					// Linked state's instance data is the parameters.
