[State Tree] StateTreeInstanceData RandomStream variable Initialization Fix/Improvement

- When duplicating the Component, RandomStream was always copied from the one on Component CDO and wasn't constructed with the current time as expected. We now initialize RandomStream at FStateTreeExecutionContext::Start(), where the InstanceData is going to persist till Stop().
- Moved RandomStream variable to FStateTreeExecutionState as it better reflects the window it persists now.
- Fixed non-deterministic default-constructed value of RandomStream variable in UStruct LogError by the above

diff --git a/Engine/Plugins/Runtime/StateTree/Source/StateTreeModule/Private/StateTreeExecutionContext.cpp b/Engine/Plugins/Runtime/StateTree/Source/StateTreeModule/Private/StateTreeExecutionContext.cpp
--- a/Engine/Plugins/Runtime/StateTree/Source/StateTreeModule/Private/StateTreeExecutionContext.cpp
+++ b/Engine/Plugins/Runtime/StateTree/Source/StateTreeModule/Private/StateTreeExecutionContext.cpp
@@ -234,7 +234,7 @@ bool FStateTreeExecutionContext::SetContextDataByName(const FName Name, FStateTr
 	return false;
 }
 
-EStateTreeRunStatus FStateTreeExecutionContext::Start(const FInstancedPropertyBag* InitialParameters)
+EStateTreeRunStatus FStateTreeExecutionContext::Start(const FInstancedPropertyBag* InitialParameters, int32 RandomSeed)
 {
 	CSV_SCOPED_TIMING_STAT_EXCLUSIVE(StateTree_Start);
 
@@ -274,6 +274,7 @@ EStateTreeRunStatus FStateTreeExecutionContext::Start(const FInstancedPropertyBa
 	InitFrame.bIsGlobalFrame = true;
 	
 	UpdateInstanceData({}, Exec.ActiveFrames);
+	Exec.RandomStream.Initialize(RandomSeed == -1 ? FPlatformTime::Cycles() : RandomSeed);
 
 	if (!CollectActiveExternalData())
 	{
@@ -3858,7 +3859,7 @@ bool FStateTreeExecutionContext::SelectStateInternal(
 				for (uint32 Index = 0; Index < LastIndex; ++Index)
 				{
 					// Get a random integer in [Index, Num)
-					const uint32 IndexToSwap = InstanceDataStorage->GetRandomStream().RandRange(Index, LastIndex);
+					const uint32 IndexToSwap = Exec.RandomStream.RandRange(Index, LastIndex);
 					RandomChildStates.Swap(Index, IndexToSwap);
 				}
 
diff --git a/Engine/Plugins/Runtime/StateTree/Source/StateTreeModule/Private/StateTreeInstanceData.cpp b/Engine/Plugins/Runtime/StateTree/Source/StateTreeModule/Private/StateTreeInstanceData.cpp
--- a/Engine/Plugins/Runtime/StateTree/Source/StateTreeModule/Private/StateTreeInstanceData.cpp
+++ b/Engine/Plugins/Runtime/StateTree/Source/StateTreeModule/Private/StateTreeInstanceData.cpp
@@ -83,10 +83,7 @@ namespace UE::StateTree
 // FStateTreeInstanceStorage
 //----------------------------------------------------------------//
 
-FStateTreeInstanceStorage::FStateTreeInstanceStorage()
-	: RandomStream(FPlatformTime::Cycles())
-{
-}
+FStateTreeInstanceStorage::FStateTreeInstanceStorage() = default;
 
 FStateTreeInstanceStorage::FStateTreeInstanceStorage(const FStateTreeInstanceStorage& Other)
 	: InstanceStructs(Other.InstanceStructs)
@@ -95,7 +92,6 @@ FStateTreeInstanceStorage::FStateTreeInstanceStorage(const FStateTreeInstanceSto
 	, EventQueue(MakeShared<FStateTreeEventQueue>(*Other.EventQueue))
 	, TransitionRequests(Other.TransitionRequests)
 	, GlobalParameters(Other.GlobalParameters)
-	, RandomStream(Other.RandomStream)
 {
 }
 
@@ -106,7 +102,6 @@ FStateTreeInstanceStorage::FStateTreeInstanceStorage(FStateTreeInstanceStorage&&
 	, EventQueue(Other.EventQueue)
 	, TransitionRequests(MoveTemp(Other.TransitionRequests))
 	, GlobalParameters(MoveTemp(Other.GlobalParameters))
-	, RandomStream(MoveTemp(Other.RandomStream))
 {
 	Other.EventQueue = MakeShared<FStateTreeEventQueue>();
 }
@@ -119,7 +114,6 @@ FStateTreeInstanceStorage& FStateTreeInstanceStorage::operator=(const FStateTree
 	EventQueue = MakeShared<FStateTreeEventQueue>(*Other.EventQueue);
 	TransitionRequests = Other.TransitionRequests;
 	GlobalParameters = Other.GlobalParameters;
-	RandomStream = Other.RandomStream;
 
 	return *this;
 }
@@ -133,7 +127,6 @@ FStateTreeInstanceStorage& FStateTreeInstanceStorage::operator=(FStateTreeInstan
 	Other.EventQueue = MakeShared<FStateTreeEventQueue>();
 	TransitionRequests = MoveTemp(Other.TransitionRequests);
 	GlobalParameters = MoveTemp(Other.GlobalParameters);
-	RandomStream = MoveTemp(Other.RandomStream);
 
 	return *this;
 }
@@ -258,11 +251,6 @@ void FStateTreeInstanceStorage::SetGlobalParameters(const FInstancedPropertyBag&
 	GlobalParameters = Parameters;
 }
 
-void FStateTreeInstanceStorage::SetRandomSeed(int32 Seed)
-{
-	RandomStream.Initialize(Seed);
-}
-
 void FStateTreeInstanceStorage::AddStructReferencedObjects(FReferenceCollector& Collector)
 {
 	Collector.AddPropertyReferencesWithStructARO(TBaseStructure<FStateTreeInstanceStorage>::Get(), this);
@@ -279,7 +267,6 @@ void FStateTreeInstanceStorage::Reset()
 	ExecutionState.Reset();
 	TemporaryInstances.Reset();
 	GlobalParameters.Reset();
-	RandomStream.Reset();
 }
 
 //----------------------------------------------------------------//
@@ -403,11 +390,6 @@ int32 FStateTreeInstanceData::GetEstimatedMemoryUsage() const
 	return Size;
 }
 
-void FStateTreeInstanceData::SetRandomSeed(int32 Seed)
-{
-	GetMutableStorage().SetRandomSeed(Seed);
-}
-
 bool FStateTreeInstanceData::Identical(const FStateTreeInstanceData* Other, uint32 PortFlags) const
 {
 	if (Other == nullptr)
diff --git a/Engine/Plugins/Runtime/StateTree/Source/StateTreeModule/Public/StateTreeExecutionContext.h b/Engine/Plugins/Runtime/StateTree/Source/StateTreeModule/Public/StateTreeExecutionContext.h
--- a/Engine/Plugins/Runtime/StateTree/Source/StateTreeModule/Public/StateTreeExecutionContext.h
+++ b/Engine/Plugins/Runtime/StateTree/Source/StateTreeModule/Public/StateTreeExecutionContext.h
@@ -142,9 +142,10 @@ public:
 	/**
 	 * Start executing.
 	 * @param InitialParameters Optional override of parameters initial values
+	 * @param RandomSeed Optional override of initial seed for RandomStream. By default FPlatformTime::Cycles() will be used.
 	 * @return Tree execution status after the start.
 	 */
-	EStateTreeRunStatus Start(const FInstancedPropertyBag* InitialParameters = nullptr);
+	EStateTreeRunStatus Start(const FInstancedPropertyBag* InitialParameters = nullptr, int32 RandomSeed = -1);
 	
 	/**
 	 * Stop executing if the tree is running.
diff --git a/Engine/Plugins/Runtime/StateTree/Source/StateTreeModule/Public/StateTreeExecutionTypes.h b/Engine/Plugins/Runtime/StateTree/Source/StateTreeModule/Public/StateTreeExecutionTypes.h
--- a/Engine/Plugins/Runtime/StateTree/Source/StateTreeModule/Public/StateTreeExecutionTypes.h
+++ b/Engine/Plugins/Runtime/StateTree/Source/StateTreeModule/Public/StateTreeExecutionTypes.h
@@ -626,6 +626,7 @@ PRAGMA_ENABLE_DEPRECATION_WARNINGS
 		CurrentPhase = EStateTreeUpdatePhase::Unset;
 		CompletedStateHandle = FStateTreeStateHandle::Invalid;
 		StateChangeCount = 0;
+		RandomStream.Reset();
 	}
 
 	UE_DEPRECATED(5.5, "Use FindAndRemoveExpiredDelayedTransitions() instead or search through DelayedTransitions directly.")
@@ -695,6 +696,10 @@ PRAGMA_ENABLE_DEPRECATION_WARNINGS
 	UPROPERTY()
 	uint16 StateChangeCount = 0;
 
+	/** Used by state tree random-based operations. */
+	UPROPERTY()
+	FRandomStream RandomStream;
+
 #if WITH_EDITORONLY_DATA
 	UE_DEPRECATED(5.3, "Use DataHandle instead.")
 	UPROPERTY()
diff --git a/Engine/Plugins/Runtime/StateTree/Source/StateTreeModule/Public/StateTreeInstanceData.h b/Engine/Plugins/Runtime/StateTree/Source/StateTreeModule/Public/StateTreeInstanceData.h
--- a/Engine/Plugins/Runtime/StateTree/Source/StateTreeModule/Public/StateTreeInstanceData.h
+++ b/Engine/Plugins/Runtime/StateTree/Source/StateTreeModule/Public/StateTreeInstanceData.h
@@ -239,15 +239,6 @@ struct STATETREEMODULE_API FStateTreeInstanceStorage
 	/** Resets the storage to initial state. */
 	void Reset();
 	
-	/** Update seed for the RandomStream used by state tree random-based operations. */
-	void SetRandomSeed(int32 Seed);
-
-	/** @return RandomStream used by state tree random-based operations. */
-	const FRandomStream& GetRandomStream() const
-	{
-		return RandomStream;
-	}
-
 	UE_DEPRECATED(5.4, "Use Num() instead.")
 	int32 NumStructs() const { return 0; }
 
@@ -287,10 +278,6 @@ protected:
 	UPROPERTY(Transient)
 	FInstancedPropertyBag GlobalParameters;
 
-	/** Used by state tree random-based operations. By Default it is initialized with seed 0 */
-	UPROPERTY()
-	FRandomStream RandomStream;
-
 	friend struct FStateTreeInstanceData;
 };
 
@@ -425,15 +412,6 @@ struct STATETREEMODULE_API FStateTreeInstanceData
 
 	int32 GetEstimatedMemoryUsage() const;
 	
-	/** @return RandomStream used by state tree random-based operations. */
-	const FRandomStream& GetRandomStream() const
-	{
-		return GetStorage().GetRandomStream();
-	}
-
-	/** Update seed for the RandomStream used by state tree random-based operations. */
-	void SetRandomSeed(int32 Seed);
-
 	/** Type traits */
 	bool Identical(const FStateTreeInstanceData* Other, uint32 PortFlags) const;
 	void AddStructReferencedObjects(FReferenceCollector& Collector);
