StateTrees: Made FStateTreeNodeBase::PostLoad available in all kinds of builds.

diff --git a/Engine/Plugins/Runtime/StateTree/Source/StateTreeModule/Private/StateTree.cpp b/Engine/Plugins/Runtime/StateTree/Source/StateTreeModule/Private/StateTree.cpp
--- a/Engine/Plugins/Runtime/StateTree/Source/StateTreeModule/Private/StateTree.cpp
+++ b/Engine/Plugins/Runtime/StateTree/Source/StateTreeModule/Private/StateTree.cpp
@@ -351,7 +351,6 @@ void UStateTree::PostLoad()
 {
 	Super::PostLoad();
 
-#if WITH_EDITOR
 	for (int32 NodeIndex = 0; NodeIndex < Nodes.Num(); ++NodeIndex)
 	{
 		FStructView NodeView = Nodes[NodeIndex];
@@ -372,7 +371,6 @@ void UStateTree::PostLoad()
 			
 		}
 	}
-#endif // WITH_EDITOR
 
 	const int32 CurrentVersion = GetLinkerCustomVersion(FStateTreeCustomVersion::GUID);
 #if WITH_EDITOR
