StateTrees: Made StateTree override schemas overridable.

diff --git a/Engine/Plugins/Runtime/StateTree/Source/StateTreeModule/Public/StateTreeReference.h b/Engine/Plugins/Runtime/StateTree/Source/StateTreeModule/Public/StateTreeReference.h
--- a/Engine/Plugins/Runtime/StateTree/Source/StateTreeModule/Public/StateTreeReference.h
+++ b/Engine/Plugins/Runtime/StateTree/Source/StateTreeModule/Public/StateTreeReference.h
@@ -144,7 +144,7 @@ private:
 	FGameplayTag StateTag;
 
 	/** State Tree and parameters to replace the linked state asset with. */
-	UPROPERTY(EditAnywhere, Category = "")
+	UPROPERTY(EditAnywhere, Category = "", meta=(SchemaCanBeOverriden))
 	FStateTreeReference StateTreeReference;
 	
 	friend class FStateTreeReferenceOverridesDetails;
