StateTree: Fix accessing state captured event

diff --git a/Engine/Plugins/Runtime/StateTree/Source/StateTreeModule/Private/StateTreeExecutionContext.cpp b/Engine/Plugins/Runtime/StateTree/Source/StateTreeModule/Private/StateTreeExecutionContext.cpp
--- a/Engine/Plugins/Runtime/StateTree/Source/StateTreeModule/Private/StateTreeExecutionContext.cpp
+++ b/Engine/Plugins/Runtime/StateTree/Source/StateTreeModule/Private/StateTreeExecutionContext.cpp
@@ -1025,14 +1025,15 @@ FStateTreeDataView FStateTreeExecutionContext::GetDataView(const FStateTreeExecu
 
 	case EStateTreeDataSourceType::StateEvent:
 		{
+			// If state selection is going, return FStateTreeEvent of the event currently captured by the state selection.
 			if (CurrentlyProcessedStateSelectionEvents)
 			{
 				if (const FCompactStateTreeState* State = CurrentFrame.StateTree->GetStateFromHandle(Handle.GetState()))
 				{
-					if (const FStateTreeEvent* Event = CurrentlyProcessedStateSelectionEvents->Events[State->Depth].Get())
+					// Events are read only, but we cannot express that in FStateTreeDataView.
+					if (FStateTreeEvent* Event = CurrentlyProcessedStateSelectionEvents->Events[State->Depth].GetMutable())
 					{
-						// const_cast because events are read only, but we cannot express that in FStateTreeDataView.
-						return FStateTreeDataView(FStructView::Make(*const_cast<FStateTreeEvent*>(Event)));
+						return FStateTreeDataView(FStructView::Make(*Event));
 					}
 				}
 
@@ -1100,9 +1101,11 @@ FStateTreeDataView FStateTreeExecutionContext::GetDataViewFromInstanceStorage(FS
 
 	case EStateTreeDataSourceType::StateEvent:
 		{
+			// Return FStateTreeEvent from shared event.
 			FStateTreeSharedEvent& SharedEvent = InstanceDataStorage.GetMutableStruct(CurrentFrame.ActiveInstanceIndexBase.Get() + Handle.GetIndex()).Get<FStateTreeSharedEvent>();
 			check(SharedEvent.IsValid());
-			return FStateTreeDataView(FStructView::Make(SharedEvent));
+			// Events are read only, but we cannot express that in FStateTreeDataView.
+			return FStateTreeDataView(FStructView::Make(*SharedEvent.GetMutable()));
 		}
 
 	case EStateTreeDataSourceType::ContextData:
diff --git a/Engine/Plugins/Runtime/StateTree/Source/StateTreeTestSuite/Private/StateTreeTest.cpp b/Engine/Plugins/Runtime/StateTree/Source/StateTreeTestSuite/Private/StateTreeTest.cpp
--- a/Engine/Plugins/Runtime/StateTree/Source/StateTreeTestSuite/Private/StateTreeTest.cpp
+++ b/Engine/Plugins/Runtime/StateTree/Source/StateTreeTestSuite/Private/StateTreeTest.cpp
@@ -1321,7 +1321,9 @@ struct FStateTreeTest_StateRequiringEvent : FAITestBase
 		const FString EnterStateStr(TEXT("EnterState"));
 
 		Status = Exec.Start();
-		AITEST_TRUE("StateTree TaskInitial should not enter state", Exec.Expect(TaskInitial.GetName(), EnterStateStr));
+		AITEST_TRUE("StateTree TaskInitial should enter state", Exec.Expect(TaskInitial.GetName(), EnterStateStr));
+		Exec.LogClear();
+
 		Exec.SendEvent(ValidTag, FConstStructView::Make(FValidPayload()));
 		Status = Exec.Tick(0.1f);
 
@@ -1373,7 +1375,13 @@ struct FStateTreeTest_PassingTransitionEventToStateSelection : FAITestBase
 		UStateTreeState& StateB = Root.AddChildState(FName(TEXT("B")));
 		StateB.bHasRequiredEventToEnter  = true;
 		StateB.RequiredEventToEnter.PayloadStruct = FStateTreeTest_PropertyStructA::StaticStruct();
-		auto& TaskB = StateB.AddTask<FTestTask_Stand>(FName(TEXT("TaskB")));
+		auto& TaskB = StateB.AddTask<FTestTask_PrintValue>(FName(TEXT("TaskB")));
+		// Test copying data from the state event. The condition properties are copied from temp instance data during selection, this gets copied from active instance data.
+		TaskB.GetInstanceData().Value = -1; // Initially -1, expected to be overridden by property binding below. 
+		EditorData.AddPropertyBinding(
+			FStateTreePropertyPath(StateB.GetEventID(), PathToPayloadMember.GetSegments()),
+			FStateTreePropertyPath(TaskB.ID, TEXT("Value")));
+		
 		TStateTreeEditorNode<FStateTreeCompareIntCondition>& BIntCond = StateB.AddEnterCondition<FStateTreeCompareIntCondition>(EGenericAICheck::Equal);
 		BIntCond.GetInstanceData().Right = 1;
 		EditorData.AddPropertyBinding(
@@ -1383,6 +1391,7 @@ struct FStateTreeTest_PassingTransitionEventToStateSelection : FAITestBase
 		// This state should be selected only initially when there's not event in the queue.
 		UStateTreeState& StateInitial = Root.AddChildState(FName(TEXT("Initial")));
 		auto& TaskInitial = StateInitial.AddTask<FTestTask_Stand>(FName(TEXT("TaskInitial")));
+		// Transition from Initial -> StateA
 		FStateTreeTransition& TransA = StateInitial.AddTransition(EStateTreeTransitionTrigger::OnEvent, FGameplayTag(), EStateTreeTransitionType::GotoState, &StateA);
 		TransA.RequiredEvent.PayloadStruct = FStateTreeTest_PropertyStructA::StaticStruct();
 		TStateTreeEditorNode<FStateTreeCompareIntCondition>& TransAIntCond = TransA.AddCondition<FStateTreeCompareIntCondition>(EGenericAICheck::Equal);
@@ -1390,7 +1399,7 @@ struct FStateTreeTest_PassingTransitionEventToStateSelection : FAITestBase
 		EditorData.AddPropertyBinding(
 			FStateTreePropertyPath(TransA.GetEventID(), PathToPayloadMember.GetSegments()),
 			FStateTreePropertyPath(TransAIntCond.ID, TEXT("Left")));
-
+		// Transition from Initial -> StateB
 		FStateTreeTransition& TransB = StateInitial.AddTransition(EStateTreeTransitionTrigger::OnEvent, FGameplayTag(), EStateTreeTransitionType::GotoState, &StateB);
 		TransB.RequiredEvent.PayloadStruct = FStateTreeTest_PropertyStructA::StaticStruct();
 		TStateTreeEditorNode<FStateTreeCompareIntCondition>& TransBIntCond = TransB.AddCondition<FStateTreeCompareIntCondition>(EGenericAICheck::Equal);
@@ -1414,13 +1423,16 @@ struct FStateTreeTest_PassingTransitionEventToStateSelection : FAITestBase
 		const FString EnterStateStr(TEXT("EnterState"));
 
 		Status = Exec.Start();
-		AITEST_TRUE("StateTree TaskInitial should not enter state", Exec.Expect(TaskInitial.GetName(), EnterStateStr));
+		AITEST_TRUE("StateTree TaskInitial should enter state", Exec.Expect(TaskInitial.GetName(), EnterStateStr));
+		Exec.LogClear();
+
+		// The conditions test for payload Value=1, the first event should not trigger transition. 
 		Exec.SendEvent(UE::StateTree::Tests::FNativeGameplayTags::Get().TestTag, FConstStructView::Make(FStateTreeTest_PropertyStructA{0}));
 		Exec.SendEvent(UE::StateTree::Tests::FNativeGameplayTags::Get().TestTag, FConstStructView::Make(FStateTreeTest_PropertyStructA{1}));
 		Status = Exec.Tick(0.1f);
 
 		AITEST_FALSE("StateTree TaskA should not enter state", Exec.Expect(TaskA.GetName(), EnterStateStr));
-		AITEST_TRUE("StateTree TaskB should enter state", Exec.Expect(TaskB.GetName(), EnterStateStr));
+		AITEST_TRUE("StateTree TaskB should enter state", Exec.Expect(TaskB.GetName(), TEXT("EnterState1"))); // TaskB decorates "EnterState" with value from the payload.
 		Exec.LogClear();
 
 		return true;
diff --git a/Engine/Plugins/Runtime/StateTree/Source/StateTreeTestSuite/Private/StateTreeTestTypes.h b/Engine/Plugins/Runtime/StateTree/Source/StateTreeTestSuite/Private/StateTreeTestTypes.h
--- a/Engine/Plugins/Runtime/StateTree/Source/StateTreeTestSuite/Private/StateTreeTestTypes.h
+++ b/Engine/Plugins/Runtime/StateTree/Source/StateTreeTestSuite/Private/StateTreeTestTypes.h
@@ -314,6 +314,9 @@ struct FTestTask_StandInstanceData
 {
 	GENERATED_BODY()
 
+	UPROPERTY(EditAnywhere, Category = "Parameter")
+	int32 Value = 0;
+	
 	UPROPERTY()
 	int32 CurrentTick = 0;
 };
