[StateTreeDebugger] fixed timelines reused for the same instances through multiple traces taken during the same game session

diff --git a/Engine/Plugins/Runtime/StateTree/Source/StateTreeEditorModule/Private/Debugger/SStateTreeDebuggerView.cpp b/Engine/Plugins/Runtime/StateTree/Source/StateTreeEditorModule/Private/Debugger/SStateTreeDebuggerView.cpp
--- a/Engine/Plugins/Runtime/StateTree/Source/StateTreeEditorModule/Private/Debugger/SStateTreeDebuggerView.cpp
+++ b/Engine/Plugins/Runtime/StateTree/Source/StateTreeEditorModule/Private/Debugger/SStateTreeDebuggerView.cpp
@@ -348,9 +348,9 @@ void SStateTreeDebuggerView::Construct(const FArguments& InArgs, const UStateTre
 
 	IStateTreeModule& StateTreeModule = FModuleManager::GetModuleChecked<IStateTreeModule>("StateTreeModule");
 	bRecording = StateTreeModule.IsTracing();
-	UE::StateTree::Delegates::OnTracingStateChanged.AddSPLambda(this, [&bRecording=bRecording](const bool bTracesEnabled)
+	UE::StateTree::Delegates::OnTracingStateChanged.AddSPLambda(this, [&bRecording=bRecording](const EStateTreeTraceStatus TraceStatus)
 		{
-			bRecording = bTracesEnabled;
+			bRecording = TraceStatus == EStateTreeTraceStatus::TracesStarted;
 		});
 
 	// Bind callbacks to the debugger delegates
@@ -936,11 +936,14 @@ void SStateTreeDebuggerView::OnDebuggerScrubStateChanged(const UE::StateTreeDebu
 		if (bShouldPopScopeStack)
 		{
 			// Pop scope and remove associated element if empty
-			TSharedPtr<FStateTreeDebuggerEventTreeElement> Scope = ScopeStack.Pop();
-			if (Scope->Children.IsEmpty())
+			if (ensureMsgf(ScopeStack.Num() > 0, TEXT("Expected to pop an entry in the scope stack but it is already empty.")))
 			{
-				TArray<TSharedPtr<FStateTreeDebuggerEventTreeElement>>& TreeElements = ScopeStack.IsEmpty() ? EventsTreeElements : ScopeStack.Top()->Children;
-				TreeElements.Remove(Scope);
+				TSharedPtr<FStateTreeDebuggerEventTreeElement> Scope = ScopeStack.Pop();
+				if (Scope->Children.IsEmpty())
+				{
+					TArray<TSharedPtr<FStateTreeDebuggerEventTreeElement>>& TreeElements = ScopeStack.IsEmpty() ? EventsTreeElements : ScopeStack.Top()->Children;
+					TreeElements.Remove(Scope);
+				}
 			}
 			// We don't want to create a child when a scope is popped.
 			continue;
diff --git a/Engine/Plugins/Runtime/StateTree/Source/StateTreeEditorModule/Private/Debugger/StateTreeDebuggerTrack.cpp b/Engine/Plugins/Runtime/StateTree/Source/StateTreeEditorModule/Private/Debugger/StateTreeDebuggerTrack.cpp
--- a/Engine/Plugins/Runtime/StateTree/Source/StateTreeEditorModule/Private/Debugger/StateTreeDebuggerTrack.cpp
+++ b/Engine/Plugins/Runtime/StateTree/Source/StateTreeEditorModule/Private/Debugger/StateTreeDebuggerTrack.cpp
@@ -58,32 +58,65 @@ bool FStateTreeDebuggerInstanceTrack::UpdateInternal()
 		const TConstArrayView<UE::StateTreeDebugger::FFrameSpan> Spans = EventCollection.FrameSpans;
 		const TConstArrayView<FStateTreeTraceEventVariantType> Events = EventCollection.Events;
 		const uint32 NumStateChanges = EventCollection.ActiveStatesChanges.Num();
-		
-		for (uint32 StateChangeIndex = 0; StateChangeIndex < NumStateChanges; ++StateChangeIndex)
+
+		TArray<UE::StateTreeDebugger::FInstanceEventCollection::FContiguousTraceInfo> TracesInfo = EventCollection.ContiguousTracesData;
+		// Append the on going trace info to "stopped" previous trace
+		if (NumStateChanges > 0)
 		{
-			const uint32 SpanIndex = EventCollection.ActiveStatesChanges[StateChangeIndex].SpanIndex;
-			const uint32 EventIndex = EventCollection.ActiveStatesChanges[StateChangeIndex].EventIndex;
-			const FStateTreeTraceActiveStatesEvent& Event = Events[EventIndex].Get<FStateTreeTraceActiveStatesEvent>();
-				
-			FString StatePath = Event.GetValueString(*StateTree);
-			UE::StateTreeDebugger::FFrameSpan Span = EventCollection.FrameSpans[SpanIndex];
-			SStateTreeDebuggerEventTimelineView::FTimelineEventData::EventWindow& Window = EventData->Windows.AddDefaulted_GetRef();
-			Window.Color = MakeRandomColor(GetTypeHash(StatePath));
-			Window.Description = FText::FromString(StatePath);
-			Window.TimeStart = Span.GetWorldTimeStart();
-
-			// When there is another state change after the current one in the list we use
-			// the end time of the previous frame span to close the window. Not using the start time
-			// since we might have a gap in case of multiple recordings during the same game session.
-			if (StateChangeIndex < NumStateChanges-1)
+			TracesInfo.Emplace(EventCollection.ActiveStatesChanges.Last().SpanIndex);
+		}
+
+		int32 StateChangeEndIndex = INDEX_NONE;
+		for (int32 TraceIndex = 0; TraceIndex < TracesInfo.Num(); TraceIndex++)
+		{
+			UE::StateTreeDebugger::FInstanceEventCollection::FContiguousTraceInfo TraceInfo = TracesInfo[TraceIndex];
+			// Start at first event for the first trace or from the end index of the previous trace 
+			const int32 StateChangeBeginIndex = (StateChangeEndIndex == INDEX_NONE) ? 0 : StateChangeEndIndex;
+
+			// Find the starting index of the next trace to stop our iteration
+			StateChangeEndIndex = EventCollection.ActiveStatesChanges.IndexOfByPredicate(
+				[LastSpanIndex = TraceInfo.LastSpanIndex](const UE::StateTreeDebugger::FInstanceEventCollection::FActiveStatesChangePair& Pair)
+				{
+					return Pair.SpanIndex > LastSpanIndex;
+				});
+
+			// When not found means we are processing the last (or the only) trace
+			if (StateChangeEndIndex == INDEX_NONE)
 			{
-				const uint32 EndSpanIndex = EventCollection.ActiveStatesChanges[StateChangeIndex+1].SpanIndex-1;
-				check(EventCollection.FrameSpans.IsValidIndex(EndSpanIndex));
-				Window.TimeEnd = EventCollection.FrameSpans[EndSpanIndex].GetWorldTimeEnd();
+				StateChangeEndIndex = NumStateChanges;
 			}
-			else
+
+			for (int32 StateChangeIndex = StateChangeBeginIndex; StateChangeIndex < StateChangeEndIndex; ++StateChangeIndex)
 			{
-				Window.TimeEnd = Debugger->IsActiveInstance(RecordingDuration, InstanceId) ? RecordingDuration : EventCollection.FrameSpans.Last().GetWorldTimeEnd();
+				const uint32 SpanIndex = EventCollection.ActiveStatesChanges[StateChangeIndex].SpanIndex;
+				const uint32 EventIndex = EventCollection.ActiveStatesChanges[StateChangeIndex].EventIndex;
+				const FStateTreeTraceActiveStatesEvent& Event = Events[EventIndex].Get<FStateTreeTraceActiveStatesEvent>();
+				
+				FString StatePath = Event.GetValueString(*StateTree);
+				UE::StateTreeDebugger::FFrameSpan Span = EventCollection.FrameSpans[SpanIndex];
+				SStateTreeDebuggerEventTimelineView::FTimelineEventData::EventWindow& Window = EventData->Windows.AddDefaulted_GetRef();
+				Window.Color = MakeRandomColor(GetTypeHash(StatePath));
+				Window.Description = FText::FromString(StatePath);
+				Window.TimeStart = Span.GetWorldTimeStart();
+
+				// For the last received event we use either the current recording duration if the instance is still active
+				if (StateChangeIndex == (NumStateChanges - 1))
+				{
+					// or the last recorded frame time.
+					Window.TimeEnd = Debugger->IsActiveInstance(RecordingDuration, InstanceId)
+						? RecordingDuration
+						: EventCollection.FrameSpans.Last().GetWorldTimeEnd();
+				}
+				else
+				{
+					// When there is another state change after the current one in the list we use it to close the window.
+					// If the event is not the last of that specific trace then we use the start time of the next span.
+					// Otherwise, we use the end time of the last frame that was part of that trace.
+					const int32 NextStateChangeSpanIndex = EventCollection.ActiveStatesChanges[StateChangeIndex+1].SpanIndex;
+					Window.TimeEnd = (StateChangeIndex < (StateChangeEndIndex - 1))
+						? EventCollection.FrameSpans[NextStateChangeSpanIndex].GetWorldTimeStart()
+						: EventCollection.FrameSpans[NextStateChangeSpanIndex - 1].GetWorldTimeEnd(); //TraceInfo.LastRecordingTime;
+				}
 			}
 		}
 
diff --git a/Engine/Plugins/Runtime/StateTree/Source/StateTreeModule/Private/Debugger/StateTreeDebugger.cpp b/Engine/Plugins/Runtime/StateTree/Source/StateTreeModule/Private/Debugger/StateTreeDebugger.cpp
--- a/Engine/Plugins/Runtime/StateTree/Source/StateTreeModule/Private/Debugger/StateTreeDebugger.cpp
+++ b/Engine/Plugins/Runtime/StateTree/Source/StateTreeModule/Private/Debugger/StateTreeDebugger.cpp
@@ -78,10 +78,10 @@ FStateTreeDebugger::FStateTreeDebugger()
 	: StateTreeModule(FModuleManager::GetModuleChecked<IStateTreeModule>("StateTreeModule"))
 	, ScrubState(EventCollections)
 {
-	TracingStateChangedHandle = UE::StateTree::Delegates::OnTracingStateChanged.AddLambda([this](const bool bTracesEnabled)
+	TracingStateChangedHandle = UE::StateTree::Delegates::OnTracingStateChanged.AddLambda([this](const EStateTreeTraceStatus TraceStatus)
 		{
 			// StateTree traces got enabled in the current process so let's analyse it if not already analysing something.
-			if (bTracesEnabled && !IsAnalysisSessionActive())
+			if (TraceStatus == EStateTreeTraceStatus::TracesStarted && !IsAnalysisSessionActive())
 			{
 				RequestAnalysisOfLatestTrace();
 			}
@@ -869,23 +869,29 @@ bool FStateTreeDebugger::ProcessEvent(const FStateTreeInstanceDebugId InstanceId
 	}
 	else
 	{
-		const TraceServices::FFrame& LastFrame = ExistingCollection->FrameSpans.Last().Frame;
+		const UE::StateTreeDebugger::FFrameSpan& LastSpan = ExistingCollection->FrameSpans.Last();
+		const TraceServices::FFrame& LastFrame = LastSpan.Frame;
+		const uint64 FrameIndexOffset = ExistingCollection->ContiguousTracesData.IsEmpty()
+			? 0
+			: (ExistingCollection->FrameSpans[ExistingCollection->ContiguousTracesData.Last().LastSpanIndex].Frame.Index + 1);
+
 		// Add new frame span for new larger frame index
-		if (Frame.Index > LastFrame.Index)
+		if (Frame.Index + FrameIndexOffset > LastFrame.Index)
 		{
 			bShouldAddFrameToSpans = true;
+
+			// Apply current offset to the frame index
+			FrameToAddInSpans.Index += FrameIndexOffset;
 		}
 		else if (Frame.Index < LastFrame.Index && Frame.StartTime > LastFrame.StartTime)
 		{
-			// Some events are buffered and can be sent from an older world recording
-			// time in case of late recording (e.g., ActiveStatesEvent) so if we want to aggregate the events
-			// with existing data we'll snap it to the most recent time to not break the timelines
-			RecordingWorldTime = FMath::Max(RecordingWorldTime, RecordingDuration);
-
 			// Frame index will restart at 0 if a new session is started,
 			// in that case we offset the frame we store to append to existing data
-			FrameToAddInSpans.Index += LastFrame.Index + 1;
 			bShouldAddFrameToSpans = true;
+
+			const UE::StateTreeDebugger::FInstanceEventCollection::FContiguousTraceInfo& TraceInfo =
+				ExistingCollection->ContiguousTracesData.Emplace_GetRef(ExistingCollection->FrameSpans.Num()-1);
+			FrameToAddInSpans.Index += ExistingCollection->FrameSpans[TraceInfo.LastSpanIndex].Frame.Index + 1;
 		}
 	}
 
diff --git a/Engine/Plugins/Runtime/StateTree/Source/StateTreeModule/Private/Debugger/StateTreeTrace.cpp b/Engine/Plugins/Runtime/StateTree/Source/StateTreeModule/Private/Debugger/StateTreeTrace.cpp
--- a/Engine/Plugins/Runtime/StateTree/Source/StateTreeModule/Private/Debugger/StateTreeTrace.cpp
+++ b/Engine/Plugins/Runtime/StateTree/Source/StateTreeModule/Private/Debugger/StateTreeTrace.cpp
@@ -214,35 +214,33 @@ struct FInstanceEventBufferedData
 
 	FInstanceEventBufferedData() = default;
 	explicit FInstanceEventBufferedData(
-		const double RecordingWorldTime,
 		const UStateTree* StateTree,
 		const FStateTreeInstanceDebugId InstanceId,
 		const FString& InstanceName)
 		: InstanceName(InstanceName)
 		, WeakStateTree(StateTree)
 		, InstanceId(InstanceId)
-		, LifetimeRecordingWorldTime(RecordingWorldTime)
 	{
 	}
 
+	void CloseRecording(double WorldTime) const
+	{
+		// Output and empty active states event at the last recorded world time to close active states
+		UE_TRACE_LOG(StateTreeDebugger, WorldTimestampEvent, StateTreeDebugChannel) << WorldTimestampEvent.WorldTime(WorldTime);
+		FActiveStates EmptyState;
+		EmptyState.Output(InstanceId);
+	}
+
 	void Trace() const
 	{
 		if (ensureMsgf(UE_TRACE_CHANNELEXPR_IS_ENABLED(StateTreeDebugChannel), TEXT("Tracing a buffered data is expected only if channel is enabled.")))
 		{
 			if (const UStateTree* StateTree = WeakStateTree.Get())
 			{
-				// Force a world time update since we are tracing an event from the past
-				UE_TRACE_LOG(StateTreeDebugger, WorldTimestampEvent, StateTreeDebugChannel)
-					<< WorldTimestampEvent.WorldTime(LifetimeRecordingWorldTime);
-
 				OutputInstanceLifetimeEvent(InstanceId, StateTree, *InstanceName, EStateTreeTraceEventType::Push);
 
 				if (ActiveStates.IsValid())
 				{
-					// Force a world time update since we are tracing an event from the past
-					UE_TRACE_LOG(StateTreeDebugger, WorldTimestampEvent, StateTreeDebugChannel)
-						<< WorldTimestampEvent.WorldTime(ActiveStatesRecordingWorldTime);
-
 					ActiveStates.Output(InstanceId);
 				}
 			}
@@ -253,8 +251,6 @@ struct FInstanceEventBufferedData
 	FString InstanceName;
 	TWeakObjectPtr<const UStateTree> WeakStateTree;
 	FStateTreeInstanceDebugId InstanceId;
-	double LifetimeRecordingWorldTime = 0;
-	double ActiveStatesRecordingWorldTime = 0;
 };
 
 /** Struct to keep track of the buffered event data and flush them. */
@@ -291,29 +287,34 @@ struct FBufferedDataList
 
 		TGuardValue<bool> GuardReentry(bFlushing, true);
 
-		if (FlushedVersion != CurrentVersion)
-		{
-			FlushedVersion = CurrentVersion;
+		const bool bTraceLifetimeEvents = FlushedVersion != CurrentVersion;
+		FlushedVersion = CurrentVersion;
 
-			// Trace asset events first since they are required for instance lifetime event types.
-			// Events are preserved in case the trace session is stopped and then a new one gets started
-			// in the same game session. In which case we need to output the ids to that new trace.
+		// Trace asset events first since they are required for instance lifetime event types.
+		// Events are preserved in case the trace session is stopped and then a new one gets started
+		// in the same game session. In which case we need to output the ids to that new trace.
+		if (bTraceLifetimeEvents)
+		{
 			for (const FAssetDebugIdEventBufferedData& AssetDebugIdEventData : AssetDebugIdEvents)
 			{
 				AssetDebugIdEventData.Trace();
 			}
+		}
 
-			// Then trace instance lifetime events since they are required for other event types.
-			// It is also associated to an older world time.
-			// Events are also preserved for the same reason as AssetDebugIdEvents.
+		TraceWorldTime();
+
+		// Then trace instance lifetime events since they are required for other event types.
+		// They are associated to an older world time but to simplify the logic on the analysis side
+		// we will send them as if the instances were created at the beginning of the recording.
+		// Events are also preserved for the same reason as AssetDebugIdEvents.
+		if (bTraceLifetimeEvents)
+		{
 			for (TPair<FStateTreeInstanceDebugId, FInstanceEventBufferedData>& Pair : InstanceLifetimeEvents)
 			{
 				Pair.Value.Trace();
 			}
 		}
 
-		TraceWorldTime();
-
 		if (InstanceId.IsValid())
 		{
 			TraceStackedPhases(InstanceId);
@@ -367,8 +368,13 @@ struct FBufferedDataList
 		}
 	}
 
-	void BumpTraceVersion()
+	void OnStoppingTraces()
 	{
+		for (TPair<FStateTreeInstanceDebugId, FInstanceEventBufferedData>& Pair : InstanceLifetimeEvents)
+		{
+			Pair.Value.CloseRecording(TracedRecordingWorldTime);
+		}
+
 		// Bump version so shareable data will be flush in the next trace (e.g. Asset ids, instance lifetime events, etc.)
 		CurrentVersion++;
 
@@ -489,12 +495,12 @@ void RegisterGlobalDelegates()
 #endif// OBJECT_TRACE_ENABLED
 		});
 
-	GTracingStateChangedDelegateHandle = UE::StateTree::Delegates::OnTracingStateChanged.AddLambda([](const bool bTracesEnabled)
+	GTracingStateChangedDelegateHandle = UE::StateTree::Delegates::OnTracingStateChanged.AddLambda([](const EStateTreeTraceStatus TraceStatus)
 	{
-		// Bump trace version when disabling traces so next trace will flush buffered events that are still relevant.
-		if (!bTracesEnabled)
+		// Traces are about to be stopped so allow the buffered events to react.
+		if (TraceStatus == EStateTreeTraceStatus::StoppingTrace)
 		{
-			GBufferedEvents.BumpTraceVersion();
+			GBufferedEvents.OnStoppingTraces();
 		}
 	});
 }
@@ -584,7 +590,7 @@ void OutputInstanceLifetimeEvent(
 	{
 		if (EventType == EStateTreeTraceEventType::Push)
 		{
-			GBufferedEvents.InstanceLifetimeEvents.Emplace(InstanceId, FInstanceEventBufferedData(GBufferedEvents.RecordingWorldTime, StateTree, InstanceId, InstanceName));
+			GBufferedEvents.InstanceLifetimeEvents.Emplace(InstanceId, FInstanceEventBufferedData(StateTree, InstanceId, InstanceName));
 		}
 		else if (EventType == EStateTreeTraceEventType::Pop)
 		{
@@ -761,7 +767,6 @@ void OutputActiveStatesEventTrace(
 		if (FInstanceEventBufferedData* ExisingBufferedData = GBufferedEvents.InstanceLifetimeEvents.Find(InstanceId))
 		{
 			ExisingBufferedData->ActiveStates= FInstanceEventBufferedData::FActiveStates(ActiveFrames);
-			ExisingBufferedData->ActiveStatesRecordingWorldTime = GBufferedEvents.RecordingWorldTime;
 		}
 	}
 }
diff --git a/Engine/Plugins/Runtime/StateTree/Source/StateTreeModule/Private/Debugger/StateTreeTraceAnalyzer.cpp b/Engine/Plugins/Runtime/StateTree/Source/StateTreeModule/Private/Debugger/StateTreeTraceAnalyzer.cpp
--- a/Engine/Plugins/Runtime/StateTree/Source/StateTreeModule/Private/Debugger/StateTreeTraceAnalyzer.cpp
+++ b/Engine/Plugins/Runtime/StateTree/Source/StateTreeModule/Private/Debugger/StateTreeTraceAnalyzer.cpp
@@ -222,34 +222,47 @@ bool FStateTreeTraceAnalyzer::OnEvent(const uint16 RouteId, EStyle Style, const
 
 			if (ensureMsgf(ActiveStates.Num() == AssetDebugIds.Num(), TEXT("Each state is expected to have a matching asset id")))
 			{
-				FStateTreeIndex16 LastAssetDebugId;
+				FStateTreeInstanceDebugId InstanceDebugId(EventData.GetValue<uint32>("InstanceId"), EventData.GetValue<uint32>("InstanceSerial"));
+				TWeakObjectPtr<const UStateTree> WeakStateTree;
 
-				FStateTreeTraceActiveStates::FAssetActiveStates* AssetActiveStates = nullptr;
-				for (int Index = 0; Index < ActiveStates.Num(); ++Index)
+				// If empty, we create the event with an empty list of states for the main StateTree.
+				if (ActiveStates.IsEmpty())
 				{
-					FStateTreeIndex16 AssetDebugId(AssetDebugIds[Index]);
-					if (AssetDebugId != LastAssetDebugId)
+					if (Provider.GetAssetFromInstanceId(InstanceDebugId, WeakStateTree))
 					{
-						TWeakObjectPtr<const UStateTree> WeakStateTree;
-						if (Provider.GetAssetFromDebugId(AssetDebugId, WeakStateTree))
-						{
-							FStateTreeTraceActiveStates::FAssetActiveStates& NewPair = Event.ActiveStates.PerAssetStates.Emplace_GetRef();
-							NewPair.WeakStateTree = WeakStateTree;
-							AssetActiveStates = &NewPair;
-							LastAssetDebugId = AssetDebugId;	
-						}
-						else
+						FStateTreeTraceActiveStates::FAssetActiveStates& NewPair = Event.ActiveStates.PerAssetStates.Emplace_GetRef();
+						NewPair.WeakStateTree = WeakStateTree;
+					}
+				}
+				else
+				{
+					FStateTreeIndex16 LastAssetDebugId;
+					FStateTreeTraceActiveStates::FAssetActiveStates* AssetActiveStates = nullptr;
+					for (int Index = 0; Index < ActiveStates.Num(); ++Index)
+					{
+						FStateTreeIndex16 AssetDebugId(AssetDebugIds[Index]);
+						if (AssetDebugId != LastAssetDebugId)
 						{
-							UE_LOG(LogStateTree, Error, TEXT("Instance frame event refers to an asset Id that wasn't added previously."));
-							continue;
+							if (Provider.GetAssetFromDebugId(AssetDebugId, WeakStateTree))
+							{
+								FStateTreeTraceActiveStates::FAssetActiveStates& NewPair = Event.ActiveStates.PerAssetStates.Emplace_GetRef();
+								NewPair.WeakStateTree = WeakStateTree;
+								AssetActiveStates = &NewPair;
+								LastAssetDebugId = AssetDebugId;
+							}
+							else
+							{
+								UE_LOG(LogStateTree, Error, TEXT("Instance frame event refers to an asset Id that wasn't added previously."));
+								continue;
+							}
 						}
-					}
 
-					check(AssetActiveStates != nullptr);
-					AssetActiveStates->ActiveStates.Push(ActiveStates[Index]);
+						check(AssetActiveStates != nullptr);
+						AssetActiveStates->ActiveStates.Push(ActiveStates[Index]);
+					}
 				}
 
-				Provider.AppendEvent(FStateTreeInstanceDebugId(EventData.GetValue<uint32>("InstanceId"), EventData.GetValue<uint32>("InstanceSerial")),
+				Provider.AppendEvent(InstanceDebugId,
 					Context.EventTime.AsSeconds(EventData.GetValue<uint64>("Cycle")),
 					FStateTreeTraceEventVariantType(TInPlaceType<FStateTreeTraceActiveStatesEvent>(), Event));
 			}
diff --git a/Engine/Plugins/Runtime/StateTree/Source/StateTreeModule/Private/Debugger/StateTreeTraceProvider.cpp b/Engine/Plugins/Runtime/StateTree/Source/StateTreeModule/Private/Debugger/StateTreeTraceProvider.cpp
--- a/Engine/Plugins/Runtime/StateTree/Source/StateTreeModule/Private/Debugger/StateTreeTraceProvider.cpp
+++ b/Engine/Plugins/Runtime/StateTree/Source/StateTreeModule/Private/Debugger/StateTreeTraceProvider.cpp
@@ -145,6 +145,18 @@ bool FStateTreeTraceProvider::GetAssetFromDebugId(const FStateTreeIndex16 AssetD
 	return ExistingPair != nullptr;
 }
 
+bool FStateTreeTraceProvider::GetAssetFromInstanceId(const FStateTreeInstanceDebugId InstanceId, TWeakObjectPtr<const UStateTree>& WeakStateTree) const
+{
+	if (const uint32* IndexPtr = InstanceIdToDebuggerEntryTimelines.Find(InstanceId))
+	{
+		check(Descriptors.Num() == EventsTimelines.Num());
+		WeakStateTree = Descriptors[*IndexPtr].StateTree;
+		return true;
+	}
+
+	return false;
+}
+
 void FStateTreeTraceProvider::GetInstances(TArray<UE::StateTreeDebugger::FInstanceDescriptor>& OutInstances) const
 {
 	OutInstances = Descriptors;
diff --git a/Engine/Plugins/Runtime/StateTree/Source/StateTreeModule/Private/StateTreeModule.cpp b/Engine/Plugins/Runtime/StateTree/Source/StateTreeModule/Private/StateTreeModule.cpp
--- a/Engine/Plugins/Runtime/StateTree/Source/StateTreeModule/Private/StateTreeModule.cpp
+++ b/Engine/Plugins/Runtime/StateTree/Source/StateTreeModule/Private/StateTreeModule.cpp
@@ -5,6 +5,7 @@
 #include "StateTreeTypes.h"
 
 #if WITH_STATETREE_DEBUGGER
+#include "Debugger/StateTreeDebuggerTypes.h"
 #include "Debugger/StateTreeTrace.h"
 #include "Debugger/StateTreeTraceModule.h"
 #include "Features/IModularFeatures.h"
@@ -188,7 +189,7 @@ bool FStateTreeModule::StartTraces(int32& OutTraceId)
 	if (UE::StateTree::Delegates::OnTracingStateChanged.IsBound())
 	{
 		UE_LOG(LogStateTree, Log, TEXT("StateTree traces enabled"));
-		UE::StateTree::Delegates::OnTracingStateChanged.Broadcast(bIsTracing);
+		UE::StateTree::Delegates::OnTracingStateChanged.Broadcast(EStateTreeTraceStatus::TracesStarted);
 	}
 
 	return bAreTracesStarted;
@@ -214,6 +215,12 @@ void FStateTreeModule::StopTraces()
 		return;
 	}
 
+	if (UE::StateTree::Delegates::OnTracingStateChanged.IsBound())
+	{
+		UE_LOG(LogStateTree, Log, TEXT("Stopping StateTree traces..."));
+		UE::StateTree::Delegates::OnTracingStateChanged.Broadcast(EStateTreeTraceStatus::StoppingTrace);
+	}
+
 	UE::Trace::ToggleChannel(TEXT("StateTreeDebugChannel"), false);
 	UE::Trace::ToggleChannel(TEXT("FrameChannel"), false);
 
@@ -236,8 +243,8 @@ void FStateTreeModule::StopTraces()
 
 	if (UE::StateTree::Delegates::OnTracingStateChanged.IsBound())
 	{
-		UE_LOG(LogStateTree, Log, TEXT("StateTree traces disabled"));
-		UE::StateTree::Delegates::OnTracingStateChanged.Broadcast(bIsTracing);
+		UE_LOG(LogStateTree, Log, TEXT("StateTree traces stopped"));
+		UE::StateTree::Delegates::OnTracingStateChanged.Broadcast(EStateTreeTraceStatus::TracesStopped);
 	}
 #endif // WITH_STATETREE_DEBUGGER
 }
diff --git a/Engine/Plugins/Runtime/StateTree/Source/StateTreeModule/Public/Debugger/StateTreeDebuggerTypes.h b/Engine/Plugins/Runtime/StateTree/Source/StateTreeModule/Public/Debugger/StateTreeDebuggerTypes.h
--- a/Engine/Plugins/Runtime/StateTree/Source/StateTreeModule/Public/Debugger/StateTreeDebuggerTypes.h
+++ b/Engine/Plugins/Runtime/StateTree/Source/StateTreeModule/Public/Debugger/StateTreeDebuggerTypes.h
@@ -21,6 +21,14 @@ enum class EStateTreeBreakpointType : uint8
 	OnTransition,
 };
 
+UENUM()
+enum class EStateTreeTraceStatus : uint8
+{
+	TracesStarted,
+	StoppingTrace,
+	TracesStopped
+};
+
 #if WITH_STATETREE_DEBUGGER
 
 class UStateTree;
@@ -131,6 +139,12 @@ struct STATETREEMODULE_API FInstanceEventCollection
 		int32 EventIndex = INDEX_NONE;
 	};
 
+	struct FContiguousTraceInfo
+	{
+		/** Indicates the index of the last spans of the trace and from which the frame index will be used to offset new events since their frames will restart at 0. */
+		int32 LastSpanIndex = INDEX_NONE;
+	};
+
 	/** Id of the instance associated to the stored events. */
 	FStateTreeInstanceDebugId InstanceId;
 
@@ -140,6 +154,9 @@ struct STATETREEMODULE_API FInstanceEventCollection
 	/** Spans for frames with events. Each span contains the frame information and the index of the first event for that frame. */
 	TArray<FFrameSpan> FrameSpans;
 
+	/** This list is only used to merge events when dealing with multiple traces related to the same tree instance. */
+	TArray<FContiguousTraceInfo> ContiguousTracesData;
+
 	/** Indices of span and event for frames with a change of activate states. */
 	TArray<FActiveStatesChangePair> ActiveStatesChanges;
 
diff --git a/Engine/Plugins/Runtime/StateTree/Source/StateTreeModule/Public/Debugger/StateTreeTraceProvider.h b/Engine/Plugins/Runtime/StateTree/Source/StateTreeModule/Public/Debugger/StateTreeTraceProvider.h
--- a/Engine/Plugins/Runtime/StateTree/Source/StateTreeModule/Public/Debugger/StateTreeTraceProvider.h
+++ b/Engine/Plugins/Runtime/StateTree/Source/StateTreeModule/Public/Debugger/StateTreeTraceProvider.h
@@ -31,6 +31,7 @@ public:
 
 	void AppendAssetDebugId(const UStateTree* InStateTree, const FStateTreeIndex16 AssetDebugId);
 	bool GetAssetFromDebugId(const FStateTreeIndex16 AssetDebugId, TWeakObjectPtr<const UStateTree>& WeakStateTree) const;
+	bool GetAssetFromInstanceId(const FStateTreeInstanceDebugId InstanceId, TWeakObjectPtr<const UStateTree>& WeakStateTree) const;
 
 protected:
 	/** IStateTreeDebuggerProvider interface */
@@ -47,6 +48,12 @@ private:
 
 	struct FStateTreeDebugIdPair
 	{
+		FStateTreeDebugIdPair(const TWeakObjectPtr<const UStateTree>& WeakStateTree, const FStateTreeIndex16 Id)
+			: WeakStateTree(WeakStateTree)
+			, Id(Id)
+		{
+		}
+
 		TWeakObjectPtr<const UStateTree> WeakStateTree;
 		FStateTreeIndex16 Id;
 	};
diff --git a/Engine/Plugins/Runtime/StateTree/Source/StateTreeModule/Public/StateTreeDelegates.h b/Engine/Plugins/Runtime/StateTree/Source/StateTreeModule/Public/StateTreeDelegates.h
--- a/Engine/Plugins/Runtime/StateTree/Source/StateTreeModule/Public/StateTreeDelegates.h
+++ b/Engine/Plugins/Runtime/StateTree/Source/StateTreeModule/Public/StateTreeDelegates.h
@@ -5,6 +5,7 @@
 #include "Delegates/Delegate.h"
 
 class UStateTree;
+enum class EStateTreeTraceStatus : uint8;
 
 namespace UE::StateTree::Delegates
 {
@@ -69,7 +70,7 @@ extern STATETREEMODULE_API FOnRequestEditorHash OnRequestEditorHash;
 #if WITH_STATETREE_DEBUGGER
 
 /** Called by the StateTree module when StateTree traces are enabled/disabled. */
-DECLARE_MULTICAST_DELEGATE_OneParam(FOnTracingStateChanged, bool /*bTracesEnabled*/);
+DECLARE_MULTICAST_DELEGATE_OneParam(FOnTracingStateChanged, EStateTreeTraceStatus);
 extern STATETREEMODULE_API FOnTracingStateChanged OnTracingStateChanged;
 
 #endif // WITH_STATETREE_DEBUGGER
