StateTree: Fix state parameter binding regression.

diff --git a/Engine/Plugins/Runtime/StateTree/Source/StateTreeEditorModule/Private/Customizations/StateTreeStateDetails.cpp b/Engine/Plugins/Runtime/StateTree/Source/StateTreeEditorModule/Private/Customizations/StateTreeStateDetails.cpp
--- a/Engine/Plugins/Runtime/StateTree/Source/StateTreeEditorModule/Private/Customizations/StateTreeStateDetails.cpp
+++ b/Engine/Plugins/Runtime/StateTree/Source/StateTreeEditorModule/Private/Customizations/StateTreeStateDetails.cpp
@@ -10,7 +10,9 @@
 #include "StateTreeEditorData.h"
 #include "StateTreeEditorNodeUtils.h"
 #include "StateTreeEditorStyle.h"
+#include "StateTreePropertyHelpers.h"
 #include "StateTreeSchema.h"
+#include "StateTreeStateParametersDetails.h"
 #include "Debugger/StateTreeDebuggerUIExtensions.h"
 #include "Widgets/Layout/SBorder.h"
 #include "Widgets/Layout/SBox.h"
@@ -29,13 +31,14 @@ TSharedRef<IDetailCustomization> FStateTreeStateDetails::MakeInstance()
 
 void FStateTreeStateDetails::CustomizeDetails(IDetailLayoutBuilder& DetailBuilder)
 {
-	// Find StateTreeEditorData associated with this panel.
-	UStateTreeEditorData* EditorData = nullptr;
 	TArray<TWeakObjectPtr<UObject>> Objects;
 	DetailBuilder.GetObjectsBeingCustomized(Objects);
+	
+	// Find StateTreeEditorData associated with this panel.
+	UStateTreeEditorData* EditorData = nullptr;
 	for (TWeakObjectPtr<UObject>& WeakObject : Objects)
 	{
-		if (const UObject* Object = WeakObject.Get())
+		if (UObject* Object = WeakObject.Get())
 		{
 			if (UStateTreeEditorData* OuterEditorData = Object->GetTypedOuter<UStateTreeEditorData>())
 			{
@@ -44,9 +47,25 @@ void FStateTreeStateDetails::CustomizeDetails(IDetailLayoutBuilder& DetailBuilde
 			}
 		}
 	}
+
+	// Find StateTreeState associated with this panel.
+	TWeakObjectPtr<UStateTreeState> WeakState;
+	for (TWeakObjectPtr<UObject>& WeakObject : Objects)
+	{
+		if (UObject* Object = WeakObject.Get())
+		{
+			if (UStateTreeState* State = Cast<UStateTreeState>(Object))
+			{
+				WeakState = State;
+				break;
+			}
+		}
+	}
+	
 	const UStateTreeSchema* Schema = EditorData ? EditorData->Schema : nullptr;
 	const FString SchemaPath = Schema ? Schema->GetClass()->GetPathName() : FString();
 	TWeakObjectPtr<UStateTreeEditorData> WeakEditorData = EditorData;
+
 	
 	const TSharedPtr<IPropertyHandle> IDProperty = DetailBuilder.GetProperty(GET_MEMBER_NAME_CHECKED(UStateTreeState, ID));
 	const TSharedPtr<IPropertyHandle> NameProperty = DetailBuilder.GetProperty(GET_MEMBER_NAME_CHECKED(UStateTreeState, Name));
@@ -186,12 +205,21 @@ void FStateTreeStateDetails::CustomizeDetails(IDetailLayoutBuilder& DetailBuilde
 		// Show parameters as a category.
 		ParametersProperty->MarkHiddenByCustomization();
 
-		TSharedPtr<IPropertyHandle> PropertyBagParametersProperty = ParametersProperty->GetChildHandle(TEXT("Parameters")); // FInstancedPropertyBag
-		check(PropertyBagParametersProperty);
+		TSharedPtr<IPropertyHandle> ParametersParametersProperty = ParametersProperty->GetChildHandle(GET_MEMBER_NAME_CHECKED(FStateTreeStateParameters, Parameters)); // FInstancedPropertyBag
+		check(ParametersParametersProperty);
+		TSharedPtr<IPropertyHandle> ParametersFixedLayoutProperty = ParametersProperty->GetChildHandle(GET_MEMBER_NAME_CHECKED(FStateTreeStateParameters, bFixedLayout));
+		check(ParametersFixedLayoutProperty);
+		TSharedPtr<IPropertyHandle> ParametersIDProperty = ParametersProperty->GetChildHandle(GET_MEMBER_NAME_CHECKED(FStateTreeStateParameters, ID));
+		check(ParametersIDProperty);
 
+		bool bFixedLayout = false;
+		ParametersFixedLayoutProperty->GetValue(bFixedLayout);
+
+		
 		const TSharedRef<SHorizontalBox> HeaderContentWidget = SNew(SHorizontalBox)
-			.IsEnabled(PropUtils.ToSharedRef(), &IPropertyUtilities::IsPropertyEditingEnabled)
-			+SHorizontalBox::Slot()
+			.IsEnabled(PropUtils.ToSharedRef(), &IPropertyUtilities::IsPropertyEditingEnabled);
+
+		HeaderContentWidget->AddSlot()
 			.VAlign(VAlign_Center)
 			.Padding(FMargin(4, 0, 0, 0))
 			.AutoWidth()
@@ -199,17 +227,31 @@ void FStateTreeStateDetails::CustomizeDetails(IDetailLayoutBuilder& DetailBuilde
 				SNew(STextBlock)
 				.TextStyle(FStateTreeEditorStyle::Get(), "StateTree.Category")
 				.Text(ParametersDisplayName)
-			]
-			+SHorizontalBox::Slot()
-			.FillWidth(1.f)
-			.HAlign(HAlign_Right)
-			.VAlign(VAlign_Center)
-			[
-				FPropertyBagDetails::MakeAddPropertyWidget(PropertyBagParametersProperty, PropUtils.ToSharedRef(), EPropertyBagPropertyType::Bool, FLinearColor(UE::StateTree::Colors::Grey)).ToSharedRef()
 			];
+
+		if (!bFixedLayout)
+		{
+			HeaderContentWidget->AddSlot()
+				.FillWidth(1.f)
+				.HAlign(HAlign_Right)
+				.VAlign(VAlign_Center)
+				[
+					FPropertyBagDetails::MakeAddPropertyWidget(
+						ParametersParametersProperty,
+						PropUtils.ToSharedRef(),
+						EPropertyBagPropertyType::Bool,
+						FLinearColor(UE::StateTree::Colors::Grey)).ToSharedRef()
+				];
+		}
+		
 		ParametersCategory.HeaderContent(HeaderContentWidget, /*FullRowContent*/true);
 
-		TSharedRef<FPropertyBagInstanceDataDetails> InstanceDetails = MakeShareable(new FPropertyBagInstanceDataDetails(PropertyBagParametersProperty, PropUtils.ToSharedRef(), false));
+		FGuid ID;
+		UE::StateTree::PropertyHelpers::GetStructValue<FGuid>(ParametersIDProperty, ID);
+
+		// Show the Value (FInstancedStruct) as child rows.
+		TSharedRef<FStateTreeStateParametersInstanceDataDetails> InstanceDetails =
+			MakeShareable(new FStateTreeStateParametersInstanceDataDetails(ParametersProperty, ParametersProperty, PropUtils, bFixedLayout, ID, WeakEditorData, WeakState));
 		ParametersCategory.AddCustomBuilder(InstanceDetails);
 	}
 	
diff --git a/Engine/Plugins/Runtime/StateTree/Source/StateTreeEditorModule/Private/Customizations/StateTreeStateParametersDetails.cpp b/Engine/Plugins/Runtime/StateTree/Source/StateTreeEditorModule/Private/Customizations/StateTreeStateParametersDetails.cpp
--- a/Engine/Plugins/Runtime/StateTree/Source/StateTreeEditorModule/Private/Customizations/StateTreeStateParametersDetails.cpp
+++ b/Engine/Plugins/Runtime/StateTree/Source/StateTreeEditorModule/Private/Customizations/StateTreeStateParametersDetails.cpp
@@ -15,140 +15,134 @@
 //----------------------------------------------------------------//
 // FStateTreeStateParametersInstanceDataDetails
 //----------------------------------------------------------------//
-
-class FStateTreeStateParametersInstanceDataDetails : public FPropertyBagInstanceDataDetails
+struct FStateTreeStateOverrideProvider : public FPropertyBagInstanceDataDetails::IPropertyBagOverrideProvider
 {
-public:
-	FStateTreeStateParametersInstanceDataDetails(
-		const TSharedPtr<IPropertyHandle>& InStructProperty,
-		const TSharedPtr<IPropertyHandle>& InParametersStructProperty,
-		const TSharedPtr<IPropertyUtilities>& InPropUtils,
-		const bool bInFixedLayout,
-		FGuid InID,
-		TWeakObjectPtr<UStateTreeEditorData> InEditorData,
-		TWeakObjectPtr<UStateTreeState> InState)
-		: FPropertyBagInstanceDataDetails(InParametersStructProperty, InPropUtils, bInFixedLayout)
-		, StructProperty(InStructProperty)
-		, WeakEditorData(InEditorData)
-		, WeakState(InState)
-		, ID(InID)
+	FStateTreeStateOverrideProvider(UStateTreeState& InState)
+		: State(InState)
 	{
 	}
-	
-	virtual void OnChildRowAdded(IDetailPropertyRow& ChildRow) override
+		
+	virtual bool IsPropertyOverridden(const FGuid PropertyID) const override
 	{
-		FPropertyBagInstanceDataDetails::OnChildRowAdded(ChildRow);
-
-		EStateTreeStateType Type = EStateTreeStateType::State;
-		if (const UStateTreeState* State = WeakState.Get())
-		{
-			Type = State->Type;
-		}
-
-		// Subtree parameters cannot be bound to, they are provided from the linked state.
-		const bool bAllowBinding = Type != EStateTreeStateType::Subtree && ID.IsValid(); 
-
-		if (bAllowBinding)
-		{
-			const TSharedPtr<IPropertyHandle> ChildPropHandle = ChildRow.GetPropertyHandle();
-			const FProperty* Property = ChildPropHandle->GetProperty();
-
-			// Set the category to Parameter so that the binding extension will pick it up.
-			static const FName CategoryName(TEXT("Category"));
-			ChildPropHandle->SetInstanceMetaData(CategoryName, TEXT("Parameter"));
+		return State.IsParametersPropertyOverridden(PropertyID);
+	}
 		
-			// Conditionally control visibility of the value field of bound properties.
+	virtual void SetPropertyOverride(const FGuid PropertyID, const bool bIsOverridden) const override
+	{
+		State.SetParametersPropertyOverridden(PropertyID, bIsOverridden);
+	}
 
-			// Pass the node ID to binding extension. Since the properties are added using AddChildStructure(), we break the hierarchy and cannot access parent.
-			ChildPropHandle->SetInstanceMetaData(UE::StateTree::PropertyBinding::StateTreeNodeIDName, LexToString(ID));
+private:
+	UStateTreeState& State;
+};
 
-			FStateTreePropertyPath Path(ID, *Property->GetFName().ToString());
-			
-			const auto IsValueVisible = TAttribute<EVisibility>::Create([Path, WeakEditorData = WeakEditorData]() -> EVisibility
-				{
-					bool bHasBinding = false;
-					if (UStateTreeEditorData* EditorData = WeakEditorData.Get())
-					{
-						if (const FStateTreeEditorPropertyBindings* EditorPropBindings = EditorData->GetPropertyEditorBindings())
-						{
-							bHasBinding = EditorPropBindings->HasPropertyBinding(Path); 
-						}
-					}
+//----------------------------------------------------------------//
+// FStateTreeStateParametersInstanceDataDetails
+//----------------------------------------------------------------//
 
-					return bHasBinding ? EVisibility::Collapsed : EVisibility::Visible;
-				});
+FStateTreeStateParametersInstanceDataDetails::FStateTreeStateParametersInstanceDataDetails(
+	const TSharedPtr<IPropertyHandle>& InStructProperty,
+	const TSharedPtr<IPropertyHandle>& InParametersStructProperty,
+	const TSharedPtr<IPropertyUtilities>& InPropUtils,
+	const bool bInFixedLayout,
+	FGuid InID,
+	TWeakObjectPtr<UStateTreeEditorData> InEditorData,
+	TWeakObjectPtr<UStateTreeState> InState)
+	: FPropertyBagInstanceDataDetails(InParametersStructProperty, InPropUtils, bInFixedLayout)
+	, StructProperty(InStructProperty)
+	, WeakEditorData(InEditorData)
+	, WeakState(InState)
+	, ID(InID)
+{
+}
+	
+void FStateTreeStateParametersInstanceDataDetails::OnChildRowAdded(IDetailPropertyRow& ChildRow)
+{
+	FPropertyBagInstanceDataDetails::OnChildRowAdded(ChildRow);
 
-			FDetailWidgetDecl* ValueWidgetDecl = ChildRow.CustomValueWidget();
-			const TSharedRef<SBox> WrappedValueWidget = SNew(SBox)
-				.Visibility(IsValueVisible)
-				[
-					ValueWidgetDecl->Widget
-				];
-			ValueWidgetDecl->Widget = WrappedValueWidget;
-		}
+	EStateTreeStateType Type = EStateTreeStateType::State;
+	if (const UStateTreeState* State = WeakState.Get())
+	{
+		Type = State->Type;
 	}
 
-	struct FStateTreeStateOverrideProvider : public IPropertyBagOverrideProvider
+	// Subtree parameters cannot be bound to, they are provided from the linked state.
+	const bool bAllowBinding = Type != EStateTreeStateType::Subtree && ID.IsValid(); 
+
+	if (bAllowBinding)
 	{
-		FStateTreeStateOverrideProvider(UStateTreeState& InState)
-			: State(InState)
-		{
-		}
-		
-		virtual bool IsPropertyOverridden(const FGuid PropertyID) const override
-		{
-			return State.IsParametersPropertyOverridden(PropertyID);
-		}
-		
-		virtual void SetPropertyOverride(const FGuid PropertyID, const bool bIsOverridden) const override
-		{
-			State.SetParametersPropertyOverridden(PropertyID, bIsOverridden);
-		}
+		const TSharedPtr<IPropertyHandle> ChildPropHandle = ChildRow.GetPropertyHandle();
+		const FProperty* Property = ChildPropHandle->GetProperty();
 
-	private:
-		UStateTreeState& State;
-	};
+		// Set the category to Parameter so that the binding extension will pick it up.
+		static const FName CategoryName(TEXT("Category"));
+		ChildPropHandle->SetInstanceMetaData(CategoryName, TEXT("Parameter"));
+	
+		// Conditionally control visibility of the value field of bound properties.
 
-	virtual bool HasPropertyOverrides() const override
-	{
-		if (const UStateTreeState* State = WeakState.Get())
-		{
-			return State->Type == EStateTreeStateType::Linked || State->Type == EStateTreeStateType::LinkedAsset;
-		}
-		return false;
-	}
+		// Pass the node ID to binding extension. Since the properties are added using AddChildStructure(), we break the hierarchy and cannot access parent.
+		ChildPropHandle->SetInstanceMetaData(UE::StateTree::PropertyBinding::StateTreeNodeIDName, LexToString(ID));
 
-	virtual void PreChangeOverrides() override
-	{
-		check(StructProperty);
-		StructProperty->NotifyPreChange();
+		FStateTreePropertyPath Path(ID, *Property->GetFName().ToString());
+		
+		const auto IsValueVisible = TAttribute<EVisibility>::Create([Path, WeakEditorData = WeakEditorData]() -> EVisibility
+			{
+				bool bHasBinding = false;
+				if (UStateTreeEditorData* EditorData = WeakEditorData.Get())
+				{
+					if (const FStateTreeEditorPropertyBindings* EditorPropBindings = EditorData->GetPropertyEditorBindings())
+					{
+						bHasBinding = EditorPropBindings->HasPropertyBinding(Path); 
+					}
+				}
+
+				return bHasBinding ? EVisibility::Collapsed : EVisibility::Visible;
+			});
+
+		FDetailWidgetDecl* ValueWidgetDecl = ChildRow.CustomValueWidget();
+		const TSharedRef<SBox> WrappedValueWidget = SNew(SBox)
+			.Visibility(IsValueVisible)
+			[
+				ValueWidgetDecl->Widget
+			];
+		ValueWidgetDecl->Widget = WrappedValueWidget;
 	}
+}
 
-	virtual void PostChangeOverrides() override
+bool FStateTreeStateParametersInstanceDataDetails::HasPropertyOverrides() const
+{
+	if (const UStateTreeState* State = WeakState.Get())
 	{
-		check(StructProperty);
-		StructProperty->NotifyPostChange(EPropertyChangeType::ValueSet);
-		StructProperty->NotifyFinishedChangingProperties();
+		return State->Type == EStateTreeStateType::Linked || State->Type == EStateTreeStateType::LinkedAsset;
 	}
+	return false;
+}
+
+void FStateTreeStateParametersInstanceDataDetails::PreChangeOverrides()
+{
+	check(StructProperty);
+	StructProperty->NotifyPreChange();
+}
 
-	virtual void EnumeratePropertyBags(TSharedPtr<IPropertyHandle> PropertyBagHandle, const EnumeratePropertyBagFuncRef& Func) const override
+void FStateTreeStateParametersInstanceDataDetails::PostChangeOverrides()
+{
+	check(StructProperty);
+	StructProperty->NotifyPostChange(EPropertyChangeType::ValueSet);
+	StructProperty->NotifyFinishedChangingProperties();
+}
+
+void FStateTreeStateParametersInstanceDataDetails::EnumeratePropertyBags(TSharedPtr<IPropertyHandle> PropertyBagHandle, const EnumeratePropertyBagFuncRef& Func) const
+{
+	if (UStateTreeState* State = WeakState.Get())
 	{
-		if (UStateTreeState* State = WeakState.Get())
+		if (const FInstancedPropertyBag* DefaultParameters = State->GetDefaultParameters())
 		{
-			if (const FInstancedPropertyBag* DefaultParameters = State->GetDefaultParameters())
-			{
-				FInstancedPropertyBag& Parameters = State->Parameters.Parameters;
-				FStateTreeStateOverrideProvider OverrideProvider(*State);
-				Func(*DefaultParameters, Parameters, OverrideProvider);
-			}
+			FInstancedPropertyBag& Parameters = State->Parameters.Parameters;
+			FStateTreeStateOverrideProvider OverrideProvider(*State);
+			Func(*DefaultParameters, Parameters, OverrideProvider);
 		}
 	}
-
-	TSharedPtr<IPropertyHandle> StructProperty;
-	TWeakObjectPtr<UStateTreeEditorData> WeakEditorData;
-	TWeakObjectPtr<UStateTreeState> WeakState;
-	FGuid ID;
-};
+}
 
 
 //----------------------------------------------------------------//
diff --git a/Engine/Plugins/Runtime/StateTree/Source/StateTreeEditorModule/Private/Customizations/StateTreeStateParametersDetails.h b/Engine/Plugins/Runtime/StateTree/Source/StateTreeEditorModule/Private/Customizations/StateTreeStateParametersDetails.h
--- a/Engine/Plugins/Runtime/StateTree/Source/StateTreeEditorModule/Private/Customizations/StateTreeStateParametersDetails.h
+++ b/Engine/Plugins/Runtime/StateTree/Source/StateTreeEditorModule/Private/Customizations/StateTreeStateParametersDetails.h
@@ -3,7 +3,9 @@
 #pragma once
 
 #include "IPropertyTypeCustomization.h"
+#include "PropertyBagDetails.h"
 
+class IDetailPropertyRow;
 class IPropertyHandle;
 class IPropertyUtilities;
 class UStateTree;
@@ -41,3 +43,30 @@ private:
 	TWeakObjectPtr<UStateTree> WeakStateTree = nullptr;
 	TWeakObjectPtr<UStateTreeState> WeakState = nullptr;
 };
+
+
+class FStateTreeStateParametersInstanceDataDetails : public FPropertyBagInstanceDataDetails
+{
+public:
+	FStateTreeStateParametersInstanceDataDetails(
+		const TSharedPtr<IPropertyHandle>& InStructProperty,
+		const TSharedPtr<IPropertyHandle>& InParametersStructProperty,
+		const TSharedPtr<IPropertyUtilities>& InPropUtils,
+		const bool bInFixedLayout,
+		FGuid InID,
+		TWeakObjectPtr<UStateTreeEditorData> InEditorData,
+		TWeakObjectPtr<UStateTreeState> InState);
+	
+	virtual void OnChildRowAdded(IDetailPropertyRow& ChildRow) override;
+
+	virtual bool HasPropertyOverrides() const override;
+	virtual void PreChangeOverrides() override;
+	virtual void PostChangeOverrides() override;
+	virtual void EnumeratePropertyBags(TSharedPtr<IPropertyHandle> PropertyBagHandle, const EnumeratePropertyBagFuncRef& Func) const override;
+
+private:
+	TSharedPtr<IPropertyHandle> StructProperty;
+	TWeakObjectPtr<UStateTreeEditorData> WeakEditorData;
+	TWeakObjectPtr<UStateTreeState> WeakState;
+	FGuid ID;
+};
