StateTree: Add missing check when emitting a warning in ST compiler.

diff --git a/Engine/Plugins/Runtime/StateTree/Source/StateTreeEditorModule/Private/StateTreeCompiler.cpp b/Engine/Plugins/Runtime/StateTree/Source/StateTreeEditorModule/Private/StateTreeCompiler.cpp
--- a/Engine/Plugins/Runtime/StateTree/Source/StateTreeEditorModule/Private/StateTreeCompiler.cpp
+++ b/Engine/Plugins/Runtime/StateTree/Source/StateTreeEditorModule/Private/StateTreeCompiler.cpp
@@ -409,8 +409,9 @@ bool FStateTreeCompiler::CreateStateRecursive(UStateTreeState& State, const FSta
 
 	FStateTreeCompilerLogStateScope LogStateScope(&State, Log);
 
-	if (State.Type == EStateTreeStateType::LinkedAsset
+	if ((State.Type == EStateTreeStateType::LinkedAsset
 		|| State.Type == EStateTreeStateType::Linked)
+		&& State.Children.Num() > 0)
 	{
 		Log.Reportf(EMessageSeverity::Warning,
 			TEXT("Linked State cannot have child states, because the state selection will enter to the linked state on activation."));
