[State Tree] Renaming common elements between Condition and the upcoming Utility Consideration

- Renamed EStateTreeConditionOperand and properties of that type.
- Renamed FStateTreeEditorNode::Indent Property.
- Corrected a few tooltips

diff --git a/Engine/Plugins/Experimental/Avalanche/Source/AvalancheEditor/Private/Transition/AvaTransitionExtension.cpp b/Engine/Plugins/Experimental/Avalanche/Source/AvalancheEditor/Private/Transition/AvaTransitionExtension.cpp
--- a/Engine/Plugins/Experimental/Avalanche/Source/AvalancheEditor/Private/Transition/AvaTransitionExtension.cpp
+++ b/Engine/Plugins/Experimental/Avalanche/Source/AvalancheEditor/Private/Transition/AvaTransitionExtension.cpp
@@ -34,7 +34,7 @@ namespace UE::AvaEditor::Private
 		FAvaTransitionSceneMatchCondition& Condition = ConditionNode.GetNode();
 
 		Condition.SceneComparisonType  = InComparisonType;
-		ConditionNode.ConditionOperand = EStateTreeConditionOperand::Or;
+		ConditionNode.ExpressionOperand = EStateTreeExpressionOperand::Or;
 	}
 
 	void AddTransitionCondition(UStateTreeState& InState, EAvaTransitionType InTransitionType)
diff --git a/Engine/Plugins/Experimental/Avalanche/Source/AvalancheTransitionEditor/Private/ViewModels/Condition/AvaTransitionConditionContainerViewModel.cpp b/Engine/Plugins/Experimental/Avalanche/Source/AvalancheTransitionEditor/Private/ViewModels/Condition/AvaTransitionConditionContainerViewModel.cpp
--- a/Engine/Plugins/Experimental/Avalanche/Source/AvalancheTransitionEditor/Private/ViewModels/Condition/AvaTransitionConditionContainerViewModel.cpp
+++ b/Engine/Plugins/Experimental/Avalanche/Source/AvalancheTransitionEditor/Private/ViewModels/Condition/AvaTransitionConditionContainerViewModel.cpp
@@ -17,9 +17,9 @@ namespace UE::AvaTransitionEditor::Private
 	FText GetOperandText(int32 InConditionIndex, const FStateTreeEditorNode& InEditorNode)
 	{
 		// First Conditions or Copy operands should not have any operand displayed
-		if (InConditionIndex > 0 && InEditorNode.ConditionOperand != EStateTreeConditionOperand::Copy)
+		if (InConditionIndex > 0 && InEditorNode.ExpressionOperand != EStateTreeExpressionOperand::Copy)
 		{
-			return FText::Format(INVTEXT("{0} "), UEnum::GetDisplayValueAsText(InEditorNode.ConditionOperand).ToLower());
+			return FText::Format(INVTEXT("{0} "), UEnum::GetDisplayValueAsText(InEditorNode.ExpressionOperand).ToLower());
 		}
 		return FText::GetEmpty();
 	}
@@ -87,12 +87,12 @@ FText FAvaTransitionConditionContainerViewModel::UpdateStateDescription() const
 			continue;
 		}
 
-		int8 DeltaIndent = -EditorNode->ConditionIndent;
+		int8 DeltaIndent = -EditorNode->ExpressionIndent;
 		if (ConditionViewModels.IsValidIndex(ConditionIndex + 1))
 		{
 			if (const FStateTreeEditorNode* NextEditorNode = ConditionViewModels[ConditionIndex + 1]->GetEditorNode())
 			{
-				DeltaIndent += NextEditorNode->ConditionIndent;
+				DeltaIndent += NextEditorNode->ExpressionIndent;
 			}
 		}
 
diff --git a/Engine/Plugins/Runtime/StateTree/Config/BaseStateTree.ini b/Engine/Plugins/Runtime/StateTree/Config/BaseStateTree.ini
--- a/Engine/Plugins/Runtime/StateTree/Config/BaseStateTree.ini
+++ b/Engine/Plugins/Runtime/StateTree/Config/BaseStateTree.ini
@@ -66,3 +66,7 @@
 
 +EnumRedirects=(OldName="/Script/StateTreeModule.EStateTreeDataSourceType",ValueChanges=(("LinkedStateParameterData", "StateParameterData")))
 +EnumRedirects=(OldName="/Script/StateTreeModule.EStateTreeBindableStructSource",ValueChanges=(("State", "StateParameter")))
+
++EnumRedirects=(OldName="EStateTreeConditionOperand",NewName="/Script/StateTreeModule.EStateTreeExpressionOperand")
++PropertyRedirects=(OldName="StateTreeEditorNode.ConditionIndent",NewName="StateTreeEditorNode.ExpressionIndent")
++PropertyRedirects=(OldName="StateTreeEditorNode.ConditionOperand",NewName="StateTreeEditorNode.ExpressionOperand")
\ No newline at end of file
diff --git a/Engine/Plugins/Runtime/StateTree/Source/StateTreeEditorModule/Private/Customizations/StateTreeEditorNodeDetails.cpp b/Engine/Plugins/Runtime/StateTree/Source/StateTreeEditorModule/Private/Customizations/StateTreeEditorNodeDetails.cpp
--- a/Engine/Plugins/Runtime/StateTree/Source/StateTreeEditorModule/Private/Customizations/StateTreeEditorNodeDetails.cpp
+++ b/Engine/Plugins/Runtime/StateTree/Source/StateTreeEditorModule/Private/Customizations/StateTreeEditorNodeDetails.cpp
@@ -362,8 +362,8 @@ void FStateTreeEditorNodeDetails::CustomizeHeader(TSharedRef<class IPropertyHand
 	InstanceObjectProperty = StructProperty->GetChildHandle(GET_MEMBER_NAME_CHECKED(FStateTreeEditorNode, InstanceObject));
 	IDProperty = StructProperty->GetChildHandle(GET_MEMBER_NAME_CHECKED(FStateTreeEditorNode, ID));
 
-	IndentProperty = StructProperty->GetChildHandle(TEXT("ConditionIndent"));
-	OperandProperty = StructProperty->GetChildHandle(TEXT("ConditionOperand"));
+	IndentProperty = StructProperty->GetChildHandle(GET_MEMBER_NAME_CHECKED(FStateTreeEditorNode, ExpressionIndent));
+	OperandProperty = StructProperty->GetChildHandle(GET_MEMBER_NAME_CHECKED(FStateTreeEditorNode, ExpressionOperand));
 
 	check(NodeProperty.IsValid());
 	check(InstanceProperty.IsValid());
@@ -414,7 +414,7 @@ void FStateTreeEditorNodeDetails::CustomizeHeader(TSharedRef<class IPropertyHand
 						.OnClicked(this, &FStateTreeEditorNodeDetails::HandleIndentPlus)
 						.HAlign(HAlign_Center)
 						.ContentPadding(FMargin(4.f, 4.f))
-						.ToolTipText(LOCTEXT("AddIdentTOoltip", "Browse to the current node blueprint in Content Browser"))
+						.ToolTipText(LOCTEXT("IncreaseIdentTooltip", "Increment the depth of the expression row controlling parentheses and expression order"))
 						[
 							SNew(SImage)
 							.DesiredSizeOverride(FVector2D(8.f, 8.f))
@@ -437,7 +437,7 @@ void FStateTreeEditorNodeDetails::CustomizeHeader(TSharedRef<class IPropertyHand
 						.OnClicked(this, &FStateTreeEditorNodeDetails::HandleIndentMinus)
 						.HAlign(HAlign_Center)
 						.ContentPadding(FMargin(4.f, 4.f))
-						.ToolTipText(LOCTEXT("BrowseToCurrentNodeBP", "Browse to the current node blueprint in Content Browser"))
+						.ToolTipText(LOCTEXT("DecreaseIndentTooltip", "Decrement the depth of the expression row controlling parentheses and expression order"))
 						[
 							SNew(SImage)
 							.DesiredSizeOverride(FVector2D(8.f, 8.f))
@@ -1041,7 +1041,7 @@ void FStateTreeEditorNodeDetails::SetIndent(const int32 Indent) const
 {
 	check(IndentProperty);
 	
-	IndentProperty->SetValue((uint8)FMath::Clamp(Indent, 0, UE::StateTree::MaxConditionIndent - 1));
+	IndentProperty->SetValue((uint8)FMath::Clamp(Indent, 0, UE::StateTree::MaxExpressionIndent - 1));
 }
 
 bool FStateTreeEditorNodeDetails::IsIndent(const int32 Indent) const
@@ -1094,7 +1094,7 @@ int32 FStateTreeEditorNodeDetails::GetNextIndent() const
 		return 0;
 	}
 	
-	TSharedPtr<IPropertyHandle> NextIndentProperty = NextStructProperty->GetChildHandle(TEXT("ConditionIndent"));
+	TSharedPtr<IPropertyHandle> NextIndentProperty = NextStructProperty->GetChildHandle(GET_MEMBER_NAME_CHECKED(FStateTreeEditorNode, ExpressionIndent));
 	if (!NextIndentProperty.IsValid())
 	{
 		return 0;
@@ -1115,7 +1115,7 @@ FText FStateTreeEditorNodeDetails::GetOpenParens() const
 	const int32 DeltaIndent = NextIndent - CurrIndent;
 	const int32 OpenParens = FMath::Max(0, DeltaIndent);
 
-	static_assert(UE::StateTree::MaxConditionIndent == 4);
+	static_assert(UE::StateTree::MaxExpressionIndent == 4);
 	switch (OpenParens)
 	{
 		case 1: return FText::FromString(TEXT("("));
@@ -1135,7 +1135,7 @@ FText FStateTreeEditorNodeDetails::GetCloseParens() const
 	const int32 DeltaIndent = NextIndent - CurrIndent;
 	const int32 CloseParens = FMath::Max(0, -DeltaIndent);
 
-	static_assert(UE::StateTree::MaxConditionIndent == 4);
+	static_assert(UE::StateTree::MaxExpressionIndent == 4);
 	switch (CloseParens)
 	{
 	case 1: return FText::FromString(TEXT(")"));
@@ -1158,13 +1158,13 @@ FText FStateTreeEditorNodeDetails::GetOperandText() const
 
 	uint8 Value = 0;
 	OperandProperty->GetValue(Value);
-	const EStateTreeConditionOperand Operand = (EStateTreeConditionOperand)Value;
+	const EStateTreeExpressionOperand Operand = (EStateTreeExpressionOperand)Value;
 
-	if (Operand == EStateTreeConditionOperand::And)
+	if (Operand == EStateTreeExpressionOperand::And)
 	{
 		return LOCTEXT("AndOperand", "AND");
 	}
-	else if (Operand == EStateTreeConditionOperand::Or)
+	else if (Operand == EStateTreeExpressionOperand::Or)
 	{
 		return LOCTEXT("OrOperand", "OR");
 	}
@@ -1187,13 +1187,13 @@ FSlateColor FStateTreeEditorNodeDetails::GetOperandColor() const
 
 	uint8 Value = 0; 
 	OperandProperty->GetValue(Value);
-	const EStateTreeConditionOperand Operand = (EStateTreeConditionOperand)Value;
+	const EStateTreeExpressionOperand Operand = (EStateTreeExpressionOperand)Value;
 
-	if (Operand == EStateTreeConditionOperand::And)
+	if (Operand == EStateTreeExpressionOperand::And)
 	{
 		return FStyleColors::AccentPink;
 	}
-	else if (Operand == EStateTreeConditionOperand::Or)
+	else if (Operand == EStateTreeExpressionOperand::Or)
 	{
 		return FStyleColors::AccentBlue;
 	}
@@ -1210,14 +1210,14 @@ TSharedRef<SWidget> FStateTreeEditorNodeDetails::OnGetOperandContent() const
 	FMenuBuilder MenuBuilder(true, NULL);
 
 	FUIAction AndAction(
-		FExecuteAction::CreateSP(this, &FStateTreeEditorNodeDetails::SetOperand, EStateTreeConditionOperand::And),
+		FExecuteAction::CreateSP(this, &FStateTreeEditorNodeDetails::SetOperand, EStateTreeExpressionOperand::And),
 		FCanExecuteAction(),
-		FIsActionChecked::CreateSP(this, &FStateTreeEditorNodeDetails::IsOperand, EStateTreeConditionOperand::And));
+		FIsActionChecked::CreateSP(this, &FStateTreeEditorNodeDetails::IsOperand, EStateTreeExpressionOperand::And));
 	MenuBuilder.AddMenuEntry(LOCTEXT("AndOperand", "AND"), TAttribute<FText>(), FSlateIcon(), AndAction, FName(), EUserInterfaceActionType::Check);
 
-	FUIAction OrAction(FExecuteAction::CreateSP(this, &FStateTreeEditorNodeDetails::SetOperand, EStateTreeConditionOperand::Or),
+	FUIAction OrAction(FExecuteAction::CreateSP(this, &FStateTreeEditorNodeDetails::SetOperand, EStateTreeExpressionOperand::Or),
 		FCanExecuteAction(),
-		FIsActionChecked::CreateSP(this, &FStateTreeEditorNodeDetails::IsOperand, EStateTreeConditionOperand::Or));
+		FIsActionChecked::CreateSP(this, &FStateTreeEditorNodeDetails::IsOperand, EStateTreeExpressionOperand::Or));
 	MenuBuilder.AddMenuEntry(LOCTEXT("OrOperand", "OR"), TAttribute<FText>(), FSlateIcon(), OrAction, FName(), EUserInterfaceActionType::Check);
 
 	return MenuBuilder.MakeWidget();
@@ -1228,18 +1228,18 @@ bool FStateTreeEditorNodeDetails::IsOperandEnabled() const
 	return !IsFirstItem();
 }
 
-bool FStateTreeEditorNodeDetails::IsOperand(const EStateTreeConditionOperand Operand) const
+bool FStateTreeEditorNodeDetails::IsOperand(const EStateTreeExpressionOperand Operand) const
 {
 	check(OperandProperty);
 
 	uint8 Value = 0; 
 	OperandProperty->GetValue(Value);
-	const EStateTreeConditionOperand CurrOperand = (EStateTreeConditionOperand)Value;
+	const EStateTreeExpressionOperand CurrOperand = (EStateTreeExpressionOperand)Value;
 
 	return CurrOperand == Operand;
 }
 
-void FStateTreeEditorNodeDetails::SetOperand(const EStateTreeConditionOperand Operand) const
+void FStateTreeEditorNodeDetails::SetOperand(const EStateTreeExpressionOperand Operand) const
 {
 	check(OperandProperty);
 
diff --git a/Engine/Plugins/Runtime/StateTree/Source/StateTreeEditorModule/Private/Customizations/StateTreeEditorNodeDetails.h b/Engine/Plugins/Runtime/StateTree/Source/StateTreeEditorModule/Private/Customizations/StateTreeEditorNodeDetails.h
--- a/Engine/Plugins/Runtime/StateTree/Source/StateTreeEditorModule/Private/Customizations/StateTreeEditorNodeDetails.h
+++ b/Engine/Plugins/Runtime/StateTree/Source/StateTreeEditorModule/Private/Customizations/StateTreeEditorNodeDetails.h
@@ -26,7 +26,7 @@ class UStateTreeEditorData;
 struct EVisibility;
 struct FStateTreeEditorNode;
 struct FStateTreePropertyPath;
-enum class EStateTreeConditionOperand : uint8;
+enum class EStateTreeExpressionOperand : uint8;
 
 /**
  * Type customization for nodes (Conditions, Evaluators and Tasks) in StateTreeState.
@@ -64,8 +64,8 @@ private:
 	TSharedRef<SWidget> OnGetOperandContent() const;
 	bool IsOperandEnabled() const;
 
-	void SetOperand(const EStateTreeConditionOperand Operand) const;
-	bool IsOperand(const EStateTreeConditionOperand Operand) const;
+	void SetOperand(const EStateTreeExpressionOperand Operand) const;
+	bool IsOperand(const EStateTreeExpressionOperand Operand) const;
 
 	bool IsFirstItem() const;
 	int32 GetCurrIndent() const;
diff --git a/Engine/Plugins/Runtime/StateTree/Source/StateTreeEditorModule/Private/StateTreeCompiler.cpp b/Engine/Plugins/Runtime/StateTree/Source/StateTreeEditorModule/Private/StateTreeCompiler.cpp
--- a/Engine/Plugins/Runtime/StateTree/Source/StateTreeEditorModule/Private/StateTreeCompiler.cpp
+++ b/Engine/Plugins/Runtime/StateTree/Source/StateTreeEditorModule/Private/StateTreeCompiler.cpp
@@ -469,11 +469,11 @@ bool FStateTreeCompiler::CreateConditions(UStateTreeState& State, TConstArrayVie
 		const bool bIsFirst = Index == 0;
 		const FStateTreeEditorNode& CondNode = Conditions[Index];
 		// First operand should be copy as we dont have a previous item to operate on.
-		const EStateTreeConditionOperand Operand = bIsFirst ? EStateTreeConditionOperand::Copy : CondNode.ConditionOperand;
+		const EStateTreeExpressionOperand Operand = bIsFirst ? EStateTreeExpressionOperand::Copy : CondNode.ExpressionOperand;
 		// First indent must be 0 to make the parentheses calculation match.
-		const int32 CurrIndent = bIsFirst ? 0 : FMath::Clamp((int32)CondNode.ConditionIndent, 0, UE::StateTree::MaxConditionIndent);
+		const int32 CurrIndent = bIsFirst ? 0 : FMath::Clamp((int32)CondNode.ExpressionIndent, 0, UE::StateTree::MaxExpressionIndent);
 		// Next indent, or terminate at zero.
-		const int32 NextIndent = Conditions.IsValidIndex(Index + 1) ? FMath::Clamp((int32)Conditions[Index + 1].ConditionIndent, 0, UE::StateTree::MaxConditionIndent) : 0;
+		const int32 NextIndent = Conditions.IsValidIndex(Index + 1) ? FMath::Clamp((int32)Conditions[Index + 1].ExpressionIndent, 0, UE::StateTree::MaxExpressionIndent) : 0;
 		
 		const int32 DeltaIndent = NextIndent - CurrIndent;
 
@@ -1106,7 +1106,7 @@ bool FStateTreeCompiler::ResolveTransitionState(const UStateTreeState* SourceSta
 	return true;
 }
 
-bool FStateTreeCompiler::CreateCondition(UStateTreeState& State, const FStateTreeEditorNode& CondNode, const EStateTreeConditionOperand Operand, const int8 DeltaIndent)
+bool FStateTreeCompiler::CreateCondition(UStateTreeState& State, const FStateTreeEditorNode& CondNode, const EStateTreeExpressionOperand Operand, const int8 DeltaIndent)
 {
 	if (!CondNode.Node.IsValid())
 	{
diff --git a/Engine/Plugins/Runtime/StateTree/Source/StateTreeEditorModule/Private/StateTreeEditorStyle.cpp b/Engine/Plugins/Runtime/StateTree/Source/StateTreeEditorModule/Private/StateTreeEditorStyle.cpp
--- a/Engine/Plugins/Runtime/StateTree/Source/StateTreeEditorModule/Private/StateTreeEditorStyle.cpp
+++ b/Engine/Plugins/Runtime/StateTree/Source/StateTreeEditorModule/Private/StateTreeEditorStyle.cpp
@@ -210,7 +210,7 @@ FStateTreeEditorStyle::FStateTreeEditorStyle()
 
 	Set("StateTree.Param.Background", new FSlateRoundedBoxBrush(FStyleColors::Hover, 6.f));
 	
-	// Condition Indent combo button
+	// Expression Indent combo button
 	const FButtonStyle IndentButton = FButtonStyle()
 		.SetNormal(FSlateRoundedBoxBrush(FLinearColor::Transparent, 2.0f))
 		.SetHovered(FSlateRoundedBoxBrush(FStyleColors::Background, 2.0f, FStyleColors::InputOutline, 1.0f))
diff --git a/Engine/Plugins/Runtime/StateTree/Source/StateTreeEditorModule/Public/StateTreeCompiler.h b/Engine/Plugins/Runtime/StateTree/Source/StateTreeEditorModule/Public/StateTreeCompiler.h
--- a/Engine/Plugins/Runtime/StateTree/Source/StateTreeEditorModule/Public/StateTreeCompiler.h
+++ b/Engine/Plugins/Runtime/StateTree/Source/StateTreeEditorModule/Public/StateTreeCompiler.h
@@ -6,7 +6,7 @@
 
 struct FStructView;
 
-enum class EStateTreeConditionOperand : uint8;
+enum class EStateTreeExpressionOperand : uint8;
 enum class EStateTreePropertyUsage : uint8;
 struct FStateTreeDataView;
 struct FStateTreeStateHandle;
@@ -50,7 +50,7 @@ private:
 	bool CreateStateTransitions();
 	
 	bool CreateConditions(UStateTreeState& State, TConstArrayView<FStateTreeEditorNode> Conditions);
-	bool CreateCondition(UStateTreeState& State, const FStateTreeEditorNode& CondNode, const EStateTreeConditionOperand Operand, const int8 DeltaIndent);
+	bool CreateCondition(UStateTreeState& State, const FStateTreeEditorNode& CondNode, const EStateTreeExpressionOperand Operand, const int8 DeltaIndent);
 	bool CreateTask(UStateTreeState* State, const FStateTreeEditorNode& TaskNode, const FStateTreeDataHandle TaskDataHandle);
 	bool CreateEvaluator(const FStateTreeEditorNode& EvalNode, const FStateTreeDataHandle EvalDataHandle);
 	bool GetAndValidateBindings(const FStateTreeBindableStructDesc& TargetStruct, FStateTreeDataView TargetValue, TArray<FStateTreePropertyPathBinding>& OutCopyBindings, TArray<FStateTreePropertyPathBinding>& OutReferenceBindings) const;
diff --git a/Engine/Plugins/Runtime/StateTree/Source/StateTreeEditorModule/Public/StateTreeEditorNode.h b/Engine/Plugins/Runtime/StateTree/Source/StateTreeEditorModule/Public/StateTreeEditorNode.h
--- a/Engine/Plugins/Runtime/StateTree/Source/StateTreeEditorModule/Public/StateTreeEditorNode.h
+++ b/Engine/Plugins/Runtime/StateTree/Source/StateTreeEditorModule/Public/StateTreeEditorNode.h
@@ -72,10 +72,20 @@ struct STATETREEEDITORMODULE_API FStateTreeEditorNode
 	FGuid ID;
 
 	UPROPERTY(EditDefaultsOnly, Category = Node)
-	uint8 ConditionIndent = 0;
+	uint8 ExpressionIndent = 0;
 
 	UPROPERTY(EditDefaultsOnly, Category = Node)
-	EStateTreeConditionOperand ConditionOperand = EStateTreeConditionOperand::And; 
+	EStateTreeExpressionOperand ExpressionOperand = EStateTreeExpressionOperand::And;
+
+#if WITH_EDITOR
+	UE_DEPRECATED(5.5, "Use ExpressionIndent instead.")
+	uint8 ConditionIndent = 0;
+
+PRAGMA_DISABLE_DEPRECATION_WARNINGS
+	UE_DEPRECATED(5.5, "Use ExpressionOperand instead.")
+	EStateTreeConditionOperand ConditionOperand = EStateTreeConditionOperand::And;
+PRAGMA_ENABLE_DEPRECATION_WARNINGS
+#endif //WITH_EDITOR
 };
 
 template <typename T>
diff --git a/Engine/Plugins/Runtime/StateTree/Source/StateTreeModule/Private/StateTreeExecutionContext.cpp b/Engine/Plugins/Runtime/StateTree/Source/StateTreeModule/Private/StateTreeExecutionContext.cpp
--- a/Engine/Plugins/Runtime/StateTree/Source/StateTreeModule/Private/StateTreeExecutionContext.cpp
+++ b/Engine/Plugins/Runtime/StateTree/Source/StateTreeModule/Private/StateTreeExecutionContext.cpp
@@ -2518,8 +2518,8 @@ bool FStateTreeExecutionContext::TestAllConditions(const FStateTreeExecutionFram
 		return true;
 	}
 
-	TStaticArray<EStateTreeConditionOperand, UE::StateTree::MaxConditionIndent + 1> Operands(InPlace, EStateTreeConditionOperand::Copy);
-	TStaticArray<bool, UE::StateTree::MaxConditionIndent + 1> Values(InPlace, false);
+	TStaticArray<EStateTreeExpressionOperand, UE::StateTree::MaxExpressionIndent + 1> Operands(InPlace, EStateTreeExpressionOperand::Copy);
+	TStaticArray<bool, UE::StateTree::MaxExpressionIndent + 1> Values(InPlace, false);
 
 	int32 Level = 0;
 	
@@ -2568,7 +2568,7 @@ bool FStateTreeExecutionContext::TestAllConditions(const FStateTreeExecutionFram
 
 		// Store the operand to apply when merging higher level down when returning to this level.
 		// @todo: remove this conditions in 5.1, needs resaving existing StateTrees.
-		const EStateTreeConditionOperand Operand = Index == 0 ? EStateTreeConditionOperand::Copy : Cond.Operand;
+		const EStateTreeExpressionOperand Operand = Index == 0 ? EStateTreeExpressionOperand::Copy : Cond.Operand;
 		Operands[Level] = Operand;
 
 		// Store current value at the top of the stack.
@@ -2584,17 +2584,17 @@ bool FStateTreeExecutionContext::TestAllConditions(const FStateTreeExecutionFram
 			Level--;
 			switch (Operands[Level])
 			{
-			case EStateTreeConditionOperand::Copy:
+			case EStateTreeExpressionOperand::Copy:
 				Values[Level] = Values[Level + 1];
 				break;
-			case EStateTreeConditionOperand::And:
+			case EStateTreeExpressionOperand::And:
 				Values[Level] &= Values[Level + 1];
 				break;
-			case EStateTreeConditionOperand::Or:
+			case EStateTreeExpressionOperand::Or:
 				Values[Level] |= Values[Level + 1];
 				break;
 			}
-			Operands[Level] = EStateTreeConditionOperand::Copy;
+			Operands[Level] = EStateTreeExpressionOperand::Copy;
 		}
 	}
 	
diff --git a/Engine/Plugins/Runtime/StateTree/Source/StateTreeModule/Public/StateTreeConditionBase.h b/Engine/Plugins/Runtime/StateTree/Source/StateTreeModule/Public/StateTreeConditionBase.h
--- a/Engine/Plugins/Runtime/StateTree/Source/StateTreeModule/Public/StateTreeConditionBase.h
+++ b/Engine/Plugins/Runtime/StateTree/Source/StateTreeModule/Public/StateTreeConditionBase.h
@@ -26,7 +26,7 @@ struct STATETREEMODULE_API FStateTreeConditionBase : public FStateTreeNodeBase
 	virtual bool TestCondition(FStateTreeExecutionContext& Context) const { return false; }
 
 	UPROPERTY()
-	EStateTreeConditionOperand Operand = EStateTreeConditionOperand::And;
+	EStateTreeExpressionOperand Operand = EStateTreeExpressionOperand::And;
 
 	UPROPERTY()
 	int8 DeltaIndent = 0;
diff --git a/Engine/Plugins/Runtime/StateTree/Source/StateTreeModule/Public/StateTreeTypes.h b/Engine/Plugins/Runtime/StateTree/Source/StateTreeModule/Public/StateTreeTypes.h
--- a/Engine/Plugins/Runtime/StateTree/Source/StateTreeModule/Public/StateTreeTypes.h
+++ b/Engine/Plugins/Runtime/StateTree/Source/StateTreeModule/Public/StateTreeTypes.h
@@ -18,7 +18,12 @@ struct FStateTreeEvent;
 
 namespace UE::StateTree
 {
+#if WITH_EDITOR
+	UE_DEPRECATED(5.5, "Use MaxExpressionIndent instead.")
 	inline constexpr int32 MaxConditionIndent = 4;
+#endif //WITH_EDITOR
+
+	inline constexpr int32 MaxExpressionIndent = 4;
 
 	inline const FName SchemaTag(TEXT("Schema"));
 
@@ -76,9 +81,9 @@ enum class EStateTreeTransitionType : uint8
 	NotSet UE_DEPRECATED(5.0, "Use None instead."),
 };
 
-/** Operand between conditions */
+/** Operand in an expression */
 UENUM()
-enum class EStateTreeConditionOperand : uint8
+enum class EStateTreeExpressionOperand : uint8
 {
 	/** Copy result */
 	Copy UMETA(Hidden),
@@ -90,6 +95,10 @@ enum class EStateTreeConditionOperand : uint8
 	Or,
 };
 
+#if WITH_EDITOR
+	using EStateTreeConditionOperand UE_DEPRECATED(5.5, "Use EStateTreeExpressionOperand instead.") = EStateTreeExpressionOperand;
+#endif //WITH_EDITOR
+
 UENUM()
 enum class EStateTreeStateType : uint8
 {
diff --git a/Engine/Plugins/Runtime/StateTree/Source/StateTreeTestSuite/Private/StateTreeTest.cpp b/Engine/Plugins/Runtime/StateTree/Source/StateTreeTestSuite/Private/StateTreeTest.cpp
--- a/Engine/Plugins/Runtime/StateTree/Source/StateTreeTestSuite/Private/StateTreeTest.cpp
+++ b/Engine/Plugins/Runtime/StateTree/Source/StateTreeTestSuite/Private/StateTreeTest.cpp
@@ -1014,7 +1014,7 @@ struct FStateTreeTest_LastConditionWithIndent : FAITestBase
 		auto& LastCondition = State1.AddEnterCondition<FStateTreeTestCondition>();
 
 		// Last condition has Indent
-		LastCondition.ConditionIndent = 1;
+		LastCondition.ExpressionIndent = 1;
 		
 		State1.AddTransition(EStateTreeTransitionTrigger::OnStateCompleted, EStateTreeTransitionType::Succeeded);
 
