Fix deprecation errors in StateTreeEditorNode.h

diff --git a/Engine/Plugins/Runtime/StateTree/Source/StateTreeEditorModule/Public/StateTreeEditorNode.h b/Engine/Plugins/Runtime/StateTree/Source/StateTreeEditorModule/Public/StateTreeEditorNode.h
--- a/Engine/Plugins/Runtime/StateTree/Source/StateTreeEditorModule/Public/StateTreeEditorNode.h
+++ b/Engine/Plugins/Runtime/StateTree/Source/StateTreeEditorModule/Public/StateTreeEditorNode.h
@@ -84,6 +84,13 @@ struct STATETREEEDITORMODULE_API FStateTreeEditorNode
 PRAGMA_DISABLE_DEPRECATION_WARNINGS
 	UE_DEPRECATED(5.5, "Use ExpressionOperand instead.")
 	EStateTreeConditionOperand ConditionOperand = EStateTreeConditionOperand::And;
+
+	FStateTreeEditorNode() = default;
+	FStateTreeEditorNode(const FStateTreeEditorNode&) = default;
+	FStateTreeEditorNode(FStateTreeEditorNode&&) = default;
+	FStateTreeEditorNode& operator=(const FStateTreeEditorNode&) = default;
+	FStateTreeEditorNode& operator=(FStateTreeEditorNode&&) = default;
+	~FStateTreeEditorNode() = default;
 PRAGMA_ENABLE_DEPRECATION_WARNINGS
 #endif //WITH_EDITOR
 };
