StateTree: Various fixes how the binding sources are displayed.
- Added editor only StatePath to FStateTreeBindableStructDesc to help error reporting
- Friendlify user defined struct property names in FStateTreePropertyPath::ToString()
- Fixed FStateTreeEditorNode::GetName() to return the BP node name instead of the wrapper name
- Categorize binding sources based on state (including one section for global) in the binding popup

diff --git a/Engine/Plugins/Runtime/StateTree/Source/StateTreeEditorModule/Private/Customizations/StateTreeBindingExtension.cpp b/Engine/Plugins/Runtime/StateTree/Source/StateTreeEditorModule/Private/Customizations/StateTreeBindingExtension.cpp
--- a/Engine/Plugins/Runtime/StateTree/Source/StateTreeEditorModule/Private/Customizations/StateTreeBindingExtension.cpp
+++ b/Engine/Plugins/Runtime/StateTree/Source/StateTreeEditorModule/Private/Customizations/StateTreeBindingExtension.cpp
@@ -213,32 +213,6 @@ const FStateTreeBindableStructDesc* FindStruct(TConstArrayView<FStateTreeBindabl
 	return AccessibleStructs.FindByPredicate([StructID](const FStateTreeBindableStructDesc& Desc) { return Desc.ID == StructID; });
 }
 
-
-FText GetSectionNameFromDataSource(const EStateTreeBindableStructSource Source)
-{
-	switch (Source)
-	{
-	case EStateTreeBindableStructSource::Context:
-		return LOCTEXT("Context", "Context");
-	case EStateTreeBindableStructSource::Parameter:
-		return LOCTEXT("Parameters", "Parameters");
-	case EStateTreeBindableStructSource::Evaluator:
-		return LOCTEXT("Evaluators", "Evaluators");
-	case EStateTreeBindableStructSource::GlobalTask:
-		return LOCTEXT("StateGlobalTasks", "Global Tasks");
-	case EStateTreeBindableStructSource::StateParameter:
-		return LOCTEXT("StateParameters", "State Parameters");
-	case EStateTreeBindableStructSource::Task:
-		return LOCTEXT("Tasks", "Tasks");
-	case EStateTreeBindableStructSource::TransitionEvent:
-		return LOCTEXT("TransitionEvents", "Transition Events");
-	case EStateTreeBindableStructSource::StateEvent:
-		return LOCTEXT("StateSelectionEvents", "State Selection Events");
-	default:
-		return FText::GetEmpty();
-	}
-}
-
 // @todo: there's a similar function in StateTreeNodeDetails.cpp, merge.
 FText GetPropertyTypeText(const FProperty* Property)
 {
@@ -461,7 +435,7 @@ struct FCachedBindingData : public TSharedFromThis<FCachedBindingData>
 				SourcePropertyName += SourceDesc->Name.ToString();
 				if (!SourcePath->IsPathEmpty())
 				{
-					SourcePropertyName += TEXT(" ") + SourcePath->ToString();
+					SourcePropertyName += TEXT(".") + SourcePath->ToString();
 				}
 
 				if (bIsValidBinding)
@@ -903,6 +877,8 @@ void FStateTreeBindingExtension::ExtendWidgetRow(FDetailWidgetRow& InWidgetRow,
 			EditorBindings = BindingOwner->GetPropertyEditorBindings();
 			BindingOwner->GetAccessibleStructs(TargetPath.GetStructID(), AccessibleStructs);
 
+			TMap<FString, FText> SectionNames;
+			
 			for (FStateTreeBindableStructDesc& StructDesc : AccessibleStructs)
 			{
 				const UStruct* Struct = StructDesc.Struct;
@@ -910,7 +886,16 @@ void FStateTreeBindingExtension::ExtendWidgetRow(FDetailWidgetRow& InWidgetRow,
 				FBindingContextStruct& ContextStruct = BindingContextStructs.AddDefaulted_GetRef();
 				ContextStruct.DisplayText = FText::FromString(StructDesc.Name.ToString());
 				ContextStruct.Struct = const_cast<UStruct*>(Struct);
-				ContextStruct.Section = UE::StateTree::PropertyBinding::GetSectionNameFromDataSource(StructDesc.DataSource);
+
+				// Mare sure same section names get exact same FText representation (binding widget uses IsIdentical() to compare the section names).
+				if (const FText* SectionText = SectionNames.Find(StructDesc.StatePath))
+				{
+					ContextStruct.Section = *SectionText;
+				}
+				else
+				{
+					ContextStruct.Section = SectionNames.Add(StructDesc.StatePath, FText::FromString(StructDesc.StatePath));
+				}
 			}
 		}
 
diff --git a/Engine/Plugins/Runtime/StateTree/Source/StateTreeEditorModule/Private/StateTreeCompiler.cpp b/Engine/Plugins/Runtime/StateTree/Source/StateTreeEditorModule/Private/StateTreeCompiler.cpp
--- a/Engine/Plugins/Runtime/StateTree/Source/StateTreeEditorModule/Private/StateTreeCompiler.cpp
+++ b/Engine/Plugins/Runtime/StateTree/Source/StateTreeEditorModule/Private/StateTreeCompiler.cpp
@@ -216,6 +216,7 @@ bool FStateTreeCompiler::Compile(UStateTree& InStateTree)
 	
 	// Mark parameters as binding source
 	const FStateTreeBindableStructDesc ParametersDesc = {
+			UE::StateTree::Editor::GlobalStateName,
 			TEXT("Parameters"),
 			StateTree->Parameters.GetPropertyBagStruct(),
 			FStateTreeDataHandle(EStateTreeDataSourceType::GlobalParameterData),
@@ -239,6 +240,7 @@ bool FStateTreeCompiler::Compile(UStateTree& InStateTree)
 		for (FStateTreeExternalDataDesc& Desc : StateTree->ContextDataDescs)
 		{
 			const FStateTreeBindableStructDesc ExtDataDesc = {
+					UE::StateTree::Editor::GlobalStateName,
 					Desc.Name,
 					Desc.Struct,
 					FStateTreeDataHandle(EStateTreeDataSourceType::ContextData, ContextDataIndex++),
@@ -462,7 +464,7 @@ bool FStateTreeCompiler::CreateStateRecursive(UStateTreeState& State, const FSta
 	return true;
 }
 
-bool FStateTreeCompiler::CreateConditions(UStateTreeState& State, TConstArrayView<FStateTreeEditorNode> Conditions)
+bool FStateTreeCompiler::CreateConditions(UStateTreeState& State, const FString& StatePath, TConstArrayView<FStateTreeEditorNode> Conditions)
 {
 	for (int32 Index = 0; Index < Conditions.Num(); Index++)
 	{
@@ -477,7 +479,7 @@ bool FStateTreeCompiler::CreateConditions(UStateTreeState& State, TConstArrayVie
 		
 		const int32 DeltaIndent = NextIndent - CurrIndent;
 
-		if (!CreateCondition(State, CondNode, Operand, (int8)DeltaIndent))
+		if (!CreateCondition(State, StatePath, CondNode, Operand, (int8)DeltaIndent))
 		{
 			return false;
 		}
@@ -620,9 +622,12 @@ bool FStateTreeCompiler::CreateStateTasksAndParameters()
 
 		// @todo: We should be able to skip empty parameter data.
 
+		const FString StatePath = State->GetPath(); 
+		
 		// Binding target
 		FStateTreeBindableStructDesc LinkedParamsDesc = {
-			State->Name,
+			StatePath,
+			FName("Parameters"),
 			State->Parameters.Parameters.GetPropertyBagStruct(),
 			CompactState.ParameterDataHandle,
 			EStateTreeBindableStructSource::StateParameter,
@@ -644,9 +649,12 @@ bool FStateTreeCompiler::CreateStateTasksAndParameters()
 			CompactState.RequiredEventToEnter.PayloadStruct = State->RequiredEventToEnter.PayloadStruct;
 			CompactState.bConsumeEventOnSelect = State->RequiredEventToEnter.bConsumeEventOnSelect;
 
+			const FString StatePathWithConditions = StatePath + TEXT("/EnterConditions");
+
 			FStateTreeBindableStructDesc Desc;
+			Desc.StatePath = StatePathWithConditions,
 			Desc.Struct = FStateTreeEvent::StaticStruct();
-			Desc.Name = CompactState.Name;
+			Desc.Name = FName("Enter Event");
 			Desc.ID = State->GetEventID();
 			Desc.DataSource = EStateTreeBindableStructSource::StateEvent;
 			Desc.DataHandle = FStateTreeDataHandle(EStateTreeDataSourceType::StateEvent, CompactState.EventDataIndex.Get(), CompactStateHandle);
@@ -773,7 +781,9 @@ bool FStateTreeCompiler::CreateStateTransitions()
 		check(SourceState != nullptr);
 
 		FStateTreeCompilerLogStateScope LogStateScope(SourceState, Log);
-		
+
+		const FString StatePath = SourceState->GetPath();
+
 		// Enter conditions.
 		const int32 EnterConditionsBegin = Nodes.Num();
 		if (const auto Validation = UE::StateTree::Compiler::IsValidCount16(EnterConditionsBegin); Validation.DidFail())
@@ -782,8 +792,9 @@ bool FStateTreeCompiler::CreateStateTransitions()
 			return false;
 		}
 		CompactState.EnterConditionsBegin = uint16(EnterConditionsBegin);
-		
-		if (!CreateConditions(*SourceState, SourceState->EnterConditions))
+
+		const FString StatePathWithConditions = StatePath + TEXT("/EnterConditions");
+		if (!CreateConditions(*SourceState, StatePathWithConditions, SourceState->EnterConditions))
 		{
 			Log.Reportf(EMessageSeverity::Error,
 				TEXT("Failed to create state enter condition."));
@@ -959,6 +970,8 @@ bool FStateTreeCompiler::CreateStateTransitions()
 						*TransitionTargetState.Name.ToString());
 				}
 			}
+
+			const FString StatePathWithTransition = StatePath + FString::Printf(TEXT("/Transition[%d]"), TransitionIndex - TransitionsBegin);
 			
 			if (Transition.Trigger == EStateTreeTransitionTrigger::OnEvent)
 			{
@@ -967,8 +980,9 @@ bool FStateTreeCompiler::CreateStateTransitions()
 				CompactTransition.bConsumeEventOnSelect = Transition.RequiredEvent.bConsumeEventOnSelect;
 
 				FStateTreeBindableStructDesc Desc;
+				Desc.StatePath = StatePathWithTransition;
 				Desc.Struct = FStateTreeEvent::StaticStruct();
-				Desc.Name = FName(TEXT("Transition"));
+				Desc.Name = FName(TEXT("Transition Event"));
 				Desc.ID = Transition.GetEventID();
 				Desc.DataSource = EStateTreeBindableStructSource::TransitionEvent;
 				Desc.DataHandle = FStateTreeDataHandle(EStateTreeDataSourceType::TransitionEvent, TransitionIndex);
@@ -1005,7 +1019,7 @@ bool FStateTreeCompiler::CreateStateTransitions()
 			}
 			CompactTransition.ConditionsBegin = uint16(ConditionsBegin);
 			
-			if (!CreateConditions(*SourceState, Transition.Conditions))
+			if (!CreateConditions(*SourceState, StatePathWithTransition, Transition.Conditions))
 			{
 				Log.Reportf(EMessageSeverity::Error,
 					TEXT("Failed to create condition for transition to '%s'."),
@@ -1106,7 +1120,7 @@ bool FStateTreeCompiler::ResolveTransitionState(const UStateTreeState* SourceSta
 	return true;
 }
 
-bool FStateTreeCompiler::CreateCondition(UStateTreeState& State, const FStateTreeEditorNode& CondNode, const EStateTreeExpressionOperand Operand, const int8 DeltaIndent)
+bool FStateTreeCompiler::CreateCondition(UStateTreeState& State, const FString& StatePath, const FStateTreeEditorNode& CondNode, const EStateTreeExpressionOperand Operand, const int8 DeltaIndent)
 {
 	if (!CondNode.Node.IsValid())
 	{
@@ -1115,6 +1129,7 @@ bool FStateTreeCompiler::CreateCondition(UStateTreeState& State, const FStateTre
 	}
 
 	FStateTreeBindableStructDesc StructDesc;
+	StructDesc.StatePath = StatePath;
 	StructDesc.ID = CondNode.ID;
 	StructDesc.Name = CondNode.GetName();
 	StructDesc.DataSource = EStateTreeBindableStructSource::Condition;
@@ -1295,6 +1310,7 @@ bool FStateTreeCompiler::CreateTask(UStateTreeState* State, const FStateTreeEdit
 	
 	// Create binding source struct descriptor.
 	FStateTreeBindableStructDesc StructDesc;
+	StructDesc.StatePath = State ? State->GetPath() : UE::StateTree::Editor::GlobalStateName;
 	StructDesc.ID = TaskNode.ID;
 	StructDesc.Name = TaskNode.GetName();
 	StructDesc.DataSource = EStateTreeBindableStructSource::Task;
@@ -1409,6 +1425,7 @@ bool FStateTreeCompiler::CreateEvaluator(const FStateTreeEditorNode& EvalNode, c
 
 	// Create binding source struct descriptor.
 	FStateTreeBindableStructDesc StructDesc;
+	StructDesc.StatePath = UE::StateTree::Editor::GlobalStateName;
     StructDesc.ID = EvalNode.ID;
 	StructDesc.Name = EvalNode.GetName();
 	StructDesc.DataSource = EStateTreeBindableStructSource::Evaluator;
diff --git a/Engine/Plugins/Runtime/StateTree/Source/StateTreeEditorModule/Private/StateTreeEditorData.cpp b/Engine/Plugins/Runtime/StateTree/Source/StateTreeEditorModule/Private/StateTreeEditorData.cpp
--- a/Engine/Plugins/Runtime/StateTree/Source/StateTreeEditorModule/Private/StateTreeEditorData.cpp
+++ b/Engine/Plugins/Runtime/StateTree/Source/StateTreeEditorModule/Private/StateTreeEditorData.cpp
@@ -19,6 +19,11 @@
 
 #define LOCTEXT_NAMESPACE "StateTreeEditor"
 
+namespace UE::StateTree::Editor
+{
+	const FString GlobalStateName(TEXT("Global"));
+}
+
 UStateTreeEditorData::UStateTreeEditorData()
 {
 	FStateTreeEditorColor DefaultColor;
@@ -351,11 +356,6 @@ void UStateTreeEditorData::GetAccessibleStructs(const TConstArrayView<const USta
 			}
 		}
 	}
-	
-	OutStructDescs.StableSort([](const FStateTreeBindableStructDesc& A, const FStateTreeBindableStructDesc& B)
-	{
-		return (uint8)A.DataSource < (uint8)B.DataSource;
-	});
 }
 
 FStateTreeBindableStructDesc UStateTreeEditorData::FindContextData(const UStruct* ObjectType, const FString ObjectNameHint) const
@@ -371,7 +371,7 @@ FStateTreeBindableStructDesc UStateTreeEditorData::FindContextData(const UStruct
 	{
 		if (Desc.Struct->IsChildOf(ObjectType))
 		{
-			Candidates.Emplace(Desc.Name, Desc.Struct, FStateTreeDataHandle(), EStateTreeBindableStructSource::Context, Desc.ID);
+			Candidates.Emplace(UE::StateTree::Editor::GlobalStateName, Desc.Name, Desc.Struct, FStateTreeDataHandle(), EStateTreeBindableStructSource::Context, Desc.ID);
 		}
 	}
 
@@ -893,19 +893,21 @@ EStateTreeVisitor UStateTreeEditorData::VisitStateNodes(const UStateTreeState& S
 	return bContinue ? EStateTreeVisitor::Continue : EStateTreeVisitor::Break;
 }
 
-
 EStateTreeVisitor UStateTreeEditorData::VisitStateNodes(const UStateTreeState& State, TFunctionRef<EStateTreeVisitor(const UStateTreeState* State, const FStateTreeBindableStructDesc& Desc, const FStateTreeDataView Value)> InFunc) const
 {
 	bool bContinue = true;
 
+	const FString StatePath = State.GetPath();
+	
 	if (bContinue)
 	{
 		// Bindable state parameters
 		if (State.Parameters.Parameters.IsValid())
 		{
 			FStateTreeBindableStructDesc Desc;
+			Desc.StatePath = StatePath;
 			Desc.Struct = State.Parameters.Parameters.GetPropertyBagStruct();
-			Desc.Name = State.Name;
+			Desc.Name = FName("Parameters");
 			Desc.ID = State.Parameters.ID;
 			Desc.DataSource = EStateTreeBindableStructSource::StateParameter;
 
@@ -916,13 +918,16 @@ EStateTreeVisitor UStateTreeEditorData::VisitStateNodes(const UStateTreeState& S
 		}
 	}
 
+	const FString StatePathWithConditions = StatePath + TEXT("/EnterConditions");
+	
 	if (bContinue)
 	{
 		if (State.bHasRequiredEventToEnter )
 		{
 			FStateTreeBindableStructDesc Desc;
+			Desc.StatePath = StatePathWithConditions;
 			Desc.Struct = FStateTreeEvent::StaticStruct();
-			Desc.Name = State.Name;
+			Desc.Name = FName("Enter Event");
 			Desc.ID = State.GetEventID();
 			Desc.DataSource = EStateTreeBindableStructSource::StateEvent;
 
@@ -941,6 +946,7 @@ EStateTreeVisitor UStateTreeEditorData::VisitStateNodes(const UStateTreeState& S
 			if (const FStateTreeConditionBase* Cond = Node.Node.GetPtr<FStateTreeConditionBase>())
 			{
 				FStateTreeBindableStructDesc Desc;
+				Desc.StatePath = StatePathWithConditions;
 				Desc.Struct = Cond->GetInstanceDataType();
 				Desc.Name = Node.GetName();
 				Desc.ID = Node.ID;
@@ -962,6 +968,7 @@ EStateTreeVisitor UStateTreeEditorData::VisitStateNodes(const UStateTreeState& S
 			if (const FStateTreeTaskBase* Task = Node.Node.GetPtr<FStateTreeTaskBase>())
 			{
 				FStateTreeBindableStructDesc Desc;
+				Desc.StatePath = StatePath;
 				Desc.Struct = Task->GetInstanceDataType();
 				Desc.Name = Node.GetName();
 				Desc.ID = Node.ID;
@@ -980,6 +987,7 @@ EStateTreeVisitor UStateTreeEditorData::VisitStateNodes(const UStateTreeState& S
 		if (const FStateTreeTaskBase* Task = State.SingleTask.Node.GetPtr<FStateTreeTaskBase>())
 		{
 			FStateTreeBindableStructDesc Desc;
+			Desc.StatePath = StatePath;
 			Desc.Struct = Task->GetInstanceDataType();
 			Desc.Name = State.SingleTask.GetName();
 			Desc.ID = State.SingleTask.ID;
@@ -995,13 +1003,17 @@ EStateTreeVisitor UStateTreeEditorData::VisitStateNodes(const UStateTreeState& S
 	if (bContinue)
 	{
 		// Transitions
-		for (const FStateTreeTransition& Transition : State.Transitions)
+		for (int32 TransitionIndex = 0; TransitionIndex < State.Transitions.Num(); TransitionIndex++)
 		{
+			const FStateTreeTransition& Transition = State.Transitions[TransitionIndex];
+			const FString StatePathWithTransition = StatePath + FString::Printf(TEXT("/Transition[%d]"), TransitionIndex);
+
 			if (Transition.Trigger == EStateTreeTransitionTrigger::OnEvent)
 			{
 				FStateTreeBindableStructDesc Desc;
+				Desc.StatePath = StatePathWithTransition;
 				Desc.Struct = FStateTreeEvent::StaticStruct();
-				Desc.Name = FName(TEXT("Transition"));
+				Desc.Name = FName(TEXT("Transition Event"));
 				Desc.ID = Transition.GetEventID();
 				Desc.DataSource = EStateTreeBindableStructSource::TransitionEvent;
 
@@ -1017,6 +1029,7 @@ EStateTreeVisitor UStateTreeEditorData::VisitStateNodes(const UStateTreeState& S
 				if (const FStateTreeConditionBase* Cond = Node.Node.GetPtr<FStateTreeConditionBase>())
 				{
 					FStateTreeBindableStructDesc Desc;
+					Desc.StatePath = StatePathWithTransition;
 					Desc.Struct = Cond->GetInstanceDataType();
 					Desc.Name = Node.GetName();
 					Desc.ID = Node.ID;
@@ -1096,6 +1109,7 @@ EStateTreeVisitor UStateTreeEditorData::VisitGlobalNodes(TFunctionRef<EStateTree
 	// Root parameters
 	{
 		FStateTreeBindableStructDesc Desc;
+		Desc.StatePath = UE::StateTree::Editor::GlobalStateName;
 		Desc.Struct = RootParameters.Parameters.GetPropertyBagStruct();
 		Desc.Name = FName(TEXT("Parameters"));
 		Desc.ID = RootParameters.ID;
@@ -1113,6 +1127,7 @@ EStateTreeVisitor UStateTreeEditorData::VisitGlobalNodes(TFunctionRef<EStateTree
 		for (const FStateTreeExternalDataDesc& ContextDesc : Schema->GetContextDataDescs())
 		{
 			FStateTreeBindableStructDesc Desc;
+			Desc.StatePath = UE::StateTree::Editor::GlobalStateName;
 			Desc.Struct = ContextDesc.Struct;
 			Desc.Name = ContextDesc.Name;
 			Desc.ID = ContextDesc.ID;
@@ -1132,6 +1147,7 @@ EStateTreeVisitor UStateTreeEditorData::VisitGlobalNodes(TFunctionRef<EStateTree
 		if (const FStateTreeEvaluatorBase* Evaluator = Node.Node.GetPtr<FStateTreeEvaluatorBase>())
 		{
 			FStateTreeBindableStructDesc Desc;
+			Desc.StatePath = UE::StateTree::Editor::GlobalStateName;
 			Desc.Struct = Evaluator->GetInstanceDataType();
 			Desc.Name = Node.GetName();
 			Desc.ID = Node.ID;
@@ -1150,6 +1166,7 @@ EStateTreeVisitor UStateTreeEditorData::VisitGlobalNodes(TFunctionRef<EStateTree
 		if (const FStateTreeTaskBase* Task = Node.Node.GetPtr<FStateTreeTaskBase>())
 		{
 			FStateTreeBindableStructDesc Desc;
+			Desc.StatePath = UE::StateTree::Editor::GlobalStateName;
 			Desc.Struct = Task->GetInstanceDataType();
 			Desc.Name = Node.GetName();
 			Desc.ID = Node.ID;
diff --git a/Engine/Plugins/Runtime/StateTree/Source/StateTreeEditorModule/Private/StateTreeEditorNode.cpp b/Engine/Plugins/Runtime/StateTree/Source/StateTreeEditorModule/Private/StateTreeEditorNode.cpp
new file mode 100644
--- /dev/null
+++ b/Engine/Plugins/Runtime/StateTree/Source/StateTreeEditorModule/Private/StateTreeEditorNode.cpp
@@ -0,0 +1,34 @@
+﻿// Copyright Epic Games, Inc. All Rights Reserved.
+
+#include "StateTreeEditorNode.h"
+
+#include "Blueprint/StateTreeConditionBlueprintBase.h"
+#include "Blueprint/StateTreeEvaluatorBlueprintBase.h"
+#include "Blueprint/StateTreeTaskBlueprintBase.h"
+
+FName FStateTreeEditorNode::GetName() const
+{
+	const UScriptStruct* NodeType = Node.GetScriptStruct();
+	if (!NodeType)
+	{
+		return FName();
+	}
+
+	if (const FStateTreeNodeBase* NodePtr = Node.GetPtr<FStateTreeNodeBase>())
+	{
+		if (NodePtr->Name.IsNone())
+		{
+			if (InstanceObject &&
+					(NodeType->IsChildOf(TBaseStructure<FStateTreeBlueprintTaskWrapper>::Get())
+					|| NodeType->IsChildOf(TBaseStructure<FStateTreeBlueprintEvaluatorWrapper>::Get())
+					|| NodeType->IsChildOf(TBaseStructure<FStateTreeBlueprintConditionWrapper>::Get())))
+			{
+				return FName(InstanceObject->GetClass()->GetDisplayNameText().ToString());
+			}
+			return FName(NodeType->GetDisplayNameText().ToString());
+		}
+		return NodePtr->Name;
+	}
+	
+	return FName();
+}
diff --git a/Engine/Plugins/Runtime/StateTree/Source/StateTreeEditorModule/Private/StateTreeState.cpp b/Engine/Plugins/Runtime/StateTree/Source/StateTreeEditorModule/Private/StateTreeState.cpp
--- a/Engine/Plugins/Runtime/StateTree/Source/StateTreeEditorModule/Private/StateTreeState.cpp
+++ b/Engine/Plugins/Runtime/StateTree/Source/StateTreeEditorModule/Private/StateTreeState.cpp
@@ -543,6 +543,28 @@ const UStateTreeState* UStateTreeState::GetNextSelectableSiblingState() const
 	return nullptr;
 }
 
+FString UStateTreeState::GetPath() const
+{
+	TArray<const UStateTreeState*> States;
+	for (const UStateTreeState* CurrState = this; CurrState; CurrState = CurrState->Parent)
+	{
+		States.Add(CurrState);
+	}
+	Algo::Reverse(States);
+	
+	FStringBuilderBase Result;
+	for (const UStateTreeState* CurrState : States)
+	{
+		if (Result.Len() > 0)
+		{
+			Result.Append(TEXT("/"));
+		}
+		Result.Append(CurrState->Name.ToString());
+	}
+
+	return Result.ToString();
+}
+
 FStateTreeStateLink UStateTreeState::GetLinkToState() const
 {
 	FStateTreeStateLink Link(EStateTreeTransitionType::GotoState);
diff --git a/Engine/Plugins/Runtime/StateTree/Source/StateTreeEditorModule/Public/StateTreeCompiler.h b/Engine/Plugins/Runtime/StateTree/Source/StateTreeEditorModule/Public/StateTreeCompiler.h
--- a/Engine/Plugins/Runtime/StateTree/Source/StateTreeEditorModule/Public/StateTreeCompiler.h
+++ b/Engine/Plugins/Runtime/StateTree/Source/StateTreeEditorModule/Public/StateTreeCompiler.h
@@ -49,8 +49,8 @@ private:
 	bool CreateStateTasksAndParameters();
 	bool CreateStateTransitions();
 	
-	bool CreateConditions(UStateTreeState& State, TConstArrayView<FStateTreeEditorNode> Conditions);
-	bool CreateCondition(UStateTreeState& State, const FStateTreeEditorNode& CondNode, const EStateTreeExpressionOperand Operand, const int8 DeltaIndent);
+	bool CreateConditions(UStateTreeState& State, const FString& StatePath, TConstArrayView<FStateTreeEditorNode> Conditions);
+	bool CreateCondition(UStateTreeState& State, const FString& StatePath, const FStateTreeEditorNode& CondNode, const EStateTreeExpressionOperand Operand, const int8 DeltaIndent);
 	bool CreateTask(UStateTreeState* State, const FStateTreeEditorNode& TaskNode, const FStateTreeDataHandle TaskDataHandle);
 	bool CreateEvaluator(const FStateTreeEditorNode& EvalNode, const FStateTreeDataHandle EvalDataHandle);
 	bool GetAndValidateBindings(const FStateTreeBindableStructDesc& TargetStruct, FStateTreeDataView TargetValue, TArray<FStateTreePropertyPathBinding>& OutCopyBindings, TArray<FStateTreePropertyPathBinding>& OutReferenceBindings) const;
diff --git a/Engine/Plugins/Runtime/StateTree/Source/StateTreeEditorModule/Public/StateTreeEditorData.h b/Engine/Plugins/Runtime/StateTree/Source/StateTreeEditorModule/Public/StateTreeEditorData.h
--- a/Engine/Plugins/Runtime/StateTree/Source/StateTreeEditorModule/Public/StateTreeEditorData.h
+++ b/Engine/Plugins/Runtime/StateTree/Source/StateTreeEditorModule/Public/StateTreeEditorData.h
@@ -13,6 +13,13 @@ struct FStateTreeEditorPropertyPath;
 
 class UStateTreeSchema;
 
+namespace UE::StateTree::Editor
+{
+	// Name used to describe container of global items (other items use the path to the container State).  
+	extern STATETREEEDITORMODULE_API const FString GlobalStateName;
+}
+
+
 USTRUCT()
 struct FStateTreeEditorBreakpoint
 {
diff --git a/Engine/Plugins/Runtime/StateTree/Source/StateTreeEditorModule/Public/StateTreeEditorNode.h b/Engine/Plugins/Runtime/StateTree/Source/StateTreeEditorModule/Public/StateTreeEditorNode.h
--- a/Engine/Plugins/Runtime/StateTree/Source/StateTreeEditorModule/Public/StateTreeEditorNode.h
+++ b/Engine/Plugins/Runtime/StateTree/Source/StateTreeEditorModule/Public/StateTreeEditorNode.h
@@ -36,18 +36,7 @@ struct STATETREEEDITORMODULE_API FStateTreeEditorNode
 	 * If the node has a specified name, used that, or else of return the display name of the node.
 	 * @return name of the node.
 	 */
-	FName GetName() const
-	{
-		if (const FStateTreeNodeBase* NodePtr = Node.GetPtr<FStateTreeNodeBase>())
-		{
-			if (NodePtr->Name.IsNone())
-			{
-				return FName(Node.GetScriptStruct()->GetDisplayNameText().ToString());
-			}
-			return NodePtr->Name;
-		}
-		return FName();
-	}
+	FName GetName() const;
 
 	FStateTreeDataView GetInstance() const
 	{
diff --git a/Engine/Plugins/Runtime/StateTree/Source/StateTreeEditorModule/Public/StateTreeState.h b/Engine/Plugins/Runtime/StateTree/Source/StateTreeEditorModule/Public/StateTreeState.h
--- a/Engine/Plugins/Runtime/StateTree/Source/StateTreeEditorModule/Public/StateTreeState.h
+++ b/Engine/Plugins/Runtime/StateTree/Source/StateTreeEditorModule/Public/StateTreeState.h
@@ -214,6 +214,9 @@ public:
 	const UStateTreeState* GetNextSiblingState() const;
 	const UStateTreeState* GetNextSelectableSiblingState() const;
 
+	/** @return the path of the state as string. */
+	FString GetPath() const;
+	
 	/** @return true if the property of specified ID is overridden. */
 	bool IsParametersPropertyOverridden(const FGuid PropertyID) const
 	{
diff --git a/Engine/Plugins/Runtime/StateTree/Source/StateTreeModule/Private/StateTreePropertyBindings.cpp b/Engine/Plugins/Runtime/StateTree/Source/StateTreeModule/Private/StateTreePropertyBindings.cpp
--- a/Engine/Plugins/Runtime/StateTree/Source/StateTreeModule/Private/StateTreePropertyBindings.cpp
+++ b/Engine/Plugins/Runtime/StateTree/Source/StateTreeModule/Private/StateTreePropertyBindings.cpp
@@ -126,8 +126,11 @@ FString FStateTreeBindableStructDesc::ToString() const
 {
 	FStringBuilderBase Result;
 
-	Result += UEnum::GetDisplayValueAsText(DataSource).ToString();
-	Result += TEXT(" '");
+	Result += TEXT("'");
+#if WITH_EDITORONLY_DATA
+	Result += StatePath;
+	Result += TEXT("/");
+#endif
 	Result += Name.ToString();
 	Result += TEXT("'");
 
@@ -1400,7 +1403,7 @@ bool FStateTreePropertyPath::UpdateSegmentsFromValue(const FStateTreeDataView Ba
 
 FString FStateTreePropertyPath::ToString(const int32 HighlightedSegment, const TCHAR* HighlightPrefix, const TCHAR* HighlightPostfix, const bool bOutputInstances) const
 {
-	FString Result;
+	FStringBuilderBase Result;
 	for (TEnumerateRef<const FStateTreePropertyPathSegment> Segment : EnumerateRange(Segments))
 	{
 		if (Segment.GetIndex() > 0)
@@ -1417,21 +1420,38 @@ FString FStateTreePropertyPath::ToString(const int32 HighlightedSegment, const T
 			Result += FString::Printf(TEXT("(%s)"), *GetNameSafe(Segment->GetInstanceStruct()));
 		}
 
-		if (Segment->GetArrayIndex() >= 0)
+#if WITH_EDITORONLY_DATA
+		const UStruct* ParentInstanceStruct = Segment.GetIndex() > 0 ? Segments[Segment.GetIndex() - 1].GetInstanceStruct() : nullptr;
+		if (const UUserDefinedStruct* ParentUserDefinedStruct = Cast<const UUserDefinedStruct>(ParentInstanceStruct))
 		{
-			Result += FString::Printf(TEXT("%s[%d]"), *Segment->GetName().ToString(), Segment->GetArrayIndex());
+			// Find friendly names for UDS properties (the property name itself has hash in it). 
+			const FString FriendlyName = FStructureEditorUtils::GetVariableFriendlyName(ParentUserDefinedStruct, Segment->GetPropertyGuid());
+			if (!FriendlyName.IsEmpty())
+			{
+				Result += FriendlyName;
+			}
+			else
+			{
+				Result += Segment->GetName().ToString();
+			}
 		}
 		else
+#endif			
 		{
 			Result += Segment->GetName().ToString();
 		}
 
+		if (Segment->GetArrayIndex() >= 0)
+		{
+			Result += FString::Printf(TEXT("[%d]"), Segment->GetArrayIndex());
+		}
+
 		if (Segment.GetIndex() == HighlightedSegment && HighlightPostfix)
 		{
 			Result += HighlightPostfix;
 		}
 	}
-	return Result;
+	return Result.ToString();
 }
 
 
diff --git a/Engine/Plugins/Runtime/StateTree/Source/StateTreeModule/Public/StateTreePropertyBindings.h b/Engine/Plugins/Runtime/StateTree/Source/StateTreeModule/Public/StateTreePropertyBindings.h
--- a/Engine/Plugins/Runtime/StateTree/Source/StateTreeModule/Public/StateTreePropertyBindings.h
+++ b/Engine/Plugins/Runtime/StateTree/Source/StateTreeModule/Public/StateTreePropertyBindings.h
@@ -132,6 +132,17 @@ struct STATETREEMODULE_API FStateTreeBindableStructDesc
 	FStateTreeBindableStructDesc() = default;
 
 #if WITH_EDITORONLY_DATA
+	FStateTreeBindableStructDesc(const FString& InStatePath, const FName InName, const UStruct* InStruct, const FStateTreeDataHandle InDataHandle, const EStateTreeBindableStructSource InDataSource, const FGuid InGuid)
+		: Struct(InStruct)
+		, Name(InName)
+		, DataHandle(InDataHandle)
+		, DataSource(InDataSource)
+		, ID(InGuid)
+		, StatePath(InStatePath)
+	{
+	}
+
+	UE_DEPRECATED(5.5, "Use constructor with StatePath instead.")
 	FStateTreeBindableStructDesc(const FName InName, const UStruct* InStruct, const FStateTreeDataHandle InDataHandle, const EStateTreeBindableStructSource InDataSource, const FGuid InGuid)
 		: Struct(InStruct)
 		, Name(InName)
@@ -180,6 +191,10 @@ struct STATETREEMODULE_API FStateTreeBindableStructDesc
 	/** Unique identifier of the struct. */
 	UPROPERTY()
 	FGuid ID;
+
+	/** In Editor path to State containting the data. */
+	UPROPERTY(Transient)
+	FString StatePath;
 #endif
 };
 
