[StateTree] StateTreeEditor LOCTEXT namespace text conflict

diff --git a/Engine/Plugins/Runtime/StateTree/Source/StateTreeEditorModule/Private/Customizations/StateTreeEditorNodeDetails.cpp b/Engine/Plugins/Runtime/StateTree/Source/StateTreeEditorModule/Private/Customizations/StateTreeEditorNodeDetails.cpp
--- a/Engine/Plugins/Runtime/StateTree/Source/StateTreeEditorModule/Private/Customizations/StateTreeEditorNodeDetails.cpp
+++ b/Engine/Plugins/Runtime/StateTree/Source/StateTreeEditorModule/Private/Customizations/StateTreeEditorNodeDetails.cpp
@@ -1617,7 +1617,7 @@ TSharedRef<SWidget> FStateTreeEditorNodeDetails::GenerateOptionsMenu()
 	// Paste
 	MenuBuilder.AddMenuEntry(
 		LOCTEXT("PasteItem", "Paste"),
-		LOCTEXT("PasteItemTooltip", "Paste over this item"),
+		LOCTEXT("PasteItemTooltip", "Paste into this item"),
 		FSlateIcon(FAppStyle::GetAppStyleSetName(), "GenericCommands.Paste"),
 		FUIAction(
 			FExecuteAction::CreateSP(this, &FStateTreeEditorNodeDetails::OnPasteNode),
diff --git a/Engine/Plugins/Runtime/StateTree/Source/StateTreeEditorModule/Private/Customizations/StateTreeEditorNodeUtils.cpp b/Engine/Plugins/Runtime/StateTree/Source/StateTreeEditorModule/Private/Customizations/StateTreeEditorNodeUtils.cpp
--- a/Engine/Plugins/Runtime/StateTree/Source/StateTreeEditorModule/Private/Customizations/StateTreeEditorNodeUtils.cpp
+++ b/Engine/Plugins/Runtime/StateTree/Source/StateTreeEditorModule/Private/Customizations/StateTreeEditorNodeUtils.cpp
@@ -526,7 +526,7 @@ void OnArrayNodePicked(const UStruct* InStruct, TSharedPtr<SComboButton> PickerC
 {
 	if (const TSharedPtr<IPropertyHandleArray> ArrayHandle = ArrayPropertyHandle->AsArray())
 	{
-		GEditor->BeginTransaction(LOCTEXT("SelectNode", "Add Node"));
+		GEditor->BeginTransaction(LOCTEXT("AddNode", "Add Node"));
 		ArrayPropertyHandle->NotifyPreChange();
 
 		// Add new item to the end.
