StateTree: Fixed bug where renaming a state parameter would remove bindings to it.

diff --git a/Engine/Plugins/Runtime/StateTree/Source/StateTreeEditorModule/Private/StateTreeEditorData.cpp b/Engine/Plugins/Runtime/StateTree/Source/StateTreeEditorModule/Private/StateTreeEditorData.cpp
--- a/Engine/Plugins/Runtime/StateTree/Source/StateTreeEditorModule/Private/StateTreeEditorData.cpp
+++ b/Engine/Plugins/Runtime/StateTree/Source/StateTreeEditorModule/Private/StateTreeEditorData.cpp
@@ -33,6 +33,8 @@ UStateTreeEditorData::UStateTreeEditorData()
 	DefaultColor.DisplayName = TEXT("Default Color");
 
 	Colors.Add(MoveTemp(DefaultColor));
+
+	EditorBindings.SetBindingsOwner(this);
 }
 
 void UStateTreeEditorData::PostInitProperties()
@@ -45,6 +47,7 @@ void UStateTreeEditorData::PostInitProperties()
 	OnObjectsReinstancedHandle = FCoreUObjectDelegates::OnObjectsReinstanced.AddUObject(this, &UStateTreeEditorData::OnObjectsReinstanced);
 	OnUserDefinedStructReinstancedHandle = UE::StructUtils::Delegates::OnUserDefinedStructReinstanced.AddUObject(this, &UStateTreeEditorData::OnUserDefinedStructReinstanced);
 	OnParametersChangedHandle = UE::StateTree::Delegates::OnParametersChanged.AddUObject(this, &UStateTreeEditorData::OnParametersChanged);
+	OnStateParametersChangedHandle = UE::StateTree::Delegates::OnStateParametersChanged.AddUObject(this, &UStateTreeEditorData::OnStateParametersChanged);
 #endif
 }
 
@@ -67,6 +70,11 @@ void UStateTreeEditorData::BeginDestroy()
 		UE::StateTree::Delegates::OnParametersChanged.Remove(OnParametersChangedHandle);
 		OnParametersChangedHandle.Reset();
 	}
+	if (OnStateParametersChangedHandle.IsValid())
+	{
+		UE::StateTree::Delegates::OnStateParametersChanged.Remove(OnStateParametersChangedHandle);
+		OnStateParametersChangedHandle.Reset();
+	}
 	
 	Super::BeginDestroy();
 }
@@ -147,6 +155,16 @@ void UStateTreeEditorData::OnParametersChanged(const UStateTree& StateTree)
 	}
 }
 
+void UStateTreeEditorData::OnStateParametersChanged(const UStateTree& StateTree, const FGuid StateID)
+{
+	if (const UStateTree* OwnerStateTree = GetTypedOuter<UStateTree>())
+	{
+		if (OwnerStateTree == &StateTree)
+		{
+			UpdateBindingsInstanceStructs();
+		}
+	}
+}
 
 void UStateTreeEditorData::PostLoad()
 {
diff --git a/Engine/Plugins/Runtime/StateTree/Source/StateTreeEditorModule/Private/StateTreeEditorPropertyBindings.cpp b/Engine/Plugins/Runtime/StateTree/Source/StateTreeEditorModule/Private/StateTreeEditorPropertyBindings.cpp
--- a/Engine/Plugins/Runtime/StateTree/Source/StateTreeEditorModule/Private/StateTreeEditorPropertyBindings.cpp
+++ b/Engine/Plugins/Runtime/StateTree/Source/StateTreeEditorModule/Private/StateTreeEditorPropertyBindings.cpp
@@ -17,10 +17,33 @@ UStateTreeEditorPropertyBindingsOwner::UStateTreeEditorPropertyBindingsOwner(con
 
 //////////////////////////////////////////////////////////////////////////
 
+void FStateTreeEditorPropertyBindings::SetBindingsOwner(IStateTreeEditorPropertyBindingsOwner* InBindingsOwner)
+{
+	BindingsOwner = InBindingsOwner;
+}
+
 void FStateTreeEditorPropertyBindings::AddPropertyBinding(const FStateTreePropertyPath& SourcePath, const FStateTreePropertyPath& TargetPath)
 {
 	RemovePropertyBindings(TargetPath);
-	PropertyBindings.Add(FStateTreePropertyPathBinding(SourcePath, TargetPath));
+
+	FStateTreePropertyPathBinding Binding(SourcePath, TargetPath);
+
+	// If we have bindings owner, update property path segments to capture property IDs, etc.
+	if (BindingsOwner)
+	{
+		FStateTreeDataView SourceDataView;
+		if (BindingsOwner->GetDataViewByID(Binding.GetSourcePath().GetStructID(), SourceDataView))
+		{
+			Binding.GetMutableSourcePath().UpdateSegmentsFromValue(SourceDataView);
+		}
+		FStateTreeDataView TargetDataView;
+		if (BindingsOwner->GetDataViewByID(Binding.GetTargetPath().GetStructID(), TargetDataView))
+		{
+			Binding.GetMutableTargetPath().UpdateSegmentsFromValue(TargetDataView);
+		}
+	}
+	
+	PropertyBindings.Add(Binding);
 }
 
 void FStateTreeEditorPropertyBindings::RemovePropertyBindings(const FStateTreePropertyPath& TargetPath)
diff --git a/Engine/Plugins/Runtime/StateTree/Source/StateTreeEditorModule/Private/StateTreeState.cpp b/Engine/Plugins/Runtime/StateTree/Source/StateTreeEditorModule/Private/StateTreeState.cpp
--- a/Engine/Plugins/Runtime/StateTree/Source/StateTreeEditorModule/Private/StateTreeState.cpp
+++ b/Engine/Plugins/Runtime/StateTree/Source/StateTreeEditorModule/Private/StateTreeState.cpp
@@ -216,17 +216,13 @@ void UStateTreeState::PostEditChangeChainProperty(FPropertyChangedChainEvent& Pr
 		}
 	}
 
-	// Broadcast subtree parameter layout edits so that the linked states can adapt.
+	// Broadcast subtree parameter layout edits so that the linked states can adapt, and bindings can update.
 	if (ChangePropertyPath.IsPathExact(StateParametersPath))
 	{
-		if (!(Type == EStateTreeStateType::Linked
-				|| Type == EStateTreeStateType::LinkedAsset))
+		const UStateTree* StateTree = GetTypedOuter<UStateTree>();
+		if (ensure(StateTree))
 		{
-			const UStateTree* StateTree = GetTypedOuter<UStateTree>();
-			if (ensure(StateTree))
-			{
-				UE::StateTree::Delegates::OnStateParametersChanged.Broadcast(*StateTree, ID);
-			}
+			UE::StateTree::Delegates::OnStateParametersChanged.Broadcast(*StateTree, ID);
 		}
 	}
 
diff --git a/Engine/Plugins/Runtime/StateTree/Source/StateTreeEditorModule/Public/StateTreeEditorData.h b/Engine/Plugins/Runtime/StateTree/Source/StateTreeEditorModule/Public/StateTreeEditorData.h
--- a/Engine/Plugins/Runtime/StateTree/Source/StateTreeEditorModule/Public/StateTreeEditorData.h
+++ b/Engine/Plugins/Runtime/StateTree/Source/StateTreeEditorModule/Public/StateTreeEditorData.h
@@ -83,6 +83,7 @@ public:
 	void OnObjectsReinstanced(const FReplacementObjectMap& ObjectMap);
 	void OnUserDefinedStructReinstanced(const UUserDefinedStruct& UserDefinedStruct);
 	void OnParametersChanged(const UStateTree& StateTree);
+	void OnStateParametersChanged(const UStateTree& StateTree, const FGuid StateID);
 	virtual void BeginDestroy() override;
 	virtual void PostLoad() override;
 	virtual void PostEditChangeChainProperty(FPropertyChangedChainEvent& PropertyChangedEvent) override;
@@ -268,6 +269,7 @@ private:
 	FDelegateHandle OnObjectsReinstancedHandle;
 	FDelegateHandle OnUserDefinedStructReinstancedHandle;
 	FDelegateHandle OnParametersChangedHandle;
+	FDelegateHandle OnStateParametersChangedHandle;
 #endif
 
 public:
diff --git a/Engine/Plugins/Runtime/StateTree/Source/StateTreeEditorModule/Public/StateTreeEditorPropertyBindings.h b/Engine/Plugins/Runtime/StateTree/Source/StateTreeEditorModule/Public/StateTreeEditorPropertyBindings.h
--- a/Engine/Plugins/Runtime/StateTree/Source/StateTreeEditorModule/Public/StateTreeEditorPropertyBindings.h
+++ b/Engine/Plugins/Runtime/StateTree/Source/StateTreeEditorModule/Public/StateTreeEditorPropertyBindings.h
@@ -7,6 +7,7 @@
 #include "StateTreeEditorPropertyBindings.generated.h"
 
 enum EStateTreeNodeFormatting : uint8;
+class IStateTreeEditorPropertyBindingsOwner;
 
 /**
  * Editor representation of a all property bindings in a StateTree
@@ -16,6 +17,9 @@ struct STATETREEEDITORMODULE_API FStateTreeEditorPropertyBindings
 {
 	GENERATED_BODY()
 
+	/** Sets associated bindings owner, used to validate added property paths. */
+	void SetBindingsOwner(IStateTreeEditorPropertyBindingsOwner* BindingsOwner);
+
 	/**
 	 * Adds binding between source and destination paths. Removes any bindings to TargetPath before adding the new one.
 	 * @param SourcePath Binding source property path.
@@ -90,6 +94,8 @@ private:
 
 	UPROPERTY()
 	TArray<FStateTreePropertyPathBinding> PropertyBindings;
+
+	IStateTreeEditorPropertyBindingsOwner* BindingsOwner = nullptr; 
 };
 
 
