[StateTree] Fix property not initialized properly

diff --git a/Engine/Plugins/Runtime/StateTree/Source/StateTreeModule/Private/StateTreeConsiderationBase.cpp b/Engine/Plugins/Runtime/StateTree/Source/StateTreeModule/Private/StateTreeConsiderationBase.cpp
--- a/Engine/Plugins/Runtime/StateTree/Source/StateTreeModule/Private/StateTreeConsiderationBase.cpp
+++ b/Engine/Plugins/Runtime/StateTree/Source/StateTreeModule/Private/StateTreeConsiderationBase.cpp
@@ -1,9 +1,24 @@
 // Copyright Epic Games, Inc. All Rights Reserved.
 
 #include "StateTreeConsiderationBase.h"
+#include "AlphaBlend.h"
+#include "StateTreeTypes.h"
 
 #include UE_INLINE_GENERATED_CPP_BY_NAME(StateTreeConsiderationBase)
 
+FStateTreeConsiderationResponseCurve::FStateTreeConsiderationResponseCurve()
+	: BlendOption(EAlphaBlendOption::Linear)
+	, RawScoreLowerBound(.0f)
+	, RawScoreUpperBound(.1f)
+{
+}
+
+FStateTreeConsiderationBase::FStateTreeConsiderationBase()
+	: Operand(EStateTreeExpressionOperand::And)
+	, DeltaIndent(0)
+{
+}
+
 float FStateTreeConsiderationBase::ComputeNormalizedScore(FStateTreeExecutionContext& Context) const
 {
 	const float RawScore = ComputeRawScore(Context);
diff --git a/Engine/Plugins/Runtime/StateTree/Source/StateTreeModule/Public/StateTreeConsiderationBase.h b/Engine/Plugins/Runtime/StateTree/Source/StateTreeModule/Public/StateTreeConsiderationBase.h
--- a/Engine/Plugins/Runtime/StateTree/Source/StateTreeModule/Public/StateTreeConsiderationBase.h
+++ b/Engine/Plugins/Runtime/StateTree/Source/StateTreeModule/Public/StateTreeConsiderationBase.h
@@ -2,29 +2,31 @@
 
 #pragma once
 
-#include "AlphaBlend.h"
 #include "StateTreeNodeBase.h"
-#include "StateTreeTypes.h"
 #include "StateTreeConsiderationBase.generated.h"
 
 struct FStateTreeExecutionContext;
+enum class EAlphaBlendOption : uint8;
+enum class EStateTreeExpressionOperand : uint8;
 
 USTRUCT()
 struct STATETREEMODULE_API FStateTreeConsiderationResponseCurve
 {
 	GENERATED_BODY()
 
+	FStateTreeConsiderationResponseCurve();
+
 	/* Optional Curve used to output the final normalized score.If it is set empty, the final value will be raw score normalized by bounds */
 	UPROPERTY(EditAnywhere, Category = Default)
 	EAlphaBlendOption BlendOption;
 
 	/* Lower Bound used to normalize the raw score */
 	UPROPERTY(EditAnywhere, Category = Default)
-	float RawScoreLowerBound = .0f;
+	float RawScoreLowerBound;
 
 	/* Upper Bound used to normalize the raw score */
 	UPROPERTY(EditAnywhere, Category = Default)
-	float RawScoreUpperBound = 1.f;
+	float RawScoreUpperBound;
 };
 
 /**
@@ -36,6 +38,8 @@ struct STATETREEMODULE_API FStateTreeConsiderationBase : public FStateTreeNodeBa
 {
 	GENERATED_BODY()
 
+	FStateTreeConsiderationBase();
+
 	float ComputeNormalizedScore(FStateTreeExecutionContext& Context) const;
 
 protected:
@@ -43,10 +47,10 @@ protected:
 
 public:
 	UPROPERTY()
-	EStateTreeExpressionOperand Operand = EStateTreeExpressionOperand::And;
+	EStateTreeExpressionOperand Operand;
 
 	UPROPERTY()
-	int8 DeltaIndent = 0;
+	int8 DeltaIndent;
 
 	//Response Curve used to output the final normalized score.
 	UPROPERTY(EditAnywhere, Category = Default)
