[StateTree] Fix LOCTEXT name conflict

diff --git a/Engine/Plugins/Runtime/StateTree/Source/StateTreeModule/Private/Considerations/StateTreeCommonConsiderations.cpp b/Engine/Plugins/Runtime/StateTree/Source/StateTreeModule/Private/Considerations/StateTreeCommonConsiderations.cpp
--- a/Engine/Plugins/Runtime/StateTree/Source/StateTreeModule/Private/Considerations/StateTreeCommonConsiderations.cpp
+++ b/Engine/Plugins/Runtime/StateTree/Source/StateTreeModule/Private/Considerations/StateTreeCommonConsiderations.cpp
@@ -20,7 +20,7 @@ FText FStateTreeFloatConsideration::GetDescription(const FGuid& ID, FStateTreeDa
 	{
 		//using a constant value for the float param
 		OptionalConstantKeywordText = (Formatting == EStateTreeNodeFormatting::RichText)
-			? LOCTEXT("ConstantConsideration", "<s>Constant</>")
+			? LOCTEXT("ConstantConsiderationRich", "<s>Constant</>")
 			: LOCTEXT("ConstantConsideration", "Constant");
 
 		RawScoreOrPropertyPathText = FText::AsNumber(InstanceData->RawScore);
