State Tree: added instance data mismatch auto-fix

diff --git a/Engine/Plugins/Runtime/StateTree/Source/StateTreeEditorModule/Private/Customizations/StateTreeEditorNodeDetails.cpp b/Engine/Plugins/Runtime/StateTree/Source/StateTreeEditorModule/Private/Customizations/StateTreeEditorNodeDetails.cpp
--- a/Engine/Plugins/Runtime/StateTree/Source/StateTreeEditorModule/Private/Customizations/StateTreeEditorNodeDetails.cpp
+++ b/Engine/Plugins/Runtime/StateTree/Source/StateTreeEditorModule/Private/Customizations/StateTreeEditorNodeDetails.cpp
@@ -63,44 +63,6 @@ public:
 
 namespace UE::StateTreeEditor::Internal
 {
-	/**
-	 * This function recursively instantiates instanced objects of a given struct.
-	 * It is needed to fixup nodes pasted from clipboard, which seem to give shallow copy.
-	 */
-	void InstantiateStructSubobjects(UObject& OuterObject, FStructView Struct)
-	{
-		// Empty struct, nothing to do.
-		if (!Struct.IsValid())
-		{
-			return;
-		}
-
-		for (TPropertyValueIterator<FProperty> It(Struct.GetScriptStruct(), Struct.GetMemory()); It; ++It)
-		{
-			if (const FObjectProperty* ObjectProperty = CastField<FObjectProperty>(It->Key))
-			{
-				// Duplicate instanced objects.
-				if (ObjectProperty->HasAnyPropertyFlags(CPF_InstancedReference | CPF_PersistentInstance))
-				{
-					if (UObject* Object = ObjectProperty->GetObjectPropertyValue(It->Value))
-					{
-						UObject* DuplicatedObject = DuplicateObject(Object, &OuterObject);
-						ObjectProperty->SetObjectPropertyValue(const_cast<void*>(It->Value), DuplicatedObject);
-					}
-				}
-			}
-			if (const FStructProperty* StructProperty = CastField<FStructProperty>(It->Key))
-			{
-				// If we encounter instanced struct, recursively handle it too.
-				if (StructProperty->Struct == TBaseStructure<FInstancedStruct>::Get())
-				{
-					FInstancedStruct& InstancedStruct = *static_cast<FInstancedStruct*>(const_cast<void*>(It->Value));
-					InstantiateStructSubobjects(OuterObject, InstancedStruct);
-				}
-			}
-		}
-	}
-
 	/** @return text describing the pin type, matches SPinTypeSelector. */
 	FText GetPinTypeText(const FEdGraphPinType& PinType)
 	{
@@ -770,14 +732,14 @@ void FStateTreeEditorNodeDetails::OnPasteNode()
 				*EditorNode = TempNode;
 
 				// Ensure unique instance value
-				UE::StateTreeEditor::Internal::InstantiateStructSubobjects(*OuterObject, EditorNode->Node);
+				UE::StateTreeEditor::EditorNodeUtils::InstantiateStructSubobjects(*OuterObject, EditorNode->Node);
 				if (EditorNode->InstanceObject)
 				{
 					EditorNode->InstanceObject = DuplicateObject(EditorNode->InstanceObject, OuterObject);
 				}
 				else
 				{
-					UE::StateTreeEditor::Internal::InstantiateStructSubobjects(*OuterObject, EditorNode->Instance);
+					UE::StateTreeEditor::EditorNodeUtils::InstantiateStructSubobjects(*OuterObject, EditorNode->Instance);
 				}
 				
 				const FGuid OldStructID = EditorNode->ID; 
diff --git a/Engine/Plugins/Runtime/StateTree/Source/StateTreeEditorModule/Private/Customizations/StateTreeEditorNodeUtils.cpp b/Engine/Plugins/Runtime/StateTree/Source/StateTreeEditorModule/Private/Customizations/StateTreeEditorNodeUtils.cpp
--- a/Engine/Plugins/Runtime/StateTree/Source/StateTreeEditorModule/Private/Customizations/StateTreeEditorNodeUtils.cpp
+++ b/Engine/Plugins/Runtime/StateTree/Source/StateTreeEditorModule/Private/Customizations/StateTreeEditorNodeUtils.cpp
@@ -555,6 +555,84 @@ void SetNodeType(const TSharedPtr<IPropertyHandle>& StructProperty, const UStruc
 	}	
 }
 
+void InstantiateStructSubobjects(UObject& OuterObject, FStructView Struct)
+{
+	// Empty struct, nothing to do.
+	if (!Struct.IsValid())
+	{
+		return;
+	}
+
+	for (TPropertyValueIterator<FProperty> It(Struct.GetScriptStruct(), Struct.GetMemory()); It; ++It)
+	{
+		if (const FObjectProperty* ObjectProperty = CastField<FObjectProperty>(It->Key))
+		{
+			// Duplicate instanced objects.
+			if (ObjectProperty->HasAnyPropertyFlags(CPF_InstancedReference | CPF_PersistentInstance))
+			{
+				if (UObject* Object = ObjectProperty->GetObjectPropertyValue(It->Value))
+				{
+					UObject* DuplicatedObject = DuplicateObject(Object, &OuterObject);
+					ObjectProperty->SetObjectPropertyValue(const_cast<void*>(It->Value), DuplicatedObject);
+				}
+			}
+		}
+		if (const FStructProperty* StructProperty = CastField<FStructProperty>(It->Key))
+		{
+			// If we encounter instanced struct, recursively handle it too.
+			if (StructProperty->Struct == TBaseStructure<FInstancedStruct>::Get())
+			{
+				FInstancedStruct& InstancedStruct = *static_cast<FInstancedStruct*>(const_cast<void*>(It->Value));
+				InstantiateStructSubobjects(OuterObject, InstancedStruct);
+			}
+		}
+	}
+}
+
+void ConditionalUpdateNodeInstanceData(FStateTreeEditorNode& EditorNode, UObject& InstanceOuter)
+{
+	const FStateTreeNodeBase* Node = EditorNode.Node.GetPtr<FStateTreeNodeBase>();
+	if (!Node)
+	{
+		return;
+	}
+
+	const UStruct* CurrentType = EditorNode.GetInstance().GetStruct();
+	const UStruct* DesiredType = Node->GetInstanceDataType();
+
+	// Nothing to upgrade. Instance Data Type is unchanged
+	if (CurrentType == DesiredType)
+	{
+		return;
+	}
+
+	FStateTreeEditorNode OldEditorNode = EditorNode;
+
+	EditorNode.Instance.Reset();
+	EditorNode.InstanceObject = nullptr;
+
+	if (const UScriptStruct* InstanceType = Cast<UScriptStruct>(DesiredType))
+	{
+		EditorNode.Instance.InitializeAs(InstanceType);
+	}
+	else if (const UClass* InstanceClass = Cast<UClass>(DesiredType))
+	{
+		EditorNode.InstanceObject = NewObject<UObject>(&InstanceOuter, InstanceClass);
+	}
+
+	RetainProperties(OldEditorNode, EditorNode);
+
+	// Ensure that the instanced objects on the nodes are correctly copied over (deep copy)
+	UE::StateTreeEditor::EditorNodeUtils::InstantiateStructSubobjects(InstanceOuter, EditorNode.Node);
+	if (EditorNode.InstanceObject)
+	{
+		EditorNode.InstanceObject = DuplicateObject(EditorNode.InstanceObject, &InstanceOuter);
+	}
+	else
+	{
+		UE::StateTreeEditor::EditorNodeUtils::InstantiateStructSubobjects(InstanceOuter, EditorNode.Instance);
+	}
+}
 
 void OnArrayNodePicked(const UStruct* InStruct, TSharedPtr<SComboButton> PickerCombo, TSharedPtr<IPropertyHandle> ArrayPropertyHandle, TSharedRef<IPropertyUtilities> PropUtils)
 {
diff --git a/Engine/Plugins/Runtime/StateTree/Source/StateTreeEditorModule/Private/Customizations/StateTreeEditorNodeUtils.h b/Engine/Plugins/Runtime/StateTree/Source/StateTreeEditorModule/Private/Customizations/StateTreeEditorNodeUtils.h
--- a/Engine/Plugins/Runtime/StateTree/Source/StateTreeEditorModule/Private/Customizations/StateTreeEditorNodeUtils.h
+++ b/Engine/Plugins/Runtime/StateTree/Source/StateTreeEditorModule/Private/Customizations/StateTreeEditorNodeUtils.h
@@ -8,6 +8,7 @@
 
 enum class EStateTreeConditionEvaluationMode : uint8;
 struct FStateTreeEditorNode;
+struct FStructView;
 struct EVisibility;
 class FText;
 class IPropertyHandle;
@@ -109,6 +110,21 @@ EVisibility IsIconVisible(const TSharedPtr<IPropertyHandle>& StructProperty);
  */
 void SetNodeType(const TSharedPtr<IPropertyHandle>& StructProperty, const UStruct* NewType);
 
+/**
+ * Recursively instantiates instanced objects of a given struct.
+ * It is needed to fixup nodes pasted from clipboard, which seem to give shallow copy.
+ * @param OuterObject the object the instanced objects should be outered to
+ * @param Struct the struct with the instanced objects
+ */
+void InstantiateStructSubobjects(UObject& OuterObject, FStructView Struct);
+
+/**
+ * Handles updating the Node Instance Data if there is a type mismatch
+ * @param EditorNode the node to check
+ * @param InstanceOuter the outer to use if the instance data is a uobject, or there are instanced uobjects within the instance data
+ */
+void ConditionalUpdateNodeInstanceData(FStateTreeEditorNode& EditorNode, UObject& InstanceOuter);
+
 /**
  * Creates widget combon button with plus icon (+), which summons node picker and adds the selected node to specified array. 
  * @param TooltipText Tool tip to show on he button.
diff --git a/Engine/Plugins/Runtime/StateTree/Source/StateTreeEditorModule/Private/StateTreeEditorData.cpp b/Engine/Plugins/Runtime/StateTree/Source/StateTreeEditorModule/Private/StateTreeEditorData.cpp
--- a/Engine/Plugins/Runtime/StateTree/Source/StateTreeEditorModule/Private/StateTreeEditorData.cpp
+++ b/Engine/Plugins/Runtime/StateTree/Source/StateTreeEditorModule/Private/StateTreeEditorData.cpp
@@ -10,6 +10,7 @@
 #include "Algo/LevenshteinDistance.h"
 #include "StateTreeEditorModule.h"
 #include "StateTreePropertyHelpers.h"
+#include "Customizations/StateTreeEditorNodeUtils.h"
 
 #if WITH_EDITOR
 #include "Engine/UserDefinedStruct.h"
@@ -820,6 +821,7 @@ void UStateTreeEditorData::CallPostLoadOnNodes()
 	{
 		if (FStateTreeNodeBase* EvaluatorNode = EvaluatorEditorNode.Node.GetMutablePtr<FStateTreeNodeBase>())
 		{
+			UE::StateTreeEditor::EditorNodeUtils::ConditionalUpdateNodeInstanceData(EvaluatorEditorNode, *this);
 			EvaluatorNode->PostLoad(EvaluatorEditorNode.GetInstance());
 		}
 	}
@@ -828,6 +830,7 @@ void UStateTreeEditorData::CallPostLoadOnNodes()
 	{
 		if (FStateTreeNodeBase* GlobalTaskNode = GlobalTaskEditorNode.Node.GetMutablePtr<FStateTreeNodeBase>())
 		{
+			UE::StateTreeEditor::EditorNodeUtils::ConditionalUpdateNodeInstanceData(GlobalTaskEditorNode, *this);
 			GlobalTaskNode->PostLoad(GlobalTaskEditorNode.GetInstance());
 		}
 	}
diff --git a/Engine/Plugins/Runtime/StateTree/Source/StateTreeEditorModule/Private/StateTreeState.cpp b/Engine/Plugins/Runtime/StateTree/Source/StateTreeEditorModule/Private/StateTreeState.cpp
--- a/Engine/Plugins/Runtime/StateTree/Source/StateTreeEditorModule/Private/StateTreeState.cpp
+++ b/Engine/Plugins/Runtime/StateTree/Source/StateTreeEditorModule/Private/StateTreeState.cpp
@@ -8,6 +8,7 @@
 #include "StateTreeTaskBase.h"
 #include "StateTreeDelegates.h"
 #include "StateTreePropertyHelpers.h"
+#include "Customizations/StateTreeEditorNodeUtils.h"
 
 #include UE_INLINE_GENERATED_CPP_BY_NAME(StateTreeState)
 
@@ -405,6 +406,7 @@ void UStateTreeState::PostLoad()
 	{
 		if (FStateTreeNodeBase* ConditionNode = EnterConditionEditorNode.Node.GetMutablePtr<FStateTreeNodeBase>())
 		{
+			UE::StateTreeEditor::EditorNodeUtils::ConditionalUpdateNodeInstanceData(EnterConditionEditorNode, *this);
 			ConditionNode->PostLoad(EnterConditionEditorNode.GetInstance());
 		}
 	}
@@ -413,6 +415,7 @@ void UStateTreeState::PostLoad()
 	{
 		if (FStateTreeNodeBase* ConsiderationNode = ConsiderationEditorNode.Node.GetMutablePtr<FStateTreeNodeBase>())
 		{
+			UE::StateTreeEditor::EditorNodeUtils::ConditionalUpdateNodeInstanceData(ConsiderationEditorNode, *this);
 			ConsiderationNode->PostLoad(ConsiderationEditorNode.GetInstance());
 		}
 	}
@@ -421,16 +424,29 @@ void UStateTreeState::PostLoad()
 	{
 		if (FStateTreeNodeBase* TaskNode = TaskEditorNode.Node.GetMutablePtr<FStateTreeNodeBase>())
 		{
+			UE::StateTreeEditor::EditorNodeUtils::ConditionalUpdateNodeInstanceData(TaskEditorNode, *this);
 			TaskNode->PostLoad(TaskEditorNode.GetInstance());
 		}
 	}
 
 	if (FStateTreeNodeBase* SingleTaskNode = SingleTask.Node.GetMutablePtr<FStateTreeNodeBase>())
 	{
+		UE::StateTreeEditor::EditorNodeUtils::ConditionalUpdateNodeInstanceData(SingleTask, *this);
 		SingleTaskNode->PostLoad(SingleTask.GetInstance());
 	}
-#endif // WITH_EDITOR
 
+	for (FStateTreeTransition& Transition : Transitions)
+	{
+		for (FStateTreeEditorNode& TransitionConditionEditorNode : Transition.Conditions)
+		{
+			if (FStateTreeNodeBase* ConditionNode = TransitionConditionEditorNode.Node.GetMutablePtr<FStateTreeNodeBase>())
+			{
+				UE::StateTreeEditor::EditorNodeUtils::ConditionalUpdateNodeInstanceData(TransitionConditionEditorNode, *this);
+				ConditionNode->PostLoad(TransitionConditionEditorNode.GetInstance());
+			}
+		}
+	}
+#endif // WITH_EDITOR
 }
 
 void UStateTreeState::UpdateParametersFromLinkedSubtree()
diff --git a/Engine/Plugins/Runtime/StateTree/Source/StateTreeModule/Private/StateTree.cpp b/Engine/Plugins/Runtime/StateTree/Source/StateTreeModule/Private/StateTree.cpp
--- a/Engine/Plugins/Runtime/StateTree/Source/StateTreeModule/Private/StateTree.cpp
+++ b/Engine/Plugins/Runtime/StateTree/Source/StateTreeModule/Private/StateTree.cpp
@@ -6,6 +6,7 @@
 #include "StateTreeTaskBase.h"
 #include "StateTreeEvaluatorBase.h"
 #include "StateTreeConditionBase.h"
+#include "StateTreeConsiderationBase.h"
 #include "AssetRegistry/AssetData.h"
 #include "Misc/ScopeRWLock.h"
 #include "StateTreeDelegates.h"
@@ -358,7 +359,10 @@ void UStateTree::PostLoad()
 		{
 			if (Node->InstanceTemplateIndex.IsValid())
 			{
-				FStateTreeInstanceData& SourceInstanceData = NodeView.GetPtr<FStateTreeConditionBase>() != nullptr ? SharedInstanceData : DefaultInstanceData;
+				const bool bUseSharedInstanceData = NodeView.GetPtr<FStateTreeConditionBase>() || NodeView.GetPtr<FStateTreeConsiderationBase>();
+
+				FStateTreeInstanceData& SourceInstanceData = bUseSharedInstanceData ? SharedInstanceData : DefaultInstanceData;
+
 				if (SourceInstanceData.IsObject(Node->InstanceTemplateIndex.Get()))
 				{
 					Node->PostLoad(SourceInstanceData.GetMutableObject(Node->InstanceTemplateIndex.Get()));
@@ -455,12 +459,56 @@ void UStateTree::ResetLinked()
 	PerThreadSharedInstanceData.Reset();
 }
 
+bool UStateTree::ValidateInstanceData() const
+{
+	for (FConstStructView NodeView : Nodes)
+	{
+		const FStateTreeNodeBase* Node = NodeView.GetPtr<const FStateTreeNodeBase>();
+		if (Node && Node->InstanceTemplateIndex.IsValid())
+		{
+			const UStruct* CurrentInstanceDataType;
+			{
+				const bool bUseSharedInstanceData = NodeView.GetPtr<const FStateTreeConditionBase>() || NodeView.GetPtr<const FStateTreeConsiderationBase>();
+
+				const FStateTreeInstanceData& SourceInstanceData = bUseSharedInstanceData ? SharedInstanceData : DefaultInstanceData;
+
+				if (SourceInstanceData.IsObject(Node->InstanceTemplateIndex.Get()))
+				{
+					const UObject* InstanceObject = SourceInstanceData.GetObject(Node->InstanceTemplateIndex.Get());
+					CurrentInstanceDataType = InstanceObject ? InstanceObject->GetClass() : nullptr;
+				}
+				else
+				{
+					CurrentInstanceDataType = SourceInstanceData.GetStruct(Node->InstanceTemplateIndex.Get()).GetScriptStruct();
+				}
+			}
+
+			const UStruct* DesiredInstanceDataType = Node->GetInstanceDataType();
+
+			// Use strict testing so that the users will have option to initialize data mismatch if the type changes (even if potentially compatible).
+			if (CurrentInstanceDataType != DesiredInstanceDataType)
+			{
+				UE_LOG(LogStateTree, Error, TEXT("%s: node '%s' failed. The source Instance Data type '%s' does not match '%s'"), *GetFullName(), *Node->StaticStruct()->GetName(), *GetNameSafe(CurrentInstanceDataType), *GetNameSafe(DesiredInstanceDataType));
+				return false;
+			}
+		}
+	}
+
+	return true;
+}
+
 bool UStateTree::Link()
 {
 	// Initialize the instance data default value.
 	// This data will be used to allocate runtime instance on all StateTree users.
 	ResetLinked();
 
+	// Validate that all the source instance data types matches the node instance data types
+	if (!ValidateInstanceData())
+	{
+		return false;
+	}
+
 	// Resolves nodes references to other StateTree data
 	FStateTreeLinker Linker(Schema);
 
@@ -696,9 +744,9 @@ bool UStateTree::PatchBindings()
 	for (FConstStructView NodeView : Nodes)
 	{
 		const FStateTreeNodeBase& Node = NodeView.Get<const FStateTreeNodeBase>();
-		
+
 		FStateTreeInstanceData* SourceInstanceData = &DefaultInstanceData;
-		if (NodeView.GetPtr<const FStateTreeConditionBase>())
+		if (NodeView.GetPtr<const FStateTreeConditionBase>() || NodeView.GetPtr<const FStateTreeConsiderationBase>())
 		{
 			// Conditions are stored in shared instance data.
 			SourceInstanceData = &SharedInstanceData;
diff --git a/Engine/Plugins/Runtime/StateTree/Source/StateTreeModule/Public/StateTree.h b/Engine/Plugins/Runtime/StateTree/Source/StateTreeModule/Public/StateTree.h
--- a/Engine/Plugins/Runtime/StateTree/Source/StateTreeModule/Public/StateTree.h
+++ b/Engine/Plugins/Runtime/StateTree/Source/StateTreeModule/Public/StateTree.h
@@ -232,6 +232,9 @@ private:
 	 */
 	void ResetLinked();
 
+	/** @return true if all the source instance data types match with the node's instance data types */
+	bool ValidateInstanceData() const;
+
 	bool PatchBindings();
 
 	// Data created during compilation, source data in EditorData.
