StateTree: Introduced PropertyFunctions (contributed)
- Added property functions which can be executed and chained during property binding to set values of properties

diff --git a/Engine/Plugins/Experimental/Avalanche/Source/AvalancheTransitionEditor/Private/Serialization/AvaTransitionStateSerializer.cpp b/Engine/Plugins/Experimental/Avalanche/Source/AvalancheTransitionEditor/Private/Serialization/AvaTransitionStateSerializer.cpp
--- a/Engine/Plugins/Experimental/Avalanche/Source/AvalancheTransitionEditor/Private/Serialization/AvaTransitionStateSerializer.cpp
+++ b/Engine/Plugins/Experimental/Avalanche/Source/AvalancheTransitionEditor/Private/Serialization/AvaTransitionStateSerializer.cpp
@@ -56,9 +56,12 @@ namespace UE::AvaTransitionEditor::Private
 		InEditorData.VisitStateNodes(*InState,
 			[&InEditorData, &InCopyData](const UStateTreeState*, const FStateTreeBindableStructDesc& InDesc, const FStateTreeDataView)
 			{
-				TArray<FStateTreePropertyPathBinding> NodeBindings;
+				TArray<const FStateTreePropertyPathBinding*> NodeBindings;
 				InEditorData.GetPropertyEditorBindings()->GetPropertyBindingsFor(InDesc.ID, NodeBindings);
-				InCopyData.Bindings.Append(NodeBindings);
+				for (const FStateTreePropertyPathBinding* Binding : NodeBindings)
+				{
+					InCopyData.Bindings.Add(*Binding);
+				}
 				return EStateTreeVisitor::Continue;				
 			});
 
diff --git a/Engine/Plugins/Runtime/GameplayStateTree/Source/GameplayStateTreeModule/Private/Components/StateTreeComponentSchema.cpp b/Engine/Plugins/Runtime/GameplayStateTree/Source/GameplayStateTreeModule/Private/Components/StateTreeComponentSchema.cpp
--- a/Engine/Plugins/Runtime/GameplayStateTree/Source/GameplayStateTreeModule/Private/Components/StateTreeComponentSchema.cpp
+++ b/Engine/Plugins/Runtime/GameplayStateTree/Source/GameplayStateTreeModule/Private/Components/StateTreeComponentSchema.cpp
@@ -13,6 +13,7 @@
 #include "Subsystems/WorldSubsystem.h"
 #include "Tasks/StateTreeAITask.h"
 #include "VisualLogger/VisualLogger.h"
+#include "StateTreePropertyFunctionBase.h"
 
 UStateTreeComponentSchema::UStateTreeComponentSchema()
 	: ContextActorClass(AActor::StaticClass())
@@ -25,7 +26,8 @@ bool UStateTreeComponentSchema::IsStructAllowed(const UScriptStruct* InScriptStr
 	return InScriptStruct->IsChildOf(FStateTreeConditionCommonBase::StaticStruct())
 	|| InScriptStruct->IsChildOf(FStateTreeEvaluatorCommonBase::StaticStruct())
 	|| InScriptStruct->IsChildOf(FStateTreeTaskCommonBase::StaticStruct())
-	|| InScriptStruct->IsChildOf(FStateTreeConsiderationCommonBase::StaticStruct());
+	|| InScriptStruct->IsChildOf(FStateTreeConsiderationCommonBase::StaticStruct())
+	|| InScriptStruct->IsChildOf(FStateTreePropertyFunctionCommonBase::StaticStruct());
 }
 
 bool UStateTreeComponentSchema::IsClassAllowed(const UClass* InClass) const
diff --git a/Engine/Plugins/Runtime/StateTree/Source/StateTreeEditorModule/Private/Customizations/StateTreeBindingExtension.cpp b/Engine/Plugins/Runtime/StateTree/Source/StateTreeEditorModule/Private/Customizations/StateTreeBindingExtension.cpp
--- a/Engine/Plugins/Runtime/StateTree/Source/StateTreeEditorModule/Private/Customizations/StateTreeBindingExtension.cpp
+++ b/Engine/Plugins/Runtime/StateTree/Source/StateTreeEditorModule/Private/Customizations/StateTreeBindingExtension.cpp
@@ -13,6 +13,15 @@
 #include "Widgets/Layout/SBox.h"
 #include "StateTreePropertyRef.h"
 #include "StateTreePropertyRefHelpers.h"
+#include "DetailLayoutBuilder.h"
+#include "IPropertyUtilities.h"
+#include "IDetailChildrenBuilder.h"
+#include "IStructureDataProvider.h"
+#include "StateTreeEditorNodeUtils.h"
+#include "StateTreeEditorModule.h"
+#include "StateTreeNodeClassCache.h"
+#include "StateTreePropertyFunctionBase.h"
+#include "StateTreeEditorData.h"
 
 #define LOCTEXT_NAMESPACE "StateTreeEditor"
 
@@ -208,11 +217,6 @@ EStateTreePropertyUsage MakeStructPropertyPathFromPropertyHandle(TSharedPtr<cons
 	return ResultUsage;
 }
 
-const FStateTreeBindableStructDesc* FindStruct(TConstArrayView<FStateTreeBindableStructDesc> AccessibleStructs, const FGuid StructID)
-{
-	return AccessibleStructs.FindByPredicate([StructID](const FStateTreeBindableStructDesc& Desc) { return Desc.ID == StructID; });
-}
-
 // @todo: there's a similar function in StateTreeNodeDetails.cpp, merge.
 FText GetPropertyTypeText(const FProperty* Property)
 {
@@ -279,18 +283,56 @@ struct FCachedBindingData : public TSharedFromThis<FCachedBindingData>
 		// First item in the binding chain is the index in AccessibleStructs.
 		const int32 SourceStructIndex = InBindingChain[0].ArrayIndex;
 		check(SourceStructIndex >= 0 && SourceStructIndex < AccessibleStructs.Num());
+
+		const FStateTreeBindableStructDesc& BindableStruct = AccessibleStructs[SourceStructIndex];
 				
 		TConstArrayView<FBindingChainElement> SourceBindingChain(InBindingChain.begin() + 1, InBindingChain.Num() - 1); // remove struct index.
 
 		FStateTreeDataView DataView;
-		BindingOwner->GetDataViewByID(AccessibleStructs[SourceStructIndex].ID, DataView);
+		BindingOwner->GetDataViewByID(BindableStruct.ID, DataView);
 
 		// If SourceBindingChain is empty at this stage, it means that the binding points to the source struct itself.
 		FStateTreePropertyPath SourcePath;
-		UE::StateTree::PropertyBinding::MakeStructPropertyPathFromBindingChain(AccessibleStructs[SourceStructIndex].ID, SourceBindingChain, DataView, SourcePath);
-				
+		UE::StateTree::PropertyBinding::MakeStructPropertyPathFromBindingChain(BindableStruct.ID, SourceBindingChain, DataView, SourcePath);
+
 		OwnerObject->Modify();
-		EditorBindings->AddPropertyBinding(SourcePath, TargetPath);
+
+		if (BindableStruct.DataSource == EStateTreeBindableStructSource::PropertyFunction)
+		{
+			const UScriptStruct* PropertyFunctionNodeStruct = nullptr;
+
+			BindingOwner->EnumerateBindablePropertyFunctionNodes([ID = BindableStruct.ID, &PropertyFunctionNodeStruct](const UScriptStruct* NodeStruct, const FStateTreeBindableStructDesc& Desc, const FStateTreeDataView Value)
+			{
+				if (Desc.ID == ID)
+				{
+					PropertyFunctionNodeStruct = NodeStruct;
+					return EStateTreeVisitor::Break;
+				}
+
+				return EStateTreeVisitor::Continue;
+			});
+
+			if (ensure(PropertyFunctionNodeStruct))
+			{
+				// If there are no segments, bindings leads directly into source struct's single output property. It's path has to be recovered.
+				if (SourcePath.NumSegments() == 0)
+				{
+					const FProperty* SingleOutputProperty = UE::StateTree::GetStructSingleOutputProperty(*BindableStruct.Struct);
+					check(SingleOutputProperty);
+
+					FStateTreePropertyPathSegment SingleOutputPropertySegment = FStateTreePropertyPathSegment(SingleOutputProperty->GetFName());
+					SourcePath = EditorBindings->AddFunctionPropertyBinding(PropertyFunctionNodeStruct, {SingleOutputPropertySegment}, TargetPath);
+				}
+				else
+				{
+					SourcePath = EditorBindings->AddFunctionPropertyBinding(PropertyFunctionNodeStruct, SourcePath.GetSegments(), TargetPath);
+				}
+			}
+		}
+		else
+		{
+			EditorBindings->AddPropertyBinding(SourcePath, TargetPath);
+		}
 
 		UpdateData();
 
@@ -331,8 +373,9 @@ struct FCachedBindingData : public TSharedFromThis<FCachedBindingData>
 	{
 		static FName PropertyIcon(TEXT("Kismet.Tabs.Variables"));
 
-		Text = FText::GetEmpty();
-		TooltipText = FText::GetEmpty();
+		SourceStructName = FText::GetEmpty();
+		FormatableText = FText::GetEmpty();
+		FormatableTooltipText = FText::GetEmpty();
 		Color = FLinearColor::White;
 		Image = nullptr;
 
@@ -390,18 +433,32 @@ struct FCachedBindingData : public TSharedFromThis<FCachedBindingData>
 
 
 		FTextBuilder TooltipBuilder;
-		if (const FStateTreePropertyPath* SourcePath = EditorBindings->GetPropertyBindingSource(TargetPath))
+		const FStateTreePropertyPathBinding* CurrentBinding = EditorBindings->GetBindings().FindByPredicate([this](const FStateTreePropertyPathBinding& Binding)
 		{
-			const FStateTreeBindableStructDesc* SourceDesc = UE::StateTree::PropertyBinding::FindStruct(AccessibleStructs, SourcePath->GetStructID());
-			if (SourceDesc)
+			return Binding.GetTargetPath() == TargetPath;
+		});
+
+		if (CurrentBinding)
+		{
+			const FStateTreePropertyPath& SourcePath = CurrentBinding->GetSourcePath();
+			FString SourcePropertyPathAsString = SourcePath.ToString();
+
+			// If source is a bound PropertyFunction, it will not be present in AccessibleStructs thus it has to be accessed through bindings owner.
+			FStateTreeBindableStructDesc SourceDesc;
+			if (BindingOwner->GetStructByID(SourcePath.GetStructID(), SourceDesc))
 			{
-				bool bIsValidBinding = false;
+				// Making first segment of the path invisible for the user if it's struct's single output property.
+				if (UE::StateTree::GetStructSingleOutputProperty(*SourceDesc.Struct))
+				{
+					SourcePropertyPathAsString = SourcePath.ToString(/*HighlightedSegment*/ INDEX_NONE, /*HighlightPrefix*/ nullptr, /*HighlightPostfix*/ nullptr, /*bOutputInstances*/ false, 1);
+				}
 
 				// Check that the binding is valid.
+				bool bIsValidBinding = false;
 				FStateTreeDataView SourceDataView;
 				const FProperty* SourceLeafProperty = nullptr;
 				const UStruct* SourceStruct = nullptr;
-				if (BindingOwner->GetDataViewByID(SourcePath->GetStructID(), SourceDataView)
+				if (BindingOwner->GetDataViewByID(SourcePath.GetStructID(), SourceDataView)
 					&& TargetDataView.IsValid())
 				{
 					TArray<FStateTreePropertyPathIndirection> SourceIndirections;
@@ -410,7 +467,7 @@ struct FCachedBindingData : public TSharedFromThis<FCachedBindingData>
 					// Resolve source and target properties.
 					// Source path can be empty, when the binding binds directly to a context struct/class.
 					// Target path must always point to a valid property (at least one indirection).
-					if (SourcePath->ResolveIndirectionsWithValue(SourceDataView, SourceIndirections)
+					if (SourcePath.ResolveIndirectionsWithValue(SourceDataView, SourceIndirections)
 						&& TargetPath.ResolveIndirectionsWithValue(TargetDataView, TargetIndirections)
 						&& !TargetIndirections.IsEmpty())
 					{
@@ -431,24 +488,32 @@ struct FCachedBindingData : public TSharedFromThis<FCachedBindingData>
 					}
 				}
 
-				FString SourcePropertyName;
-				SourcePropertyName += SourceDesc->Name.ToString();
-				if (!SourcePath->IsPathEmpty())
-				{
-					SourcePropertyName += TEXT(".") + SourcePath->ToString();
-				}
+				FormatableText = FText::FormatNamed(LOCTEXT("ValidSourcePath", "{SourceStruct}{PropertyPath}"), TEXT("PropertyPath"), SourcePropertyPathAsString.IsEmpty() ? FText() : FText::FromString(TEXT(".") + SourcePropertyPathAsString));
+				SourceStructName = FText::FromString(SourceDesc.Name.ToString());
 
 				if (bIsValidBinding)
 				{
-					Text = FText::FromString(SourcePropertyName);
-
-					if (SourcePath->IsPathEmpty())
+					if (SourcePropertyPathAsString.IsEmpty())
 					{
-						TooltipBuilder.AppendLineFormat(LOCTEXT("ExistingBindingTooltip", "Property is bound to {0}."), FText::FromString(SourceDesc->ToString()));
+						if (CurrentBinding->GetPropertyFunctionNode().IsValid())
+						{
+							TooltipBuilder.AppendLine(LOCTEXT("ExistingBindingToFunctionTooltip", "Property is bound to function {SourceStruct}."));
+						}
+						else
+						{
+							TooltipBuilder.AppendLine(LOCTEXT("ExistingBindingTooltip", "Property is bound to {SourceStruct}."));
+						}
 					}
 					else
 					{
-						TooltipBuilder.AppendLineFormat(LOCTEXT("ExistingBindingWithPropertyTooltip", "Property is bound to {0} property {1}."), FText::FromString(SourceDesc->ToString()), FText::FromString(SourcePath->ToString()));
+						if (CurrentBinding->GetPropertyFunctionNode().IsValid())
+						{
+							TooltipBuilder.AppendLineFormat(LOCTEXT("ExistingBindingToFunctionWithPropertyTooltip", "Property is bound to function {SourceStruct} property {PropertyPath}."), {{TEXT("PropertyPath"), FText::FromString(SourcePropertyPathAsString)}});
+						}
+						else
+						{
+							TooltipBuilder.AppendLineFormat(LOCTEXT("ExistingBindingWithPropertyTooltip", "Property is bound to {SourceStruct} property {PropertyPath}."), {{TEXT("PropertyPath"), FText::FromString(SourcePropertyPathAsString)}});
+						}
 					}
 
 					if (bIsPropertyRef) // Update the pin type with source property so that property ref that can binds to multiple types display the binded one.
@@ -471,18 +536,23 @@ struct FCachedBindingData : public TSharedFromThis<FCachedBindingData>
 						SourceType = SourceStruct->GetDisplayNameText();
 					}
 					FText TargetType = UE::StateTree::PropertyBinding::GetPropertyTypeText(Property);
-					
-					Text = FText::FromString(SourcePropertyName);
 
-					if (SourcePath->IsPathEmpty())
+					if (SourcePath.IsPathEmpty())
 					{
-						TooltipBuilder.AppendLineFormat(LOCTEXT("MismatchingBindingTooltip", "Property is bound to {0}, but binding source type '{1}' does not match property type '{2}'."),
-							FText::FromString(SourceDesc->ToString()), SourceType, TargetType);
+						TooltipBuilder.AppendLineFormat(LOCTEXT("MismatchingBindingTooltip", "Property is bound to {SourceStruct}, but binding source type '{SourceType}' does not match property type '{TargetType}'."),
+							{
+								{TEXT("SourceType"), SourceType},
+								{TEXT("TargetType"), TargetType}
+							});
 					}
 					else
 					{
-						TooltipBuilder.AppendLineFormat(LOCTEXT("MismatchingBindingTooltipWithProperty", "Property is bound to {0} property {1}, but binding source type '{2}' does not match property type '{3}'."),
-							FText::FromString(SourceDesc->ToString()), FText::FromString(SourcePath->ToString()), SourceType, TargetType);
+						TooltipBuilder.AppendLineFormat(LOCTEXT("MismatchingBindingTooltipWithProperty", "Property is bound to {SourceStruct} property {PropertyPath}, but binding source type '{SourceType}' does not match property type '{TargetType}'."),
+							{
+								{TEXT("PropertyPath"), FText::FromString(SourcePropertyPathAsString)},
+								{TEXT("SourceType"), SourceType},
+								{TEXT("TargetType"), TargetType}
+							});
 					}
 
 					Image = FCoreStyle::Get().GetBrush("Icons.ErrorWithColor");
@@ -492,18 +562,18 @@ struct FCachedBindingData : public TSharedFromThis<FCachedBindingData>
 			else
 			{
 				// Missing source
-				Text = FText::Format(LOCTEXT("MissingSource", "???.{0}"), FText::FromString(SourcePath->ToString()));
-				TooltipBuilder.AppendLineFormat(LOCTEXT("MissingBindingTooltip", "Missing binding source for property path '{0}'."), FText::FromString(SourcePath->ToString()));
+				FormatableText = FText::Format(LOCTEXT("MissingSource", "???.{0}"), FText::FromString(SourcePropertyPathAsString));
+				TooltipBuilder.AppendLineFormat(LOCTEXT("MissingBindingTooltip", "Missing binding source for property path '{0}'."), FText::FromString(SourcePropertyPathAsString));
 				Image = FCoreStyle::Get().GetBrush("Icons.ErrorWithColor");
 				Color = FLinearColor::White;
 			}
 
-			CachedSourcePath = *SourcePath;
+			CachedSourcePath = SourcePath;
 		}
 		else
 		{
 			// No bindings
-			Text = FText::GetEmpty();
+			FormatableText = FText::GetEmpty();
 			TooltipBuilder.AppendLineFormat(LOCTEXT("BindTooltip", "Bind {0} to value from another property."), UE::StateTree::PropertyBinding::GetPropertyTypeText(Property));
 
 			Image = FAppStyle::GetBrush(PropertyIcon);
@@ -528,12 +598,12 @@ struct FCachedBindingData : public TSharedFromThis<FCachedBindingData>
 			}
 		}
 
-		TooltipText = TooltipBuilder.ToText();
+		FormatableTooltipText = TooltipBuilder.ToText();
 
 		bIsDataCached = true;
 	}
 
-	bool CanBindToContextStruct(const UStruct* InStruct)
+	bool CanBindToContextStruct(const UStruct* InStruct, int32 InStructIndex)
 	{
 		ConditionallyUpdateData();
 
@@ -547,6 +617,7 @@ struct FCachedBindingData : public TSharedFromThis<FCachedBindingData>
 					&& AccessibleStruct.DataSource != EStateTreeBindableStructSource::Parameter
 					&& AccessibleStruct.DataSource != EStateTreeBindableStructSource::TransitionEvent
 					&& AccessibleStruct.DataSource != EStateTreeBindableStructSource::StateEvent
+					&& AccessibleStruct.DataSource != EStateTreeBindableStructSource::PropertyFunction
 					&& AccessibleStruct.Struct == InStruct;
 			});
 
@@ -556,6 +627,24 @@ struct FCachedBindingData : public TSharedFromThis<FCachedBindingData>
 			}
 		}
 
+		check(AccessibleStructs.IsValidIndex(InStructIndex));
+		// Binding directly into PropertyFunction's struct is allowed if it contains a compatible single output property.
+		if (AccessibleStructs[InStructIndex].DataSource == EStateTreeBindableStructSource::PropertyFunction)
+		{
+			IStateTreeEditorPropertyBindingsOwner* BindingOwner = Cast<IStateTreeEditorPropertyBindingsOwner>(WeakOwnerObject.Get());
+			FStateTreeDataView DataView;
+			// If DataView exists, struct is an instance of already bound function.
+			if (BindingOwner == nullptr || BindingOwner->GetDataViewByID(AccessibleStructs[InStructIndex].ID, DataView))
+			{
+				return false;
+			}
+
+			if (const FProperty* SingleOutputProperty = UE::StateTree::GetStructSingleOutputProperty(*AccessibleStructs[InStructIndex].Struct))
+			{
+				return CanBindToProperty(SingleOutputProperty, {FBindingChainElement(nullptr, InStructIndex), FBindingChainElement(const_cast<FProperty*>(SingleOutputProperty))});
+			}
+		}
+
 		return ArePropertyAndContextStructCompatible(InStruct, PropertyHandle->GetProperty());
 	}
 			
@@ -585,7 +674,14 @@ struct FCachedBindingData : public TSharedFromThis<FCachedBindingData>
 		check(AccessibleStructs.IsValidIndex(SourceStructIndex));
 
 		FStateTreeDataView SourceDataView;
-		BindingOwner->GetDataViewByID(AccessibleStructs[SourceStructIndex].ID, SourceDataView);
+		if (AccessibleStructs[SourceStructIndex].DataSource == EStateTreeBindableStructSource::PropertyFunction)
+		{
+			SourceDataView = FStateTreeDataView(AccessibleStructs[SourceStructIndex].Struct, nullptr);
+		}
+		else
+		{
+			BindingOwner->GetDataViewByID(AccessibleStructs[SourceStructIndex].ID, SourceDataView);
+		}
 
 		FStateTreePropertyPath SourcePath;
 		UE::StateTree::PropertyBinding::MakeStructPropertyPathFromBindingChain(AccessibleStructs[SourceStructIndex].ID, InBindingChain, SourceDataView, SourcePath);
@@ -609,12 +705,39 @@ struct FCachedBindingData : public TSharedFromThis<FCachedBindingData>
 
 		ConditionallyUpdateData();
 
-		if (UE::StateTree::PropertyRefHelpers::IsPropertyRef(*PropertyHandle->GetProperty()))
+		const int32 SourceStructIndex = InBindingChain[0].ArrayIndex;
+		check(AccessibleStructs.IsValidIndex(SourceStructIndex));
+		const FStateTreeBindableStructDesc& StructDesc = AccessibleStructs[SourceStructIndex];
+
+		if (StructDesc.DataSource == EStateTreeBindableStructSource::PropertyFunction)
 		{
-			const int32 SourceStructIndex = InBindingChain[0].ArrayIndex;
-			check(AccessibleStructs.IsValidIndex(SourceStructIndex));
+			IStateTreeEditorPropertyBindingsOwner* BindingOwner = Cast<IStateTreeEditorPropertyBindingsOwner>(WeakOwnerObject.Get());
+			FStateTreeDataView DataView;
+			// If DataView exists, struct is an instance of already bound function.
+			if (BindingOwner == nullptr || BindingOwner->GetDataViewByID(AccessibleStructs[SourceStructIndex].ID, DataView))
+			{
+				return false;
+			}
+
+			// To avoid duplicates, PropertyFunction struct's children are not allowed to be bound if it contains a compatible single output property.
+			if (const FProperty* SingleOutputProperty = UE::StateTree::GetStructSingleOutputProperty(*StructDesc.Struct))
+			{
+				if (CanBindToProperty(SingleOutputProperty, {FBindingChainElement(nullptr, SourceStructIndex), FBindingChainElement(const_cast<FProperty*>(SingleOutputProperty))}))
+				{
+					return false;
+				}
+			}
 
-			if (!UE::StateTree::PropertyRefHelpers::IsPropertyAccessibleForPropertyRef(*SourceProperty, InBindingChain, AccessibleStructs[SourceStructIndex]))
+			// Binding to non-output PropertyFunctions properties is not allowed.
+			if (InBindingChain.Num() == 1 && UE::StateTree::GetUsageFromMetaData(SourceProperty) != EStateTreePropertyUsage::Output)
+			{
+				return false;
+			}
+		}
+
+		if (UE::StateTree::PropertyRefHelpers::IsPropertyRef(*PropertyHandle->GetProperty()) && !UE::StateTree::PropertyRefHelpers::IsPropertyAccessibleForPropertyRef(*SourceProperty, InBindingChain, StructDesc))
+		{
+			if (!UE::StateTree::PropertyRefHelpers::IsPropertyAccessibleForPropertyRef(*SourceProperty, InBindingChain, StructDesc))
 			{
 				return false;
 			}
@@ -732,24 +855,143 @@ struct FCachedBindingData : public TSharedFromThis<FCachedBindingData>
 	FText GetText()
 	{
 		ConditionallyUpdateData();
-		return Text;
+
+		// Bound PropertyFunction is allowed to override it's display name.
+		if (IStateTreeEditorPropertyBindingsOwner* BindingOwner = Cast<IStateTreeEditorPropertyBindingsOwner>(WeakOwnerObject.Get()))
+		{
+			if(FStateTreeEditorPropertyBindings* EditorBindings = BindingOwner->GetPropertyEditorBindings())
+			{
+				const FStateTreePropertyPathBinding* CurrentBinding = EditorBindings->GetBindings().FindByPredicate([this](const FStateTreePropertyPathBinding& Binding)
+				{
+					return Binding.GetTargetPath() == TargetPath;
+				});
+
+				if (CurrentBinding)
+				{
+					const FConstStructView PropertyFunctionEditorNodeView = CurrentBinding->GetPropertyFunctionNode();
+					if(PropertyFunctionEditorNodeView.IsValid())
+					{
+						const FStateTreeEditorNode& EditorNode = PropertyFunctionEditorNodeView.Get<const FStateTreeEditorNode>();
+						if (const FStateTreeNodeBase* Node = EditorNode.Node.GetPtr<const FStateTreeNodeBase>())
+						{
+							const FText Description = Node->GetDescription(CachedSourcePath.GetStructID(), EditorNode.GetInstance(), FStateTreeBindingLookup(BindingOwner), EStateTreeNodeFormatting::Text);
+							if (!Description.IsEmpty())
+							{
+								return FText::FormatNamed(FormatableText, TEXT("SourceStruct"), Description);
+							}
+						}
+					}
+				}
+			}
+		}
+
+		return FText::FormatNamed(FormatableText, TEXT("SourceStruct"), SourceStructName);
 	}
 	
 	FText GetTooltipText()
 	{
 		ConditionallyUpdateData();
-		return TooltipText;
+
+		// If the source property is a PropertyFunction and it overrides it's display name, it's been used in the tooltip text.
+		if (IStateTreeEditorPropertyBindingsOwner* BindingOwner = Cast<IStateTreeEditorPropertyBindingsOwner>(WeakOwnerObject.Get()))
+		{
+			if(FStateTreeEditorPropertyBindings* EditorBindings = BindingOwner->GetPropertyEditorBindings())
+			{
+				const FStateTreePropertyPathBinding* CurrentBinding = EditorBindings->GetBindings().FindByPredicate([this](const FStateTreePropertyPathBinding& Binding)
+				{
+					return Binding.GetTargetPath() == TargetPath;
+				});
+
+				if (CurrentBinding)
+				{
+					const FConstStructView PropertyFunctionEditorNodeView = CurrentBinding->GetPropertyFunctionNode();
+					if(PropertyFunctionEditorNodeView.IsValid())
+					{
+						const FStateTreeEditorNode& EditorNode = PropertyFunctionEditorNodeView.Get<const FStateTreeEditorNode>();
+						if(const FStateTreeNodeBase* Node = EditorNode.Node.GetPtr<const FStateTreeNodeBase>())
+						{
+							const FText Description = Node->GetDescription(CachedSourcePath.GetStructID(), EditorNode.GetInstance(), FStateTreeBindingLookup(BindingOwner), EStateTreeNodeFormatting::Text);
+							if (!Description.IsEmpty())
+							{
+								return FText::FormatNamed(FormatableTooltipText, TEXT("SourceStruct"), Description);
+							}
+						}
+					}
+				}
+			}
+		}
+
+		return FText::FormatNamed(FormatableTooltipText, TEXT("SourceStruct"), SourceStructName);
 	}
 	
 	FLinearColor GetColor()
 	{
 		ConditionallyUpdateData();
+
+		// Bound PropertyFunction is allowed to override it's icon color if the binding leads directly into it's single output property.
+		if (CachedSourcePath.NumSegments() == 1)
+		{
+			if (IStateTreeEditorPropertyBindingsOwner* BindingOwner = Cast<IStateTreeEditorPropertyBindingsOwner>(WeakOwnerObject.Get()))
+			{
+				if(FStateTreeEditorPropertyBindings* EditorBindings = BindingOwner->GetPropertyEditorBindings())
+				{
+					const FStateTreePropertyPathBinding* CurrentBinding = EditorBindings->GetBindings().FindByPredicate([this](const FStateTreePropertyPathBinding& Binding)
+					{
+						return Binding.GetTargetPath() == TargetPath;
+					});
+
+					if (CurrentBinding)
+					{
+						const FConstStructView PropertyFunctionEditorNodeView = CurrentBinding->GetPropertyFunctionNode();
+						if (PropertyFunctionEditorNodeView.IsValid())
+						{
+							const FStateTreeEditorNode& EditorNode = PropertyFunctionEditorNodeView.Get<const FStateTreeEditorNode>();
+							if (const FStateTreeNodeBase* Node = EditorNode.Node.GetPtr<const FStateTreeNodeBase>())
+							{
+								if (Node && UE::StateTree::GetStructSingleOutputProperty(*Node->GetInstanceDataType()))
+								{
+									return Node->GetIconColor();
+								}
+					
+							}
+						}
+					}
+				}
+			}
+		}
+
 		return Color;
 	}
 	
 	const FSlateBrush* GetImage()
 	{
 		ConditionallyUpdateData();
+
+		// Bound PropertyFunction is allowed to override it's icon.
+		if (IStateTreeEditorPropertyBindingsOwner* BindingOwner = Cast<IStateTreeEditorPropertyBindingsOwner>(WeakOwnerObject.Get()))
+		{
+				if(FStateTreeEditorPropertyBindings* EditorBindings = BindingOwner->GetPropertyEditorBindings())
+				{
+					const FStateTreePropertyPathBinding* CurrentBinding = EditorBindings->GetBindings().FindByPredicate([this](const FStateTreePropertyPathBinding& Binding)
+					{
+						return Binding.GetTargetPath() == TargetPath;
+					});
+
+					if (CurrentBinding)
+					{
+						const FConstStructView PropertyFunctionEditorNodeView = CurrentBinding->GetPropertyFunctionNode();
+						if (PropertyFunctionEditorNodeView.IsValid())
+						{
+							const FStateTreeEditorNode& EditorNode = PropertyFunctionEditorNodeView.Get<const FStateTreeEditorNode>();
+							if (const FStateTreeNodeBase* Node = EditorNode.Node.GetPtr<const FStateTreeNodeBase>())
+							{
+								return UE::StateTreeEditor::EditorNodeUtils::ParseIcon(Node->GetIconName()).GetIcon();
+							}
+						}
+					}
+				}
+		}
+
 		return Image;
 	}
 	
@@ -798,8 +1040,15 @@ private:
 	TSharedPtr<const IPropertyHandle> PropertyHandle;
 	TArray<FStateTreeBindableStructDesc> AccessibleStructs;
 
-	FText Text;
-	FText TooltipText;
+	/* Default name of the source struct. */
+	FText SourceStructName;
+
+	/* Binding's display name text. Expects it's source struct name to be injected before use. */
+	FText FormatableText;
+
+	/* Binding's tooltip text. Expects it's source struct name to be injected before use. */
+	FText FormatableTooltipText;
+
 	FLinearColor Color = FLinearColor::White;
 	const FSlateBrush* Image = nullptr;
 	
@@ -812,6 +1061,76 @@ bool IsPropertyBindable(const FProperty& Property)
 			&& (!Property.HasAnyPropertyFlags(CPF_NativeAccessSpecifierPrivate | CPF_NativeAccessSpecifierProtected) || Property.GetBoolMetaData(FBlueprintMetadata::MD_AllowPrivateAccess));
 }
 
+/* Provides PropertyFunctionNode instance for a property node. */
+class FStateTreePropertyFunctionNodeProvider : public IStructureDataProvider
+{
+public:
+	FStateTreePropertyFunctionNodeProvider(IStateTreeEditorPropertyBindingsOwner& InBindingsOwner, FStateTreePropertyPath InTargetPath)
+		: BindingsOwner(Cast<UObject>(&InBindingsOwner))
+		, TargetPath(MoveTemp(InTargetPath))
+	{}
+	
+	virtual bool IsValid() const override
+	{
+		return GetPropertyFunctionEditorNodeView(BindingsOwner.Get(), TargetPath).IsValid();
+	};
+
+	virtual const UStruct* GetBaseStructure() const override
+	{
+		return FStateTreeEditorNode::StaticStruct();
+	}
+	
+	virtual void GetInstances(TArray<TSharedPtr<FStructOnScope>>& OutInstances, const UStruct* ExpectedBaseStructure) const override
+	{
+		if (ExpectedBaseStructure)
+		{
+			const FStructView Node = GetPropertyFunctionEditorNodeView(BindingsOwner.Get(), TargetPath);
+
+			if (Node.IsValid() && Node.GetScriptStruct()->IsChildOf(ExpectedBaseStructure))
+			{
+				OutInstances.Add(MakeShared<FStructOnScope>(Node.GetScriptStruct(), Node.GetMemory()));
+			}
+		}
+	}
+
+	static bool IsBoundToValidPropertyFunction(UObject& InBindingsOwner, const FStateTreePropertyPath& InTargetPath)
+	{
+		return GetPropertyFunctionEditorNodeView(&InBindingsOwner, InTargetPath).IsValid();
+	}
+	
+private:
+	static FStructView GetPropertyFunctionEditorNodeView(UObject* RawBindingsOwner, const FStateTreePropertyPath& InTargetPath)
+	{
+		if(IStateTreeEditorPropertyBindingsOwner* Owner = Cast<IStateTreeEditorPropertyBindingsOwner>(RawBindingsOwner))
+		{
+			FStateTreeEditorPropertyBindings* EditorBindings = Owner->GetPropertyEditorBindings();
+			FStateTreePropertyPathBinding* FoundBinding = EditorBindings->GetMutableBindings().FindByPredicate([&InTargetPath](const FStateTreePropertyPathBinding& Binding)
+			{
+				return Binding.GetTargetPath() == InTargetPath;
+			});
+
+			
+			if (FoundBinding)
+			{
+				const FStructView EditorNodeView = FoundBinding->GetMutablePropertyFunctionNode();
+				if (EditorNodeView.IsValid())
+				{
+					const FStateTreeEditorNode& EditorNode = EditorNodeView.Get<FStateTreeEditorNode>();
+					if (EditorNode.Node.IsValid() && EditorNode.Instance.IsValid())
+					{
+						return EditorNodeView;
+					}
+				}
+			}
+		}
+
+		return FStructView();
+	}
+
+	TWeakObjectPtr<UObject> BindingsOwner;
+	FStateTreePropertyPath TargetPath;
+};
+
 } // UE::StateTree::PropertyBinding
 
 bool FStateTreeBindingExtension::IsPropertyExtendable(const UClass* InObjectClass, const IPropertyHandle& PropertyHandle) const
@@ -877,8 +1196,13 @@ void FStateTreeBindingExtension::ExtendWidgetRow(FDetailWidgetRow& InWidgetRow,
 			EditorBindings = BindingOwner->GetPropertyEditorBindings();
 			BindingOwner->GetAccessibleStructs(TargetPath.GetStructID(), AccessibleStructs);
 
+			BindingOwner->EnumerateBindablePropertyFunctionNodes([&AccessibleStructs](const UScriptStruct* NodeStruct, const FStateTreeBindableStructDesc& Desc, const FStateTreeDataView Value)
+			{
+				AccessibleStructs.Add(Desc);
+				return EStateTreeVisitor::Continue;
+			});
+
 			TMap<FString, FText> SectionNames;
-			
 			for (FStateTreeBindableStructDesc& StructDesc : AccessibleStructs)
 			{
 				const UStruct* Struct = StructDesc.Struct;
@@ -896,6 +1220,22 @@ void FStateTreeBindingExtension::ExtendWidgetRow(FDetailWidgetRow& InWidgetRow,
 				{
 					ContextStruct.Section = SectionNames.Add(StructDesc.StatePath, FText::FromString(StructDesc.StatePath));
 				}
+				
+				// PropertyFunction overrides it's struct's icon color.
+				if (StructDesc.DataSource == EStateTreeBindableStructSource::PropertyFunction)
+				{
+					if (const FProperty* OutputProperty = UE::StateTree::GetStructSingleOutputProperty(*StructDesc.Struct))
+					{
+						const UEdGraphSchema_K2* Schema = GetDefault<UEdGraphSchema_K2>();
+						check(Schema);
+
+						FEdGraphPinType PinType;
+						if (Schema->ConvertPropertyToPinType(OutputProperty, PinType))
+						{
+							ContextStruct.Color = Schema->GetPinTypeColor(PinType);
+						}
+					}
+				}
 			}
 		}
 
@@ -932,7 +1272,7 @@ void FStateTreeBindingExtension::ExtendWidgetRow(FDetailWidgetRow& InWidgetRow,
 
 	Args.OnCanBindToContextStructWithIndex = FOnCanBindToContextStructWithIndex::CreateLambda([CachedBindingData](const UStruct* InStruct, int32 InStructIndex)
 		{
-			return CachedBindingData->CanBindToContextStruct(InStruct);
+			return CachedBindingData->CanBindToContextStruct(InStruct, InStructIndex);
 		});
 
 	Args.OnCanAcceptPropertyOrChildrenWithBindingChain = FOnCanAcceptPropertyOrChildrenWithBindingChain::CreateLambda([CachedBindingData](FProperty* InProperty, TConstArrayView<FBindingChainElement> InBindingChain)
@@ -945,14 +1285,16 @@ void FStateTreeBindingExtension::ExtendWidgetRow(FDetailWidgetRow& InWidgetRow,
 			return true;
 		});
 
-	Args.OnAddBinding = FOnAddBinding::CreateLambda([CachedBindingData](FName InPropertyName, TConstArrayView<FBindingChainElement> InBindingChain)
+	Args.OnAddBinding = FOnAddBinding::CreateLambda([CachedBindingData, &InDetailBuilder](FName InPropertyName, TConstArrayView<FBindingChainElement> InBindingChain)
 		{
 			CachedBindingData->AddBinding(InBindingChain);
+			InDetailBuilder.GetPropertyUtilities()->RequestForceRefresh();
 		});
 
-	Args.OnRemoveBinding = FOnRemoveBinding::CreateLambda([CachedBindingData](FName InPropertyName)
+	Args.OnRemoveBinding = FOnRemoveBinding::CreateLambda([CachedBindingData, &InDetailBuilder](FName InPropertyName)
 		{
 			CachedBindingData->RemoveBinding();
+			InDetailBuilder.GetPropertyUtilities()->RequestForceRefresh();
 		});
 
 	Args.OnCanRemoveBinding = FOnCanRemoveBinding::CreateLambda([EditorBindings, TargetPath](FName InPropertyName)
@@ -999,4 +1341,43 @@ void FStateTreeBindingExtension::ExtendWidgetRow(FDetailWidgetRow& InWidgetRow,
 	];
 }
 
+void FStateTreeBindingsChildrenCustomization::CustomizeChildren(IDetailChildrenBuilder& ChildrenBuilder, TSharedPtr<IPropertyHandle> InPropertyHandle)
+{
+	TArray<UObject*> OuterObjects;
+	InPropertyHandle->GetOuterObjects(OuterObjects);
+	if (OuterObjects.Num() == 1)
+	{
+		FStateTreePropertyPath TargetPath;
+		UE::StateTree::PropertyBinding::MakeStructPropertyPathFromPropertyHandle(InPropertyHandle, TargetPath);
+
+		using FStateTreePropertyFunctionNodeProvider = UE::StateTree::PropertyBinding::FStateTreePropertyFunctionNodeProvider;
+		UObject* BindingsOwner = UE::StateTree::PropertyBinding::FindEditorBindingsOwner(OuterObjects[0]);
+		if (BindingsOwner && FStateTreePropertyFunctionNodeProvider::IsBoundToValidPropertyFunction(*BindingsOwner, TargetPath))
+		{
+			// Bound PropertyFunction takes control over property's children composition.
+			const TSharedPtr<FStateTreePropertyFunctionNodeProvider> StructProvider = MakeShared<FStateTreePropertyFunctionNodeProvider>(*CastChecked<IStateTreeEditorPropertyBindingsOwner>(BindingsOwner), MoveTemp(TargetPath));
+			// Create unique name to persists expansion state.
+			const FName UniqueName = FName(LexToString(TargetPath.GetStructID()) + TargetPath.ToString());
+			ChildrenBuilder.AddChildStructure(InPropertyHandle.ToSharedRef(), StructProvider, UniqueName);
+		}
+	}
+}
+
+bool FStateTreeBindingsChildrenCustomization::ShouldCustomizeChildren(TSharedRef<IPropertyHandle> InPropertyHandle)
+{
+	TArray<UObject*> OuterObjects;
+	InPropertyHandle->GetOuterObjects(OuterObjects);
+	if (OuterObjects.Num() == 1)
+	{
+		// Bound property's children composition gets overridden.
+		FStateTreePropertyPath TargetPath;
+		UE::StateTree::PropertyBinding::MakeStructPropertyPathFromPropertyHandle(InPropertyHandle, TargetPath);
+		IStateTreeEditorPropertyBindingsOwner* BindingOwner = Cast<IStateTreeEditorPropertyBindingsOwner>(UE::StateTree::PropertyBinding::FindEditorBindingsOwner(OuterObjects[0]));
+		FStateTreeEditorPropertyBindings* EditorBindings = BindingOwner->GetPropertyEditorBindings();
+		return EditorBindings->HasPropertyBinding(TargetPath);
+	}
+
+	return false;
+}
+
 #undef LOCTEXT_NAMESPACE
diff --git a/Engine/Plugins/Runtime/StateTree/Source/StateTreeEditorModule/Private/Customizations/StateTreeBindingExtension.h b/Engine/Plugins/Runtime/StateTree/Source/StateTreeEditorModule/Private/Customizations/StateTreeBindingExtension.h
--- a/Engine/Plugins/Runtime/StateTree/Source/StateTreeEditorModule/Private/Customizations/StateTreeBindingExtension.h
+++ b/Engine/Plugins/Runtime/StateTree/Source/StateTreeEditorModule/Private/Customizations/StateTreeBindingExtension.h
@@ -3,6 +3,7 @@
 #pragma once
 
 #include "IDetailPropertyExtensionHandler.h"
+#include "IDetailPropertyChildrenCustomizationHandler.h"
 
 enum class EStateTreePropertyUsage : uint8;
 
@@ -49,3 +50,11 @@ public:
 	virtual bool IsPropertyExtendable(const UClass* InObjectClass, const IPropertyHandle& PropertyHandle) const override;
 	virtual void ExtendWidgetRow(FDetailWidgetRow& InWidgetRow, const IDetailLayoutBuilder& InDetailBuilder, const UClass* InObjectClass, TSharedPtr<IPropertyHandle> PropertyHandle) override;
 };
+
+/* Overrides bound property's children composition. */
+class FStateTreeBindingsChildrenCustomization : public IDetailPropertyChildrenCustomizationHandler
+{
+public:
+	bool ShouldCustomizeChildren(TSharedRef<IPropertyHandle> InPropertyHandle) override;
+	void CustomizeChildren(IDetailChildrenBuilder& ChildrenBuilder, TSharedPtr<IPropertyHandle> InPropertyHandle) override;
+};
diff --git a/Engine/Plugins/Runtime/StateTree/Source/StateTreeEditorModule/Private/Customizations/StateTreeEditorNodeDetails.cpp b/Engine/Plugins/Runtime/StateTree/Source/StateTreeEditorModule/Private/Customizations/StateTreeEditorNodeDetails.cpp
--- a/Engine/Plugins/Runtime/StateTree/Source/StateTreeEditorModule/Private/Customizations/StateTreeEditorNodeDetails.cpp
+++ b/Engine/Plugins/Runtime/StateTree/Source/StateTreeEditorModule/Private/Customizations/StateTreeEditorNodeDetails.cpp
@@ -39,6 +39,7 @@
 #include "Styling/SlateTypes.h"
 #include "TextStyleDecorator.h"
 #include "Framework/Application/SlateApplication.h"
+#include "StateTreePropertyFunctionBase.h"
 
 #define LOCTEXT_NAMESPACE "StateTreeEditor"
 
@@ -63,6 +64,31 @@ public:
 
 namespace UE::StateTreeEditor::Internal
 {
+	/* Returns true if provided property is direct or indirect child of PropertyFunction */
+	bool IsOwnedByPropertyFunctionNode(TSharedPtr<IPropertyHandle> Property)
+	{
+		while (Property)
+		{
+			if (FStructProperty* StructProperty = CastField<FStructProperty>(Property->GetProperty()))
+			{
+				if (StructProperty->Struct == FStateTreeEditorNode::StaticStruct())
+				{
+					if (const FStateTreeEditorNode* Node = UE::StateTreeEditor::EditorNodeUtils::GetCommonNode(Property))
+					{
+						if (const UScriptStruct* ScriptStruct = Node->Node.GetScriptStruct())
+						{
+							return ScriptStruct->IsChildOf<FStateTreePropertyFunctionBase>();
+						}
+					}
+				}	
+			}
+
+			Property = Property->GetParentHandle();
+		}
+
+		return false;
+	}
+
 	/** @return text describing the pin type, matches SPinTypeSelector. */
 	FText GetPinTypeText(const FEdGraphPinType& PinType)
 	{
@@ -115,6 +141,13 @@ namespace UE::StateTreeEditor::Internal
 		const EStateTreePropertyUsage Usage = UE::StateTree::GetUsageFromMetaData(ChildPropHandle->GetProperty());
 		const FProperty* Property = ChildPropHandle->GetProperty();
 		
+		// Hide output properties for PropertyFunctionNode.
+		if (Usage == EStateTreePropertyUsage::Output && UE::StateTreeEditor::Internal::IsOwnedByPropertyFunctionNode(ChildPropHandle))
+		{
+			ChildRow.Visibility(EVisibility::Hidden);
+			return;
+		}
+
 		// Conditionally control visibility of the value field of bound properties.
 		if (Usage != EStateTreePropertyUsage::Invalid && ID.IsValid())
 		{
@@ -132,13 +165,8 @@ namespace UE::StateTreeEditor::Internal
 					return EditorPropBindings->HasPropertyBinding(Path) ? EVisibility::Collapsed : EVisibility::Visible;
 				});
 
-			bool bShowChildren = true;
-			
 			if (Usage == EStateTreePropertyUsage::Input || Usage == EStateTreePropertyUsage::Output || Usage == EStateTreePropertyUsage::Context)
-			{
-				// Do not show children for input, output and context.
-				bShowChildren = false;
-				
+			{				
 				FEdGraphPinType PinType;
 				const UEdGraphSchema_K2* Schema = GetDefault<UEdGraphSchema_K2>();
 
@@ -216,7 +244,7 @@ namespace UE::StateTreeEditor::Internal
 				}
 				
 				ChildRow
-					.CustomWidget(bShowChildren)
+					.CustomWidget(true)
 					.NameContent()
 					[
 						SNew(SHorizontalBox)
@@ -336,20 +364,25 @@ void FStateTreeEditorNodeDetails::CustomizeHeader(TSharedRef<class IPropertyHand
 
 	UE::StateTreeEditor::EditorNodeUtils::GetNodeBaseScriptStructAndClass(StructProperty, BaseScriptStruct, BaseClass);
 
-	const FIsResetToDefaultVisible IsResetVisible = FIsResetToDefaultVisible::CreateSP(this, &FStateTreeEditorNodeDetails::ShouldResetToDefault);
-	const FResetToDefaultHandler ResetHandler = FResetToDefaultHandler::CreateSP(this, &FStateTreeEditorNodeDetails::ResetToDefault);
-	const FResetToDefaultOverride ResetOverride = FResetToDefaultOverride::Create(IsResetVisible, ResetHandler);
-
 	UE::StateTree::Delegates::OnIdentifierChanged.AddSP(this, &FStateTreeEditorNodeDetails::OnIdentifierChanged);
 	OnBindingChangedHandle = UE::StateTree::PropertyBinding::OnStateTreePropertyBindingChanged.AddRaw(this, &FStateTreeEditorNodeDetails::OnBindingChanged);
-
 	FindOuterObjects();
 
+	// Don't draw the header if it's a PropertyFunction.
+	if (UE::StateTreeEditor::Internal::IsOwnedByPropertyFunctionNode(StructProperty))
+	{
+		return;
+	}
+
+	const FIsResetToDefaultVisible IsResetVisible = FIsResetToDefaultVisible::CreateSP(this, &FStateTreeEditorNodeDetails::ShouldResetToDefault);
+	const FResetToDefaultHandler ResetHandler = FResetToDefaultHandler::CreateSP(this, &FStateTreeEditorNodeDetails::ResetToDefault);
+	const FResetToDefaultOverride ResetOverride = FResetToDefaultOverride::Create(IsResetVisible, ResetHandler);
+
 	auto IndentColor = [this]() -> FSlateColor
 	{
 		return (RowBorder && RowBorder->IsHovered()) ? FSlateColor::UseForeground() : FSlateColor(FLinearColor::Transparent);
 	};
-	
+
 	HeaderRow
 		.WholeRowContent()
 		.VAlign(VAlign_Center)
@@ -849,7 +882,7 @@ void FStateTreeEditorNodeDetails::CustomizeChildren(TSharedRef<class IPropertyHa
 			TSharedPtr<IPropertyHandle> PropertyHandle;
 			EStateTreePropertyUsage Usage = EStateTreePropertyUsage::Invalid;
 		};
-		
+
 		TArray<FSortedChild> SortedChildren;
 		for (uint32 Index = 0; Index < NumChildren; Index++)
 		{
diff --git a/Engine/Plugins/Runtime/StateTree/Source/StateTreeEditorModule/Private/StateTreeCompiler.cpp b/Engine/Plugins/Runtime/StateTree/Source/StateTreeEditorModule/Private/StateTreeCompiler.cpp
--- a/Engine/Plugins/Runtime/StateTree/Source/StateTreeEditorModule/Private/StateTreeCompiler.cpp
+++ b/Engine/Plugins/Runtime/StateTree/Source/StateTreeEditorModule/Private/StateTreeCompiler.cpp
@@ -13,6 +13,7 @@
 #include "StateTreePropertyRef.h"
 #include "StateTreePropertyRefHelpers.h"
 #include "StateTreePropertyHelpers.h"
+#include "StateTreePropertyFunctionBase.h"
 
 
 namespace UE::StateTree::Compiler
@@ -519,7 +520,7 @@
 	}
 	StateTree->EvaluatorsNum = uint16(EvaluatorsNum);
 
-	return true;
+	return CreateBindingsForNodes(EditorData->Evaluators, FStateTreeIndex16(EvaluatorsBegin));
 }
 
 bool FStateTreeCompiler::CreateGlobalTasks()
@@ -564,7 +565,7 @@
 	}
 	StateTree->GlobalTasksNum = uint16(GlobalTasksNum);
 
-	return true;
+	return CreateBindingsForNodes(EditorData->GlobalTasks, FStateTreeIndex16(GlobalTasksBegin));
 }
 
 bool FStateTreeCompiler::CreateStateTasksAndParameters()
@@ -676,37 +677,28 @@
 			return false;
 		}
 
-		int32 BatchIndex = INDEX_NONE;
-
 		// Subtrees parameters cannot have bindings
-		if(State->Type != EStateTreeStateType::Subtree)
+		if (State->Type != EStateTreeStateType::Subtree)
 		{
-			// Check that the bindings for this struct are still all valid.
-			TArray<FStateTreePropertyPathBinding> CopyBindings;
-			TArray<FStateTreePropertyPathBinding> ReferenceBindings;
-			if (!GetAndValidateBindings(LinkedParamsDesc, FStateTreeDataView(CompactStateTreeParameters.Parameters.GetMutableValue()), CopyBindings, ReferenceBindings))
+			FStateTreeIndex16 PropertyFunctionsBegin(Nodes.Num());
+			if (!CreatePropertyFunctionsForStruct(LinkedParamsDesc.ID))
 			{
 				return false;
 			}
 
-			if (!BindingsCompiler.CompileBatch(LinkedParamsDesc, CopyBindings, BatchIndex))
+			FStateTreeIndex16 PropertyFunctionsEnd(Nodes.Num());
+		
+			if (PropertyFunctionsBegin == PropertyFunctionsEnd)
 			{
-				return false;
+				PropertyFunctionsBegin = FStateTreeIndex16::Invalid;
+				PropertyFunctionsEnd = FStateTreeIndex16::Invalid;
 			}
 
-			if (!BindingsCompiler.CompileReferences(LinkedParamsDesc, ReferenceBindings, FStateTreeDataView(CompactStateTreeParameters.Parameters.GetMutableValue()), IDToStructValue))
+			if (!CreateBindingsForStruct(LinkedParamsDesc, FStateTreeDataView(CompactStateTreeParameters.Parameters.GetMutableValue()), PropertyFunctionsBegin, PropertyFunctionsEnd, CompactState.ParameterBindingsBatch))
 			{
 				return false;
 			}
 		}
-			
-		if (const auto Validation = UE::StateTree::Compiler::IsValidIndex16(BatchIndex); Validation.DidFail())
-		{
-			Validation.Log(Log, TEXT("BatchIndex"), LinkedParamsDesc);
-			return false;
-		}
 
-		CompactState.ParameterBindingsBatch = FStateTreeIndex16(BatchIndex);
-
 		// Create tasks
 		const int32 TasksBegin = Nodes.Num();
@@ -765,6 +757,11 @@
 
 		CompactState.TasksNum = uint8(TasksNum);
 		CompactState.InstanceDataNum = uint8(InstanceDataNum);
+
+		if (!CreateBindingsForNodes(Tasks, FStateTreeIndex16(TasksBegin)))
+		{
+			return false;
+		}
 	}
 	
 	return true;
@@ -809,6 +806,11 @@
 		}
 		CompactState.EnterConditionsNum = uint8(EnterConditionsNum);
 
+		if (!CreateBindingsForNodes(SourceState->EnterConditions, FStateTreeIndex16(EnterConditionsBegin)))
+		{
+			return false;
+		}
+
 		// Linked state
 		if (SourceState->Type == EStateTreeStateType::Linked)
 		{
@@ -1034,6 +1036,11 @@
 				return false;
 			}
 			CompactTransition.ConditionsNum = uint8(ConditionsNum);
+
+			if (!CreateBindingsForNodes(Transition.Conditions, FStateTreeIndex16(ConditionsBegin)))
+			{
+				return false;
+			}
 		}
 		
 		const int32 TransitionsNum = StateTree->Transitions.Num() - TransitionsBegin;
@@ -1050,137 +1057,185 @@
 	return true;
 }
 
-bool FStateTreeCompiler::ResolveTransitionState(const UStateTreeState* SourceState, const FStateTreeStateLink& Link, FStateTreeStateHandle& OutTransitionHandle) const 
+bool FStateTreeCompiler::CreateBindingsForNodes(TConstArrayView<FStateTreeEditorNode> EditorNodes, FStateTreeIndex16 NodesBegin)
 {
-	if (Link.LinkType == EStateTreeTransitionType::GotoState)
+	check(NodesBegin.IsValid());
+
+	int32 NodeIndex = NodesBegin.Get();
+	for (const FStateTreeEditorNode& EditorNode : EditorNodes)
 	{
-		// Warn if goto state points to another subtree.
-		if (const UStateTreeState* TargetState = GetState(Link.ID))
+		// Node might be an empty line in Editor.
+		if (!EditorNode.Node.IsValid())
 		{
-			if (SourceState && TargetState->GetRootState() != SourceState->GetRootState())
-			{
-				Log.Reportf(EMessageSeverity::Warning,
-					TEXT("Target state '%s' is in different subtree. Verify that this is intentional."),
-					*Link.Name.ToString());
-			}
+			continue;
+		}
 
-			if (TargetState->SelectionBehavior == EStateTreeStateSelectionBehavior::None)
-			{
-				Log.Reportf(EMessageSeverity::Error,
-					TEXT("The target State '%s' is not selectable, it's selection behavior is set to None."),
-					*Link.Name.ToString());
-				return false;
-			}
+		FStateTreeNodeBase& Node = Nodes[NodeIndex++].GetMutable<FStateTreeNodeBase>();
+
+		const FStateTreeBindableStructDesc* BindableStruct = BindingsCompiler.GetSourceStructDescByID(EditorNode.ID);
+		check(BindableStruct);
+
+		FStateTreeIndex16 PropertyFunctionsBegin(Nodes.Num());
+		if (!CreatePropertyFunctionsForStruct(EditorNode.ID))
+		{
+			return false;
 		}
+		FStateTreeIndex16 PropertyFunctionsEnd(Nodes.Num());
 		
-		OutTransitionHandle = GetStateHandle(Link.ID);
-		if (!OutTransitionHandle.IsValid())
+		if (PropertyFunctionsBegin == PropertyFunctionsEnd)
 		{
-			Log.Reportf(EMessageSeverity::Error,
-				TEXT("Failed to resolve transition to state '%s'."),
-				*Link.Name.ToString());
+			PropertyFunctionsBegin = FStateTreeIndex16::Invalid;
+			PropertyFunctionsEnd = FStateTreeIndex16::Invalid;
+		}
+
+		if (!CreateBindingsForStruct(*BindableStruct, EditorNode.GetInstance(), PropertyFunctionsBegin, PropertyFunctionsEnd, Node.BindingsBatch))
+		{
 			return false;
 		}
 	}
-	else if (Link.LinkType == EStateTreeTransitionType::NextState || Link.LinkType == EStateTreeTransitionType::NextSelectableState)
+
+	return true;
+}
+
+bool FStateTreeCompiler::CreateBindingsForStruct(const FStateTreeBindableStructDesc& TargetStruct, FStateTreeDataView TargetValue, FStateTreeIndex16 PropertyFuncsBegin, FStateTreeIndex16 PropertyFuncsEnd, FStateTreeIndex16& OutBatchIndex)
+{
+	TArray<FStateTreePropertyPathBinding> CopyBindings;
+	TArray<FStateTreePropertyPathBinding> ReferenceBindings;
+
+	// Check that the bindings for this struct are still all valid.
+	if (!GetAndValidateBindings(TargetStruct, TargetValue, CopyBindings, ReferenceBindings))
 	{
-		// Find next state.
-		const UStateTreeState* NextState = SourceState ? SourceState->GetNextSelectableSiblingState() : nullptr;
-		if (NextState == nullptr)
+		return false;
+	}
+
+	int32 BatchIndex = INDEX_NONE;
+
+	// Compile batch copy for this struct, we pass in all the bindings, the compiler will pick up the ones for the target structs.
+	if (!BindingsCompiler.CompileBatch(TargetStruct, CopyBindings, PropertyFuncsBegin, PropertyFuncsEnd, BatchIndex))
+	{
+		return false;
+	}
+
+	if (const auto Validation = UE::StateTree::Compiler::IsValidIndex16(BatchIndex); Validation.DidFail())
+	{
+		Validation.Log(Log, TEXT("BatchIndex"), TargetStruct);
+		return false;
+	}
+
+	OutBatchIndex = FStateTreeIndex16(BatchIndex);
+
+	if (!BindingsCompiler.CompileReferences(TargetStruct, ReferenceBindings, TargetValue, IDToStructValue))
+	{
+		return false;
+	}
+
+	return true;
+}
+
+bool FStateTreeCompiler::CreatePropertyFunctionsForStruct(FGuid StructID)
+{
+	for (const FStateTreePropertyPathBinding& Binding : EditorData->EditorBindings.GetBindings())
+	{
+		if (Binding.GetTargetPath().GetStructID() != StructID)
 		{
-			Log.Reportf(EMessageSeverity::Error,
-				TEXT("Failed to resolve transition, there's no selectable next state."));
-			return false;
+			continue;
 		}
-		OutTransitionHandle = GetStateHandle(NextState->ID);
-		if (!OutTransitionHandle.IsValid())
+
+		const FConstStructView NodeView = Binding.GetPropertyFunctionNode();
+		if (!NodeView.IsValid())
 		{
-			Log.Reportf(EMessageSeverity::Error,
-				TEXT("Failed to resolve transition next state, no handle found for '%s'."),
-				*NextState->Name.ToString());
+			continue;
+		}
+
+		const FStateTreeEditorNode& FuncEditorNode = NodeView.Get<const FStateTreeEditorNode>();
+		if (!CreatePropertyFunction(FuncEditorNode))
+		{
 			return false;
 		}
 	}
-	else if(Link.LinkType == EStateTreeTransitionType::Failed)
-	{
-		OutTransitionHandle = FStateTreeStateHandle::Failed;
-		return true;
-	}
-	else if(Link.LinkType == EStateTreeTransitionType::Succeeded)
-	{
-		OutTransitionHandle = FStateTreeStateHandle::Succeeded;
-		return true;
-	}
-	else if(Link.LinkType == EStateTreeTransitionType::None)
-	{
-		OutTransitionHandle = FStateTreeStateHandle::Invalid;
-		return true;
-	}
-	
+
 	return true;
 }
 
-bool FStateTreeCompiler::CreateCondition(UStateTreeState& State, const FString& StatePath, const FStateTreeEditorNode& CondNode, const EStateTreeExpressionOperand Operand, const int8 DeltaIndent)
+bool FStateTreeCompiler::CreatePropertyFunction(const FStateTreeEditorNode& FuncEditorNode)
 {
-	if (!CondNode.Node.IsValid())
+	if (!CreatePropertyFunctionsForStruct(FuncEditorNode.ID))
 	{
-		// Empty line in conditions array, just silently ignore.
-		return true;
+		return false;
 	}
 
+	FStateTreeNodeBase* Node = CreateNodeWithSharedInstanceData(nullptr, UE::StateTree::Editor::PropertyFunctionStateName, FuncEditorNode, EStateTreeBindableStructSource::PropertyFunction);
+	if(Node == nullptr)
+	{
+		return false;
+	}
+
+	const FStateTreeBindableStructDesc* BindableStruct = BindingsCompiler.GetSourceStructDescByID(FuncEditorNode.ID);
+	check(BindableStruct);
+
+	return CreateBindingsForStruct(*BindableStruct, FuncEditorNode.GetInstance(), FStateTreeIndex16::Invalid, FStateTreeIndex16::Invalid, Node->BindingsBatch);
+}
+
+template<class T>
+T* FStateTreeCompiler::CreateNodeWithSharedInstanceData(UStateTreeState* State, const FString& StatePath, const FStateTreeEditorNode& EditorNode, EStateTreeBindableStructSource StructSource)
+{
+	if (!EditorNode.Node.IsValid())
+	{
+		return nullptr;
+	}
+
+	check(EditorNode.Node.GetScriptStruct()->IsChildOf<T>());
+
 	FStateTreeBindableStructDesc StructDesc;
 	StructDesc.StatePath = StatePath;
-	StructDesc.ID = CondNode.ID;
-	StructDesc.Name = CondNode.GetName();
-	StructDesc.DataSource = EStateTreeBindableStructSource::Condition;
+	StructDesc.ID = EditorNode.ID;
+	StructDesc.Name = EditorNode.GetName();
+	StructDesc.DataSource = StructSource;
 
 	// Check that item has valid instance initialized.
-	if (!CondNode.Instance.IsValid() && CondNode.InstanceObject == nullptr)
+	if (!EditorNode.Instance.IsValid() && EditorNode.InstanceObject == nullptr)
 	{
 		Log.Reportf(EMessageSeverity::Error, StructDesc,
-			TEXT("Malformed condition, missing instance value."));
-		return false;
+			TEXT("Malformed node, missing instance value."));
+		return nullptr;
 	}
 
-	// Copy the condition
-	IDToNode.Add(CondNode.ID, Nodes.Num());
-	FInstancedStruct& Node = Nodes.Add_GetRef(CondNode.Node);
-	InstantiateStructSubobjects(Node);
+	// Copy the node
+	IDToNode.Add(EditorNode.ID, Nodes.Num());
+	FInstancedStruct& RawNode = Nodes.Add_GetRef(EditorNode.Node);
+	InstantiateStructSubobjects(RawNode);
 
-	FStateTreeConditionBase& Cond = Node.GetMutable<FStateTreeConditionBase>();
+	FStateTreeNodeBase& Node = RawNode.GetMutable<FStateTreeNodeBase>();
 
-	Cond.Operand = Operand;
-	Cond.DeltaIndent = DeltaIndent;
+	// Update node name as description for runtime.
+	Node.Name = EditorNode.GetName();
 
-	// Update condition name as description for runtime.
-	Cond.Name = CondNode.GetName();
-
 	FStateTreeDataView InstanceDataView;
 	
-	if (CondNode.Instance.IsValid())
+	if (EditorNode.Instance.IsValid())
 	{
 		// Struct instance
-		const int32 InstanceIndex = SharedInstanceStructs.Add(CondNode.Instance);
+		const int32 InstanceIndex = SharedInstanceStructs.Add(EditorNode.Instance);
 		InstantiateStructSubobjects(SharedInstanceStructs[InstanceIndex]);
 
 		// Create binding source struct descriptor.
-		StructDesc.Struct = CondNode.Instance.GetScriptStruct();
+		StructDesc.Struct = EditorNode.Instance.GetScriptStruct();
 
 		if (const auto Validation = UE::StateTree::Compiler::IsValidIndex16(InstanceIndex); Validation.DidFail())
 		{
 			Validation.Log(Log, TEXT("InstanceIndex"), StructDesc);
-			return false;
+			return nullptr;
 		}
-		Cond.InstanceTemplateIndex = FStateTreeIndex16(InstanceIndex);
-		Cond.InstanceDataHandle = FStateTreeDataHandle(EStateTreeDataSourceType::SharedInstanceData, InstanceIndex);
+		Node.InstanceTemplateIndex = FStateTreeIndex16(InstanceIndex);
+		Node.InstanceDataHandle = FStateTreeDataHandle(EStateTreeDataSourceType::SharedInstanceData, InstanceIndex);
 		InstanceDataView = FStateTreeDataView(SharedInstanceStructs[InstanceIndex]);
 	}
 	else
 	{
 		// Object Instance
-		check(CondNode.InstanceObject != nullptr);
+		check(EditorNode.InstanceObject != nullptr);
 
-		UObject* Instance = DuplicateObject(CondNode.InstanceObject, StateTree);
+		UObject* Instance = DuplicateObject(EditorNode.InstanceObject, StateTree);
 		FInstancedStruct Wrapper;
 		Wrapper.InitializeAs<FStateTreeInstanceObjectWrapper>(Instance);
 		const int32 InstanceIndex = SharedInstanceStructs.Add(MoveTemp(Wrapper));
@@ -1191,53 +1246,115 @@
 		if (const auto Validation = UE::StateTree::Compiler::IsValidIndex16(InstanceIndex); Validation.DidFail())
 		{
 			Validation.Log(Log, TEXT("InstanceIndex"), StructDesc);
-			return false;
+			return nullptr;
 		}
-		Cond.InstanceTemplateIndex = FStateTreeIndex16(InstanceIndex);
-		Cond.InstanceDataHandle = FStateTreeDataHandle(EStateTreeDataSourceType::SharedInstanceDataObject, InstanceIndex);
+		Node.InstanceTemplateIndex = FStateTreeIndex16(InstanceIndex);
+		Node.InstanceDataHandle = FStateTreeDataHandle(EStateTreeDataSourceType::SharedInstanceDataObject, InstanceIndex);
 		InstanceDataView = FStateTreeDataView(Instance);
 	}
 
-	StructDesc.DataHandle = Cond.InstanceDataHandle;
+	StructDesc.DataHandle = Node.InstanceDataHandle;
 	
-	if (!CompileAndValidateNode(&State, StructDesc, Node, InstanceDataView))
+	if (!CompileAndValidateNode(State, StructDesc, RawNode, InstanceDataView))
 	{
-		return false;
+		return nullptr;
 	}
 
 	// Mark the struct as binding source.
 	BindingsCompiler.AddSourceStruct(StructDesc);
+	
+	check(RawNode.GetScriptStruct()->IsChildOf<T>());
+	return RawNode.GetMutablePtr<T>();
+}
 
-	// Check that the bindings for this struct are still all valid.
-	TArray<FStateTreePropertyPathBinding> CopyBindings;
-	TArray<FStateTreePropertyPathBinding> ReferenceBindings;
-	if (!GetAndValidateBindings(StructDesc, InstanceDataView, CopyBindings, ReferenceBindings))
+bool FStateTreeCompiler::ResolveTransitionState(const UStateTreeState* SourceState, const FStateTreeStateLink& Link, FStateTreeStateHandle& OutTransitionHandle) const 
+{
+	if (Link.LinkType == EStateTreeTransitionType::GotoState)
 	{
-		return false;
+		// Warn if goto state points to another subtree.
+		if (const UStateTreeState* TargetState = GetState(Link.ID))
+		{
+			if (SourceState && TargetState->GetRootState() != SourceState->GetRootState())
+			{
+				Log.Reportf(EMessageSeverity::Warning,
+					TEXT("Target state '%s' is in different subtree. Verify that this is intentional."),
+					*Link.Name.ToString());
+			}
+
+			if (TargetState->SelectionBehavior == EStateTreeStateSelectionBehavior::None)
+			{
+				Log.Reportf(EMessageSeverity::Error,
+					TEXT("The target State '%s' is not selectable, it's selection behavior is set to None."),
+					*Link.Name.ToString());
+				return false;
+			}
+		}
+		
+		OutTransitionHandle = GetStateHandle(Link.ID);
+		if (!OutTransitionHandle.IsValid())
+		{
+			Log.Reportf(EMessageSeverity::Error,
+				TEXT("Failed to resolve transition to state '%s'."),
+				*Link.Name.ToString());
+			return false;
+		}
 	}
-
-	// Compile batch copy for this struct, we pass in all the bindings, the compiler will pick up the ones for the target structs.
-	int32 BatchIndex = INDEX_NONE;
-	if (!BindingsCompiler.CompileBatch(StructDesc, CopyBindings, BatchIndex))
+	else if (Link.LinkType == EStateTreeTransitionType::NextState || Link.LinkType == EStateTreeTransitionType::NextSelectableState)
 	{
-		return false;
+		// Find next state.
+		const UStateTreeState* NextState = SourceState ? SourceState->GetNextSelectableSiblingState() : nullptr;
+		if (NextState == nullptr)
+		{
+			Log.Reportf(EMessageSeverity::Error,
+				TEXT("Failed to resolve transition, there's no selectable next state."));
+			return false;
+		}
+		OutTransitionHandle = GetStateHandle(NextState->ID);
+		if (!OutTransitionHandle.IsValid())
+		{
+			Log.Reportf(EMessageSeverity::Error,
+				TEXT("Failed to resolve transition next state, no handle found for '%s'."),
+				*NextState->Name.ToString());
+			return false;
+		}
 	}
-
-	if (!BindingsCompiler.CompileReferences(StructDesc, ReferenceBindings, InstanceDataView, IDToStructValue))
+	else if(Link.LinkType == EStateTreeTransitionType::Failed)
 	{
-		return false;
+		OutTransitionHandle = FStateTreeStateHandle::Failed;
+		return true;
 	}
-
-	if (const auto Validation = UE::StateTree::Compiler::IsValidIndex16(BatchIndex); Validation.DidFail())
+	else if(Link.LinkType == EStateTreeTransitionType::Succeeded)
 	{
-		Validation.Log(Log, TEXT("BatchIndex"), StructDesc);
-		return false;
+		OutTransitionHandle = FStateTreeStateHandle::Succeeded;
+		return true;
 	}
-	Cond.BindingsBatch = FStateTreeIndex16(BatchIndex);
+	else if(Link.LinkType == EStateTreeTransitionType::None)
+	{
+		OutTransitionHandle = FStateTreeStateHandle::Invalid;
+		return true;
+	}
 	
 	return true;
 }
 
+bool FStateTreeCompiler::CreateCondition(UStateTreeState& State, const FString& StatePath, const FStateTreeEditorNode& CondNode, const EStateTreeExpressionOperand Operand, const int8 DeltaIndent)
+{
+	if (!CondNode.Node.IsValid())
+	{
+		// Empty line in conditions array, just silently ignore.
+		return true;
+	}
+
+	if (FStateTreeConditionBase* Cond = CreateNodeWithSharedInstanceData<FStateTreeConditionBase>(&State, StatePath, CondNode, EStateTreeBindableStructSource::Condition))
+	{
+		Cond->Operand = Operand;
+		Cond->DeltaIndent = DeltaIndent;
+		return true;
+	}
+
+	return false;
+}
+
 bool FStateTreeCompiler::CompileAndValidateNode(const UStateTreeState* SourceState, const FStateTreeBindableStructDesc& NodeDesc, FStructView NodeView, const FStateTreeDataView InstanceData)
 {
 	if (!NodeView.IsValid())
@@ -1385,33 +1502,6 @@
 	// Mark the instance as binding source.
 	BindingsCompiler.AddSourceStruct(StructDesc);
 	
-	// Check that the bindings for this struct are still all valid.
-	TArray<FStateTreePropertyPathBinding> CopyBindings;
-	TArray<FStateTreePropertyPathBinding> ReferenceBindings;
-	if (!GetAndValidateBindings(StructDesc, InstanceDataView, CopyBindings, ReferenceBindings))
-	{
-		return false;
-	}
-
-	// Compile batch copy for this struct, we pass in all the bindings, the compiler will pick up the ones for the target structs.
-	int32 BatchIndex = INDEX_NONE;
-	if (!BindingsCompiler.CompileBatch(StructDesc, CopyBindings, BatchIndex))
-	{
-		return false;
-	}
-
-	if (!BindingsCompiler.CompileReferences(StructDesc, ReferenceBindings, InstanceDataView, IDToStructValue))
-	{
-		return false;
-	}
-
-	if (const auto Validation = UE::StateTree::Compiler::IsValidIndex16(BatchIndex); Validation.DidFail())
-	{
-		Validation.Log(Log, TEXT("BatchIndex"), StructDesc);
-		return false;
-	}
-	Task.BindingsBatch = FStateTreeIndex16(BatchIndex);
-	
 	return true;
 }
 
@@ -1500,33 +1590,6 @@
 	// Mark the instance as binding source.
 	BindingsCompiler.AddSourceStruct(StructDesc);
 
-	// Check that the bindings for this struct are still all valid.
-	TArray<FStateTreePropertyPathBinding> CopyBindings;
-	TArray<FStateTreePropertyPathBinding> ReferenceBindings;
-	if (!GetAndValidateBindings(StructDesc, InstanceDataView, CopyBindings, ReferenceBindings))
-	{
-		return false;
-	}
-
-	// Compile batch copy for this struct, we pass in all the bindings, the compiler will pick up the ones for the target structs.
-	int32 BatchIndex = INDEX_NONE;
-	if (!BindingsCompiler.CompileBatch(StructDesc, CopyBindings, BatchIndex))
-	{
-		return false;
-	}
-
-	if (!BindingsCompiler.CompileReferences(StructDesc, ReferenceBindings, InstanceDataView, IDToStructValue))
-	{
-		return false;
-	}
-
-	if (const auto Validation = UE::StateTree::Compiler::IsValidIndex16(BatchIndex); Validation.DidFail())
-	{
-		Validation.Log(Log, TEXT("BatchIndex"), StructDesc);
-		return false;
-	}
-	Eval.BindingsBatch = FStateTreeIndex16(BatchIndex);
-
 	return true;
 }
 
@@ -1728,7 +1791,7 @@
 		FStateTreePropertyPathBinding BindingCopy(Binding);
 		BindingCopy.SetSourceDataHandle(SourceStruct->DataHandle);
 
-		// Special case fo AnyEnum. StateTreeBindingExtension allows AnyEnums to bind to other enum types.
+		// Special case for AnyEnum. StateTreeBindingExtension allows AnyEnums to bind to other enum types.
 		// The actual copy will be done via potential type promotion copy, into the value property inside the AnyEnum.
 		// We amend the paths here to point to the 'Value' property.
 		const bool bSourceIsAnyEnum = IsPropertyOfTypeOrChild(*TBaseStructure<FStateTreeAnyEnum>::Get(), *SourceStruct, Binding.GetSourcePath());
diff --git a/Engine/Plugins/Runtime/StateTree/Source/StateTreeEditorModule/Private/StateTreeEditor.cpp b/Engine/Plugins/Runtime/StateTree/Source/StateTreeEditorModule/Private/StateTreeEditor.cpp
--- a/Engine/Plugins/Runtime/StateTree/Source/StateTreeEditorModule/Private/StateTreeEditor.cpp
+++ b/Engine/Plugins/Runtime/StateTree/Source/StateTreeEditorModule/Private/StateTreeEditor.cpp
@@ -338,6 +338,7 @@ TSharedRef<SDockTab> FStateTreeEditor::SpawnTab_SelectionDetails(const FSpawnTab
 	SelectionDetailsView->OnFinishedChangingProperties().AddSP(this, &FStateTreeEditor::OnSelectionFinishedChangingProperties);
 
 	SelectionDetailsView->SetExtensionHandler(MakeShared<FStateTreeBindingExtension>());
+	SelectionDetailsView->SetChildrenCustomizationHandler(MakeShared<FStateTreeBindingsChildrenCustomization>());
 
 	TSharedRef<SDockTab> SpawnedTab = SNew(SDockTab)
 		.Label(NSLOCTEXT("StateTreeEditor", "SelectionDetailsTab", "Details"))
@@ -361,6 +362,7 @@ TSharedRef<SDockTab> FStateTreeEditor::SpawnTab_AssetDetails(const FSpawnTabArgs
 	AssetDetailsView->OnFinishedChangingProperties().AddSP(this, &FStateTreeEditor::OnAssetFinishedChangingProperties);
 
 	AssetDetailsView->SetExtensionHandler(MakeShared<FStateTreeBindingExtension>());
+	AssetDetailsView->SetChildrenCustomizationHandler(MakeShared<FStateTreeBindingsChildrenCustomization>());
 
 	TSharedRef<SDockTab> SpawnedTab = SNew(SDockTab)
 		.Label(NSLOCTEXT("StateTreeEditor", "AssetDetailsTabLabel", "Asset Details"))
diff --git a/Engine/Plugins/Runtime/StateTree/Source/StateTreeEditorModule/Private/StateTreeEditorData.cpp b/Engine/Plugins/Runtime/StateTree/Source/StateTreeEditorModule/Private/StateTreeEditorData.cpp
--- a/Engine/Plugins/Runtime/StateTree/Source/StateTreeEditorModule/Private/StateTreeEditorData.cpp
+++ b/Engine/Plugins/Runtime/StateTree/Source/StateTreeEditorModule/Private/StateTreeEditorData.cpp
@@ -11,6 +11,8 @@
 #include "StateTreeEditorModule.h"
 #include "StateTreePropertyHelpers.h"
 #include "Customizations/StateTreeEditorNodeUtils.h"
+#include "StateTreeNodeClassCache.h"
+#include "StateTreePropertyFunctionBase.h"
 
 #if WITH_EDITOR
 #include "Engine/UserDefinedStruct.h"
@@ -24,6 +26,31 @@
 namespace UE::StateTree::Editor
 {
 	const FString GlobalStateName(TEXT("Global"));
+	const FString PropertyFunctionStateName(TEXT("Property Functions"));
+	const FName ParametersNodeName(TEXT("Parameters"));
+
+	bool IsPropertyFunctionOwnedByNode(FGuid NodeID, FGuid PropertyFuncID, const FStateTreeEditorPropertyBindings& EditorBindings)
+	{
+		for (const FStateTreePropertyPathBinding& Binding : EditorBindings.GetBindings())
+		{
+			const FGuid TargetID = Binding.GetTargetPath().GetStructID();			
+			if (TargetID == NodeID)
+			{
+				return true;
+			}
+
+			FConstStructView NodeView = Binding.GetPropertyFunctionNode();
+			if (const FStateTreeEditorNode* Node = NodeView.GetPtr<const FStateTreeEditorNode>())
+			{
+				if (Node->ID == PropertyFuncID)
+				{
+					PropertyFuncID = TargetID;
+				}
+			}
+		}
+
+		return false;
+	}
 }
 
 UStateTreeEditorData::UStateTreeEditorData()
@@ -260,7 +287,7 @@ void UStateTreeEditorData::PostEditChangeChainProperty(FPropertyChangedChainEven
 
 	}
 
-	UE::StateTree::PropertyHelpers::DispatchPostEditToNodes(*this, PropertyChangedEvent);
+	UE::StateTree::PropertyHelpers::DispatchPostEditToNodes(*this, PropertyChangedEvent, *this);
 }
 #endif // WITH_EDITOR
 
@@ -304,7 +331,6 @@ void UStateTreeEditorData::GetAccessibleStructs(const TConstArrayView<const USta
 		return EStateTreeVisitor::Continue;
 	});
 
-
 	if (BaseProgress == EStateTreeVisitor::Continue)
 	{
 		TArray<FStateTreeBindableStructDesc, TInlineAllocator<32>> BindableDescs;
@@ -316,7 +342,7 @@ void UStateTreeEditorData::GetAccessibleStructs(const TConstArrayView<const USta
 				continue;
 			}
 			
-			const EStateTreeVisitor StateProgress = VisitStateNodes(*State, [&OutStructDescs, &BindableDescs, &Path, TargetStructID]
+			const EStateTreeVisitor StateProgress = VisitStateNodes(*State, [&OutStructDescs, &BindableDescs, &Path, TargetStructID, this]
 				(const UStateTreeState* State, const FStateTreeBindableStructDesc& Desc, const FStateTreeDataView Value)
 				{
 					// Stop iterating as soon as we find the target node.
@@ -365,6 +391,17 @@ void UStateTreeEditorData::GetAccessibleStructs(const TConstArrayView<const USta
 							}
 							break;
 						}
+
+						case EStateTreeBindableStructSource::PropertyFunction:
+						{
+							if (State == Path.Last())
+							{
+								if (UE::StateTree::Editor::IsPropertyFunctionOwnedByNode(TargetStructID, Desc.ID, EditorBindings))
+								{
+									BindableDescs.Add(Desc);
+								}
+							}
+						}
 					}
 							
 					return EStateTreeVisitor::Continue;
@@ -437,6 +474,51 @@ FStateTreeBindableStructDesc UStateTreeEditorData::FindContextData(const UStruct
 	return Candidates[HighestScoreIndex];
 }
 
+EStateTreeVisitor UStateTreeEditorData::EnumerateBindablePropertyFunctionNodes(TFunctionRef<EStateTreeVisitor(const UScriptStruct* NodeStruct, const FStateTreeBindableStructDesc& Desc, const FStateTreeDataView Value)> InFunc) const
+{
+	if (Schema == nullptr)
+	{
+		return EStateTreeVisitor::Continue;
+	}
+
+	FStateTreeEditorModule& EditorModule = FModuleManager::GetModuleChecked<FStateTreeEditorModule>(TEXT("StateTreeEditorModule"));
+	FStateTreeNodeClassCache* ClassCache = EditorModule.GetNodeClassCache().Get();
+	check(ClassCache);
+
+	TArray<TSharedPtr<FStateTreeNodeClassData>> StructNodes;	
+	ClassCache->GetStructs(FStateTreePropertyFunctionBase::StaticStruct(), StructNodes);
+	for (const TSharedPtr<FStateTreeNodeClassData>& NodeClassData : StructNodes)
+	{
+		if (const UScriptStruct* NodeStruct = NodeClassData->GetScriptStruct())
+		{
+			if (NodeStruct == FStateTreePropertyFunctionBase::StaticStruct() || NodeStruct->HasMetaData(TEXT("Hidden")))
+			{
+				continue;
+			}
+
+			if (Schema->IsStructAllowed(NodeStruct))
+			{
+				if (const UStruct* InstanceDataStruct = NodeClassData->GetInstanceDataStruct())
+				{
+					FStateTreeBindableStructDesc Desc;
+					Desc.Struct = InstanceDataStruct;
+					Desc.ID = FGuid::NewDeterministicGuid(NodeStruct->GetName());
+					Desc.DataSource = EStateTreeBindableStructSource::PropertyFunction;
+					Desc.Name = FName(NodeStruct->GetDisplayNameText().ToString());
+					Desc.StatePath = UE::StateTree::Editor::PropertyFunctionStateName;
+
+					if (InFunc(NodeStruct, Desc, FStateTreeDataView(InstanceDataStruct, nullptr)) == EStateTreeVisitor::Break)
+					{
+						return EStateTreeVisitor::Break;
+					}
+				}
+			}
+		}
+	}
+
+	return EStateTreeVisitor::Continue;
+}
+
 bool UStateTreeEditorData::GetStructByID(const FGuid StructID, FStateTreeBindableStructDesc& OutStructDesc) const
 {
 	VisitAllNodes([&OutStructDesc, StructID](const UStateTreeState* State, const FStateTreeBindableStructDesc& Desc, const FStateTreeDataView Value)
@@ -840,12 +922,26 @@ EStateTreeVisitor UStateTreeEditorData::VisitStateNodes(const UStateTreeState& S
 {
 	bool bContinue = true;
 
+	auto VisitFuncNodes = [&InFunc, &State, this](const FGuid StructID, const FName NodeName)
+	{
+		const FString StatePath = FString::Printf(TEXT("%s/%s"), *State.GetPath(), *NodeName.ToString());
+		return VisitStructBoundPropertyFunctions(StructID, StatePath, [&InFunc, &State](const FStateTreeEditorNode& EditorNode, const FStateTreeBindableStructDesc& Desc, FStateTreeDataView Value)
+		{
+			return InFunc(&State, EditorNode.ID, EditorNode.GetName(), EStateTreeNodeType::PropertyFunction, EditorNode.Node.GetScriptStruct(), EditorNode.GetInstance().GetStruct());
+		});
+	};
+
 	if (bContinue)
 	{
 		// Enter conditions
 		for (const FStateTreeEditorNode& Node : State.EnterConditions)
 		{
-			if (const FStateTreeConditionBase* Cond = Node.Node.GetPtr<FStateTreeConditionBase>())
+			if (VisitFuncNodes(Node.ID, Node.GetName()) == EStateTreeVisitor::Break)
+			{
+				bContinue = false;
+				break;
+			}
+			else if (const FStateTreeConditionBase* Cond = Node.Node.GetPtr<FStateTreeConditionBase>())
 			{
 				if (InFunc(&State, Node.ID, Node.Node.GetScriptStruct()->GetFName(), EStateTreeNodeType::EnterCondition, Node.Node.GetScriptStruct(), Cond->GetInstanceDataType()) == EStateTreeVisitor::Break)
 				{
@@ -860,7 +956,12 @@ EStateTreeVisitor UStateTreeEditorData::VisitStateNodes(const UStateTreeState& S
 		// Tasks
 		for (const FStateTreeEditorNode& Node : State.Tasks)
 		{
-			if (const FStateTreeTaskBase* Task = Node.Node.GetPtr<FStateTreeTaskBase>())
+			if (VisitFuncNodes(Node.ID, Node.GetName()) == EStateTreeVisitor::Break)
+			{
+				bContinue = false;
+				break;
+			}
+			else if (const FStateTreeTaskBase* Task = Node.Node.GetPtr<FStateTreeTaskBase>())
 			{
 				if (InFunc(&State, Node.ID, Task->Name, EStateTreeNodeType::Task, Node.Node.GetScriptStruct(), Task->GetInstanceDataType()) == EStateTreeVisitor::Break)
 				{
@@ -874,7 +975,11 @@ EStateTreeVisitor UStateTreeEditorData::VisitStateNodes(const UStateTreeState& S
 	{
 		if (const FStateTreeTaskBase* Task = State.SingleTask.Node.GetPtr<FStateTreeTaskBase>())
 		{
-			if (InFunc(&State, State.SingleTask.ID, Task->Name, EStateTreeNodeType::Task, State.SingleTask.Node.GetScriptStruct(), Task->GetInstanceDataType()) == EStateTreeVisitor::Break)
+			if (VisitFuncNodes(State.SingleTask.ID, State.SingleTask.GetName()) == EStateTreeVisitor::Break)
+			{
+				bContinue = false;
+			}
+			else if (InFunc(&State, State.SingleTask.ID, Task->Name, EStateTreeNodeType::Task, State.SingleTask.Node.GetScriptStruct(), Task->GetInstanceDataType()) == EStateTreeVisitor::Break)
 			{
 				bContinue = false;
 			}
@@ -888,7 +993,12 @@ EStateTreeVisitor UStateTreeEditorData::VisitStateNodes(const UStateTreeState& S
 		{
 			for (const FStateTreeEditorNode& Node : Transition.Conditions)
 			{
-				if (const FStateTreeConditionBase* Cond = Node.Node.GetPtr<FStateTreeConditionBase>())
+				if (VisitFuncNodes(Node.ID, Node.GetName()) == EStateTreeVisitor::Break)
+				{
+					bContinue = false;
+					break;
+				}
+				else if (const FStateTreeConditionBase* Cond = Node.Node.GetPtr<FStateTreeConditionBase>())
 				{
 					if (InFunc(&State, Node.ID, Node.Node.GetScriptStruct()->GetFName(), EStateTreeNodeType::TransitionCondition, Node.Node.GetScriptStruct(), Cond->GetInstanceDataType()) == EStateTreeVisitor::Break)
 					{
@@ -905,7 +1015,11 @@ EStateTreeVisitor UStateTreeEditorData::VisitStateNodes(const UStateTreeState& S
 		if (State.Type != EStateTreeStateType::Subtree
 			&& State.Parameters.Parameters.IsValid())
 		{
-			if (InFunc(&State, State.Parameters.ID, State.Name, EStateTreeNodeType::StateParameters, nullptr, State.Parameters.Parameters.GetPropertyBagStruct()) == EStateTreeVisitor::Break)
+			if (VisitFuncNodes(State.Parameters.ID, UE::StateTree::Editor::ParametersNodeName) == EStateTreeVisitor::Break)
+			{
+				bContinue = false;
+			}
+			else if (InFunc(&State, State.Parameters.ID, State.Name, EStateTreeNodeType::StateParameters, nullptr, State.Parameters.Parameters.GetPropertyBagStruct()) == EStateTreeVisitor::Break)
 			{
 				bContinue = false;
 			}
@@ -917,6 +1031,15 @@ EStateTreeVisitor UStateTreeEditorData::VisitStateNodes(const UStateTreeState& S
 
 EStateTreeVisitor UStateTreeEditorData::VisitStateNodes(const UStateTreeState& State, TFunctionRef<EStateTreeVisitor(const UStateTreeState* State, const FStateTreeBindableStructDesc& Desc, const FStateTreeDataView Value)> InFunc) const
 {
+	auto VisitFuncNodes = [&InFunc, &State, this](const FGuid StructID, const FName NodeName)
+	{
+		const FString StatePath = FString::Printf(TEXT("%s/%s"), *State.GetPath(), *NodeName.ToString());
+		return VisitStructBoundPropertyFunctions(StructID, StatePath, [&InFunc, &State](const FStateTreeEditorNode& EditorNode, const FStateTreeBindableStructDesc& Desc, FStateTreeDataView Value)
+		{
+			return InFunc(&State, Desc, Value);
+		});
+	};
+
 	bool bContinue = true;
 
 	const FString StatePath = State.GetPath();
@@ -926,17 +1049,24 @@ EStateTreeVisitor UStateTreeEditorData::VisitStateNodes(const UStateTreeState& S
 		// Bindable state parameters
 		if (State.Parameters.Parameters.IsValid())
 		{
-			FStateTreeBindableStructDesc Desc;
-			Desc.StatePath = StatePath;
-			Desc.Struct = State.Parameters.Parameters.GetPropertyBagStruct();
-			Desc.Name = FName("Parameters");
-			Desc.ID = State.Parameters.ID;
-			Desc.DataSource = EStateTreeBindableStructSource::StateParameter;
-
-			if (InFunc(&State, Desc, FStateTreeDataView(const_cast<FInstancedPropertyBag&>(State.Parameters.Parameters).GetMutableValue())) == EStateTreeVisitor::Break)
+			if (VisitFuncNodes(State.Parameters.ID, UE::StateTree::Editor::ParametersNodeName) == EStateTreeVisitor::Break)
 			{
 				bContinue = false;
 			}
+			else
+			{
+				FStateTreeBindableStructDesc Desc;
+				Desc.StatePath = StatePath;
+				Desc.Struct = State.Parameters.Parameters.GetPropertyBagStruct();
+				Desc.Name = FName("Parameters");
+				Desc.ID = State.Parameters.ID;
+				Desc.DataSource = EStateTreeBindableStructSource::StateParameter;
+
+				if (InFunc(&State, Desc, FStateTreeDataView(const_cast<FInstancedPropertyBag&>(State.Parameters.Parameters).GetMutableValue())) == EStateTreeVisitor::Break)
+				{
+					bContinue = false;
+				}
+			}
 		}
 	}
 
@@ -959,13 +1089,18 @@ EStateTreeVisitor UStateTreeEditorData::VisitStateNodes(const UStateTreeState& S
 			}
 		}
 	}
-	
+
 	if (bContinue)
 	{
 		// Enter conditions
 		for (const FStateTreeEditorNode& Node : State.EnterConditions)
 		{
-			if (const FStateTreeConditionBase* Cond = Node.Node.GetPtr<FStateTreeConditionBase>())
+			if (VisitFuncNodes(Node.ID, Node.GetName()) == EStateTreeVisitor::Break)
+			{
+				bContinue = false;
+				break;
+			}
+			else if (const FStateTreeConditionBase* Cond = Node.Node.GetPtr<FStateTreeConditionBase>())
 			{
 				FStateTreeBindableStructDesc Desc;
 				Desc.StatePath = StatePathWithConditions;
@@ -1010,7 +1145,12 @@ EStateTreeVisitor UStateTreeEditorData::VisitStateNodes(const UStateTreeState& S
 		// Tasks
 		for (const FStateTreeEditorNode& Node : State.Tasks)
 		{
-			if (const FStateTreeTaskBase* Task = Node.Node.GetPtr<FStateTreeTaskBase>())
+			if (VisitFuncNodes(Node.ID, Node.GetName()) == EStateTreeVisitor::Break)
+			{
+				bContinue = false;
+				break;
+			}
+			else if (const FStateTreeTaskBase* Task = Node.Node.GetPtr<FStateTreeTaskBase>())
 			{
 				FStateTreeBindableStructDesc Desc;
 				Desc.StatePath = StatePath;
@@ -1031,17 +1171,24 @@ EStateTreeVisitor UStateTreeEditorData::VisitStateNodes(const UStateTreeState& S
 	{
 		if (const FStateTreeTaskBase* Task = State.SingleTask.Node.GetPtr<FStateTreeTaskBase>())
 		{
-			FStateTreeBindableStructDesc Desc;
-			Desc.StatePath = StatePath;
-			Desc.Struct = Task->GetInstanceDataType();
-			Desc.Name = State.SingleTask.GetName();
-			Desc.ID = State.SingleTask.ID;
-			Desc.DataSource = EStateTreeBindableStructSource::Task;
-
-			if (InFunc(&State, Desc, State.SingleTask.GetInstance()) == EStateTreeVisitor::Break)
+			if (VisitFuncNodes(State.SingleTask.ID, State.SingleTask.GetName()) == EStateTreeVisitor::Break)
 			{
 				bContinue = false;
 			}
+			else
+			{
+				FStateTreeBindableStructDesc Desc;
+				Desc.StatePath = StatePath;
+				Desc.Struct = Task->GetInstanceDataType();
+				Desc.Name = State.SingleTask.GetName();
+				Desc.ID = State.SingleTask.ID;
+				Desc.DataSource = EStateTreeBindableStructSource::Task;
+
+				if (InFunc(&State, Desc, State.SingleTask.GetInstance()) == EStateTreeVisitor::Break)
+				{
+					bContinue = false;
+				}
+			}
 		}
 
 	}
@@ -1071,7 +1218,12 @@ EStateTreeVisitor UStateTreeEditorData::VisitStateNodes(const UStateTreeState& S
 
 			for (const FStateTreeEditorNode& Node : Transition.Conditions)
 			{
-				if (const FStateTreeConditionBase* Cond = Node.Node.GetPtr<FStateTreeConditionBase>())
+				if (VisitFuncNodes(Node.ID, Node.GetName()) == EStateTreeVisitor::Break)
+				{
+					bContinue = false;
+					break;
+				}
+				else if (const FStateTreeConditionBase* Cond = Node.Node.GetPtr<FStateTreeConditionBase>())
 				{
 					FStateTreeBindableStructDesc Desc;
 					Desc.StatePath = StatePathWithTransition;
@@ -1093,6 +1245,40 @@ EStateTreeVisitor UStateTreeEditorData::VisitStateNodes(const UStateTreeState& S
 	return bContinue ? EStateTreeVisitor::Continue : EStateTreeVisitor::Break;
 }
 
+EStateTreeVisitor UStateTreeEditorData::VisitStructBoundPropertyFunctions(FGuid StructID, const FString& StatePath, TFunctionRef<EStateTreeVisitor(const FStateTreeEditorNode& EditorNode, const FStateTreeBindableStructDesc& Desc, const FStateTreeDataView Value)> InFunc) const
+{
+	TArray<const FStateTreePropertyPathBinding*> Bindings;
+	EditorBindings.GetPropertyBindingsFor(StructID, Bindings);
+
+	for (const FStateTreePropertyPathBinding* Binding : Bindings)
+	{
+		const FConstStructView FunctionNodeView = Binding->GetPropertyFunctionNode();
+		if (const FStateTreeEditorNode* FunctionNode = FunctionNodeView.GetPtr<const FStateTreeEditorNode>())
+		{
+			if (VisitStructBoundPropertyFunctions(FunctionNode->ID, StatePath, InFunc) == EStateTreeVisitor::Break)
+			{
+				return EStateTreeVisitor::Break;
+			}
+
+			FStateTreeBindableStructDesc Desc;
+			Desc.Struct = FunctionNode->GetInstance().GetStruct();
+			if (const UStruct* NodeStruct = FunctionNode->Node.GetScriptStruct())
+			{
+				Desc.ID = FunctionNode->ID;
+				Desc.DataSource = EStateTreeBindableStructSource::PropertyFunction;
+				Desc.Name = FName(NodeStruct->GetDisplayNameText().ToString());
+				Desc.StatePath = FString::Printf(TEXT("%s/%s"), *StatePath, *UE::StateTree::Editor::PropertyFunctionStateName);
+
+				if (InFunc(*FunctionNode, Desc, FunctionNode->GetInstance()) == EStateTreeVisitor::Break)
+				{
+					return EStateTreeVisitor::Break;
+				}
+			}
+		}
+	}
+
+	return EStateTreeVisitor::Continue;
+}
 
 EStateTreeVisitor UStateTreeEditorData::VisitHierarchy(TFunctionRef<EStateTreeVisitor(UStateTreeState& State, UStateTreeState* ParentState)> InFunc) const
 {
@@ -1186,9 +1372,23 @@ EStateTreeVisitor UStateTreeEditorData::VisitGlobalNodes(TFunctionRef<EStateTree
 		}
 	}
 
+	auto VisitFuncNodesFunc = [&InFunc, this](const FStateTreeEditorNode& Node)
+	{
+		const FString StatePath = FString::Printf(TEXT("%s/%s"), *UE::StateTree::Editor::GlobalStateName, *Node.GetName().ToString());
+		return VisitStructBoundPropertyFunctions(Node.ID, StatePath, [&InFunc](const FStateTreeEditorNode& EditorNode, const FStateTreeBindableStructDesc& Desc, FStateTreeDataView Value)
+		{
+			return InFunc(nullptr, Desc, Value);
+		});
+	};
+
 	// Evaluators
 	for (const FStateTreeEditorNode& Node : Evaluators)
 	{
+		if (VisitFuncNodesFunc(Node) == EStateTreeVisitor::Break)
+		{
+			return EStateTreeVisitor::Break;
+		}
+
 		if (const FStateTreeEvaluatorBase* Evaluator = Node.Node.GetPtr<FStateTreeEvaluatorBase>())
 		{
 			FStateTreeBindableStructDesc Desc;
@@ -1208,6 +1408,11 @@ EStateTreeVisitor UStateTreeEditorData::VisitGlobalNodes(TFunctionRef<EStateTree
 	// Global tasks
 	for (const FStateTreeEditorNode& Node : GlobalTasks)
 	{
+		if (VisitFuncNodesFunc(Node) == EStateTreeVisitor::Break)
+		{
+			return EStateTreeVisitor::Break;
+		}
+
 		if (const FStateTreeTaskBase* Task = Node.Node.GetPtr<FStateTreeTaskBase>())
 		{
 			FStateTreeBindableStructDesc Desc;
@@ -1242,7 +1447,12 @@ EStateTreeVisitor UStateTreeEditorData::VisitAllNodes(TFunctionRef<EStateTreeVis
 		return EStateTreeVisitor::Break;
 	}
 
-	return VisitHierarchyNodes(InFunc);
+	if (VisitHierarchyNodes(InFunc) == EStateTreeVisitor::Break)
+	{
+		return EStateTreeVisitor::Break;
+	}
+
+	return EStateTreeVisitor::Continue;
 }
 
 #if WITH_STATETREE_TRACE_DEBUGGER
diff --git a/Engine/Plugins/Runtime/StateTree/Source/StateTreeEditorModule/Private/StateTreeEditorModule.cpp b/Engine/Plugins/Runtime/StateTree/Source/StateTreeEditorModule/Private/StateTreeEditorModule.cpp
--- a/Engine/Plugins/Runtime/StateTree/Source/StateTreeEditorModule/Private/StateTreeEditorModule.cpp
+++ b/Engine/Plugins/Runtime/StateTree/Source/StateTreeEditorModule/Private/StateTreeEditorModule.cpp
@@ -5,6 +5,7 @@
 #include "Blueprint/StateTreeConsiderationBlueprintBase.h"
 #include "Blueprint/StateTreeEvaluatorBlueprintBase.h"
 #include "Blueprint/StateTreeTaskBlueprintBase.h"
+#include "StateTreePropertyFunctionBase.h"
 #include "Customizations/StateTreeAnyEnumDetails.h"
 #include "Customizations/StateTreeEditorColorDetails.h"
 #include "Customizations/StateTreeEditorDataDetails.h"
@@ -156,6 +157,7 @@ TSharedPtr<FStateTreeNodeClassCache> FStateTreeEditorModule::GetNodeClassCache()
 		NodeClassCache->AddRootScriptStruct(FStateTreeTaskBase::StaticStruct());
 		NodeClassCache->AddRootScriptStruct(FStateTreeConditionBase::StaticStruct());
 		NodeClassCache->AddRootScriptStruct(FStateTreeConsiderationBase::StaticStruct());
+		NodeClassCache->AddRootScriptStruct(FStateTreePropertyFunctionBase::StaticStruct());
 		NodeClassCache->AddRootClass(UStateTreeEvaluatorBlueprintBase::StaticClass());
 		NodeClassCache->AddRootClass(UStateTreeTaskBlueprintBase::StaticClass());
 		NodeClassCache->AddRootClass(UStateTreeConditionBlueprintBase::StaticClass());
diff --git a/Engine/Plugins/Runtime/StateTree/Source/StateTreeEditorModule/Private/StateTreeEditorPropertyBindings.cpp b/Engine/Plugins/Runtime/StateTree/Source/StateTreeEditorModule/Private/StateTreeEditorPropertyBindings.cpp
--- a/Engine/Plugins/Runtime/StateTree/Source/StateTreeEditorModule/Private/StateTreeEditorPropertyBindings.cpp
+++ b/Engine/Plugins/Runtime/StateTree/Source/StateTreeEditorModule/Private/StateTreeEditorPropertyBindings.cpp
@@ -5,6 +5,8 @@
 #include "Misc/EnumerateRange.h"
 #include "PropertyPathHelpers.h"
 #include "StateTreeNodeBase.h"
+#include "StateTreePropertyFunctionBase.h"
+#include "StateTreeEditorNode.h"
 
 #include UE_INLINE_GENERATED_CPP_BY_NAME(StateTreeEditorPropertyBindings)
 
@@ -46,6 +48,33 @@ void FStateTreeEditorPropertyBindings::AddPropertyBinding(const FStateTreeProper
 	PropertyBindings.Add(Binding);
 }
 
+FStateTreePropertyPath FStateTreeEditorPropertyBindings::AddFunctionPropertyBinding(const UScriptStruct* PropertyFunctionNodeStruct, TConstArrayView<FStateTreePropertyPathSegment> SourcePathSegments, const FStateTreePropertyPath& TargetPath)
+{
+	check(PropertyFunctionNodeStruct->IsChildOf<FStateTreePropertyFunctionBase>());
+
+	FInstancedStruct PropertyFunctionNode(FStateTreeEditorNode::StaticStruct());
+	FStateTreeEditorNode& PropertyFunction = PropertyFunctionNode.GetMutable<FStateTreeEditorNode>();
+	const FGuid NodeID = FGuid::NewGuid();
+	PropertyFunction.ID = NodeID;
+	PropertyFunction.Node.InitializeAs(PropertyFunctionNodeStruct);
+	const FStateTreePropertyFunctionBase& Function = PropertyFunction.Node.Get<FStateTreePropertyFunctionBase>();
+	if (const UScriptStruct* InstanceType = Cast<const UScriptStruct>(Function.GetInstanceDataType()))
+	{
+		PropertyFunction.Instance.InitializeAs(InstanceType);
+	}
+
+	RemovePropertyBindings(TargetPath);
+	FStateTreePropertyPath SourcePath = FStateTreePropertyPath(NodeID, SourcePathSegments);
+	PropertyBindings.Emplace(MoveTemp(PropertyFunctionNode), SourcePath, TargetPath);
+	return SourcePath;
+}
+
+void FStateTreeEditorPropertyBindings::AddPropertyBinding(const FStateTreePropertyPathBinding& Binding)
+{
+	RemovePropertyBindings(Binding.GetTargetPath());
+	PropertyBindings.Add(Binding);
+}
+
 void FStateTreeEditorPropertyBindings::RemovePropertyBindings(const FStateTreePropertyPath& TargetPath)
 {
 	PropertyBindings.RemoveAll([&TargetPath](const FStateTreePropertyPathBinding& Binding)
@@ -86,13 +115,15 @@ const FStateTreePropertyPath* FStateTreeEditorPropertyBindings::GetPropertyBindi
 	return Binding ? &Binding->GetSourcePath() : nullptr;
 }
 
-void FStateTreeEditorPropertyBindings::GetPropertyBindingsFor(const FGuid StructID, TArray<FStateTreePropertyPathBinding>& OutBindings) const
+void FStateTreeEditorPropertyBindings::GetPropertyBindingsFor(const FGuid StructID, TArray<const FStateTreePropertyPathBinding*>& OutBindings) const
 {
-	OutBindings = PropertyBindings.FilterByPredicate([StructID](const FStateTreePropertyPathBinding& Binding)
+	for (const FStateTreePropertyPathBinding& Binding : PropertyBindings)
+	{
+		if (Binding.GetSourcePath().GetStructID().IsValid() && Binding.GetTargetPath().GetStructID() == StructID)
 		{
-			return Binding.GetSourcePath().GetStructID().IsValid()
-				&& Binding.GetTargetPath().GetStructID() == StructID;
-		});
+			OutBindings.Add(&Binding);
+		}
+	}
 }
 
 void FStateTreeEditorPropertyBindings::RemoveUnusedBindings(const TMap<FGuid, const FStateTreeDataView>& ValidStructs)
@@ -205,6 +236,13 @@ const FStateTreeEditorPropertyPath* FStateTreeEditorPropertyBindings::GetPropert
 	}
 	return nullptr;
 }
+
+void FStateTreeEditorPropertyBindings::GetPropertyBindingsFor(const FGuid StructID, TArray<FStateTreePropertyPathBinding>& OutBindings) const
+{
+	TArray<const FStateTreePropertyPathBinding*> NodeBindings;
+	GetPropertyBindingsFor(StructID, NodeBindings);
+	Algo::Transform(NodeBindings, OutBindings, [](const FStateTreePropertyPathBinding* BindingPtr) { return *BindingPtr; });
+}
 PRAGMA_ENABLE_DEPRECATION_WARNINGS
 
 void FStateTreeEditorPropertyBindings::RemoveUnusedBindings(const TMap<FGuid, const UStruct*>& ValidStructs)
@@ -260,17 +298,55 @@ FText FStateTreeBindingLookup::GetPropertyPathDisplayName(const FStateTreeProper
 	const FStateTreeEditorPropertyBindings* EditorBindings = BindingOwner->GetPropertyEditorBindings();
 	check(EditorBindings);
 
-	FString Result;
+	FString StructName;
+	int32 FirstSegmentToStringify = 0;
 
-	FStateTreeBindableStructDesc Struct;
-	if (BindingOwner->GetStructByID(InPath.GetStructID(), Struct))
+	// If path's struct is a PropertyFunction, let it override a display name.
 	{
-		Result = Struct.Name.ToString();
+		const FStateTreePropertyPathBinding* BindingToPath = EditorBindings->GetBindings().FindByPredicate([&InPath](const FStateTreePropertyPathBinding& Binding)
+		{
+			return Binding.GetSourcePath() == InPath;
+		});
+
+		if (BindingToPath && BindingToPath->GetPropertyFunctionNode().IsValid())
+		{
+			const FConstStructView PropertyFuncEditorNodeView = BindingToPath->GetPropertyFunctionNode();
+			const FStateTreeEditorNode& EditorNode = PropertyFuncEditorNodeView.Get<const FStateTreeEditorNode>();
+
+			if (!EditorNode.Node.IsValid())
+			{
+				return LOCTEXT("Unlinked", "???");
+			}
+
+			const FStateTreeNodeBase& Node = EditorNode.Node.Get<FStateTreeNodeBase>();
+
+			// Skipping an output property if there's only one of them.
+			if (UE::StateTree::GetStructSingleOutputProperty(*Node.GetInstanceDataType()))
+			{
+				FirstSegmentToStringify = 1;
+			}
+
+			const FText Description = Node.GetDescription(BindingToPath->GetSourcePath().GetStructID(), EditorNode.GetInstance(), *this, Formatting);
+			if (!Description.IsEmpty())
+			{
+				StructName = Description.ToString();
+			}
+		}
+	
+		if (StructName.IsEmpty())
+		{
+			FStateTreeBindableStructDesc Struct;
+			if (BindingOwner->GetStructByID(InPath.GetStructID(), Struct))
+			{
+				StructName = Struct.Name.ToString();
+			}
+		}
 	}
 
-	if (!InPath.IsPathEmpty())
+	FString Result = MoveTemp(StructName);
+	if (InPath.NumSegments() > FirstSegmentToStringify)
 	{
-		Result += TEXT(".") + InPath.ToString();
+		Result += TEXT(".") + InPath.ToString(/*HighlightedSegment*/ INDEX_NONE, /*HighlightPrefix*/ nullptr, /*HighlightPostfix*/ nullptr, /*bOutputInstances*/ false, FirstSegmentToStringify);
 	}
 
 	return FText::FromString(Result);
diff --git a/Engine/Plugins/Runtime/StateTree/Source/StateTreeEditorModule/Private/StateTreeEditorStyle.cpp b/Engine/Plugins/Runtime/StateTree/Source/StateTreeEditorModule/Private/StateTreeEditorStyle.cpp
--- a/Engine/Plugins/Runtime/StateTree/Source/StateTreeEditorModule/Private/StateTreeEditorStyle.cpp
+++ b/Engine/Plugins/Runtime/StateTree/Source/StateTreeEditorModule/Private/StateTreeEditorStyle.cpp
@@ -359,6 +359,7 @@ FStateTreeEditorStyle::FStateTreeEditorStyle()
 		Set("Node.RunParallel", new IMAGE_BRUSH_SVG("Icons/RunParallel", CoreStyleConstants::Icon16x16));
 		Set("Node.Task", new IMAGE_BRUSH_SVG("Icons/Task", CoreStyleConstants::Icon16x16));
 		Set("Node.Text", new IMAGE_BRUSH_SVG("Icons/Text", CoreStyleConstants::Icon16x16));
+		Set("Node.Function", new IMAGE_BRUSH_SVG("Icons/Function", CoreStyleConstants::Icon16x16));
 	}
 }
 
diff --git a/Engine/Plugins/Runtime/StateTree/Source/StateTreeEditorModule/Private/StateTreeNodeClassCache.cpp b/Engine/Plugins/Runtime/StateTree/Source/StateTreeEditorModule/Private/StateTreeNodeClassCache.cpp
--- a/Engine/Plugins/Runtime/StateTree/Source/StateTreeEditorModule/Private/StateTreeNodeClassCache.cpp
+++ b/Engine/Plugins/Runtime/StateTree/Source/StateTreeEditorModule/Private/StateTreeNodeClassCache.cpp
@@ -13,6 +13,7 @@
 #include "Logging/MessageLog.h"
 #include "AssetRegistry/AssetRegistryModule.h"
 #include "Misc/EnumerateRange.h"
+#include "StateTreeNodeBase.h"
 
 #include UE_INLINE_GENERATED_CPP_BY_NAME(StateTreeNodeClassCache)
 
@@ -80,6 +81,21 @@ UStruct* FStateTreeNodeClassData::GetStruct(bool bSilent)
 	return Ret;
 }
 
+const UStruct* FStateTreeNodeClassData::GetInstanceDataStruct(bool bSilent /*= false*/)
+{
+	if (!InstanceDataStruct.IsValid())
+	{
+		if (const UScriptStruct* ScriptStruct = GetScriptStruct(bSilent))
+		{
+			TInstancedStruct<FStateTreeNodeBase> NodeInstance;
+			NodeInstance.InitializeAsScriptStruct(ScriptStruct);
+
+			InstanceDataStruct = NodeInstance.Get().GetInstanceDataType();
+		}
+	}
+
+	return InstanceDataStruct.Get();
+}
 
 //----------------------------------------------------------------------//
 //  FStateTreeNodeClassCache
diff --git a/Engine/Plugins/Runtime/StateTree/Source/StateTreeEditorModule/Private/StateTreePropertyBindingCompiler.cpp b/Engine/Plugins/Runtime/StateTree/Source/StateTreeEditorModule/Private/StateTreePropertyBindingCompiler.cpp
--- a/Engine/Plugins/Runtime/StateTree/Source/StateTreeEditorModule/Private/StateTreePropertyBindingCompiler.cpp
+++ b/Engine/Plugins/Runtime/StateTree/Source/StateTreeEditorModule/Private/StateTreePropertyBindingCompiler.cpp
@@ -24,7 +24,7 @@ bool FStateTreePropertyBindingCompiler::Init(FStateTreePropertyBindings& InPrope
 	return true;
 }
 
-bool FStateTreePropertyBindingCompiler::CompileBatch(const FStateTreeBindableStructDesc& TargetStruct, TConstArrayView<FStateTreePropertyPathBinding> BatchPropertyBindings, int32& OutBatchIndex)
+bool FStateTreePropertyBindingCompiler::CompileBatch(const FStateTreeBindableStructDesc& TargetStruct, TConstArrayView<FStateTreePropertyPathBinding> BatchPropertyBindings, FStateTreeIndex16 PropertyFuncsBegin, FStateTreeIndex16 PropertyFuncsEnd, int32& OutBatchIndex)
 {
 	check(Log);
 	check(PropertyBindings);
@@ -123,8 +123,10 @@ bool FStateTreePropertyBindingCompiler::CompileBatch(const FStateTreeBindableStr
 
 		FStateTreePropertyCopyBatch& Batch = PropertyBindings->CopyBatches.AddDefaulted_GetRef();
 		Batch.TargetStruct = TargetStruct;
-		Batch.BindingsBegin = IntCastChecked<uint16>(BindingsBegin);
-		Batch.BindingsEnd = IntCastChecked<uint16>(BindingsEnd);
+		Batch.BindingsBegin = FStateTreeIndex16(BindingsBegin);
+		Batch.BindingsEnd = FStateTreeIndex16(BindingsEnd);
+		Batch.PropertyFunctionsBegin = PropertyFuncsBegin;
+		Batch.PropertyFunctionsEnd = PropertyFuncsEnd;
 		OutBatchIndex = PropertyBindings->CopyBatches.Num() - 1;
 	}
 
diff --git a/Engine/Plugins/Runtime/StateTree/Source/StateTreeEditorModule/Private/StateTreePropertyHelpers.cpp b/Engine/Plugins/Runtime/StateTree/Source/StateTreeEditorModule/Private/StateTreePropertyHelpers.cpp
--- a/Engine/Plugins/Runtime/StateTree/Source/StateTreeEditorModule/Private/StateTreePropertyHelpers.cpp
+++ b/Engine/Plugins/Runtime/StateTree/Source/StateTreeEditorModule/Private/StateTreePropertyHelpers.cpp
@@ -5,13 +5,80 @@
 #include "Hash/Blake3.h"
 #include "Misc/StringBuilder.h"
 #include "UObject/Field.h"
+#include "StateTreePropertyBindings.h"
+#include "StateTreeEditorData.h"
 
 namespace UE::StateTree::PropertyHelpers
 {
+namespace Internal
+{
+	void DispatchPostEditToEditorNode(FPropertyChangedChainEvent& InPropertyChangedEvent, const TDoubleLinkedList<FProperty*>::TDoubleLinkedListNode* InEditorNodeInChain, FStateTreeEditorNode& InEditorNode)
+	{
+		if (FStateTreeNodeBase* StateTreeNode = InEditorNode.Node.GetMutablePtr<FStateTreeNodeBase>())
+		{
+			// Check that the path contains EditorNode's: Node, Instance or Instance Object
+			if (const TDoubleLinkedList<FProperty*>::TDoubleLinkedListNode* EditorNodeMemberPropNode = InEditorNodeInChain->GetNextNode())
+			{
+				// Check that we have a changed property on one of the above properties.
+				if (const TDoubleLinkedList<FProperty*>::TDoubleLinkedListNode* ActiveMemberPropNode = EditorNodeMemberPropNode->GetNextNode()) 
+				{
+					// Update the event
+					const FProperty* EditorNodeChildMember = EditorNodeMemberPropNode->GetValue();
+					check(EditorNodeChildMember);
+
+					// Take copy of the event, we'll modify it.
+					FEditPropertyChain PropertyChainCopy;
+					for (const TDoubleLinkedList<FProperty*>::TDoubleLinkedListNode* Node = InPropertyChangedEvent.PropertyChain.GetHead(); Node; Node = Node->GetNextNode())
+					{
+						PropertyChainCopy.AddTail(Node->GetValue());
+					}
+					FPropertyChangedChainEvent PropertyChangedEvent(PropertyChainCopy, InPropertyChangedEvent);
+
+					PropertyChangedEvent.SetActiveMemberProperty(ActiveMemberPropNode->GetValue());
+					PropertyChangedEvent.PropertyChain.SetActiveMemberPropertyNode(PropertyChangedEvent.MemberProperty);
+
+					// To be consistent with the other property chain callbacks, do not cross object boundary.
+					const TDoubleLinkedList<FProperty*>::TDoubleLinkedListNode* ActivePropNode = ActiveMemberPropNode;
+					while (ActivePropNode->GetNextNode())
+					{
+						if (CastField<FObjectProperty>(ActivePropNode->GetValue()))
+						{
+							break;
+						}
+						ActivePropNode = ActivePropNode->GetNextNode();
+					}
+							
+					PropertyChangedEvent.Property = ActivePropNode->GetValue();
+					PropertyChangedEvent.PropertyChain.SetActivePropertyNode(PropertyChangedEvent.Property);
+
+					if (EditorNodeChildMember->GetFName() == GET_MEMBER_NAME_CHECKED(FStateTreeEditorNode, Node))
+					{
+						StateTreeNode->PostEditNodeChangeChainProperty(PropertyChangedEvent, InEditorNode.GetInstance());
+					}
+					else if (EditorNodeChildMember->GetFName() == GET_MEMBER_NAME_CHECKED(FStateTreeEditorNode, Instance))
+					{
+						if (InEditorNode.Instance.IsValid())
+						{
+							StateTreeNode->PostEditInstanceDataChangeChainProperty(PropertyChangedEvent, FStateTreeDataView(InEditorNode.Instance));
+						}
+					}
+					else if (EditorNodeChildMember->GetFName() == GET_MEMBER_NAME_CHECKED(FStateTreeEditorNode, InstanceObject))
+					{
+						if (InEditorNode.InstanceObject)
+						{
+							StateTreeNode->PostEditInstanceDataChangeChainProperty(PropertyChangedEvent, FStateTreeDataView(InEditorNode.InstanceObject));
+						}
+					}
+				}
+			}
+		}
+	}
+}
+
 
-void DispatchPostEditToNodes(UObject& Owner, FPropertyChangedChainEvent& InPropertyChangedEvent)
+void DispatchPostEditToNodes(UObject& Owner, FPropertyChangedChainEvent& InPropertyChangedEvent, UStateTreeEditorData& EditorData)
 {
-	// Walk back from the changed property and look for first FStateTreeEditorNode, and call the node specific post edit methods.
+	// Walk through changed property chain and look for first FStateTreeEditorNode, and call the node specific post edit methods.
 	
 	const TDoubleLinkedList<FProperty*>::TDoubleLinkedListNode* CurrentPropNode = InPropertyChangedEvent.PropertyChain.GetHead();
 	const FProperty* HeadProperty = CurrentPropNode->GetValue();
@@ -21,7 +88,11 @@ void DispatchPostEditToNodes(UObject& Owner, FPropertyChangedChainEvent& InPrope
 		return;
 	}
 	
+	FStateTreeEditorNode* LastEditorNode = nullptr;
+	const TDoubleLinkedList<FProperty*>::TDoubleLinkedListNode* LastEditorNodeInChain = nullptr;
+
 	uint8* CurrentAddress = reinterpret_cast<uint8*>(&Owner);
+	FStateTreePropertyPath TargetPath;
 	while (CurrentPropNode)
 	{
 		const FProperty* CurrentProperty = CurrentPropNode->GetValue();
@@ -37,88 +108,111 @@ void DispatchPostEditToNodes(UObject& Owner, FPropertyChangedChainEvent& InPrope
 				return;
 			}
 
+			if (TargetPath.GetStructID().IsValid())
+			{
+				TargetPath.AddPathSegment(ArrayProperty->GetFName(), Index);
+			}
+
 			CurrentAddress = Helper.GetRawPtr(Index);
 			CurrentProperty = ArrayProperty->Inner;
 		}
 
+		FStateTreePropertyPathSegment PathSegment(CurrentProperty->GetFName());
 		if (const FStructProperty* StructProperty = CastField<FStructProperty>(CurrentProperty))
 		{
 			if (StructProperty->Struct == FInstancedStruct::StaticStruct())
 			{
 				FInstancedStruct& InstancedStruct = *reinterpret_cast<FInstancedStruct*>(CurrentAddress);
 				CurrentAddress = InstancedStruct.GetMutableMemory();
+
+				PathSegment.SetInstanceStruct(InstancedStruct.GetScriptStruct());
 			}
 			else if (StructProperty->Struct == FStateTreeEditorNode::StaticStruct())
 			{
-				FStateTreeEditorNode& EditorNode = *reinterpret_cast<FStateTreeEditorNode*>(CurrentAddress);
-				if (FStateTreeNodeBase* StateTreeNode = EditorNode.Node.GetMutablePtr<FStateTreeNodeBase>())
+				if (TargetPath.GetStructID().IsValid())
 				{
-					// Check that the path contains EditorNode's: Node, Instance or Instance Object
-					if (const TDoubleLinkedList<FProperty*>::TDoubleLinkedListNode* EditorNodeMemberPropNode = CurrentPropNode->GetNextNode())
+					FStateTreePropertyPathBinding* FoundBinding = EditorData.GetPropertyEditorBindings()->GetMutableBindings().FindByPredicate([&TargetPath](const FStateTreePropertyPathBinding& Binding)
 					{
-						// Check that we have a changed property on one of the above properties.
-						if (const TDoubleLinkedList<FProperty*>::TDoubleLinkedListNode* ActiveMemberPropNode = EditorNodeMemberPropNode->GetNextNode()) 
-						{
-							// Update the event
-							const FProperty* EditorNodeChildMember = EditorNodeMemberPropNode->GetValue();
-							check(EditorNodeChildMember);
-
-							// Take copy of the event, we'll modify it.
-							FEditPropertyChain PropertyChainCopy;
-							for (const TDoubleLinkedList<FProperty*>::TDoubleLinkedListNode* Node = InPropertyChangedEvent.PropertyChain.GetHead(); Node; Node = Node->GetNextNode())
-							{
-								PropertyChainCopy.AddTail(Node->GetValue());
-							}
-							FPropertyChangedChainEvent PropertyChangedEvent(PropertyChainCopy, InPropertyChangedEvent);
-
-							PropertyChangedEvent.SetActiveMemberProperty(ActiveMemberPropNode->GetValue());
-							PropertyChangedEvent.PropertyChain.SetActiveMemberPropertyNode(PropertyChangedEvent.MemberProperty);
-
-							// To be consistent with the other property chain callbacks, do not cross object boundary.
-							const TDoubleLinkedList<FProperty*>::TDoubleLinkedListNode* ActivePropNode = ActiveMemberPropNode;
-							while (ActivePropNode->GetNextNode())
-							{
-								if (CastField<FObjectProperty>(ActivePropNode->GetValue()))
-								{
-									break;
-								}
-								ActivePropNode = ActivePropNode->GetNextNode();
-							}
-							
-							PropertyChangedEvent.Property = ActivePropNode->GetValue();
-							PropertyChangedEvent.PropertyChain.SetActivePropertyNode(PropertyChangedEvent.Property);
-
-							if (EditorNodeChildMember->GetFName() == GET_MEMBER_NAME_CHECKED(FStateTreeEditorNode, Node))
-							{
-								StateTreeNode->PostEditNodeChangeChainProperty(PropertyChangedEvent, EditorNode.GetInstance());
-							}
-							else if (EditorNodeChildMember->GetFName() == GET_MEMBER_NAME_CHECKED(FStateTreeEditorNode, Instance))
-							{
-								if (EditorNode.Instance.IsValid())
-								{
-									StateTreeNode->PostEditInstanceDataChangeChainProperty(PropertyChangedEvent, FStateTreeDataView(EditorNode.Instance));
-								}
-							}
-							else if (EditorNodeChildMember->GetFName() == GET_MEMBER_NAME_CHECKED(FStateTreeEditorNode, InstanceObject))
-							{
-								if (EditorNode.InstanceObject)
-								{
-									StateTreeNode->PostEditInstanceDataChangeChainProperty(PropertyChangedEvent, FStateTreeDataView(EditorNode.InstanceObject));
-								}
-							}
-						}
+						return TargetPath == Binding.GetTargetPath();
+					});
+
+					if (!ensure(FoundBinding && FoundBinding->GetPropertyFunctionNode().IsValid()))
+					{
+						return;
+					}
+
+					CurrentAddress = FoundBinding->GetMutablePropertyFunctionNode().GetMemory();
+					TargetPath.Reset();
+				}
+
+				LastEditorNode = reinterpret_cast<FStateTreeEditorNode*>(CurrentAddress);
+				LastEditorNodeInChain = CurrentPropNode;
+				TargetPath.SetStructID(LastEditorNode->ID);
+
+				CurrentPropNode = CurrentPropNode->GetNextNode();
+				if (CurrentPropNode)
+				{
+					const FName EditorNodeChildMemberName = CurrentPropNode->GetValue()->GetFName();
+					if (EditorNodeChildMemberName == GET_MEMBER_NAME_CHECKED(FStateTreeEditorNode, Instance) || EditorNodeChildMemberName == GET_MEMBER_NAME_CHECKED(FStateTreeEditorNode, InstanceObject))
+					{
+						CurrentAddress = LastEditorNode->GetInstance().GetMutableMemory();
+						CurrentPropNode = CurrentPropNode->GetNextNode();
+						continue;
 					}
 				}
 
 				break;
 			}
+			else if (StructProperty->Struct == FStateTreeStateParameters::StaticStruct())
+			{
+				FStateTreeStateParameters& StateParameters = *reinterpret_cast<FStateTreeStateParameters*>(CurrentAddress);
+				check(!TargetPath.GetStructID().IsValid());
+				TargetPath.SetStructID(StateParameters.ID);
 
-			CurrentPropNode = CurrentPropNode->GetNextNode();
+				CurrentPropNode = CurrentPropNode->GetNextNode();
+				if (CurrentPropNode && CurrentPropNode->GetValue()->GetFName() == GET_MEMBER_NAME_CHECKED(FStateTreeStateParameters, Parameters))
+				{
+					CurrentPropNode = CurrentPropNode->GetNextNode();
+					if (CurrentPropNode && CurrentPropNode->GetValue()->GetFName() == TEXT("Value"))
+					{
+						CurrentAddress = StateParameters.Parameters.GetMutableValue().GetMemory();
+						CurrentPropNode = CurrentPropNode->GetNextNode();
+						continue;
+					}
+				}
+
+				return;
+			}
 		}
-		else
+		else if (const FObjectProperty* ObjectProperty = CastField<FObjectProperty>(CurrentProperty))
 		{
-			break;
+			if (!TargetPath.GetStructID().IsValid())
+			{
+				return;
+			}
+
+			if (UObject* Object = *reinterpret_cast<UObject**>(CurrentAddress))
+			{
+				CurrentAddress = reinterpret_cast<uint8*>(Object);
+				PathSegment.SetInstanceStruct(Object->GetClass());
+			}
+			else
+			{
+				break;
+			}
+		}
+
+		if (TargetPath.GetStructID().IsValid())
+		{
+			TargetPath.AddPathSegment(PathSegment);
 		}
+
+		CurrentPropNode = CurrentPropNode->GetNextNode();
+	}
+
+	if (LastEditorNode && LastEditorNodeInChain)
+	{
+		Internal::DispatchPostEditToEditorNode(InPropertyChangedEvent, LastEditorNodeInChain, *LastEditorNode);
 	}
 }
 
diff --git a/Engine/Plugins/Runtime/StateTree/Source/StateTreeEditorModule/Private/StateTreePropertyHelpers.h b/Engine/Plugins/Runtime/StateTree/Source/StateTreeEditorModule/Private/StateTreePropertyHelpers.h
--- a/Engine/Plugins/Runtime/StateTree/Source/StateTreeEditorModule/Private/StateTreePropertyHelpers.h
+++ b/Engine/Plugins/Runtime/StateTree/Source/StateTreeEditorModule/Private/StateTreePropertyHelpers.h
@@ -8,6 +8,7 @@
 
 struct FGuid;
 struct FSlateBrush;
+class UStateTreeEditorData;
 
 #define LOCTEXT_NAMESPACE "StateTreeEditor"
 
@@ -17,7 +18,7 @@ namespace UE::StateTree::PropertyHelpers {
  * Dispatches PostEditChange to all FState
  * Assumes property chain head is member property of Owner. 
  */
-void DispatchPostEditToNodes(UObject& Owner, FPropertyChangedChainEvent& PropertyChangedEvent);
+void DispatchPostEditToNodes(UObject& Owner, FPropertyChangedChainEvent& PropertyChangedEvent, UStateTreeEditorData& EditorData);
 
 /** Makes deterministic ID from the owners property path, a property path (or any string), and a seed value (e.g. array index). */
 FGuid MakeDeterministicID(const UObject& Owner, const FString& PropertyPath, const uint64 Seed);
diff --git a/Engine/Plugins/Runtime/StateTree/Source/StateTreeEditorModule/Private/StateTreeState.cpp b/Engine/Plugins/Runtime/StateTree/Source/StateTreeEditorModule/Private/StateTreeState.cpp
--- a/Engine/Plugins/Runtime/StateTree/Source/StateTreeEditorModule/Private/StateTreeState.cpp
+++ b/Engine/Plugins/Runtime/StateTree/Source/StateTreeEditorModule/Private/StateTreeState.cpp
@@ -357,7 +357,10 @@ void UStateTreeState::PostEditChangeChainProperty(FPropertyChangedChainEvent& Pr
 		}
 	}
 
-	UE::StateTree::PropertyHelpers::DispatchPostEditToNodes(*this, PropertyChangedEvent);
+	if (UStateTreeEditorData* TreeData = GetTypedOuter<UStateTreeEditorData>())
+	{
+		UE::StateTree::PropertyHelpers::DispatchPostEditToNodes(*this, PropertyChangedEvent, *TreeData);
+	}
 }
 
 void UStateTreeState::PostLoad()
diff --git a/Engine/Plugins/Runtime/StateTree/Source/StateTreeEditorModule/Private/StateTreeViewModel.cpp b/Engine/Plugins/Runtime/StateTree/Source/StateTreeEditorModule/Private/StateTreeViewModel.cpp
--- a/Engine/Plugins/Runtime/StateTree/Source/StateTreeEditorModule/Private/StateTreeViewModel.cpp
+++ b/Engine/Plugins/Runtime/StateTree/Source/StateTreeEditorModule/Private/StateTreeViewModel.cpp
@@ -49,7 +49,7 @@ namespace UE::StateTree::Editor
 	};
 
 
-	void CollectBindingsRecursive(UStateTreeEditorData* TreeData, UStateTreeState* State, TArray<FStateTreePropertyPathBinding>& AllBindings)
+	void CollectBindingsCopiesRecursive(UStateTreeEditorData* TreeData, UStateTreeState* State, TArray<FStateTreePropertyPathBinding>& AllBindings)
 	{
 		if (!State)
 		{
@@ -58,15 +58,15 @@ namespace UE::StateTree::Editor
 		
 		TreeData->VisitStateNodes(*State, [TreeData, &AllBindings](const UStateTreeState* State, const FStateTreeBindableStructDesc& Desc, const FStateTreeDataView Value)
 		{
-			TArray<FStateTreePropertyPathBinding> NodeBindings;
+			TArray<const FStateTreePropertyPathBinding*> NodeBindings;
 			TreeData->GetPropertyEditorBindings()->GetPropertyBindingsFor(Desc.ID, NodeBindings);
-			AllBindings.Append(NodeBindings);
+			Algo::Transform(NodeBindings, AllBindings, [](const FStateTreePropertyPathBinding* BindingPtr) { return *BindingPtr; });
 			return EStateTreeVisitor::Continue;				
 		});
 
 		for (UStateTreeState* ChildState : State->Children)
 		{
-			CollectBindingsRecursive(TreeData, ChildState, AllBindings);
+			CollectBindingsCopiesRecursive(TreeData, ChildState, AllBindings);
 		}
 	}
 
@@ -91,7 +91,7 @@ namespace UE::StateTree::Editor
 			UObject* ThisOuter = State->GetOuter();
 			UExporter::ExportToOutputDevice(&Context, State, nullptr, Archive, TEXT("copy"), 0, PPF_ExportsNotFullyQualified | PPF_Copy | PPF_Delimited, false, ThisOuter);
 
-			CollectBindingsRecursive(TreeData, State, ClipboardBindings->Bindings);
+			CollectBindingsCopiesRecursive(TreeData, State, ClipboardBindings->Bindings);
 		}
 
 		UExporter::ExportToOutputDevice(&Context, ClipboardBindings, nullptr, Archive, TEXT("copy"), 0, PPF_ExportsNotFullyQualified | PPF_Copy | PPF_Delimited, false);
@@ -952,25 +952,31 @@ void FStateTreeViewModel::PasteStatesAsChildrenFromText(const FString& TextToImp
 	// Copy property bindings for the duplicated states.
 	if (Factory.ClipboardBindings)
 	{
+		for (FStateTreePropertyPathBinding& Binding : Factory.ClipboardBindings->Bindings)
+		{
+			if (Binding.GetPropertyFunctionNode().IsValid())
+			{
+				UE::StateTree::Editor::FixNodesAfterDuplication(TArrayView<FStateTreeEditorNode>(Binding.GetMutablePropertyFunctionNode().GetPtr<FStateTreeEditorNode>(), 1), IDsMap, Links);
+			}
+		}
+
 		for (const TPair<FGuid, FGuid>& Entry : IDsMap)
 		{
 			const FGuid OldTargetID = Entry.Key;
 			const FGuid NewTargetID = Entry.Value;
 			
-			for (const FStateTreePropertyPathBinding& Binding : Factory.ClipboardBindings->Bindings)
+			for (FStateTreePropertyPathBinding& Binding : Factory.ClipboardBindings->Bindings)
 			{
 				if (Binding.GetTargetPath().GetStructID() == OldTargetID)
 				{
-					FStateTreePropertyPath TargetPath(Binding.GetTargetPath());
-					TargetPath.SetStructID(NewTargetID);
-					
-					FStateTreePropertyPath SourcePath(Binding.GetSourcePath());
+					Binding.GetMutableTargetPath().SetStructID(NewTargetID);
+
 					if (const FGuid* NewSourceID = IDsMap.Find(Binding.GetSourcePath().GetStructID()))
 					{
-						SourcePath.SetStructID(*NewSourceID);
+						Binding.GetMutableSourcePath().SetStructID(*NewSourceID);
 					}
 					
-					TreeData->GetPropertyEditorBindings()->AddPropertyBinding(SourcePath, TargetPath);
+					TreeData->GetPropertyEditorBindings()->AddPropertyBinding(MoveTemp(Binding));
 				}
 			}
 		}
diff --git a/Engine/Plugins/Runtime/StateTree/Source/StateTreeEditorModule/Public/StateTreeCompiler.h b/Engine/Plugins/Runtime/StateTree/Source/StateTreeEditorModule/Public/StateTreeCompiler.h
--- a/Engine/Plugins/Runtime/StateTree/Source/StateTreeEditorModule/Public/StateTreeCompiler.h
+++ b/Engine/Plugins/Runtime/StateTree/Source/StateTreeEditorModule/Public/StateTreeCompiler.h
@@ -48,9 +48,19 @@ private:
 	bool CreateGlobalTasks();
 	bool CreateStateTasksAndParameters();
 	bool CreateStateTransitions();
+
+	bool CreateBindingsForNodes(TConstArrayView<FStateTreeEditorNode> EditorNodes, FStateTreeIndex16 NodesBegin);
+	bool CreateBindingsForStruct(const FStateTreeBindableStructDesc& TargetStruct, FStateTreeDataView TargetValue, FStateTreeIndex16 PropertyFuncsBegin, FStateTreeIndex16 PropertyFuncsEnd, FStateTreeIndex16& OutBatchIndex);
+
+	bool CreatePropertyFunctionsForStruct(FGuid StructID);
+	bool CreatePropertyFunction(const FStateTreeEditorNode& FuncEditorNode);
 	
 	bool CreateConditions(UStateTreeState& State, const FString& StatePath, TConstArrayView<FStateTreeEditorNode> Conditions);
 	bool CreateCondition(UStateTreeState& State, const FString& StatePath, const FStateTreeEditorNode& CondNode, const EStateTreeExpressionOperand Operand, const int8 DeltaIndent);
+
+	template<class T = FStateTreeNodeBase>
+	T* CreateNodeWithSharedInstanceData(UStateTreeState* State, const FString& StatePath, const FStateTreeEditorNode& EditorNode, EStateTreeBindableStructSource StructSource);
+
 	bool CreateTask(UStateTreeState* State, const FStateTreeEditorNode& TaskNode, const FStateTreeDataHandle TaskDataHandle);
 	bool CreateEvaluator(const FStateTreeEditorNode& EvalNode, const FStateTreeDataHandle EvalDataHandle);
 	bool GetAndValidateBindings(const FStateTreeBindableStructDesc& TargetStruct, FStateTreeDataView TargetValue, TArray<FStateTreePropertyPathBinding>& OutCopyBindings, TArray<FStateTreePropertyPathBinding>& OutReferenceBindings) const;
diff --git a/Engine/Plugins/Runtime/StateTree/Source/StateTreeEditorModule/Public/StateTreeEditorData.h b/Engine/Plugins/Runtime/StateTree/Source/StateTreeEditorModule/Public/StateTreeEditorData.h
--- a/Engine/Plugins/Runtime/StateTree/Source/StateTreeEditorModule/Public/StateTreeEditorData.h
+++ b/Engine/Plugins/Runtime/StateTree/Source/StateTreeEditorModule/Public/StateTreeEditorData.h
@@ -17,6 +17,9 @@ namespace UE::StateTree::Editor
 {
 	// Name used to describe container of global items (other items use the path to the container State).  
 	extern STATETREEEDITORMODULE_API const FString GlobalStateName;
+
+	// Name used to describe container of property functions.
+	extern STATETREEEDITORMODULE_API const FString PropertyFunctionStateName;
 }
 
 
@@ -69,6 +72,8 @@ public:
 	virtual const FStateTreeEditorPropertyBindings* GetPropertyEditorBindings() const override { return &EditorBindings; }
 	virtual FStateTreeBindableStructDesc FindContextData(const UStruct* ObjectType, const FString ObjectNameHint) const override;
 
+	virtual EStateTreeVisitor EnumerateBindablePropertyFunctionNodes(TFunctionRef<EStateTreeVisitor(const UScriptStruct* NodeStruct, const FStateTreeBindableStructDesc& Desc, const FStateTreeDataView Value)> InFunc) const;
+
 	// ~IStateTreeEditorPropertyBindingsOwner
 
 	/**
@@ -103,34 +108,40 @@ public:
 
 	/**
 	* Iterates over all structs that are related to binding
-	* @param InFunc function called at each node, should return true if visiting is continued or false to stop.
+	* @param InFunc function called at each node, should return Continue if visiting is continued or Break to stop.
 	*/
 	EStateTreeVisitor VisitHierarchy(TFunctionRef<EStateTreeVisitor(UStateTreeState& State, UStateTreeState* ParentState)> InFunc) const;
 
 	/**
 	 * Iterates over all structs at the global level (context, tree parameters, evaluators, global tasks) that are related to binding.
-	 * @param InFunc function called at each node, should return true if visiting is continued or false to stop.
+	 * @param InFunc function called at each node, should return Continue if visiting is continued or Break to stop.
 	 */
 	EStateTreeVisitor VisitGlobalNodes(TFunctionRef<EStateTreeVisitor(const UStateTreeState* State, const FStateTreeBindableStructDesc& Desc, const FStateTreeDataView Value)> InFunc) const;
 
 	/**
 	 * Iterates over all structs in the state hierarchy that are related to binding.
-	 * @param InFunc function called at each node, should return true if visiting is continued or false to stop.
+	 * @param InFunc function called at each node, should return Continue if visiting is continued or Break to stop.
 	 */
 	EStateTreeVisitor VisitHierarchyNodes(TFunctionRef<EStateTreeVisitor(const UStateTreeState* State, const FStateTreeBindableStructDesc& Desc, const FStateTreeDataView Value)> InFunc) const;
 
 	/**
 	 * Iterates over all structs that are related to binding.
-	 * @param InFunc function called at each node, should return true if visiting is continued or false to stop.
+	 * @param InFunc function called at each node, should return Continue if visiting is continued or Break to stop.
 	 */
 	EStateTreeVisitor VisitAllNodes(TFunctionRef<EStateTreeVisitor(const UStateTreeState* State, const FStateTreeBindableStructDesc& Desc, const FStateTreeDataView Value)> InFunc) const;
 
 	/**
 	 * Iterates over all nodes in a given state.
-	 * @param InFunc function called at each node, should return true if visiting is continued or false to stop.
+	 * @param InFunc function called at each node, should return Continue if visiting is continued or Break to stop.
 	 */
 	EStateTreeVisitor VisitStateNodes(const UStateTreeState& State, TFunctionRef<EStateTreeVisitor(const UStateTreeState* State, const FStateTreeBindableStructDesc& Desc, const FStateTreeDataView Value)> InFunc) const;
 
+	/**
+	 * Iterates recursively over all property functions of the provided node. Also nested ones.
+	 * @param InFunc function called at each node, should return Continue if visiting is continued or Break to stop.
+	 */
+	EStateTreeVisitor VisitStructBoundPropertyFunctions(FGuid StructID, const FString& StatePath, TFunctionRef<EStateTreeVisitor(const FStateTreeEditorNode& EditorNode, const FStateTreeBindableStructDesc& Desc, const FStateTreeDataView Value)> InFunc) const;
+
 	/**
 	 * Returns array of nodes along the execution path, up to the TargetStruct.
 	 * @param Path The states to visit during the check
@@ -218,6 +229,14 @@ public:
 		EditorBindings.AddPropertyBinding(SourcePath, TargetPath);
 	}
 
+	/**
+	 * Adds property binding to PropertyFunction of provided type.
+	 */
+	void AddPropertyBinding(const UScriptStruct* PropertyFunctionNodeStruct, TConstArrayView<FStateTreePropertyPathSegment> SourcePathSegments, const FStateTreePropertyPath& TargetPath)
+	{
+		EditorBindings.AddFunctionPropertyBinding(PropertyFunctionNodeStruct, SourcePathSegments, TargetPath);
+	}
+
 PRAGMA_DISABLE_DEPRECATION_WARNINGS
 	UE_DEPRECATED(5.3, "Use version with FStateTreePropertyPath instead.")
 	void AddPropertyBinding(const FStateTreeEditorPropertyPath& SourcePath, const FStateTreeEditorPropertyPath& TargetPath);
diff --git a/Engine/Plugins/Runtime/StateTree/Source/StateTreeEditorModule/Public/StateTreeEditorNode.h b/Engine/Plugins/Runtime/StateTree/Source/StateTreeEditorModule/Public/StateTreeEditorNode.h
--- a/Engine/Plugins/Runtime/StateTree/Source/StateTreeEditorModule/Public/StateTreeEditorNode.h
+++ b/Engine/Plugins/Runtime/StateTree/Source/StateTreeEditorModule/Public/StateTreeEditorNode.h
@@ -13,6 +13,7 @@ enum class EStateTreeNodeType : uint8
 	Task,
 	TransitionCondition,
 	StateParameters,
+	PropertyFunction,
 };
 
 /**
diff --git a/Engine/Plugins/Runtime/StateTree/Source/StateTreeEditorModule/Public/StateTreeEditorPropertyBindings.h b/Engine/Plugins/Runtime/StateTree/Source/StateTreeEditorModule/Public/StateTreeEditorPropertyBindings.h
--- a/Engine/Plugins/Runtime/StateTree/Source/StateTreeEditorModule/Public/StateTreeEditorPropertyBindings.h
+++ b/Engine/Plugins/Runtime/StateTree/Source/StateTreeEditorModule/Public/StateTreeEditorPropertyBindings.h
@@ -4,10 +4,12 @@
 
 #include "UObject/Interface.h"
 #include "StateTreePropertyBindings.h"
+#include "StateTreeEditorNode.h"
 #include "StateTreeEditorPropertyBindings.generated.h"
 
 enum EStateTreeNodeFormatting : uint8;
 class IStateTreeEditorPropertyBindingsOwner;
+enum class EStateTreeVisitor : uint8;
 
 /**
  * Editor representation of a all property bindings in a StateTree
@@ -27,6 +29,21 @@ struct STATETREEEDITORMODULE_API FStateTreeEditorPropertyBindings
 	 */
 	void AddPropertyBinding(const FStateTreePropertyPath& SourcePath, const FStateTreePropertyPath& TargetPath);
 	
+	/**
+	 * Adds binding.
+	 * @param Binding Binding to be added.
+	*/
+	void AddPropertyBinding(const FStateTreePropertyPathBinding& Binding);
+
+	/**
+	 * Adds binding between PropertyFunction of the provided type and destination path.
+	 * @param PropertyFunctionNodeStruct Struct of PropertyFunction.
+	 * @param SourcePathSegments Binding source property path segments.
+	 * @param TargetPath Binding target property path.
+	 * @return Constructed binding source property path.
+	 */
+	FStateTreePropertyPath AddFunctionPropertyBinding(const UScriptStruct* PropertyFunctionNodeStruct, TConstArrayView<FStateTreePropertyPathSegment> SourcePathSegments, const FStateTreePropertyPath& TargetPath);
+
 	/**
 	 * Removes all bindings to target path.
 	 * @param TargetPath Target property path.
@@ -53,11 +70,11 @@ struct STATETREEEDITORMODULE_API FStateTreeEditorPropertyBindings
 	const FStateTreePropertyPath* GetPropertyBindingSource(const FStateTreePropertyPath& TargetPath) const;
 	
 	/**
-	 * Returns all bindings for a specified structs based in struct ID.
+	 * Returns all pointers to bindings for a specified structs based in struct ID.
 	 * @param StructID ID of the struct to find bindings for.
 	 * @param OutBindings Bindings for specified struct.
 	 */
-	void GetPropertyBindingsFor(const FGuid StructID, TArray<FStateTreePropertyPathBinding>& OutBindings) const;
+	void GetPropertyBindingsFor(const FGuid StructID, TArray<const FStateTreePropertyPathBinding*>& OutBindings) const;
 	
 	/**
 	 * Removes bindings which do not point to valid structs IDs.
@@ -88,6 +105,9 @@ PRAGMA_DISABLE_DEPRECATION_WARNINGS
 
 	UE_DEPRECATED(5.3, "Use RemoveUnusedBindings with values instead.")
 	void RemoveUnusedBindings(const TMap<FGuid, const UStruct*>& ValidStructs);
+
+	UE_DEPRECATED(5.5, "Use GetPropertyBindingsFor returning pointers instead.")
+	void GetPropertyBindingsFor(const FGuid StructID, TArray<FStateTreePropertyPathBinding>& OutBindings) const;
 PRAGMA_ENABLE_DEPRECATION_WARNINGS
 	
 private:
@@ -144,6 +164,8 @@ class STATETREEEDITORMODULE_API IStateTreeEditorPropertyBindingsOwner
 
 	/** @return Pointer to editor property bindings. */
 	virtual const FStateTreeEditorPropertyBindings* GetPropertyEditorBindings() const PURE_VIRTUAL(IStateTreeEditorPropertyBindingsOwner::GetPropertyEditorBindings, return nullptr; );
+
+	virtual EStateTreeVisitor EnumerateBindablePropertyFunctionNodes(TFunctionRef<EStateTreeVisitor(const UScriptStruct* NodeStruct, const FStateTreeBindableStructDesc& Desc, const FStateTreeDataView Value)> InFunc) const PURE_VIRTUAL(IStateTreeEditorPropertyBindingsOwner::EnumerateBindablePropertyFunctionNodes, return static_cast<EStateTreeVisitor>(0); );
 };
 
 // TODO: We should merge this with IStateTreeEditorPropertyBindingsOwner and FStateTreeEditorPropertyBindings.
diff --git a/Engine/Plugins/Runtime/StateTree/Source/StateTreeEditorModule/Public/StateTreeNodeClassCache.h b/Engine/Plugins/Runtime/StateTree/Source/StateTreeEditorModule/Public/StateTreeNodeClassCache.h
--- a/Engine/Plugins/Runtime/StateTree/Source/StateTreeEditorModule/Public/StateTreeNodeClassCache.h
+++ b/Engine/Plugins/Runtime/StateTree/Source/StateTreeEditorModule/Public/StateTreeNodeClassCache.h
@@ -7,6 +7,7 @@
 enum class EReloadCompleteReason;
 
 struct FAssetData;
+struct FStateTreeNodeBase;
 
 /**
  * Describes a class or struct.
@@ -34,6 +35,8 @@ struct STATETREEEDITORMODULE_API FStateTreeNodeClassData
 		return Cast<UScriptStruct>(GetStruct(bSilent));
 	}
 
+	const UStruct* GetInstanceDataStruct(bool bSilent = false);
+
 private:
 
 	/** Pointer to described struct or class. */
@@ -47,6 +50,9 @@ private:
 
 	/** Package of the asset if it's not loaded yet. */
 	FString ClassPackageName;
+
+	/** Pointer to described node's instance data struct or class. */
+	TWeakObjectPtr<const UStruct> InstanceDataStruct;
 };
 
 /**
diff --git a/Engine/Plugins/Runtime/StateTree/Source/StateTreeEditorModule/Public/StateTreePropertyBindingCompiler.h b/Engine/Plugins/Runtime/StateTree/Source/StateTreeEditorModule/Public/StateTreePropertyBindingCompiler.h
--- a/Engine/Plugins/Runtime/StateTree/Source/StateTreeEditorModule/Public/StateTreePropertyBindingCompiler.h
+++ b/Engine/Plugins/Runtime/StateTree/Source/StateTreeEditorModule/Public/StateTreePropertyBindingCompiler.h
@@ -29,10 +29,12 @@ struct STATETREEEDITORMODULE_API FStateTreePropertyBindingCompiler
 	  * Compiles a batch of property copies.
 	  * @param TargetStruct - Description of the structs which contains the target properties.
 	  * @param PropertyBindings - Array of bindings to compile, all bindings that point to TargetStructs will be added to the batch.
+	  * @param PropertyFuncsBegin - Index of the first PropertyFunction belonging to this batch.
+	  * @param PropertyFuncsEnd - Index of the last PropertyFunction belonging to this batch.
 	  * @param OutBatchIndex - Resulting batch index, if index is INDEX_NONE, no bindings were found and no batch was generated.
 	  * @return True on success, false on failure.
 	 */
-	[[nodiscard]] bool CompileBatch(const FStateTreeBindableStructDesc& TargetStruct, TConstArrayView<FStateTreePropertyPathBinding> PropertyBindings, int32& OutBatchIndex);
+	[[nodiscard]] bool CompileBatch(const FStateTreeBindableStructDesc& TargetStruct, TConstArrayView<FStateTreePropertyPathBinding> PropertyBindings, FStateTreeIndex16 PropertyFuncsBegin, FStateTreeIndex16 PropertyFuncsEnd, int32& OutBatchIndex);
 
 	/**
 	  * Compiles references for selected struct
@@ -99,6 +101,12 @@ PRAGMA_DISABLE_DEPRECATION_WARNINGS
 	 */
 	UE_DEPRECATED(5.3, "Use FStateTreePropertyBindings::GetPropertyCompatibility instead.")
 	static EPropertyAccessCompatibility GetPropertyCompatibility(const FProperty* FromProperty, const FProperty* ToProperty);
+
+	UE_DEPRECATED(5.5, "Use CompileBatch with PropertyFuncsBegin and PropertyFuncsEnd instead.")
+	[[nodiscard]] bool CompileBatch(const FStateTreeBindableStructDesc& InTargetStruct, TConstArrayView<FStateTreePropertyPathBinding> InPropertyBindings, int32& OutBatchIndex) 
+	{
+		return CompileBatch(InTargetStruct, InPropertyBindings, FStateTreeIndex16::Invalid, FStateTreeIndex16::Invalid, OutBatchIndex);
+	}
 PRAGMA_ENABLE_DEPRECATION_WARNINGS
 
 protected:
diff --git a/Engine/Plugins/Runtime/StateTree/Source/StateTreeModule/Private/PropertyFunctions/StateTreeBooleanAlgebraPropertyFunctions.cpp b/Engine/Plugins/Runtime/StateTree/Source/StateTreeModule/Private/PropertyFunctions/StateTreeBooleanAlgebraPropertyFunctions.cpp
new file mode 100644
--- /dev/null
+++ b/Engine/Plugins/Runtime/StateTree/Source/StateTreeModule/Private/PropertyFunctions/StateTreeBooleanAlgebraPropertyFunctions.cpp
@@ -0,0 +1,66 @@
+// Copyright Epic Games, Inc. All Rights Reserved.
+
+#include "StateTreeBooleanAlgebraPropertyFunctions.h"
+#include "StateTreeExecutionContext.h"
+#include "StateTreeNodeDescriptionHelpers.h"
+
+#define LOCTEXT_NAMESPACE "StateTree"
+
+namespace UE::StateTree::BooleanPropertyFunctions::Internal
+{
+#if WITH_EDITOR
+	FText GetDescriptionForOperation(FText OperationText, const FGuid& ID, FStateTreeDataView InstanceDataView, const IStateTreeBindingLookup& BindingLookup, EStateTreeNodeFormatting Formatting)
+	{
+		const FStateTreeBooleanOperationPropertyFunctionInstanceData& InstanceData = InstanceDataView.Get<FStateTreeBooleanOperationPropertyFunctionInstanceData>();
+
+		FText LeftValue = BindingLookup.GetBindingSourceDisplayName(FStateTreePropertyPath(ID, GET_MEMBER_NAME_CHECKED(FStateTreeBooleanOperationPropertyFunctionInstanceData, bLeft)), Formatting);
+		if (LeftValue.IsEmpty())
+		{
+			LeftValue = UE::StateTree::DescHelpers::GetBoolText(InstanceData.bLeft, Formatting);
+		}
+
+		FText RightValue = BindingLookup.GetBindingSourceDisplayName(FStateTreePropertyPath(ID, GET_MEMBER_NAME_CHECKED(FStateTreeBooleanOperationPropertyFunctionInstanceData, bRight)), Formatting);
+		if (RightValue.IsEmpty())
+		{
+			RightValue = UE::StateTree::DescHelpers::GetBoolText(InstanceData.bRight, Formatting);
+		}
+
+		const FText Format = (Formatting == EStateTreeNodeFormatting::RichText)
+			? LOCTEXT("BoolFuncRich", "({Left} <s>{Operation}</> {Right})")
+			: LOCTEXT("BoolFunc", "({Left} {Operation} {Right})");
+
+		return FText::FormatNamed(Format,
+			TEXT("Left"), LeftValue,
+			TEXT("Operation"), OperationText,
+			TEXT("Right"), RightValue);
+	}
+#endif // WITH_EDITOR
+}
+
+void FStateTreeBooleanAndPropertyFunction::Execute(FStateTreeExecutionContext& Context) const
+{
+	FStateTreeBooleanOperationPropertyFunctionInstanceData& InstanceData = Context.GetInstanceData(*this);
+	InstanceData.bResult = InstanceData.bLeft && InstanceData.bRight;
+}
+
+#if WITH_EDITOR
+FText FStateTreeBooleanAndPropertyFunction::GetDescription(const FGuid& ID, FStateTreeDataView InstanceDataView, const IStateTreeBindingLookup& BindingLookup, EStateTreeNodeFormatting Formatting) const
+{
+	return UE::StateTree::BooleanPropertyFunctions::Internal::GetDescriptionForOperation(LOCTEXT("BoolAnd", "and"), ID, InstanceDataView, BindingLookup, Formatting);
+}
+#endif // WITH_EDITOR
+
+void FStateTreeBooleanOrPropertyFunction::Execute(FStateTreeExecutionContext& Context) const
+{
+	FStateTreeBooleanOperationPropertyFunctionInstanceData& InstanceData = Context.GetInstanceData(*this);
+	InstanceData.bResult = InstanceData.bLeft || InstanceData.bRight;
+}
+
+#if WITH_EDITOR
+FText FStateTreeBooleanOrPropertyFunction::GetDescription(const FGuid& ID, FStateTreeDataView InstanceDataView, const IStateTreeBindingLookup& BindingLookup, EStateTreeNodeFormatting Formatting) const
+{
+	return UE::StateTree::BooleanPropertyFunctions::Internal::GetDescriptionForOperation(LOCTEXT("BoolOr", "or"), ID, InstanceDataView, BindingLookup, Formatting);
+}
+#endif // WITH_EDITOR
+
+#undef LOCTEXT_NAMESPACE
diff --git a/Engine/Plugins/Runtime/StateTree/Source/StateTreeModule/Private/PropertyFunctions/StateTreeBooleanAlgebraPropertyFunctions.h b/Engine/Plugins/Runtime/StateTree/Source/StateTreeModule/Private/PropertyFunctions/StateTreeBooleanAlgebraPropertyFunctions.h
new file mode 100644
--- /dev/null
+++ b/Engine/Plugins/Runtime/StateTree/Source/StateTreeModule/Private/PropertyFunctions/StateTreeBooleanAlgebraPropertyFunctions.h
@@ -0,0 +1,61 @@
+// Copyright Epic Games, Inc. All Rights Reserved.
+
+#pragma once
+
+#include "StateTreePropertyFunctionBase.h"
+#include "StateTreeBooleanAlgebraPropertyFunctions.generated.h"
+
+struct FStateTreeExecutionContext;
+
+USTRUCT()
+struct STATETREEMODULE_API FStateTreeBooleanOperationPropertyFunctionInstanceData
+{
+	GENERATED_BODY()
+
+	UPROPERTY(EditAnywhere, Category = Param)
+	bool bLeft = false;
+
+	UPROPERTY(EditAnywhere, Category = Param)
+	bool bRight = false;
+
+	UPROPERTY(EditAnywhere, Category = Output)
+	bool bResult = false;
+};
+
+/**
+ * Performs 'And' operation on two booleans.
+ */
+USTRUCT(DisplayName = "And")
+struct STATETREEMODULE_API FStateTreeBooleanAndPropertyFunction : public FStateTreePropertyFunctionCommonBase
+{
+	GENERATED_BODY()
+
+	using FInstanceDataType = FStateTreeBooleanOperationPropertyFunctionInstanceData;
+
+	virtual const UStruct* GetInstanceDataType() const override { return FInstanceDataType::StaticStruct(); }
+
+	virtual void Execute(FStateTreeExecutionContext& Context) const override;
+
+#if WITH_EDITOR
+	virtual FText GetDescription(const FGuid& ID, FStateTreeDataView InstanceDataView, const IStateTreeBindingLookup& BindingLookup, EStateTreeNodeFormatting Formatting) const;
+#endif
+};
+
+/**
+ * Performs 'Or' operation on two booleans.
+ */
+USTRUCT(DisplayName = "Or")
+struct STATETREEMODULE_API FStateTreeBooleanOrPropertyFunction : public FStateTreePropertyFunctionCommonBase
+{
+	GENERATED_BODY()
+
+	using FInstanceDataType = FStateTreeBooleanOperationPropertyFunctionInstanceData;
+
+	virtual const UStruct* GetInstanceDataType() const override { return FInstanceDataType::StaticStruct(); }
+
+	virtual void Execute(FStateTreeExecutionContext& Context) const override;
+
+#if WITH_EDITOR
+	virtual FText GetDescription(const FGuid& ID, FStateTreeDataView InstanceDataView, const IStateTreeBindingLookup& BindingLookup, EStateTreeNodeFormatting Formatting) const;
+#endif
+};
\ No newline at end of file
diff --git a/Engine/Plugins/Runtime/StateTree/Source/StateTreeModule/Private/StateTree.cpp b/Engine/Plugins/Runtime/StateTree/Source/StateTreeModule/Private/StateTree.cpp
--- a/Engine/Plugins/Runtime/StateTree/Source/StateTreeModule/Private/StateTree.cpp
+++ b/Engine/Plugins/Runtime/StateTree/Source/StateTreeModule/Private/StateTree.cpp
@@ -15,6 +15,8 @@
 #include "StructUtilsDelegates.h"
 #include "Misc/EnumerateRange.h"
 #include "UObject/AssetRegistryTagsContext.h"
+#include "StateTreePropertyFunctionBase.h"
+
 #if WITH_EDITOR
 #include "Editor.h"
 #include "Engine/UserDefinedStruct.h"
@@ -359,10 +361,10 @@ void UStateTree::PostLoad()
 		{
 			if (Node->InstanceTemplateIndex.IsValid())
 			{
-				const bool bUseSharedInstanceData = NodeView.GetPtr<FStateTreeConditionBase>() || NodeView.GetPtr<FStateTreeConsiderationBase>();
-
-				FStateTreeInstanceData& SourceInstanceData = bUseSharedInstanceData ? SharedInstanceData : DefaultInstanceData;
-
+				const bool bIsUsingSharedInstanceData = NodeView.GetScriptStruct()->IsChildOf<FStateTreeConditionBase>()
+														|| NodeView.GetScriptStruct()->IsChildOf<FStateTreeConsiderationBase>()
+														|| NodeView.GetScriptStruct()->IsChildOf<FStateTreePropertyFunctionBase>();
+				FStateTreeInstanceData& SourceInstanceData = bIsUsingSharedInstanceData ? SharedInstanceData : DefaultInstanceData;
 				if (SourceInstanceData.IsObject(Node->InstanceTemplateIndex.Get()))
 				{
 					Node->PostLoad(SourceInstanceData.GetMutableObject(Node->InstanceTemplateIndex.Get()));
@@ -372,7 +374,6 @@ void UStateTree::PostLoad()
 					Node->PostLoad(SourceInstanceData.GetMutableStruct(Node->InstanceTemplateIndex.Get()));
 				}
 			}
-			
 		}
 	}
 
@@ -746,9 +747,11 @@ bool UStateTree::PatchBindings()
 		const FStateTreeNodeBase& Node = NodeView.Get<const FStateTreeNodeBase>();
 
 		FStateTreeInstanceData* SourceInstanceData = &DefaultInstanceData;
-		if (NodeView.GetPtr<const FStateTreeConditionBase>() || NodeView.GetPtr<const FStateTreeConsiderationBase>())
+		if (NodeView.GetScriptStruct()->IsChildOf<FStateTreeConditionBase>()
+            || NodeView.GetScriptStruct()->IsChildOf<FStateTreeConsiderationBase>()
+            || NodeView.GetScriptStruct()->IsChildOf<FStateTreePropertyFunctionBase>())
 		{
-			// Conditions are stored in shared instance data.
+			// Conditions, Considerations, and PropertyFunctions are stored in shared instance data.
 			SourceInstanceData = &SharedInstanceData;
 		}
 
@@ -796,7 +799,7 @@ bool UStateTree::PatchBindings()
 		}
 
 		FString ErrorMsg;
-		for (int32 Index = Batch.BindingsBegin; Index != Batch.BindingsEnd; Index++)
+		for (int32 Index = Batch.BindingsBegin.Get(); Index != Batch.BindingsEnd.Get(); Index++)
 		{
 			FStateTreePropertyPathBinding& Binding = PropertyPathBindings[Index];
 
@@ -1067,6 +1070,5 @@ void UStateTree::CompileIfChanged()
 		}
 	}
 }
-
 #endif // WITH_EDITOR
 
diff --git a/Engine/Plugins/Runtime/StateTree/Source/StateTreeModule/Private/StateTreeExecutionContext.cpp b/Engine/Plugins/Runtime/StateTree/Source/StateTreeModule/Private/StateTreeExecutionContext.cpp
--- a/Engine/Plugins/Runtime/StateTree/Source/StateTreeModule/Private/StateTreeExecutionContext.cpp
+++ b/Engine/Plugins/Runtime/StateTree/Source/StateTreeModule/Private/StateTreeExecutionContext.cpp
@@ -4,6 +4,7 @@
 #include "StateTreeTaskBase.h"
 #include "StateTreeEvaluatorBase.h"
 #include "StateTreeConditionBase.h"
+#include "StateTreePropertyFunctionBase.h"
 #include "StateTreeReference.h"
 #include "Containers/StaticArray.h"
 #include "Debugger/StateTreeTrace.h"
@@ -1292,6 +1293,12 @@ bool FStateTreeExecutionContext::CopyBatchOnActiveInstances(const FStateTreeExec
 	const FStateTreePropertyCopyBatch& Batch = CurrentFrame.StateTree->PropertyBindings.GetBatch(BindingsBatch);
 	check(TargetView.GetStruct() == Batch.TargetStruct.Struct);
 
+	if (Batch.PropertyFunctionsBegin != Batch.PropertyFunctionsEnd)
+	{
+		check(Batch.PropertyFunctionsBegin.IsValid() && Batch.PropertyFunctionsEnd.IsValid());
+		EvaluatePropertyFunctionsOnActiveInstances(ParentFrame, CurrentFrame, Batch.PropertyFunctionsBegin, Batch.PropertyFunctionsEnd.Get() - Batch.PropertyFunctionsBegin.Get());
+	}
+
 	bool bSucceed = true;
 	for (const FStateTreePropertyCopy& Copy : CurrentFrame.StateTree->PropertyBindings.GetBatchCopies(Batch))
 	{
@@ -1306,6 +1313,12 @@ bool FStateTreeExecutionContext::CopyBatchWithValidation(const FStateTreeExecuti
 	const FStateTreePropertyCopyBatch& Batch = CurrentFrame.StateTree->PropertyBindings.GetBatch(BindingsBatch);
 	check(TargetView.GetStruct() == Batch.TargetStruct.Struct);
 
+	if (Batch.PropertyFunctionsBegin != Batch.PropertyFunctionsEnd)
+	{
+		check(Batch.PropertyFunctionsBegin.IsValid() && Batch.PropertyFunctionsEnd.IsValid());
+		EvaluatePropertyFunctionsWithValidation(ParentFrame, CurrentFrame, Batch.PropertyFunctionsBegin, Batch.PropertyFunctionsEnd.Get() - Batch.PropertyFunctionsBegin.Get());
+	}
+
 	bool bSucceed = true;
 	for (const FStateTreePropertyCopy& Copy : CurrentFrame.StateTree->PropertyBindings.GetBatchCopies(Batch))
 	{
@@ -2601,6 +2614,56 @@ bool FStateTreeExecutionContext::TestAllConditions(const FStateTreeExecutionFram
 	return Values[0];
 }
 
+void FStateTreeExecutionContext::EvaluatePropertyFunctionsOnActiveInstances(const FStateTreeExecutionFrame* CurrentParentFrame, const FStateTreeExecutionFrame& CurrentFrame, FStateTreeIndex16 FuncsBegin, uint16 FuncsNum)
+{
+	for (int32 FuncIndex = FuncsBegin.Get(); FuncIndex < FuncsBegin.Get() + FuncsNum; ++FuncIndex)
+	{
+		const FStateTreePropertyFunctionBase& Func = CurrentFrame.StateTree->Nodes[FuncIndex].Get<const FStateTreePropertyFunctionBase>();
+		const FStateTreeDataView FuncInstanceView = GetDataView(CurrentParentFrame, CurrentFrame, Func.InstanceDataHandle);
+		FNodeInstanceDataScope DataScope(*this, Func.InstanceDataHandle, FuncInstanceView);
+
+		// Copy bound properties.
+		if (Func.BindingsBatch.IsValid())
+		{
+			// Use validated copy, since we test in situations where the sources are not always valid (e.g. enter conditions may try to access inactive parent state). 
+			CopyBatchOnActiveInstances(CurrentParentFrame, CurrentFrame, FuncInstanceView, Func.BindingsBatch);
+		}
+			
+		Func.Execute(*this);
+			
+		// Reset copied properties that might contain object references.
+		if (Func.BindingsBatch.IsValid())
+		{
+			CurrentFrame.StateTree->PropertyBindings.ResetObjects(Func.BindingsBatch, FuncInstanceView);
+		}
+	}
+}
+
+void FStateTreeExecutionContext::EvaluatePropertyFunctionsWithValidation(const FStateTreeExecutionFrame* CurrentParentFrame, const FStateTreeExecutionFrame& CurrentFrame, FStateTreeIndex16 FuncsBegin, uint16 FuncsNum)
+{
+	for (int32 FuncIndex = FuncsBegin.Get(); FuncIndex < FuncsBegin.Get() + FuncsNum; ++FuncIndex)
+	{
+		const FStateTreePropertyFunctionBase& Func = CurrentFrame.StateTree->Nodes[FuncIndex].Get<const FStateTreePropertyFunctionBase>();
+		const FStateTreeDataView FuncInstanceView = GetDataView(CurrentParentFrame, CurrentFrame, Func.InstanceDataHandle);
+		FNodeInstanceDataScope DataScope(*this, Func.InstanceDataHandle, FuncInstanceView);
+
+		// Copy bound properties.
+		if (Func.BindingsBatch.IsValid())
+		{
+			// Use validated copy, since we test in situations where the sources are not always valid (e.g. enter conditions may try to access inactive parent state). 
+			CopyBatchWithValidation(CurrentParentFrame, CurrentFrame, FuncInstanceView, Func.BindingsBatch);
+		}
+			
+		Func.Execute(*this);
+			
+		// Reset copied properties that might contain object references.
+		if (Func.BindingsBatch.IsValid())
+		{
+			CurrentFrame.StateTree->PropertyBindings.ResetObjects(Func.BindingsBatch, FuncInstanceView);
+		}
+	}
+}
+
 FString FStateTreeExecutionContext::DebugGetEventsAsString() const
 {
 	FStringBuilderBase StrBuilder;
diff --git a/Engine/Plugins/Runtime/StateTree/Source/StateTreeModule/Private/StateTreePropertyBindings.cpp b/Engine/Plugins/Runtime/StateTree/Source/StateTreeModule/Private/StateTreePropertyBindings.cpp
--- a/Engine/Plugins/Runtime/StateTree/Source/StateTreeModule/Private/StateTreePropertyBindings.cpp
+++ b/Engine/Plugins/Runtime/StateTree/Source/StateTreeModule/Private/StateTreePropertyBindings.cpp
@@ -69,6 +69,25 @@ namespace UE::StateTree
 
 		return EStateTreePropertyUsage::Parameter;
 	}
+
+	const FProperty* GetStructSingleOutputProperty(const UStruct& InStruct)
+	{
+		const FProperty* FuncOutputProperty = nullptr;
+		for (TFieldIterator<FProperty> PropIt(&InStruct, EFieldIteratorFlags::IncludeSuper); PropIt; ++PropIt)
+		{
+			if (GetUsageFromMetaData(*PropIt) == EStateTreePropertyUsage::Output)
+			{
+				if (FuncOutputProperty)
+				{
+					return nullptr;
+				}
+
+				FuncOutputProperty = *PropIt;
+			}
+		}
+
+		return FuncOutputProperty;
+	}
 #endif
 
 } // UE::StateTree
@@ -213,7 +232,7 @@ bool FStateTreePropertyBindings::ResolvePaths()
 	
 	for (const FStateTreePropertyCopyBatch& Batch : CopyBatches)
 	{
-		for (int32 i = Batch.BindingsBegin; i != Batch.BindingsEnd; i++)
+		for (int32 i = Batch.BindingsBegin.Get(); i != Batch.BindingsEnd.Get(); i++)
 		{
 			const FStateTreePropertyPathBinding& Binding = PropertyPathBindings[i];
 			
@@ -1174,7 +1193,7 @@ bool FStateTreePropertyBindings::ResetObjects(const FStateTreeIndex16 TargetBatc
 
 	bool bResult = true;
 	
-	for (int32 i = Batch.BindingsBegin; i != Batch.BindingsEnd; i++)
+	for (int32 i = Batch.BindingsBegin.Get(); i != Batch.BindingsEnd.Get(); i++)
 	{
 		const FStateTreePropertyCopy& Copy = PropertyCopies[i];
 		// Copies that fail to be resolved (i.e. property path does not resolve, types changed) will be marked as None, skip them.
@@ -1254,7 +1273,7 @@ void FStateTreePropertyBindings::DebugPrintInternalLayout(FString& OutString) co
 		OutString += FString::Printf(TEXT("  | %-40s | %-40s | %8s [%3d:%-3d[ |\n"),
 									 CopyBatch.TargetStruct.Struct ? *CopyBatch.TargetStruct.Struct->GetName() : TEXT("null"),
 									 *CopyBatch.TargetStruct.Name.ToString(),
-									 TEXT(""), CopyBatch.BindingsBegin, CopyBatch.BindingsEnd);
+									 TEXT(""), CopyBatch.BindingsBegin.Get(), CopyBatch.BindingsEnd.Get());
 	}
 
 	/** Array of property bindings, resolved into arrays of copies before use. */
@@ -1401,10 +1420,10 @@ bool FStateTreePropertyPath::UpdateSegmentsFromValue(const FStateTreeDataView Ba
 	return true;
 }
 
-FString FStateTreePropertyPath::ToString(const int32 HighlightedSegment, const TCHAR* HighlightPrefix, const TCHAR* HighlightPostfix, const bool bOutputInstances) const
+FString FStateTreePropertyPath::ToString(const int32 HighlightedSegment, const TCHAR* HighlightPrefix, const TCHAR* HighlightPostfix, const bool bOutputInstances, const int32 FirstSegment) const
 {
 	FStringBuilderBase Result;
-	for (TEnumerateRef<const FStateTreePropertyPathSegment> Segment : EnumerateRange(Segments))
+	for (TEnumerateRef<const FStateTreePropertyPathSegment> Segment : EnumerateRange(TConstArrayView<FStateTreePropertyPathSegment>(Segments).Mid(FMath::Max(FirstSegment, 0))))
 	{
 		if (Segment.GetIndex() > 0)
 		{
diff --git a/Engine/Plugins/Runtime/StateTree/Source/StateTreeModule/Private/StateTreePropertyFunctionBase.cpp b/Engine/Plugins/Runtime/StateTree/Source/StateTreeModule/Private/StateTreePropertyFunctionBase.cpp
new file mode 100644
--- /dev/null
+++ b/Engine/Plugins/Runtime/StateTree/Source/StateTreeModule/Private/StateTreePropertyFunctionBase.cpp
@@ -0,0 +1,29 @@
+// Copyright Epic Games, Inc. All Rights Reserved.
+
+#include "StateTreePropertyFunctionBase.h"
+#include "StateTreePropertyBindings.h"
+
+#if WITH_EDITOR
+#include "EdGraphSchema_K2.h"
+#endif
+
+#include UE_INLINE_GENERATED_CPP_BY_NAME(StateTreePropertyFunctionBase)
+
+#if WITH_EDITOR
+FColor FStateTreePropertyFunctionBase::GetIconColor() const
+{
+	if (const FProperty* OutputProperty = UE::StateTree::GetStructSingleOutputProperty(*GetInstanceDataType()))
+	{
+		const UEdGraphSchema_K2* Schema = GetDefault<UEdGraphSchema_K2>();
+		check(Schema);
+
+		FEdGraphPinType PinType;
+		if (Schema->ConvertPropertyToPinType(OutputProperty, PinType))
+		{
+			return Schema->GetPinTypeColor(PinType).ToFColor(true);
+		}
+	}
+
+	return Super::GetIconColor();
+}
+#endif
\ No newline at end of file
diff --git a/Engine/Plugins/Runtime/StateTree/Source/StateTreeModule/Private/StateTreePropertyRefHelpers.cpp b/Engine/Plugins/Runtime/StateTree/Source/StateTreeModule/Private/StateTreePropertyRefHelpers.cpp
--- a/Engine/Plugins/Runtime/StateTree/Source/StateTreeModule/Private/StateTreePropertyRefHelpers.cpp
+++ b/Engine/Plugins/Runtime/StateTree/Source/StateTreeModule/Private/StateTreePropertyRefHelpers.cpp
@@ -234,6 +234,7 @@ namespace UE::StateTree::PropertyRefHelpers
 		case EStateTreeBindableStructSource::Context:
 		case EStateTreeBindableStructSource::Condition:
 		case EStateTreeBindableStructSource::Consideration:
+		case EStateTreeBindableStructSource::PropertyFunction:
 			return false;
 
 		case EStateTreeBindableStructSource::GlobalTask:
diff --git a/Engine/Plugins/Runtime/StateTree/Source/StateTreeModule/Public/StateTreeExecutionContext.h b/Engine/Plugins/Runtime/StateTree/Source/StateTreeModule/Public/StateTreeExecutionContext.h
--- a/Engine/Plugins/Runtime/StateTree/Source/StateTreeModule/Public/StateTreeExecutionContext.h
+++ b/Engine/Plugins/Runtime/StateTree/Source/StateTreeModule/Public/StateTreeExecutionContext.h
@@ -575,6 +575,13 @@ protected:
 	 */
 	bool TestAllConditions(const FStateTreeExecutionFrame* CurrentParentFrame, const FStateTreeExecutionFrame& CurrentFrame, const int32 ConditionsOffset, const int32 ConditionsNum);
 
+	/* Evaluate all function at given range. Should be used only on active instances, assumes valid handles and does not consider temporary instances. */
+	void EvaluatePropertyFunctionsOnActiveInstances(const FStateTreeExecutionFrame* CurrentParentFrame, const FStateTreeExecutionFrame& CurrentFrame, FStateTreeIndex16 FuncsBegin, uint16 FuncsNum);
+
+
+	/* Evaluate all function at given range. This version validates the data handles and looks up temporary instances. */
+	void EvaluatePropertyFunctionsWithValidation(const FStateTreeExecutionFrame* CurrentParentFrame, const FStateTreeExecutionFrame& CurrentFrame, FStateTreeIndex16 FuncsBegin, uint16 FuncsNum);
+
 	/**
 	 * Requests transition to a specified state with specified priority.
 	 */
diff --git a/Engine/Plugins/Runtime/StateTree/Source/StateTreeModule/Public/StateTreePropertyBindings.h b/Engine/Plugins/Runtime/StateTree/Source/StateTreeModule/Public/StateTreePropertyBindings.h
--- a/Engine/Plugins/Runtime/StateTree/Source/StateTreeModule/Public/StateTreePropertyBindings.h
+++ b/Engine/Plugins/Runtime/StateTree/Source/StateTreeModule/Public/StateTreePropertyBindings.h
@@ -5,6 +5,8 @@
 #include "StateTreeNodeBase.h"
 #include "StateTreeTypes.h"
 #include "StateTreePropertyRefHelpers.h"
+#include "StructView.h"
+#include "StateTreeIndexTypes.h"
 #include "StateTreePropertyBindings.generated.h"
 
 class FProperty;
@@ -37,6 +39,8 @@ enum class EStateTreeBindableStructSource : uint8
 	TransitionEvent,
 	/** Source is StateTree event used by state selection */
 	StateEvent,
+	/** Source is Property Function */
+	PropertyFunction,
 };
 
 
@@ -393,8 +397,9 @@ struct STATETREEMODULE_API FStateTreePropertyPath
 	 * @param HighlightPrefix String to append before highlighted segment
 	 * @param HighlightPostfix String to append after highlighted segment
 	 * @param bOutputInstances if true, the instance struct types will be output. 
+	 * @param FirstSegment Index of the first path segment to be stringified.
 	 */
-	FString ToString(const int32 HighlightedSegment = INDEX_NONE, const TCHAR* HighlightPrefix = nullptr, const TCHAR* HighlightPostfix = nullptr, const bool bOutputInstances = false) const;
+	FString ToString(const int32 HighlightedSegment = INDEX_NONE, const TCHAR* HighlightPrefix = nullptr, const TCHAR* HighlightPostfix = nullptr, const bool bOutputInstances = false, const int32 FirstSegment = 0) const;
 
 	/**
 	 * Resolves the property path against base struct type. The path is assumed to be relative to the BaseStruct.
@@ -531,7 +536,6 @@ PRAGMA_DISABLE_DEPRECATION_WARNINGS
 PRAGMA_ENABLE_DEPRECATION_WARNINGS
 };
 
-
 /**
  * Representation of a property binding in StateTree
  */
@@ -556,6 +560,15 @@ PRAGMA_DISABLE_DEPRECATION_WARNINGS
 	{
 	}
 
+#if WITH_EDITOR
+FStateTreePropertyPathBinding(FConstStructView InFunctionNodeStruct, const FStateTreePropertyPath& InSourcePath, const FStateTreePropertyPath& InTargetPath)
+		: SourcePropertyPath(InSourcePath)
+		, TargetPropertyPath(InTargetPath)
+		, PropertyFunctionNode(InFunctionNodeStruct)
+	{
+	}
+#endif
+
 	UE_DEPRECATED(5.4, "Use constructor with DataHandle instead.")
 	FStateTreePropertyPathBinding(const FStateTreeIndex16 InCompiledSourceStructIndex, const FStateTreePropertyPath& InSourcePath, const FStateTreePropertyPath& InTargetPath)
 		: SourcePropertyPath(InSourcePath)
@@ -583,6 +596,11 @@ PRAGMA_ENABLE_DEPRECATION_WARNINGS
 	void SetSourceDataHandle(const FStateTreeDataHandle NewSourceDataHandle) { SourceDataHandle = NewSourceDataHandle; }
 	FStateTreeDataHandle GetSourceDataHandle() const { return SourceDataHandle; }
 
+#if WITH_EDITOR
+	FConstStructView GetPropertyFunctionNode() const { return FConstStructView(PropertyFunctionNode); }
+	FStructView GetMutablePropertyFunctionNode() { return FStructView(PropertyFunctionNode); }
+#endif
+
 private:
 	/** Source property path of the binding */
 	UPROPERTY()
@@ -596,8 +614,11 @@ private:
 	UPROPERTY()
 	FStateTreeDataHandle SourceDataHandle = FStateTreeDataHandle::Invalid;
 
-public:
 #if WITH_EDITORONLY_DATA
+	/**	Instance of bound PropertyFunction. */
+	UPROPERTY()
+	FInstancedStruct PropertyFunctionNode;
+
 PRAGMA_DISABLE_DEPRECATION_WARNINGS
 	UPROPERTY()
 	FStateTreeEditorPropertyPath SourcePath_DEPRECATED;
@@ -816,11 +837,19 @@ struct STATETREEMODULE_API FStateTreePropertyCopyBatch
 
 	/** Index to first binding/copy. */
 	UPROPERTY()
-	uint16 BindingsBegin = 0;
+	FStateTreeIndex16 BindingsBegin;
 
 	/** Index to one past the last binding/copy. */
 	UPROPERTY()
-	uint16 BindingsEnd = 0;
+	FStateTreeIndex16 BindingsEnd;
+
+	/** Index to first property function. */
+	UPROPERTY()
+	FStateTreeIndex16 PropertyFunctionsBegin;
+
+	/** Index to one past the last property function. */
+	UPROPERTY()
+	FStateTreeIndex16 PropertyFunctionsEnd;
 };
 
 using FStateTreePropCopyBatch UE_DEPRECATED(5.3, "Deprecated struct. Please use FStateTreePropertyCopy instead.") = FStateTreePropertyCopyBatch;
@@ -903,12 +932,12 @@ struct STATETREEMODULE_API FStateTreePropertyBindings
 	/** @return All the property copies for a specific batch. */
 	TConstArrayView<FStateTreePropertyCopy> GetBatchCopies(const FStateTreePropertyCopyBatch& Batch) const
 	{
-		const int32 Count = (int32)Batch.BindingsEnd - (int32)Batch.BindingsBegin;
+		const int32 Count = Batch.BindingsEnd.Get() - Batch.BindingsBegin.Get();
 		if (Count == 0)
 		{
 			return {};
 		}
-		return MakeArrayView(&PropertyCopies[Batch.BindingsBegin], Count);
+		return MakeArrayView(&PropertyCopies[Batch.BindingsBegin.Get()], Count);
 	}
 
 	/**
@@ -1075,6 +1104,9 @@ namespace UE::StateTree
 	 * @return found usage type, or EStateTreePropertyUsage::Invalid if not found.
 	 */
 	STATETREEMODULE_API EStateTreePropertyUsage GetUsageFromMetaData(const FProperty* Property);
+
+	/** @return struct's property which is the only one marked as Output. Returns null otherwise. */
+	STATETREEMODULE_API const FProperty* GetStructSingleOutputProperty(const UStruct& InStruct);
 #endif
 } // UE::StateTree
 
diff --git a/Engine/Plugins/Runtime/StateTree/Source/StateTreeModule/Public/StateTreePropertyFunctionBase.h b/Engine/Plugins/Runtime/StateTree/Source/StateTreeModule/Public/StateTreePropertyFunctionBase.h
new file mode 100644
--- /dev/null
+++ b/Engine/Plugins/Runtime/StateTree/Source/StateTreeModule/Public/StateTreePropertyFunctionBase.h
@@ -0,0 +1,40 @@
+// Copyright Epic Games, Inc. All Rights Reserved.
+
+#pragma once
+
+#include "StateTreeNodeBase.h"
+#include "StateTreePropertyFunctionBase.generated.h"
+
+struct FStateTreeExecutionContext;
+
+/**
+ * Base struct for all property functions.
+ * PropertyFunction is a node which is executed just before evaluating owner's bindings.
+ * Only Output properties are allowed to be bound to.
+ */
+USTRUCT(meta = (Hidden))
+struct STATETREEMODULE_API FStateTreePropertyFunctionBase : public FStateTreeNodeBase
+{
+	GENERATED_BODY()
+
+	/**
+	 * Called right before evaluating bindings for the owning node.
+	 * @param Context Reference to current execution context.
+	 */
+	virtual void Execute(FStateTreeExecutionContext& Context) const {}
+
+#if WITH_EDITOR
+	virtual FName GetIconName() const override
+	{
+		return FName("StateTreeEditorStyle|Node.Function");
+	}
+
+	virtual FColor GetIconColor() const override;
+#endif
+};
+
+USTRUCT(meta = (Hidden))
+struct STATETREEMODULE_API FStateTreePropertyFunctionCommonBase : public FStateTreePropertyFunctionBase
+{
+	GENERATED_BODY()
+};
\ No newline at end of file
diff --git a/Engine/Plugins/Runtime/StateTree/Source/StateTreeTestSuite/Private/StateTreeTest.cpp b/Engine/Plugins/Runtime/StateTree/Source/StateTreeTestSuite/Private/StateTreeTest.cpp
--- a/Engine/Plugins/Runtime/StateTree/Source/StateTreeTestSuite/Private/StateTreeTest.cpp
+++ b/Engine/Plugins/Runtime/StateTree/Source/StateTreeTestSuite/Private/StateTreeTest.cpp
@@ -1797,7 +1797,7 @@ struct FStateTreeTest_BindingsCompiler : FAITestBase
 		PropertyBindings.Add(UE::StateTree::Tests::MakeBinding(SourceADesc.ID, TEXT("Array"), TargetDesc.ID, TEXT("Array")));
 
 		int32 CopyBatchIndex = INDEX_NONE;
-		const bool bCompileBatchResult = BindingCompiler.CompileBatch(TargetDesc, PropertyBindings, CopyBatchIndex);
+		const bool bCompileBatchResult = BindingCompiler.CompileBatch(TargetDesc, PropertyBindings, FStateTreeIndex16::Invalid, FStateTreeIndex16::Invalid, CopyBatchIndex);
 		AITEST_TRUE("CompileBatch should succeed", bCompileBatchResult);
 		AITEST_NOT_EQUAL("CopyBatchIndex should not be INDEX_NONE", CopyBatchIndex, (int32)INDEX_NONE);
 
@@ -1849,6 +1849,83 @@ struct FStateTreeTest_BindingsCompiler : FAITestBase
 };
 IMPLEMENT_AI_INSTANT_TEST(FStateTreeTest_BindingsCompiler, "System.StateTree.BindingsCompiler");
 
+struct FStateTreeTest_PropertyFunctions : FAITestBase
+{
+	virtual bool InstantTest() override
+	{
+		UStateTree& StateTree = UE::StateTree::Tests::NewStateTree(&GetWorld());
+		UStateTreeEditorData& EditorData = *Cast<UStateTreeEditorData>(StateTree.EditorData);
+		UStateTreeState& Root = EditorData.AddSubTree(FName(TEXT("Root")));
+		FStateTreePropertyPathSegment PathSegmentToFuncResult = FStateTreePropertyPathSegment(TEXT("Result"));
+
+		// Condition with property function binding.
+		{
+			TStateTreeEditorNode<FStateTreeCompareIntCondition>& EnterCond = Root.AddEnterCondition<FStateTreeCompareIntCondition>(EGenericAICheck::Equal);
+			EnterCond.GetInstanceData().Right = 1;
+			EditorData.AddPropertyBinding(CastChecked<UScriptStruct>(FTestPropertyFunction::StaticStruct()), {PathSegmentToFuncResult}, FStateTreePropertyPath(EnterCond.ID, TEXT("Left")));
+		}
+
+		// Task with multiple nested property function bindings.
+		auto& TaskA = Root.AddTask<FTestTask_PrintAndResetValue>(FName(TEXT("TaskA")));
+		constexpr int32 TaskAPropertyFunctionsAmount = 10;
+		{
+			EditorData.AddPropertyBinding(CastChecked<UScriptStruct>(FTestPropertyFunction::StaticStruct()), {PathSegmentToFuncResult}, FStateTreePropertyPath(TaskA.ID, TEXT("Value")));
+		
+			for (int32 i = 0; i < TaskAPropertyFunctionsAmount - 1; ++i)
+			{
+				const FStateTreePropertyPathBinding& LastBinding = EditorData.GetPropertyEditorBindings()->GetBindings().Last();
+				const FGuid LastBindingPropertyFuncID = LastBinding.GetPropertyFunctionNode().Get<const FStateTreeEditorNode>().ID;
+				EditorData.AddPropertyBinding(CastChecked<UScriptStruct>(FTestPropertyFunction::StaticStruct()), {PathSegmentToFuncResult}, FStateTreePropertyPath(LastBindingPropertyFuncID, TEXT("Input")));
+			}
+		}
+
+		// Task bound to state parameter with multiple nested property function bindings.
+		auto& TaskB = Root.AddTask<FTestTask_PrintAndResetValue>(FName(TEXT("TaskB")));
+		constexpr int32 ParameterPropertyFunctionsAmount = 5;
+		{
+			Root.Parameters.Parameters.AddProperty(FName(TEXT("Int")), EPropertyBagPropertyType::Int32);
+			const FStateTreePropertyPath PathToProperty = FStateTreePropertyPath(Root.Parameters.ID, TEXT("Int"));
+			EditorData.AddPropertyBinding(PathToProperty, FStateTreePropertyPath(TaskB.ID, TEXT("Value")));
+			EditorData.AddPropertyBinding(CastChecked<UScriptStruct>(FTestPropertyFunction::StaticStruct()), {PathSegmentToFuncResult}, PathToProperty);
+		
+			for (int32 i = 0; i < ParameterPropertyFunctionsAmount - 1; ++i)
+			{
+				const FStateTreePropertyPathBinding& LastBinding = EditorData.GetPropertyEditorBindings()->GetBindings().Last();
+				const FGuid LastBindingPropertyFuncID = LastBinding.GetPropertyFunctionNode().Get<const FStateTreeEditorNode>().ID;
+				EditorData.AddPropertyBinding(CastChecked<UScriptStruct>(FTestPropertyFunction::StaticStruct()), {PathSegmentToFuncResult}, FStateTreePropertyPath(LastBindingPropertyFuncID, TEXT("Input")));
+			}
+		}
+
+		FStateTreeCompilerLog Log;
+		FStateTreeCompiler Compiler(Log);
+		const bool bResult = Compiler.Compile(StateTree);
+		AITEST_TRUE("StateTree should get compiled", bResult);
+
+		FStateTreeInstanceData InstanceData;
+		FTestStateTreeExecutionContext Exec(StateTree, StateTree, InstanceData);
+		const bool bInitSucceeded = Exec.IsValid();
+		AITEST_TRUE("StateTree should init", bInitSucceeded);
+
+		Exec.Start();
+		AITEST_TRUE(*FString::Printf(TEXT("StateTree TaskA should enter state with value %d"), TaskAPropertyFunctionsAmount), Exec.Expect(TaskA.GetName(), *FString::Printf(TEXT("EnterState%d"), TaskAPropertyFunctionsAmount)));
+		AITEST_TRUE(*FString::Printf(TEXT("StateTree TaskB should enter state with value %d"), ParameterPropertyFunctionsAmount), Exec.Expect(TaskB.GetName(), *FString::Printf(TEXT("EnterState%d"), ParameterPropertyFunctionsAmount)));
+		Exec.LogClear();
+
+		Exec.Tick(0.1f);
+		AITEST_TRUE(*FString::Printf(TEXT("StateTree TaskA should tick with value %d"), TaskAPropertyFunctionsAmount), Exec.Expect(TaskA.GetName(), *FString::Printf(TEXT("Tick%d"), TaskAPropertyFunctionsAmount)));
+		AITEST_TRUE(*FString::Printf(TEXT("StateTree TaskB should tick with value %d"), ParameterPropertyFunctionsAmount), Exec.Expect(TaskB.GetName(), *FString::Printf(TEXT("Tick%d"), ParameterPropertyFunctionsAmount)));
+		Exec.LogClear();
+
+		Exec.Stop(EStateTreeRunStatus::Stopped);
+		AITEST_TRUE(*FString::Printf(TEXT("StateTree TaskA should exit state with value %d"), TaskAPropertyFunctionsAmount), Exec.Expect(TaskA.GetName(), *FString::Printf(TEXT("ExitState%d"), TaskAPropertyFunctionsAmount)));
+		AITEST_TRUE(*FString::Printf(TEXT("StateTree TaskB should exit state with value %d"), ParameterPropertyFunctionsAmount), Exec.Expect(TaskB.GetName(), *FString::Printf(TEXT("ExitState%d"), ParameterPropertyFunctionsAmount)));
+		Exec.LogClear();
+
+		return true;
+	}
+};
+IMPLEMENT_AI_INSTANT_TEST(FStateTreeTest_PropertyFunctions, "System.StateTree.PropertyFunctions");
+
 struct FStateTreeTest_CopyObjects : FAITestBase
 {
 	virtual bool InstantTest() override
@@ -1895,12 +1972,12 @@ struct FStateTreeTest_CopyObjects : FAITestBase
 		PropertyBindings.Add(UE::StateTree::Tests::MakeBinding(SourceDesc.ID, TEXT("Class"), TargetBDesc.ID, TEXT("SoftClass")));
 		
 		int32 TargetACopyBatchIndex = INDEX_NONE;
-		const bool bCompileBatchResultA = BindingCompiler.CompileBatch(TargetADesc, PropertyBindings, TargetACopyBatchIndex);
+		const bool bCompileBatchResultA = BindingCompiler.CompileBatch(TargetADesc, PropertyBindings, FStateTreeIndex16::Invalid, FStateTreeIndex16::Invalid, TargetACopyBatchIndex);
 		AITEST_TRUE("CompileBatchResultA should succeed", bCompileBatchResultA);
 		AITEST_NOT_EQUAL("TargetACopyBatchIndex should not be INDEX_NONE", TargetACopyBatchIndex, (int32)INDEX_NONE);
 
 		int32 TargetBCopyBatchIndex = INDEX_NONE;
-		const bool bCompileBatchResultB = BindingCompiler.CompileBatch(TargetBDesc, PropertyBindings, TargetBCopyBatchIndex);
+		const bool bCompileBatchResultB = BindingCompiler.CompileBatch(TargetBDesc, PropertyBindings, FStateTreeIndex16::Invalid, FStateTreeIndex16::Invalid, TargetBCopyBatchIndex);
 		AITEST_TRUE("CompileBatchResultB should succeed", bCompileBatchResultB);
 		AITEST_NOT_EQUAL("TargetBCopyBatchIndex should not be INDEX_NONE", TargetBCopyBatchIndex, (int32)INDEX_NONE);
 
diff --git a/Engine/Plugins/Runtime/StateTree/Source/StateTreeTestSuite/Private/StateTreeTestTypes.h b/Engine/Plugins/Runtime/StateTree/Source/StateTreeTestSuite/Private/StateTreeTestTypes.h
--- a/Engine/Plugins/Runtime/StateTree/Source/StateTreeTestSuite/Private/StateTreeTestTypes.h
+++ b/Engine/Plugins/Runtime/StateTree/Source/StateTreeTestSuite/Private/StateTreeTestTypes.h
@@ -8,6 +8,7 @@
 #include "StateTreeExecutionContext.h"
 #include "StateTreeLinker.h"
 #include "StateTreePropertyRef.h"
+#include "StateTreePropertyFunctionBase.h"
 #include "StateTreeTestTypes.generated.h"
 
 class UStateTree;
@@ -250,6 +251,30 @@ struct FTestTask_PrintValue : public FStateTreeTaskBase
 	TStateTreeExternalDataHandle<FStateTreeTestLog> LogHandle;
 };
 
+USTRUCT()
+struct FTestTask_PrintAndResetValue : public FTestTask_PrintValue
+{
+	GENERATED_BODY()
+
+	using FTestTask_PrintValue::FTestTask_PrintValue;
+
+	virtual EStateTreeRunStatus EnterState(FStateTreeExecutionContext& Context, const FStateTreeTransitionResult& Transition) const override
+	{
+		EStateTreeRunStatus Status = Super::EnterState(Context, Transition);
+		FInstanceDataType& InstanceData = Context.GetInstanceData(*this);
+		InstanceData.Value = 0;
+		return Status;
+	}
+
+	virtual EStateTreeRunStatus Tick(FStateTreeExecutionContext& Context, const float DeltaTime) const override
+	{
+		EStateTreeRunStatus Status = Super::Tick(Context, DeltaTime);
+		FInstanceDataType& InstanceData = Context.GetInstanceData(*this);
+		InstanceData.Value = 0;
+		return Status;
+	};
+};
+
 
 USTRUCT()
 struct FTestTask_StopTreeInstanceData
@@ -592,3 +617,34 @@ struct FStateTreeTest_PropertyCopyObjects
 	UPROPERTY(EditAnywhere, Category = "")
 	TSoftClassPtr<UObject> SoftClass;
 };
+
+USTRUCT()
+struct FTestPropertyFunction_InstanceData
+{
+	GENERATED_BODY()
+
+	UPROPERTY()
+	int32 Input = 0;
+
+	UPROPERTY()
+	int32 Result = 0;
+};
+
+USTRUCT()
+struct FTestPropertyFunction : public FStateTreePropertyFunctionBase
+{
+	GENERATED_BODY()
+
+	using FInstanceDataType = FTestPropertyFunction_InstanceData;
+
+	FTestPropertyFunction() = default;
+	FTestPropertyFunction(const FName InName) { Name = InName; }
+	
+	virtual const UStruct* GetInstanceDataType() const override { return FInstanceDataType::StaticStruct(); }
+
+	virtual void Execute(FStateTreeExecutionContext& Context) const override
+	{
+		FTestPropertyFunction_InstanceData& InstanceData = Context.GetInstanceData<FTestPropertyFunction_InstanceData>(*this);
+		InstanceData.Result = InstanceData.Input + 1;
+	}
+};
