StateTree: Updated FStateTreePropertyFunctionBase comment.

diff --git a/Engine/Plugins/Runtime/StateTree/Source/StateTreeModule/Public/StateTreePropertyFunctionBase.h b/Engine/Plugins/Runtime/StateTree/Source/StateTreeModule/Public/StateTreePropertyFunctionBase.h
--- a/Engine/Plugins/Runtime/StateTree/Source/StateTreeModule/Public/StateTreePropertyFunctionBase.h
+++ b/Engine/Plugins/Runtime/StateTree/Source/StateTreeModule/Public/StateTreePropertyFunctionBase.h
@@ -10,7 +10,29 @@ struct FStateTreeExecutionContext;
 /**
  * Base struct for all property functions.
  * PropertyFunction is a node which is executed just before evaluating owner's bindings.
- * Only Output properties are allowed to be bound to.
+ *
+ * The property function's instance data is expected to have one property marked as output.
+ * This property is used to find which properties the function can be used for,
+ * and that property is hidden in the UI. It is expected that there's just one output property.
+ *
+ * Example:
+ *
+ *	USTRUCT()
+ *	struct FStateTreeBooleanOperationPropertyFunctionInstanceData
+ *	{
+ *		GENERATED_BODY()
+ *
+ *		UPROPERTY(EditAnywhere, Category = Param)
+ *		bool bLeft = false;
+ *
+ *		UPROPERTY(EditAnywhere, Category = Param)
+ *		bool bRight = false;
+ *
+ *      // This property is used to find which properties the function can be used for.
+ *		UPROPERTY(EditAnywhere, Category = Output)
+ *		bool bResult = false;
+ *	};
+ *
  */
 USTRUCT(meta = (Hidden))
 struct STATETREEMODULE_API FStateTreePropertyFunctionBase : public FStateTreeNodeBase
