StateTrees: Display node tooltips in editor (contributed)

diff --git a/Engine/Plugins/Runtime/StateTree/Source/StateTreeEditorModule/Private/Customizations/StateTreeEditorNodeDetails.cpp b/Engine/Plugins/Runtime/StateTree/Source/StateTreeEditorModule/Private/Customizations/StateTreeEditorNodeDetails.cpp
--- a/Engine/Plugins/Runtime/StateTree/Source/StateTreeEditorModule/Private/Customizations/StateTreeEditorNodeDetails.cpp
+++ b/Engine/Plugins/Runtime/StateTree/Source/StateTreeEditorModule/Private/Customizations/StateTreeEditorNodeDetails.cpp
@@ -520,6 +520,7 @@ void FStateTreeEditorNodeDetails::CustomizeHeader(TSharedRef<class IPropertyHand
 								.TextStyle(&FStateTreeEditorStyle::Get().GetWidgetStyle<FTextBlockStyle>("StateTree.Node.Normal"))
 								.OverflowPolicy(ETextOverflowPolicy::Ellipsis)
 								.Visibility(this, &FStateTreeEditorNodeDetails::IsNodeDescriptionVisible)
+								.ToolTipText(this, &FStateTreeEditorNodeDetails::GetNodeTooltip)
 								+SRichTextBlock::Decorator(FTextStyleDecorator::Create(TEXT(""), FStateTreeEditorStyle::Get().GetWidgetStyle<FTextBlockStyle>("StateTree.Node.Normal")))
 								+SRichTextBlock::Decorator(FTextStyleDecorator::Create(TEXT("b"), FStateTreeEditorStyle::Get().GetWidgetStyle<FTextBlockStyle>("StateTree.Node.Bold")))
 								+SRichTextBlock::Decorator(FTextStyleDecorator::Create(TEXT("s"), FStateTreeEditorStyle::Get().GetWidgetStyle<FTextBlockStyle>("StateTree.Node.Subdued")))
@@ -1378,6 +1379,51 @@ EVisibility FStateTreeEditorNodeDetails::IsNodeDescriptionVisible() const
 	return EVisibility::Visible;
 }
 
+FText FStateTreeEditorNodeDetails::GetNodeTooltip() const
+{
+	check(StructProperty);
+	if (!EditorData)
+	{
+		return FText::GetEmpty();
+	}
+	
+	TArray<void*> RawNodeData;
+	StructProperty->AccessRawData(RawNodeData);
+	if (RawNodeData.Num() == 1)
+	{
+		FText NameText;
+		FText PathText;
+		FText DescText;
+
+		if (const FStateTreeEditorNode* Node = static_cast<FStateTreeEditorNode*>(RawNodeData[0]))
+		{
+			const UStruct* Struct = Node->GetInstance().GetStruct();
+			if (Struct == nullptr || !Struct->IsChildOf<UStateTreeNodeBlueprintBase>())
+			{
+				Struct = Node->Node.GetScriptStruct();
+			}
+
+			if (Struct)
+			{
+				static const FName NAME_Tooltip(TEXT("Tooltip"));
+				const FText StructToolTipText = Struct->HasMetaData(NAME_Tooltip) ? Struct->GetToolTipText() : FText::GetEmpty();
+
+				FTextBuilder TooltipBuilder;
+				TooltipBuilder.AppendLineFormat(LOCTEXT("NodeTooltip", "{0} ({1})"), Struct->GetDisplayNameText(), FText::FromString(Struct->GetPathName()));
+
+				if (!StructToolTipText.IsEmpty())
+				{
+					TooltipBuilder.AppendLine();
+					TooltipBuilder.AppendLine(StructToolTipText);
+				}
+				return TooltipBuilder.ToText();
+			}
+		}
+	}
+
+	return FText::GetEmpty();
+}
+
 FText FStateTreeEditorNodeDetails::GetName() const
 {
 	check(StructProperty);
diff --git a/Engine/Plugins/Runtime/StateTree/Source/StateTreeEditorModule/Private/Customizations/StateTreeEditorNodeDetails.h b/Engine/Plugins/Runtime/StateTree/Source/StateTreeEditorModule/Private/Customizations/StateTreeEditorNodeDetails.h
--- a/Engine/Plugins/Runtime/StateTree/Source/StateTreeEditorModule/Private/Customizations/StateTreeEditorNodeDetails.h
+++ b/Engine/Plugins/Runtime/StateTree/Source/StateTreeEditorModule/Private/Customizations/StateTreeEditorNodeDetails.h
@@ -87,6 +87,8 @@ private:
 	FReply OnDescriptionClicked(const FGeometry& MyGeometry, const FPointerEvent& MouseEvent) const;
 	FText GetNodeDescription() const;
 	EVisibility IsNodeDescriptionVisible() const;
+	
+	FText GetNodeTooltip() const;
 
 	FReply OnRowMouseDown(const FGeometry& MyGeometry, const FPointerEvent& MouseEvent);
 	FReply OnRowMouseUp(const FGeometry& MyGeometry, const FPointerEvent& MouseEvent);
