StateTree: Allow FStateTreeIndex16 to serialize from uint16

diff --git a/Engine/Plugins/Runtime/StateTree/Source/StateTreeModule/Public/StateTreeIndexTypes.cpp b/Engine/Plugins/Runtime/StateTree/Source/StateTreeModule/Public/StateTreeIndexTypes.cpp
--- a/Engine/Plugins/Runtime/StateTree/Source/StateTreeModule/Public/StateTreeIndexTypes.cpp
+++ b/Engine/Plugins/Runtime/StateTree/Source/StateTreeModule/Public/StateTreeIndexTypes.cpp
@@ -1,12 +1,23 @@
 // Copyright Epic Games, Inc. All Rights Reserved.
 
 #include "StateTreeIndexTypes.h"
+#include "StateTreeTypes.h"
 
 bool FStateTreeIndex16::SerializeFromMismatchedTag(const FPropertyTag& Tag, FStructuredArchive::FSlot Slot)
 {
-	// Support loading from Index8.
-	if (Tag.GetType().IsStruct(FStateTreeIndex8::StaticStruct()->GetFName()))
+	if (Tag.Type == NAME_UInt16Property)
 	{
+		// Support loading from uint16.
+		// Note: 0xffff is silently read as invalid value.
+		uint16 OldValue = 0;
+		Slot << OldValue;
+
+		*this = FStateTreeIndex16(OldValue);
+		return true;
+	}
+	else if (Tag.GetType().IsStruct(FStateTreeIndex8::StaticStruct()->GetFName()))
+	{
+		// Support loading from Index8.
 		FStateTreeIndex8 OldValue;
 		FStateTreeIndex8::StaticStruct()->SerializeItem(Slot, &OldValue, nullptr);
 
