StateTree: Fixed calculating unique states between frames.

diff --git a/Engine/Plugins/Runtime/StateTree/Source/StateTreeModule/Private/StateTreeExecutionContext.cpp b/Engine/Plugins/Runtime/StateTree/Source/StateTreeModule/Private/StateTreeExecutionContext.cpp
--- a/Engine/Plugins/Runtime/StateTree/Source/StateTreeModule/Private/StateTreeExecutionContext.cpp
+++ b/Engine/Plugins/Runtime/StateTree/Source/StateTreeModule/Private/StateTreeExecutionContext.cpp
@@ -1494,7 +1494,7 @@ void FStateTreeExecutionContext::CaptureNewStateEvents(TConstArrayView<FStateTre
 			{
 				for (int32 StateIndex = 0; StateIndex < NewFrame.ActiveStates.Num(); ++StateIndex)
 				{
-					if (!PrevFrame.ActiveStates.IsValidIndex(StateIndex) || PrevFrame.ActiveStates[StateIndex] == NewFrame.ActiveStates[StateIndex])
+					if (!PrevFrame.ActiveStates.IsValidIndex(StateIndex) || PrevFrame.ActiveStates[StateIndex] != NewFrame.ActiveStates[StateIndex])
 					{
 						UniqueStates = TConstArrayView<FStateTreeStateHandle>(&NewFrame.ActiveStates[StateIndex], NewFrame.ActiveStates.Num() - StateIndex);
 						break;
