[State Tree Editor] Fix accessing nullptr crash when deleting blueprint task/condition assets while other ST assets are referencing them

diff --git a/Engine/Plugins/Runtime/StateTree/Source/StateTreeModule/Private/Blueprint/StateTreeConditionBlueprintBase.cpp b/Engine/Plugins/Runtime/StateTree/Source/StateTreeModule/Private/Blueprint/StateTreeConditionBlueprintBase.cpp
--- a/Engine/Plugins/Runtime/StateTree/Source/StateTreeModule/Private/Blueprint/StateTreeConditionBlueprintBase.cpp
+++ b/Engine/Plugins/Runtime/StateTree/Source/StateTreeModule/Private/Blueprint/StateTreeConditionBlueprintBase.cpp
@@ -61,19 +61,27 @@ FText FStateTreeBlueprintConditionWrapper::GetDescription(const FGuid& ID, FStat
 
 FName FStateTreeBlueprintConditionWrapper::GetIconName() const
 {
-	if (const UStateTreeNodeBlueprintBase* NodeCDO = GetDefault<const UStateTreeNodeBlueprintBase>(ConditionClass))
+	if (ConditionClass)
 	{
-		return NodeCDO->GetIconName();
+		if (const UStateTreeNodeBlueprintBase* NodeCDO = GetDefault<const UStateTreeNodeBlueprintBase>(ConditionClass))
+		{
+			return NodeCDO->GetIconName();
+		}
 	}
+
 	return FStateTreeConditionBase::GetIconName();
 }
 
 FColor FStateTreeBlueprintConditionWrapper::GetIconColor() const
 {
-	if (const UStateTreeNodeBlueprintBase* NodeCDO = GetDefault<const UStateTreeNodeBlueprintBase>(ConditionClass))
+	if (ConditionClass)
 	{
-		return NodeCDO->GetIconColor();
+		if (const UStateTreeNodeBlueprintBase* NodeCDO = GetDefault<const UStateTreeNodeBlueprintBase>(ConditionClass))
+		{
+			return NodeCDO->GetIconColor();
+		}
 	}
+
 	return FStateTreeConditionBase::GetIconColor();
 }
 #endif
\ No newline at end of file
diff --git a/Engine/Plugins/Runtime/StateTree/Source/StateTreeModule/Private/Blueprint/StateTreeConsiderationBlueprintBase.cpp b/Engine/Plugins/Runtime/StateTree/Source/StateTreeModule/Private/Blueprint/StateTreeConsiderationBlueprintBase.cpp
--- a/Engine/Plugins/Runtime/StateTree/Source/StateTreeModule/Private/Blueprint/StateTreeConsiderationBlueprintBase.cpp
+++ b/Engine/Plugins/Runtime/StateTree/Source/StateTreeModule/Private/Blueprint/StateTreeConsiderationBlueprintBase.cpp
@@ -54,19 +54,27 @@ FText FStateTreeBlueprintConsiderationWrapper::GetDescription(const FGuid& ID, F
 
 FName FStateTreeBlueprintConsiderationWrapper::GetIconName() const
 {
-	if (const UStateTreeNodeBlueprintBase* NodeCDO = GetDefault<const UStateTreeNodeBlueprintBase>(ConsiderationClass))
+	if (ConsiderationClass)
 	{
-		return NodeCDO->GetIconName();
+		if (const UStateTreeNodeBlueprintBase* NodeCDO = GetDefault<const UStateTreeNodeBlueprintBase>(ConsiderationClass))
+		{
+			return NodeCDO->GetIconName();
+		}
 	}
+
 	return FStateTreeConsiderationBase::GetIconName();
 }
 
 FColor FStateTreeBlueprintConsiderationWrapper::GetIconColor() const
 {
-	if (const UStateTreeNodeBlueprintBase* NodeCDO = GetDefault<const UStateTreeNodeBlueprintBase>(ConsiderationClass))
+	if (ConsiderationClass)
 	{
-		return NodeCDO->GetIconColor();
+		if (const UStateTreeNodeBlueprintBase* NodeCDO = GetDefault<const UStateTreeNodeBlueprintBase>(ConsiderationClass))
+		{
+			return NodeCDO->GetIconColor();
+		}
 	}
+
 	return FStateTreeConsiderationBase::GetIconColor();
 }
 #endif //WITH_EDITOR
diff --git a/Engine/Plugins/Runtime/StateTree/Source/StateTreeModule/Private/Blueprint/StateTreeTaskBlueprintBase.cpp b/Engine/Plugins/Runtime/StateTree/Source/StateTreeModule/Private/Blueprint/StateTreeTaskBlueprintBase.cpp
--- a/Engine/Plugins/Runtime/StateTree/Source/StateTreeModule/Private/Blueprint/StateTreeTaskBlueprintBase.cpp
+++ b/Engine/Plugins/Runtime/StateTree/Source/StateTreeModule/Private/Blueprint/StateTreeTaskBlueprintBase.cpp
@@ -163,19 +163,27 @@ FText FStateTreeBlueprintTaskWrapper::GetDescription(const FGuid& ID, FStateTree
 
 FName FStateTreeBlueprintTaskWrapper::GetIconName() const
 {
-	if (const UStateTreeNodeBlueprintBase* NodeCDO = GetDefault<const UStateTreeNodeBlueprintBase>(TaskClass))
+	if (TaskClass)
 	{
-		return NodeCDO->GetIconName();
+		if (const UStateTreeNodeBlueprintBase* NodeCDO = GetDefault<const UStateTreeNodeBlueprintBase>(TaskClass))
+		{
+			return NodeCDO->GetIconName();
+		}
 	}
+
 	return FStateTreeTaskBase::GetIconName();
 }
 
 FColor FStateTreeBlueprintTaskWrapper::GetIconColor() const
 {
-	if (const UStateTreeNodeBlueprintBase* NodeCDO = GetDefault<const UStateTreeNodeBlueprintBase>(TaskClass))
+	if (TaskClass)
 	{
-		return NodeCDO->GetIconColor();
+		if (const UStateTreeNodeBlueprintBase* NodeCDO = GetDefault<const UStateTreeNodeBlueprintBase>(TaskClass))
+		{
+			return NodeCDO->GetIconColor();
+		}
 	}
+
 	return FStateTreeTaskBase::GetIconColor();
 }
 #endif
\ No newline at end of file
