State Tree: Modified extension point for prop handlers to pass in a details view ref instead

diff --git a/Engine/Plugins/Experimental/Avalanche/Source/AvalancheTransitionEditor/Private/Views/SAvaTransitionSelectionDetails.cpp b/Engine/Plugins/Experimental/Avalanche/Source/AvalancheTransitionEditor/Private/Views/SAvaTransitionSelectionDetails.cpp
--- a/Engine/Plugins/Experimental/Avalanche/Source/AvalancheTransitionEditor/Private/Views/SAvaTransitionSelectionDetails.cpp
+++ b/Engine/Plugins/Experimental/Avalanche/Source/AvalancheTransitionEditor/Private/Views/SAvaTransitionSelectionDetails.cpp
@@ -39,9 +39,7 @@ void SAvaTransitionSelectionDetails::Construct(const FArguments& InArgs, const T
 		DetailsView->SetIsPropertyEditingEnabledDelegate(FIsPropertyEditingEnabled::CreateLambda([]{ return false; }));
 	}
 
-	FStateTreeEditorModule& StateTreeEditorModule = FModuleManager::LoadModuleChecked<FStateTreeEditorModule>("StateTreeEditorModule");
-	DetailsView->SetExtensionHandler(StateTreeEditorModule.CreateStateTreeDetailPropertyExtensionHandler());
-
+	FStateTreeEditorModule::SetDetailPropertyHandlers(*DetailsView);
 	OnSelectionChanged(InSelection->GetSelectedItems());
 
 	OnParametersChangedHandle = UE::StateTree::Delegates::OnParametersChanged.AddSP(this, &SAvaTransitionSelectionDetails::Refresh);
diff --git a/Engine/Plugins/Runtime/StateTree/Source/StateTreeEditorModule/Private/StateTreeEditor.cpp b/Engine/Plugins/Runtime/StateTree/Source/StateTreeEditorModule/Private/StateTreeEditor.cpp
--- a/Engine/Plugins/Runtime/StateTree/Source/StateTreeEditorModule/Private/StateTreeEditor.cpp
+++ b/Engine/Plugins/Runtime/StateTree/Source/StateTreeEditorModule/Private/StateTreeEditor.cpp
@@ -337,8 +337,7 @@ TSharedRef<SDockTab> FStateTreeEditor::SpawnTab_SelectionDetails(const FSpawnTab
 	SelectionDetailsView->SetObject(nullptr);
 	SelectionDetailsView->OnFinishedChangingProperties().AddSP(this, &FStateTreeEditor::OnSelectionFinishedChangingProperties);
 
-	SelectionDetailsView->SetExtensionHandler(MakeShared<FStateTreeBindingExtension>());
-	SelectionDetailsView->SetChildrenCustomizationHandler(MakeShared<FStateTreeBindingsChildrenCustomization>());
+	FStateTreeEditorModule::SetDetailPropertyHandlers(*SelectionDetailsView);
 
 	TSharedRef<SDockTab> SpawnedTab = SNew(SDockTab)
 		.Label(NSLOCTEXT("StateTreeEditor", "SelectionDetailsTab", "Details"))
@@ -361,8 +360,7 @@ TSharedRef<SDockTab> FStateTreeEditor::SpawnTab_AssetDetails(const FSpawnTabArgs
 	AssetDetailsView->SetObject(StateTree ? StateTree->EditorData : nullptr);
 	AssetDetailsView->OnFinishedChangingProperties().AddSP(this, &FStateTreeEditor::OnAssetFinishedChangingProperties);
 
-	AssetDetailsView->SetExtensionHandler(MakeShared<FStateTreeBindingExtension>());
-	AssetDetailsView->SetChildrenCustomizationHandler(MakeShared<FStateTreeBindingsChildrenCustomization>());
+	FStateTreeEditorModule::SetDetailPropertyHandlers(*AssetDetailsView);
 
 	TSharedRef<SDockTab> SpawnedTab = SNew(SDockTab)
 		.Label(NSLOCTEXT("StateTreeEditor", "AssetDetailsTabLabel", "Asset Details"))
diff --git a/Engine/Plugins/Runtime/StateTree/Source/StateTreeEditorModule/Private/StateTreeEditorModule.cpp b/Engine/Plugins/Runtime/StateTree/Source/StateTreeEditorModule/Private/StateTreeEditorModule.cpp
--- a/Engine/Plugins/Runtime/StateTree/Source/StateTreeEditorModule/Private/StateTreeEditorModule.cpp
+++ b/Engine/Plugins/Runtime/StateTree/Source/StateTreeEditorModule/Private/StateTreeEditorModule.cpp
@@ -149,9 +149,10 @@ TSharedRef<IStateTreeEditor> FStateTreeEditorModule::CreateStateTreeEditor(const
 	return NewEditor;
 }
 
-TSharedRef<IDetailPropertyExtensionHandler> FStateTreeEditorModule::CreateStateTreeDetailPropertyExtensionHandler()
+void FStateTreeEditorModule::SetDetailPropertyHandlers(IDetailsView& DetailsView)
 {
-	return MakeShared<FStateTreeBindingExtension>();
+	DetailsView.SetExtensionHandler(MakeShared<FStateTreeBindingExtension>());
+	DetailsView.SetChildrenCustomizationHandler(MakeShared<FStateTreeBindingsChildrenCustomization>());
 }
 
 TSharedPtr<FStateTreeNodeClassCache> FStateTreeEditorModule::GetNodeClassCache()
diff --git a/Engine/Plugins/Runtime/StateTree/Source/StateTreeEditorModule/Public/StateTreeEditorModule.h b/Engine/Plugins/Runtime/StateTree/Source/StateTreeEditorModule/Public/StateTreeEditorModule.h
--- a/Engine/Plugins/Runtime/StateTree/Source/StateTreeEditorModule/Public/StateTreeEditorModule.h
+++ b/Engine/Plugins/Runtime/StateTree/Source/StateTreeEditorModule/Public/StateTreeEditorModule.h
@@ -26,7 +26,8 @@ public:
 	/** Creates an instance of StateTree editor. Only virtual so that it can be called across the DLL boundary. */
 	virtual TSharedRef<IStateTreeEditor> CreateStateTreeEditor(const EToolkitMode::Type Mode, const TSharedPtr< class IToolkitHost >& InitToolkitHost, UStateTree* StateTree);
 
-	TSharedRef<IDetailPropertyExtensionHandler> CreateStateTreeDetailPropertyExtensionHandler();
+	/** Sets the Details View with required State Tree Detail Property Handlers */
+	static void SetDetailPropertyHandlers(IDetailsView& DetailsView);
 
 	virtual TSharedPtr<FExtensibilityManager> GetMenuExtensibilityManager() override { return MenuExtensibilityManager; }
 	virtual TSharedPtr<FExtensibilityManager> GetToolBarExtensibilityManager() override { return ToolBarExtensibilityManager; }
