UE State tree
- Added transition recording functionality to enable snapshots for state tree transitional events
- Primary use case is for keeping client-side state trees in sync with their server counterpart

#rnx

diff --git a/Engine/Plugins/Runtime/StateTree/Source/StateTreeModule/Private/StateTreeExecutionContext.cpp b/Engine/Plugins/Runtime/StateTree/Source/StateTreeModule/Private/StateTreeExecutionContext.cpp
--- a/Engine/Plugins/Runtime/StateTree/Source/StateTreeModule/Private/StateTreeExecutionContext.cpp
+++ b/Engine/Plugins/Runtime/StateTree/Source/StateTreeModule/Private/StateTreeExecutionContext.cpp
@@ -103,7 +103,7 @@ FStateTreeExecutionContext::FCurrentlyProcessedFrameScope::~FCurrentlyProcessedF
 }
 
 
-FStateTreeExecutionContext::FStateTreeExecutionContext(UObject& InOwner, const UStateTree& InStateTree, FStateTreeInstanceData& InInstanceData, const FOnCollectStateTreeExternalData& InCollectExternalDataDelegate)
+FStateTreeExecutionContext::FStateTreeExecutionContext(UObject& InOwner, const UStateTree& InStateTree, FStateTreeInstanceData& InInstanceData, const FOnCollectStateTreeExternalData& InCollectExternalDataDelegate, const EStateTreeRecordTransitions RecordTransitions)
 	: Owner(InOwner)
 	, RootStateTree(InStateTree)
 	, InstanceData(InInstanceData)
@@ -118,6 +118,8 @@ FStateTreeExecutionContext::FStateTreeExecutionContext(UObject& InOwner, const U
 		check(InstanceDataStorage);
 		
 		EventQueue = InstanceData.GetSharedMutableEventQueue();
+
+		bRecordTransitions = RecordTransitions == EStateTreeRecordTransitions::Yes;
 	}
 	else
 	{
@@ -1052,6 +1054,15 @@ FStateTreeDataView FStateTreeExecutionContext::GetDataView(const FStateTreeExecu
 	}
 }
 
+EStateTreeRunStatus FStateTreeExecutionContext::ForceTransition(const FRecordedStateTreeTransitionResult& Transition)
+{
+	FStateTreeTransitionResult TransitionResult = FStateTreeTransitionResult(Transition);
+
+	ExitState(TransitionResult);
+
+	return EnterState(TransitionResult);
+}
+
 FStateTreeDataView FStateTreeExecutionContext::GetDataViewFromInstanceStorage(FStateTreeInstanceStorage& InstanceDataStorage, FStateTreeInstanceStorage* CurrentlyProcessedSharedInstanceStorage, const FStateTreeExecutionFrame* ParentFrame, const FStateTreeExecutionFrame& CurrentFrame, const FStateTreeDataHandle Handle)
 {
 	switch (Handle.GetSource())
@@ -1541,6 +1552,11 @@ EStateTreeRunStatus FStateTreeExecutionContext::EnterState(FStateTreeTransitionR
 
 	FStateTreeExecutionState& Exec = GetExecState();
 
+	if (bRecordTransitions)
+	{
+		RecordedTransitions.Add(FRecordedStateTreeTransitionResult(Transition));
+	}
+
 	// Allocate new tasks.
 	UpdateInstanceData(Exec.ActiveFrames, Transition.NextActiveFrames);
 
diff --git a/Engine/Plugins/Runtime/StateTree/Source/StateTreeModule/Private/StateTreeExecutionTypes.cpp b/Engine/Plugins/Runtime/StateTree/Source/StateTreeModule/Private/StateTreeExecutionTypes.cpp
--- a/Engine/Plugins/Runtime/StateTree/Source/StateTreeModule/Private/StateTreeExecutionTypes.cpp
+++ b/Engine/Plugins/Runtime/StateTree/Source/StateTreeModule/Private/StateTreeExecutionTypes.cpp
@@ -6,4 +6,77 @@ const FStateTreeExternalDataHandle FStateTreeExternalDataHandle::Invalid = FStat
 
 #if WITH_STATETREE_TRACE
 const FStateTreeInstanceDebugId FStateTreeInstanceDebugId::Invalid = FStateTreeInstanceDebugId();
-#endif // WITH_STATETREE_TRACE
\ No newline at end of file
+#endif // WITH_STATETREE_TRACE
+
+FStateTreeTransitionResult::FStateTreeTransitionResult(const FRecordedStateTreeTransitionResult& RecordedTransition)
+{
+	for (int32 RecordedFrameIndex = 0; RecordedFrameIndex < RecordedTransition.NextActiveFrames.Num(); RecordedFrameIndex++)
+	{
+		const FRecordedStateTreeExecutionFrame& RecordedExecutionFrame = RecordedTransition.NextActiveFrames[RecordedFrameIndex];
+		NextActiveFrames.Add(RecordedExecutionFrame);
+
+		FStateTreeFrameStateSelectionEvents StateTreeFrameStateSelectionEvents;
+
+		for (int32 EventIdx = 0; EventIdx < RecordedExecutionFrame.EventIndices.Num(); EventIdx++)
+		{
+			if (RecordedTransition.NextActiveFrameEvents.IsValidIndex(EventIdx))
+			{
+				const FStateTreeEvent& RecordedStateTreeEvent = RecordedTransition.NextActiveFrameEvents[EventIdx];
+				StateTreeFrameStateSelectionEvents.Events[EventIdx] = FStateTreeSharedEvent(RecordedStateTreeEvent);
+			}
+		}
+
+		NextActiveFrameEvents.Add(MoveTemp(StateTreeFrameStateSelectionEvents));
+	}
+
+	SourceState = RecordedTransition.SourceState;
+	TargetState = RecordedTransition.SourceState;
+	Priority = RecordedTransition.Priority;
+	SourceStateTree = RecordedTransition.SourceStateTree;
+	SourceRootState = RecordedTransition.SourceRootState;
+}
+
+FRecordedStateTreeTransitionResult::FRecordedStateTreeTransitionResult(const FStateTreeTransitionResult& Transition)
+{
+	check(Transition.NextActiveFrames.Num() == Transition.NextActiveFrameEvents.Num());
+
+	for (int32 FrameIndex = 0; FrameIndex < Transition.NextActiveFrames.Num(); FrameIndex++)
+	{
+		const FStateTreeExecutionFrame& ExecutionFrame = Transition.NextActiveFrames[FrameIndex];
+		const FStateTreeFrameStateSelectionEvents& StateSelectionEvents = Transition.NextActiveFrameEvents[FrameIndex];
+
+		FRecordedStateTreeExecutionFrame& RecordedFrame = NextActiveFrames.Emplace_GetRef(ExecutionFrame);
+
+		for (int32 StateIndex = 0; StateIndex < ExecutionFrame.ActiveStates.Num(); StateIndex++)
+		{
+			const FStateTreeEvent* Event = StateSelectionEvents.Events[StateIndex].Get();
+			if (Event)
+			{
+				const int32 EventIndex = NextActiveFrameEvents.Add(*Event);
+				RecordedFrame.EventIndices[StateIndex] = static_cast<uint8>(EventIndex);
+			}
+		}
+	}
+
+	SourceState = Transition.SourceState;
+	TargetState = Transition.SourceState;
+	Priority = Transition.Priority;
+	SourceStateTree = Transition.SourceStateTree;
+	SourceRootState = Transition.SourceRootState;
+}
+
+FStateTreeExecutionFrame::FStateTreeExecutionFrame(const FRecordedStateTreeExecutionFrame& RecordedExecutionFrame)
+{
+	StateTree = RecordedExecutionFrame.StateTree;
+	RootState = RecordedExecutionFrame.RootState;
+	ActiveStates = RecordedExecutionFrame.ActiveStates;
+	bIsGlobalFrame = RecordedExecutionFrame.bIsGlobalFrame;
+}
+
+FRecordedStateTreeExecutionFrame::FRecordedStateTreeExecutionFrame(const FStateTreeExecutionFrame& ExecutionFrame)
+{
+	StateTree = ExecutionFrame.StateTree;
+	RootState = ExecutionFrame.RootState;
+	ActiveStates = ExecutionFrame.ActiveStates;
+	bIsGlobalFrame = ExecutionFrame.bIsGlobalFrame;
+}
diff --git a/Engine/Plugins/Runtime/StateTree/Source/StateTreeModule/Public/StateTreeExecutionContext.h b/Engine/Plugins/Runtime/StateTree/Source/StateTreeModule/Public/StateTreeExecutionContext.h
--- a/Engine/Plugins/Runtime/StateTree/Source/StateTreeModule/Public/StateTreeExecutionContext.h
+++ b/Engine/Plugins/Runtime/StateTree/Source/StateTreeModule/Public/StateTreeExecutionContext.h
@@ -92,7 +92,7 @@ DECLARE_DELEGATE_RetVal_FourParams(bool, FOnCollectStateTreeExternalData, const
 struct STATETREEMODULE_API FStateTreeExecutionContext
 {
 public:
-	FStateTreeExecutionContext(UObject& InOwner, const UStateTree& InStateTree, FStateTreeInstanceData& InInstanceData, const FOnCollectStateTreeExternalData& CollectExternalDataCallback = {});
+	FStateTreeExecutionContext(UObject& InOwner, const UStateTree& InStateTree, FStateTreeInstanceData& InInstanceData, const FOnCollectStateTreeExternalData& CollectExternalDataCallback = {}, const EStateTreeRecordTransitions RecordTransitions = EStateTreeRecordTransitions::No);
 	/** Construct an execution context from a parent context and another tree. Useful to run a subtree from the parent context with the same schema. */
 	FStateTreeExecutionContext(const FStateTreeExecutionContext& InContextToCopy, const UStateTree& InStateTree, FStateTreeInstanceData& InInstanceData);
 	virtual ~FStateTreeExecutionContext();
@@ -441,6 +441,17 @@ public:
 		return GetDataViewFromInstanceStorage(InstanceDataStorage, CurrentlyProcessedSharedInstanceStorage, ParentFrame, CurrentFrame, Handle);
 	}
 
+	/**
+	* Forces transition to a state from a previously recorded state tree transition result.
+	* Primarily used for replication purposes so that a client state tree stay in sync with its server counterpart.
+	* @param Recorded state transition to run on the state tree.
+	* @return The new run status for the state tree.
+	*/
+	EStateTreeRunStatus ForceTransition(const FRecordedStateTreeTransitionResult& Transition);
+
+	/** Returns the recorded transitions for this context. */
+	TConstArrayView<FRecordedStateTreeTransitionResult> GetRecordedTransitions() const { return RecordedTransitions; }
+
 protected:
 	/**
 	 * Describes a result of States Selection.
@@ -884,4 +895,10 @@ protected:
 		FStateTreeDataHandle SavedNodeDataHandle;
 		FStateTreeDataView SavedNodeInstanceData;
 	};
+
+	/** If true, the state tree context will create snapshots of transition events and capture them within RecordedTransitions for later use. */
+	bool bRecordTransitions = false;
+
+	/** Captured snapshots for transition results that can be used to recreate transitions. This array is only populated if bRecordTransitions is true. */
+	TArray<FRecordedStateTreeTransitionResult> RecordedTransitions;
 };
diff --git a/Engine/Plugins/Runtime/StateTree/Source/StateTreeModule/Public/StateTreeExecutionTypes.h b/Engine/Plugins/Runtime/StateTree/Source/StateTreeModule/Public/StateTreeExecutionTypes.h
--- a/Engine/Plugins/Runtime/StateTree/Source/StateTreeModule/Public/StateTreeExecutionTypes.h
+++ b/Engine/Plugins/Runtime/StateTree/Source/StateTreeModule/Public/StateTreeExecutionTypes.h
@@ -554,6 +554,9 @@ struct STATETREEMODULE_API FStateTreeExecutionFrame
 {
 	GENERATED_BODY()
 
+	FStateTreeExecutionFrame() = default;
+	FStateTreeExecutionFrame(const FRecordedStateTreeExecutionFrame& RecordedExecutionFrame);
+
 	bool IsSameFrame(const FStateTreeExecutionFrame& OtherFrame) const
 	{
 		return StateTree == OtherFrame.StateTree && RootState == OtherFrame.RootState;
@@ -739,6 +742,7 @@ PRAGMA_DISABLE_DEPRECATION_WARNINGS
 	FStateTreeTransitionResult& operator=(FStateTreeTransitionResult&&) = default;
 PRAGMA_ENABLE_DEPRECATION_WARNINGS
 
+	FStateTreeTransitionResult(const FRecordedStateTreeTransitionResult& RecordedTransition);
 	
 	void Reset()
 	{
@@ -761,11 +765,11 @@ PRAGMA_ENABLE_DEPRECATION_WARNINGS
 	UPROPERTY(EditDefaultsOnly, Category = "Default", BlueprintReadOnly)
 	EStateTreeRunStatus CurrentRunStatus = EStateTreeRunStatus::Unset;
 
-	/** Transition source state */
+	/** Transition source state. */
 	UPROPERTY(EditDefaultsOnly, Category = "Default", BlueprintReadOnly)
 	FStateTreeStateHandle SourceState = FStateTreeStateHandle::Invalid;
 
-	/** Transition target state */
+	/** Transition target state. */
 	UPROPERTY(EditDefaultsOnly, Category = "Default", BlueprintReadOnly)
 	FStateTreeStateHandle TargetState = FStateTreeStateHandle::Invalid;
 
@@ -798,4 +802,89 @@ PRAGMA_ENABLE_DEPRECATION_WARNINGS
 	UPROPERTY()
 	FStateTreeActiveStates NextActiveStates_DEPRECATED;
 #endif	
+};
+
+/*
+ * Enumeration for the different transition recording types.
+ * This is used by the execution context to capture transition snapshots if set to record.
+*/
+UENUM()
+enum class EStateTreeRecordTransitions : uint8
+{
+	No,
+	Yes
+};
+
+/*
+* Captured state tree execution frame that can be cached for recording purposes.
+* Held in FRecordedStateTreeTransitionResult for its NextActiveFrames.
+*/
+USTRUCT()
+struct STATETREEMODULE_API FRecordedStateTreeExecutionFrame
+{
+	GENERATED_BODY()
+
+	FRecordedStateTreeExecutionFrame() = default;
+	FRecordedStateTreeExecutionFrame(const FStateTreeExecutionFrame& ExecutionFrame);
+
+	/** The State Tree used for ticking this frame. */
+	UPROPERTY()
+	TObjectPtr<const UStateTree> StateTree = nullptr;
+
+	/** The root state of the frame (e.g. Root state or a subtree). */
+	UPROPERTY()
+	FStateTreeStateHandle RootState = FStateTreeStateHandle::Root; 
+	
+	/** Active states in this frame. */
+	UPROPERTY()
+	FStateTreeActiveStates ActiveStates;
+
+	/** If true, the global tasks of the State Tree should be handle in this frame. */
+	UPROPERTY()
+	uint8 bIsGlobalFrame : 1 = false;
+
+	/** Captured indices of the events we've recorded. */
+	TStaticArray<uint8, FStateTreeActiveStates::MaxStates> EventIndices;
+};
+
+/*
+* Captured state tree transition result that can be cached for recording purposes.
+* Primarily, when transitions are recorded through this structure, we can replicate them down
+* to clients to keep our state tree in sync.
+*/
+USTRUCT()
+struct STATETREEMODULE_API FRecordedStateTreeTransitionResult
+{
+	GENERATED_BODY()
+
+	FRecordedStateTreeTransitionResult() = default;
+	FRecordedStateTreeTransitionResult(const FStateTreeTransitionResult& Transition);
+
+   	/** States selected as result of the transition. */
+    UPROPERTY()
+	TArray<FRecordedStateTreeExecutionFrame> NextActiveFrames;
+
+	/** Captured events from the transition that we've recorded */
+	UPROPERTY()
+	TArray<FStateTreeEvent> NextActiveFrameEvents;
+
+	/** Transition source state. */
+	UPROPERTY()
+	FStateTreeStateHandle SourceState = FStateTreeStateHandle::Invalid;
+
+	/** Transition target state. */
+	UPROPERTY()
+	FStateTreeStateHandle TargetState = FStateTreeStateHandle::Invalid;
+	
+	/** Priority of the transition that caused the state change. */
+	UPROPERTY()
+	EStateTreeTransitionPriority Priority = EStateTreeTransitionPriority::None;
+
+	/** StateTree asset that was active when the transition was requested. */
+	UPROPERTY()
+	TObjectPtr<const UStateTree> SourceStateTree = nullptr;
+
+	/** Root state the execution frame where the transition was requested. */
+	UPROPERTY()
+	FStateTreeStateHandle SourceRootState = FStateTreeStateHandle::Invalid;
 };
\ No newline at end of file
