State Tree: Added binding creation support in Binding Extension
- this also considers property refs

diff --git a/Engine/Plugins/Experimental/StructUtils/Source/StructUtilsEditor/Private/Customizations/PropertyBagDetails.cpp b/Engine/Plugins/Experimental/StructUtils/Source/StructUtilsEditor/Private/Customizations/PropertyBagDetails.cpp
--- a/Engine/Plugins/Experimental/StructUtils/Source/StructUtilsEditor/Private/Customizations/PropertyBagDetails.cpp
+++ b/Engine/Plugins/Experimental/StructUtils/Source/StructUtilsEditor/Private/Customizations/PropertyBagDetails.cpp
@@ -23,7 +23,128 @@
 
 ////////////////////////////////////
 
-namespace UE::StructUtils::Private
+namespace UE::StructUtils
 {
 
+/** Sets property descriptor based on a Blueprint pin type. */
+void SetPropertyDescFromPin(FPropertyBagPropertyDesc& Desc, const FEdGraphPinType& PinType)
+{
+	const UPropertyBagSchema* Schema = GetDefault<UPropertyBagSchema>();
+	check(Schema);
+
+	// remove any existing containers
+	Desc.ContainerTypes.Reset();
+
+	// Fill Container types, if any
+	switch (PinType.ContainerType)
+	{
+	case EPinContainerType::Array:
+		Desc.ContainerTypes.Add(EPropertyBagContainerType::Array);
+		break;
+	case EPinContainerType::Set:
+		ensureMsgf(false, TEXT("Unsuported container type [Set] "));
+		break;
+	case EPinContainerType::Map:
+		ensureMsgf(false, TEXT("Unsuported container type [Map] "));
+		break;
+	default:
+		break;
+	}
+	
+	// Value type
+	if (PinType.PinCategory == UEdGraphSchema_K2::PC_Boolean)
+	{
+		Desc.ValueType = EPropertyBagPropertyType::Bool;
+		Desc.ValueTypeObject = nullptr;
+	}
+	else if (PinType.PinCategory == UEdGraphSchema_K2::PC_Byte)
+	{
+		if (UEnum* Enum = Cast<UEnum>(PinType.PinSubCategoryObject))
+		{
+			Desc.ValueType = EPropertyBagPropertyType::Enum;
+			Desc.ValueTypeObject = PinType.PinSubCategoryObject.Get();
+		}
+		else
+		{
+			Desc.ValueType = EPropertyBagPropertyType::Byte;
+			Desc.ValueTypeObject = nullptr;
+		}
+	}
+	else if (PinType.PinCategory == UEdGraphSchema_K2::PC_Int)
+	{
+		Desc.ValueType = EPropertyBagPropertyType::Int32;
+		Desc.ValueTypeObject = nullptr;
+	}
+	else if (PinType.PinCategory == UEdGraphSchema_K2::PC_Int64)
+	{
+		Desc.ValueType = EPropertyBagPropertyType::Int64;
+		Desc.ValueTypeObject = nullptr;
+	}
+	else if (PinType.PinCategory == UEdGraphSchema_K2::PC_Real)
+	{
+		if (PinType.PinSubCategory == UEdGraphSchema_K2::PC_Float)
+		{
+			Desc.ValueType = EPropertyBagPropertyType::Float;
+			Desc.ValueTypeObject = nullptr;
+		}
+		else if (PinType.PinSubCategory == UEdGraphSchema_K2::PC_Double)
+		{
+			Desc.ValueType = EPropertyBagPropertyType::Double;
+			Desc.ValueTypeObject = nullptr;
+		}		
+	}
+	else if (PinType.PinCategory == UEdGraphSchema_K2::PC_Name)
+	{
+		Desc.ValueType = EPropertyBagPropertyType::Name;
+		Desc.ValueTypeObject = nullptr;
+	}
+	else if (PinType.PinCategory == UEdGraphSchema_K2::PC_String)
+	{
+		Desc.ValueType = EPropertyBagPropertyType::String;
+		Desc.ValueTypeObject = nullptr;
+	}
+	else if (PinType.PinCategory == UEdGraphSchema_K2::PC_Text)
+	{
+		Desc.ValueType = EPropertyBagPropertyType::Text;
+		Desc.ValueTypeObject = nullptr;
+	}
+	else if (PinType.PinCategory == UEdGraphSchema_K2::PC_Enum)
+	{
+		Desc.ValueType = EPropertyBagPropertyType::Enum;
+		Desc.ValueTypeObject = PinType.PinSubCategoryObject.Get();
+	}
+	else if (PinType.PinCategory == UEdGraphSchema_K2::PC_Struct)
+	{
+		Desc.ValueType = EPropertyBagPropertyType::Struct;
+		Desc.ValueTypeObject = PinType.PinSubCategoryObject.Get();
+	}
+	else if (PinType.PinCategory == UEdGraphSchema_K2::PC_Object)
+	{
+		Desc.ValueType = EPropertyBagPropertyType::Object;
+		Desc.ValueTypeObject = PinType.PinSubCategoryObject.Get();
+	}
+	else if (PinType.PinCategory == UEdGraphSchema_K2::PC_SoftObject)
+	{
+		Desc.ValueType = EPropertyBagPropertyType::SoftObject;
+		Desc.ValueTypeObject = PinType.PinSubCategoryObject.Get();
+	}
+	else if (PinType.PinCategory == UEdGraphSchema_K2::PC_Class)
+	{
+		Desc.ValueType = EPropertyBagPropertyType::Class;
+		Desc.ValueTypeObject = PinType.PinSubCategoryObject.Get();
+	}
+	else if (PinType.PinCategory == UEdGraphSchema_K2::PC_SoftClass)
+	{
+		Desc.ValueType = EPropertyBagPropertyType::SoftClass;
+		Desc.ValueTypeObject = PinType.PinSubCategoryObject.Get();
+	}
+	else
+	{
+		ensureMsgf(false, TEXT("Unhandled pin category %s"), *PinType.PinCategory.ToString());
+	}
+}
+
+namespace Private
+{
+
 /** @return true property handle holds struct property of type T.  */ 
@@ -302,124 +423,6 @@
 	return PinType;
 }
 
-/** Sets property descriptor based on a Blueprint pin type. */
-void SetPropertyDescFromPin(FPropertyBagPropertyDesc& Desc, const FEdGraphPinType& PinType)
-{
-	const UPropertyBagSchema* Schema = GetDefault<UPropertyBagSchema>();
-	check(Schema);
-
-	// remove any existing containers
-	Desc.ContainerTypes.Reset();
-
-	// Fill Container types, if any
-	switch (PinType.ContainerType)
-	{
-	case EPinContainerType::Array:
-		Desc.ContainerTypes.Add(EPropertyBagContainerType::Array);
-		break;
-	case EPinContainerType::Set:
-		ensureMsgf(false, TEXT("Unsuported container type [Set] "));
-		break;
-	case EPinContainerType::Map:
-		ensureMsgf(false, TEXT("Unsuported container type [Map] "));
-		break;
-	default:
-		break;
-	}
-	
-	// Value type
-	if (PinType.PinCategory == UEdGraphSchema_K2::PC_Boolean)
-	{
-		Desc.ValueType = EPropertyBagPropertyType::Bool;
-		Desc.ValueTypeObject = nullptr;
-	}
-	else if (PinType.PinCategory == UEdGraphSchema_K2::PC_Byte)
-	{
-		if (UEnum* Enum = Cast<UEnum>(PinType.PinSubCategoryObject))
-		{
-			Desc.ValueType = EPropertyBagPropertyType::Enum;
-			Desc.ValueTypeObject = PinType.PinSubCategoryObject.Get();
-		}
-		else
-		{
-			Desc.ValueType = EPropertyBagPropertyType::Byte;
-			Desc.ValueTypeObject = nullptr;
-		}
-	}
-	else if (PinType.PinCategory == UEdGraphSchema_K2::PC_Int)
-	{
-		Desc.ValueType = EPropertyBagPropertyType::Int32;
-		Desc.ValueTypeObject = nullptr;
-	}
-	else if (PinType.PinCategory == UEdGraphSchema_K2::PC_Int64)
-	{
-		Desc.ValueType = EPropertyBagPropertyType::Int64;
-		Desc.ValueTypeObject = nullptr;
-	}
-	else if (PinType.PinCategory == UEdGraphSchema_K2::PC_Real)
-	{
-		if (PinType.PinSubCategory == UEdGraphSchema_K2::PC_Float)
-		{
-			Desc.ValueType = EPropertyBagPropertyType::Float;
-			Desc.ValueTypeObject = nullptr;
-		}
-		else if (PinType.PinSubCategory == UEdGraphSchema_K2::PC_Double)
-		{
-			Desc.ValueType = EPropertyBagPropertyType::Double;
-			Desc.ValueTypeObject = nullptr;
-		}		
-	}
-	else if (PinType.PinCategory == UEdGraphSchema_K2::PC_Name)
-	{
-		Desc.ValueType = EPropertyBagPropertyType::Name;
-		Desc.ValueTypeObject = nullptr;
-	}
-	else if (PinType.PinCategory == UEdGraphSchema_K2::PC_String)
-	{
-		Desc.ValueType = EPropertyBagPropertyType::String;
-		Desc.ValueTypeObject = nullptr;
-	}
-	else if (PinType.PinCategory == UEdGraphSchema_K2::PC_Text)
-	{
-		Desc.ValueType = EPropertyBagPropertyType::Text;
-		Desc.ValueTypeObject = nullptr;
-	}
-	else if (PinType.PinCategory == UEdGraphSchema_K2::PC_Enum)
-	{
-		Desc.ValueType = EPropertyBagPropertyType::Enum;
-		Desc.ValueTypeObject = PinType.PinSubCategoryObject.Get();
-	}
-	else if (PinType.PinCategory == UEdGraphSchema_K2::PC_Struct)
-	{
-		Desc.ValueType = EPropertyBagPropertyType::Struct;
-		Desc.ValueTypeObject = PinType.PinSubCategoryObject.Get();
-	}
-	else if (PinType.PinCategory == UEdGraphSchema_K2::PC_Object)
-	{
-		Desc.ValueType = EPropertyBagPropertyType::Object;
-		Desc.ValueTypeObject = PinType.PinSubCategoryObject.Get();
-	}
-	else if (PinType.PinCategory == UEdGraphSchema_K2::PC_SoftObject)
-	{
-		Desc.ValueType = EPropertyBagPropertyType::SoftObject;
-		Desc.ValueTypeObject = PinType.PinSubCategoryObject.Get();
-	}
-	else if (PinType.PinCategory == UEdGraphSchema_K2::PC_Class)
-	{
-		Desc.ValueType = EPropertyBagPropertyType::Class;
-		Desc.ValueTypeObject = PinType.PinSubCategoryObject.Get();
-	}
-	else if (PinType.PinCategory == UEdGraphSchema_K2::PC_SoftClass)
-	{
-		Desc.ValueType = EPropertyBagPropertyType::SoftClass;
-		Desc.ValueTypeObject = PinType.PinSubCategoryObject.Get();
-	}
-	else
-	{
-		ensureMsgf(false, TEXT("Unhandled pin category %s"), *PinType.PinCategory.ToString());
-	}
-}
-
 template<typename TFunc>
 void ApplyChangesToPropertyDescs(const FText& SessionName, const TSharedPtr<IPropertyHandle>& StructProperty, const TSharedPtr<IPropertyUtilities>& PropUtils, TFunc&& Function)
 {
@@ -769,6 +772,7 @@
 
 } // UE::StructUtils::Private
 
+} // UE::StructUtils
 
 //----------------------------------------------------------------//
 //  FPropertyBagInstanceDataDetails
@@ -1245,7 +1249,7 @@
 				const FProperty* Property = ChildPropertyHandle ? ChildPropertyHandle->GetProperty() : nullptr;
 				if (FPropertyBagPropertyDesc* Desc = PropertyDescs.FindByPredicate([Property](const FPropertyBagPropertyDesc& Desc){ return Desc.CachedProperty == Property; }))
 				{
-					UE::StructUtils::Private::SetPropertyDescFromPin(*Desc, PinType);
+					UE::StructUtils::SetPropertyDescFromPin(*Desc, PinType);
 				}
 			});
 	};
diff --git a/Engine/Plugins/Experimental/StructUtils/Source/StructUtilsEditor/Public/PropertyBagDetails.h b/Engine/Plugins/Experimental/StructUtils/Source/StructUtilsEditor/Public/PropertyBagDetails.h
--- a/Engine/Plugins/Experimental/StructUtils/Source/StructUtilsEditor/Public/PropertyBagDetails.h
+++ b/Engine/Plugins/Experimental/StructUtils/Source/StructUtilsEditor/Public/PropertyBagDetails.h
@@ -17,6 +17,11 @@ class IDetailPropertyRow;
 class SInlineEditableTextBlock;
 class SWidget;
 
+namespace UE::StructUtils
+{
+	STRUCTUTILSEDITOR_API void SetPropertyDescFromPin(FPropertyBagPropertyDesc& Desc, const FEdGraphPinType& PinType);
+}
+
 /**
  * Type customization for FInstancedPropertyBag.
  */
diff --git a/Engine/Plugins/Runtime/StateTree/Source/StateTreeEditorModule/Private/Customizations/StateTreeBindingExtension.cpp b/Engine/Plugins/Runtime/StateTree/Source/StateTreeEditorModule/Private/Customizations/StateTreeBindingExtension.cpp
--- a/Engine/Plugins/Runtime/StateTree/Source/StateTreeEditorModule/Private/Customizations/StateTreeBindingExtension.cpp
+++ b/Engine/Plugins/Runtime/StateTree/Source/StateTreeEditorModule/Private/Customizations/StateTreeBindingExtension.cpp
@@ -10,6 +10,7 @@
 #include "StateTreeNodeBase.h"
 #include "Styling/AppStyle.h"
 #include "UObject/EnumProperty.h"
+#include "Widgets/Images/SImage.h"
 #include "Widgets/Layout/SBox.h"
 #include "StateTreePropertyRef.h"
 #include "StateTreePropertyRefHelpers.h"
@@ -17,6 +18,8 @@
 #include "IPropertyUtilities.h"
 #include "IDetailChildrenBuilder.h"
 #include "IStructureDataProvider.h"
+#include "Kismet2/BlueprintEditorUtils.h"
+#include "PropertyBagDetails.h"
 #include "StateTreeEditorNodeUtils.h"
 #include "StateTreeEditorModule.h"
 #include "StateTreeNodeClassCache.h"
@@ -27,7 +30,17 @@
 
 namespace UE::StateTree::PropertyBinding
 {
-	
+
+/** Information for the types gathered from a FStateTreePropertyRef property meta-data */
+struct FRefTypeInfo
+{
+	/** Display Name Text of the Ref Type */
+	FText TypeNameText;
+
+	/** Ref Type expressed as a Pin Type */
+	FEdGraphPinType PinType;
+};
+
 const FName StateTreeNodeIDName(TEXT("StateTreeNodeID"));
 const FName AllowAnyBindingName(TEXT("AllowAnyBinding"));
 
@@ -238,6 +251,118 @@ FText GetPropertyTypeText(const FProperty* Property)
 	return UEdGraphSchema_K2::GetCategoryText(PinType.PinCategory, NAME_None, true);
 }
 
+TSharedRef<SWidget> MakeContextStructWidget(const FStateTreeBindableStructDesc& InContextStruct)
+{
+	FEdGraphPinType PinType;
+
+	UStruct* Struct = const_cast<UStruct*>(InContextStruct.Struct.Get());
+
+	if (UClass* Class = Cast<UClass>(Struct))
+	{
+		PinType.PinCategory = UEdGraphSchema_K2::PC_Object;
+		PinType.PinSubCategory = NAME_None;
+		PinType.PinSubCategoryObject = Class;
+	}
+	else if (UScriptStruct* ScriptStruct = Cast<UScriptStruct>(Struct))
+	{
+		PinType.PinCategory = UEdGraphSchema_K2::PC_Struct;
+		PinType.PinSubCategory = NAME_None;
+		PinType.PinSubCategoryObject = ScriptStruct;
+	}
+
+	const FSlateBrush* Icon = FBlueprintEditorUtils::GetIconFromPin(PinType, true);
+	const FLinearColor IconColor = GetDefault<UEdGraphSchema_K2>()->GetPinTypeColor(PinType);
+
+	return SNew(SHorizontalBox)
+		+ SHorizontalBox::Slot()
+		.AutoWidth()
+		[
+			SNew(SSpacer)
+			.Size(FVector2D(18.0f, 0.0f))
+		]
+		+ SHorizontalBox::Slot()
+		.AutoWidth()
+		.VAlign(VAlign_Center)
+		.Padding(1.0f, 0.0f)
+		[
+			SNew(SImage)
+			.Image(Icon)
+			.ColorAndOpacity(IconColor)
+		]
+		+ SHorizontalBox::Slot()
+		.AutoWidth()
+		.VAlign(VAlign_Center)
+		.Padding(4.0f, 0.0f)
+		[
+			SNew(STextBlock)
+			.Text(FText::FromName(InContextStruct.Name))
+		];
+}
+
+TSharedRef<SWidget> MakeBindingPropertyInfoWidget(const FText& InDisplayText, const FEdGraphPinType& InPinType)
+{
+	const FSlateBrush* Icon = FBlueprintEditorUtils::GetIconFromPin(InPinType, /*bIsLarge*/true);
+	const FLinearColor IconColor = GetDefault<UEdGraphSchema_K2>()->GetPinTypeColor(InPinType);
+
+	return SNew(SHorizontalBox)
+		+SHorizontalBox::Slot()
+		.AutoWidth()
+		[
+			SNew(SSpacer)
+			.Size(FVector2D(18.0f, 0.0f))
+		]
+		+SHorizontalBox::Slot()
+		.AutoWidth()
+		.VAlign(VAlign_Center)
+		.Padding(1.0f, 0.0f)
+		[
+			SNew(SImage)
+			.Image(Icon)
+			.ColorAndOpacity(IconColor)
+		]
+		+SHorizontalBox::Slot()
+		.AutoWidth()
+		.VAlign(VAlign_Center)
+		.Padding(4.0f, 0.0f)
+		[
+			SNew(STextBlock)
+			.Text(InDisplayText)
+		];
+}
+	
+/** Helper struct to Begin/End Sections */
+struct FSectionHelper
+{
+	FSectionHelper(FMenuBuilder& MenuBuilder)
+		: MenuBuilder(MenuBuilder)
+	{
+	}
+	~FSectionHelper()
+	{
+		if (bSectionOpened)
+		{
+			MenuBuilder.EndSection();
+		}
+	}
+	void SetSection(const FText& InSection)
+	{
+		if (!InSection.IdenticalTo(CurrentSection))
+		{
+			if (bSectionOpened)
+			{
+				MenuBuilder.EndSection();
+			}
+			CurrentSection = InSection;
+			MenuBuilder.BeginSection(NAME_None, CurrentSection);
+			bSectionOpened = true;
+		}
+	}
+private:
+	FText CurrentSection;
+	FMenuBuilder& MenuBuilder;
+	bool bSectionOpened = false;
+};
+
 FOnStateTreePropertyBindingChanged STATETREEEDITORMODULE_API OnStateTreePropertyBindingChanged;
 
 struct FCachedBindingData : public TSharedFromThis<FCachedBindingData>
@@ -368,6 +493,173 @@ struct FCachedBindingData : public TSharedFromThis<FCachedBindingData>
 		UE::StateTree::PropertyBinding::OnStateTreePropertyBindingChanged.Broadcast(SourcePath, TargetPath);
 	}
 
+	bool CanCreateParameter(const FStateTreeBindableStructDesc& InStructDesc, TArray<TSharedPtr<const FRefTypeInfo>>& OutRefTypeInfos) const
+	{
+		const FProperty* Property = PropertyHandle->GetProperty();
+		if (!Property)
+		{
+			return false;
+		}
+
+		IStateTreeEditorPropertyBindingsOwner* BindingOwner = Cast<IStateTreeEditorPropertyBindingsOwner>(WeakOwnerObject.Get());
+		if (!BindingOwner)
+		{
+			return false;
+		}
+
+		if (!BindingOwner->CanCreateParameter(InStructDesc.ID))
+		{
+			return false;
+		}
+
+		// Add the PropertyRef property type with its RefTypes
+		const FStructProperty* StructProperty = CastField<const FStructProperty>(Property);
+		if (StructProperty && StructProperty->Struct && StructProperty->Struct->IsChildOf(FStateTreePropertyRef::StaticStruct()))
+		{
+			TArray<FEdGraphPinType, TInlineAllocator<1>> PinTypes;
+
+			const bool bCanTargetRefArray = PropertyHandle->HasMetaData(PropertyRefHelpers::CanRefToArrayName);
+
+			if (StructProperty->Struct->IsChildOf(FStateTreeBlueprintPropertyRef::StaticStruct()))
+			{
+				void* PropertyRefAddress = nullptr;
+				if (PropertyHandle->GetValueData(PropertyRefAddress) == FPropertyAccess::Result::Success)
+				{
+					check(PropertyRefAddress);
+					PinTypes.Add(PropertyRefHelpers::GetBlueprintPropertyRefInternalTypeAsPin(*static_cast<const FStateTreeBlueprintPropertyRef*>(PropertyRefAddress)));
+				}
+			}
+			else
+			{
+				PinTypes = PropertyRefHelpers::GetPropertyRefInternalTypesAsPins(*Property);
+			}
+
+			// If Property supports Arrays, add the Array version of these pin types
+			if (PropertyHandle->HasMetaData(PropertyRefHelpers::CanRefToArrayName))
+			{
+				const int32 PinTypeNum = PinTypes.Num();
+				for (int32 Index = 0; Index < PinTypeNum; ++Index)
+				{
+					const FEdGraphPinType& SourcePinType = PinTypes[Index];
+					if (!SourcePinType.IsArray())
+					{
+						FEdGraphPinType& PinType = PinTypes.Emplace_GetRef(SourcePinType);
+						PinType.ContainerType = EPinContainerType::Array;
+					}
+				}
+			}
+
+			const UEdGraphSchema_K2* Schema = GetDefault<UEdGraphSchema_K2>();
+
+			for (const FEdGraphPinType& PinType : PinTypes)
+			{
+				TSharedRef<FRefTypeInfo> RefTypeInfo = MakeShared<FRefTypeInfo>();
+				RefTypeInfo->PinType = PinType;
+
+				FString TypeName;
+				if(UObject* SubCategoryObject = PinType.PinSubCategoryObject.Get()) 
+				{
+					TypeName = SubCategoryObject->GetName();
+				}
+				else
+				{
+					TypeName = PinType.PinCategory.ToString() + TEXT(" ") + PinType.PinSubCategory.ToString();
+				}
+
+				RefTypeInfo->TypeNameText = FText::FromString(TypeName);
+				OutRefTypeInfos.Emplace(MoveTemp(RefTypeInfo));
+			}
+		}
+
+		return true;
+	}
+
+	void PromoteToParameter(FName InPropertyName, FStateTreeBindableStructDesc InStructDesc, TSharedPtr<const FRefTypeInfo> InPropertyInfoOverride)
+	{
+		if (!TargetPath.GetStructID().IsValid())
+		{
+			return;
+		}
+		
+		UObject* OwnerObject = WeakOwnerObject.Get();
+		if (!OwnerObject)
+		{
+			return;
+		}
+
+		IStateTreeEditorPropertyBindingsOwner* BindingOwner = Cast<IStateTreeEditorPropertyBindingsOwner>(OwnerObject);
+		if (!BindingOwner)
+		{
+			return;
+		}
+
+		const FProperty* Property = PropertyHandle->GetProperty();
+		if (!Property)
+		{
+			return;
+		}
+
+		const FProperty* TargetProperty = nullptr;
+		const void* TargetContainerAddress = nullptr;
+
+		FStateTreeDataView TargetDataView;
+		if (BindingOwner->GetDataViewByID(TargetPath.GetStructID(), TargetDataView) && TargetDataView.IsValid())
+		{
+			TArray<FStateTreePropertyPathIndirection> TargetIndirections;
+			if (ensure(TargetPath.ResolveIndirectionsWithValue(TargetDataView, TargetIndirections)))
+			{
+				const FStateTreePropertyPathIndirection& LastIndirection = TargetIndirections.Last();
+				TargetProperty = LastIndirection.GetProperty();
+				TargetContainerAddress = LastIndirection.GetContainerAddress();
+			}
+		}
+
+		FStateTreeEditorPropertyBindings* EditorBindings = BindingOwner->GetPropertyEditorBindings();
+		if (!EditorBindings)
+		{
+			return;
+		}
+
+		const FGuid StructID = InStructDesc.ID;
+
+		TArray<FStateTreeEditorPropertyCreationDesc, TFixedAllocator<1>> PropertyCreationDescs;
+		{
+			FStateTreeEditorPropertyCreationDesc& PropertyCreationDesc = PropertyCreationDescs.AddDefaulted_GetRef();
+
+			if (InPropertyInfoOverride)
+			{
+				PropertyCreationDesc.PropertyDesc.Name = InPropertyName;
+				UE::StructUtils::SetPropertyDescFromPin(PropertyCreationDesc.PropertyDesc, InPropertyInfoOverride->PinType);
+			}
+			else
+			{
+				PropertyCreationDesc.PropertyDesc = FPropertyBagPropertyDesc(InPropertyName, Property);
+			}
+
+			// Create desc based on the Target Property, but without the meta-data.
+			// This functionality mirrors the user action of adding a new property from the UI, where meta-data is not available.
+			// Additionally, meta-data like EditCondition is not desirable here
+			PropertyCreationDesc.PropertyDesc.MetaClass = nullptr;
+			PropertyCreationDesc.PropertyDesc.MetaData.Reset();
+
+			// Set the Property & Container Address to copy
+			if (TargetProperty && TargetContainerAddress)
+			{
+				PropertyCreationDesc.SourceProperty = TargetProperty;
+				PropertyCreationDesc.SourceContainerAddress = TargetContainerAddress;
+			}
+		}
+
+		OwnerObject->Modify();
+		BindingOwner->CreateParameters(StructID, /*InOut*/PropertyCreationDescs);
+
+		// Use the name in PropertyDescs, as it might contain a different name than the desired InPropertyName (for uniquness)
+		FStateTreePropertyPath SourcePath(StructID, PropertyCreationDescs[0].PropertyDesc.Name);
+		EditorBindings->AddPropertyBinding(SourcePath, TargetPath);
+
+		UpdateData();
+		UE::StateTree::PropertyBinding::OnStateTreePropertyBindingChanged.Broadcast(SourcePath, TargetPath);
+	}
 
 	void UpdateData()
 	{
@@ -1335,12 +1627,120 @@ void FStateTreeBindingExtension::ExtendWidgetRow(FDetailWidgetRow& InWidgetRow,
 	Args.bAllowArrayElementBindings = false;
 	Args.bAllowUObjectFunctions = false;
 
+	if (CanPromoteToParameter(InPropertyHandle))
+	{
+		Args.MenuExtender = MakeShared<FExtender>();
+		Args.MenuExtender->AddMenuExtension(
+			TEXT("BindingActions"),
+			EExtensionHook::After,
+			nullptr,
+			FMenuExtensionDelegate::CreateLambda([CachedBindingData, AccessibleStructs = MoveTemp(AccessibleStructs), InPropertyHandle](FMenuBuilder& MenuBuilder)
+			{
+				MenuBuilder.AddSubMenu(
+					LOCTEXT("PromoteToParameter", "Promote to Parameter"),
+					LOCTEXT("PromoteToParameterTooltip", "Create a new parameter of the same type as the property, copy value over, and bind the property to the new parameter."),
+					FNewMenuDelegate::CreateLambda([&CachedBindingData, &AccessibleStructs, &InPropertyHandle](FMenuBuilder& InMenuBuilder)
+					{
+						using namespace UE::StateTree::PropertyBinding;
+
+						const FProperty* Property = InPropertyHandle->GetProperty();
+						check(Property);
+						const FName PropertyName = Property->GetFName();
+
+						TSharedRef<FCachedBindingData> CachedBindingDataRef = CachedBindingData.ToSharedRef();
+
+						FSectionHelper SectionHelper(InMenuBuilder);
+						for (const FStateTreeBindableStructDesc& ContextStruct : AccessibleStructs)
+						{
+							TArray<TSharedPtr<const FRefTypeInfo>> RefTypeInfos;
+							if (CachedBindingData->CanCreateParameter(ContextStruct, /*out*/RefTypeInfos))
+							{
+								SectionHelper.SetSection(FText::FromString(ContextStruct.StatePath));
+
+								if (RefTypeInfos.IsEmpty())
+								{
+									InMenuBuilder.AddMenuEntry(FExecuteAction::CreateSP(CachedBindingDataRef, &FCachedBindingData::PromoteToParameter, PropertyName, ContextStruct, TSharedPtr<const FRefTypeInfo>()),
+										MakeContextStructWidget(ContextStruct));
+								}
+								else
+								{
+									InMenuBuilder.AddSubMenu(MakeContextStructWidget(ContextStruct),
+										FNewMenuDelegate::CreateLambda([&CachedBindingDataRef, PropertyName, &ContextStruct, RefTypeInfos = MoveTemp(RefTypeInfos)](FMenuBuilder& InSubMenuBuilder)
+										{
+											FSectionHelper SectionHelper(InSubMenuBuilder);
+											SectionHelper.SetSection(LOCTEXT("RefTypeParams", "Reference Types"));
+											for (const TSharedPtr<const FRefTypeInfo>& RefTypeInfo : RefTypeInfos)
+											{
+												InSubMenuBuilder.AddMenuEntry(FExecuteAction::CreateSP(CachedBindingDataRef, &FCachedBindingData::PromoteToParameter, PropertyName, ContextStruct, RefTypeInfo),
+													MakeBindingPropertyInfoWidget(RefTypeInfo->TypeNameText, RefTypeInfo->PinType));
+											}
+										}));
+								}
+							}
+						}
+					})
+				);
+			})
+		);
+	}
+
 	InWidgetRow.ExtensionContent()
 	[
 		PropertyAccessEditor.MakePropertyBindingWidget(BindingContextStructs, Args)
 	];
 }
 
+bool FStateTreeBindingExtension::CanPromoteToParameter(const TSharedPtr<IPropertyHandle>& InPropertyHandle) const
+{
+	const FProperty* Property = InPropertyHandle->GetProperty();
+	if (!Property)
+	{
+		return false;
+	}
+
+	// Property Bag picker only detects Blueprint Types, so only allow properties that are blueprint types
+	// FPropertyBagInstanceDataDetails::OnPropertyNameContent uses SPinTypeSelector to generate the property type picker.
+	// UEdGraphSchema_K2::GetVariableTypeTree (GatherPinsImpl: FindEnums, FindStructs, FindObjectsAndInterfaces) is used there which only allows bp types.
+	// The below behavior mirrors the behavior in the pin gathering but for properties
+
+	if (const FEnumProperty* EnumProperty = CastField<FEnumProperty>(Property))
+	{
+		if (!UEdGraphSchema_K2::IsAllowableBlueprintVariableType(EnumProperty->GetEnum()))
+		{
+			return false;
+		}
+	}
+	else if (const FStructProperty* StructProperty = CastField<FStructProperty>(Property))
+	{
+		// Support Property Refs as even though these aren't bp types, the actual types that would be added are the ones in the meta-data RefType
+		if (StructProperty->Struct && StructProperty->Struct->IsChildOf(FStateTreePropertyRef::StaticStruct()))
+		{
+			return true;
+		}
+
+		if (!UEdGraphSchema_K2::IsAllowableBlueprintVariableType(StructProperty->Struct))
+		{
+			return false;
+		}
+	}
+	else if (const FObjectPropertyBase* ObjectProperty = CastField<FObjectPropertyBase>(Property))
+	{
+		if (!UEdGraphSchema_K2::IsAllowableBlueprintVariableType(ObjectProperty->PropertyClass))
+		{
+			return false;
+		}
+	}
+	else if (const FInterfaceProperty* InterfaceProperty = CastField<FInterfaceProperty>(Property))
+	{
+		if (!UEdGraphSchema_K2::IsAllowableBlueprintVariableType(InterfaceProperty->InterfaceClass))
+		{
+			return false;
+		}
+	}
+
+	return true;
+}
+
 void FStateTreeBindingsChildrenCustomization::CustomizeChildren(IDetailChildrenBuilder& ChildrenBuilder, TSharedPtr<IPropertyHandle> InPropertyHandle)
 {
 	TArray<UObject*> OuterObjects;
diff --git a/Engine/Plugins/Runtime/StateTree/Source/StateTreeEditorModule/Private/Customizations/StateTreeBindingExtension.h b/Engine/Plugins/Runtime/StateTree/Source/StateTreeEditorModule/Private/Customizations/StateTreeBindingExtension.h
--- a/Engine/Plugins/Runtime/StateTree/Source/StateTreeEditorModule/Private/Customizations/StateTreeBindingExtension.h
+++ b/Engine/Plugins/Runtime/StateTree/Source/StateTreeEditorModule/Private/Customizations/StateTreeBindingExtension.h
@@ -49,6 +49,9 @@ public:
 	// IDetailPropertyExtensionHandler interface
 	virtual bool IsPropertyExtendable(const UClass* InObjectClass, const IPropertyHandle& PropertyHandle) const override;
 	virtual void ExtendWidgetRow(FDetailWidgetRow& InWidgetRow, const IDetailLayoutBuilder& InDetailBuilder, const UClass* InObjectClass, TSharedPtr<IPropertyHandle> PropertyHandle) override;
+
+private:
+	bool CanPromoteToParameter(const TSharedPtr<IPropertyHandle>& InPropertyHandle) const;
 };
 
 /* Overrides bound property's children composition. */
diff --git a/Engine/Plugins/Runtime/StateTree/Source/StateTreeEditorModule/Private/StateTreeEditorData.cpp b/Engine/Plugins/Runtime/StateTree/Source/StateTreeEditorModule/Private/StateTreeEditorData.cpp
--- a/Engine/Plugins/Runtime/StateTree/Source/StateTreeEditorModule/Private/StateTreeEditorData.cpp
+++ b/Engine/Plugins/Runtime/StateTree/Source/StateTreeEditorModule/Private/StateTreeEditorData.cpp
@@ -520,6 +520,86 @@ EStateTreeVisitor UStateTreeEditorData::EnumerateBindablePropertyFunctionNodes(T
 	return EStateTreeVisitor::Continue;
 }
 
+bool UStateTreeEditorData::CanCreateParameter(const FGuid StructID) const
+{
+	if (RootParameters.ID == StructID)
+	{
+		return true;
+	}
+
+	bool bFoundStructID = false;
+
+	VisitHierarchy([&StructID, &bFoundStructID](UStateTreeState& State, UStateTreeState* ParentState)->EStateTreeVisitor
+	{
+		if (State.Parameters.ID == StructID)
+		{
+			bFoundStructID = true;
+			return EStateTreeVisitor::Break;
+		}
+		return EStateTreeVisitor::Continue;
+	});
+
+	return bFoundStructID;
+}
+
+void UStateTreeEditorData::CreateParameters(const FGuid StructID, TArrayView<FStateTreeEditorPropertyCreationDesc> InOutCreationDescs)
+{
+	if (InOutCreationDescs.IsEmpty())
+	{
+		return;
+	}
+
+	auto CreateProperties = [&InOutCreationDescs](FInstancedPropertyBag& PropertyBag)
+	{
+		TArray<FPropertyBagPropertyDesc, TInlineAllocator<1>> PropertyDescs;
+		PropertyDescs.Reserve(InOutCreationDescs.Num());
+
+		// Generate unique names for the incoming property descs to avoid changing the existing properties in the bag
+		for (FStateTreeEditorPropertyCreationDesc& CreationDesc : InOutCreationDescs)
+		{
+			int32 Index = CreationDesc.PropertyDesc.Name.GetNumber();
+			while (PropertyBag.FindPropertyDescByName(CreationDesc.PropertyDesc.Name))
+			{
+				CreationDesc.PropertyDesc.Name = FName(CreationDesc.PropertyDesc.Name, Index++);
+			}
+			PropertyDescs.Add(CreationDesc.PropertyDesc);
+		}
+
+		PropertyBag.AddProperties(PropertyDescs);
+
+		for (const FStateTreeEditorPropertyCreationDesc& CreationDesc : InOutCreationDescs)
+		{
+			// Attempt to copy the value from the Source Property / Addr to Property Desc
+			// There could be Type Mismatches if the Descs don't match the Source Property, but attempt to do it on all property descs
+			if (CreationDesc.SourceProperty && CreationDesc.SourceContainerAddress)
+			{
+				PropertyBag.SetValue(CreationDesc.PropertyDesc.Name, CreationDesc.SourceProperty, CreationDesc.SourceContainerAddress);
+			}
+		}
+	};
+
+	const UStateTree* StateTree = GetTypedOuter<UStateTree>();
+	checkf(StateTree, TEXT("UStateTreeEditorData should only be allocated within a UStateTree"));
+
+	if (RootParameters.ID == StructID)
+	{
+		CreateProperties(RootParameters.Parameters);
+		UE::StateTree::Delegates::OnParametersChanged.Broadcast(*StateTree);
+		return;
+	}
+
+	VisitHierarchy([&StructID, &CreateProperties, StateTree](UStateTreeState& State, UStateTreeState* ParentState)->EStateTreeVisitor
+	{
+		if (State.Parameters.ID == StructID)
+		{
+			CreateProperties(State.Parameters.Parameters);
+			UE::StateTree::Delegates::OnStateParametersChanged.Broadcast(*StateTree, State.ID);
+			return EStateTreeVisitor::Break;
+		}
+		return EStateTreeVisitor::Continue;
+	});
+}
+
 bool UStateTreeEditorData::GetStructByID(const FGuid StructID, FStateTreeBindableStructDesc& OutStructDesc) const
 {
 	VisitAllNodes([&OutStructDesc, StructID](const UStateTreeState* State, const FStateTreeBindableStructDesc& Desc, const FStateTreeDataView Value)
diff --git a/Engine/Plugins/Runtime/StateTree/Source/StateTreeEditorModule/Public/StateTreeEditorData.h b/Engine/Plugins/Runtime/StateTree/Source/StateTreeEditorModule/Public/StateTreeEditorData.h
--- a/Engine/Plugins/Runtime/StateTree/Source/StateTreeEditorModule/Public/StateTreeEditorData.h
+++ b/Engine/Plugins/Runtime/StateTree/Source/StateTreeEditorModule/Public/StateTreeEditorData.h
@@ -74,6 +74,8 @@ public:
 
 	virtual EStateTreeVisitor EnumerateBindablePropertyFunctionNodes(TFunctionRef<EStateTreeVisitor(const UScriptStruct* NodeStruct, const FStateTreeBindableStructDesc& Desc, const FStateTreeDataView Value)> InFunc) const;
 
+	virtual bool CanCreateParameter(const FGuid StructID) const override;
+	virtual void CreateParameters(const FGuid StructID, TArrayView<FStateTreeEditorPropertyCreationDesc> InOutCreationDescs) override;
 	// ~IStateTreeEditorPropertyBindingsOwner
 
 	/**
diff --git a/Engine/Plugins/Runtime/StateTree/Source/StateTreeEditorModule/Public/StateTreeEditorPropertyBindings.h b/Engine/Plugins/Runtime/StateTree/Source/StateTreeEditorModule/Public/StateTreeEditorPropertyBindings.h
--- a/Engine/Plugins/Runtime/StateTree/Source/StateTreeEditorModule/Public/StateTreeEditorPropertyBindings.h
+++ b/Engine/Plugins/Runtime/StateTree/Source/StateTreeEditorModule/Public/StateTreeEditorPropertyBindings.h
@@ -125,6 +125,19 @@ class UStateTreeEditorPropertyBindingsOwner : public UInterface
 	GENERATED_UINTERFACE_BODY()
 };
 
+/** Struct of Parameters used to Create a Property */
+struct FStateTreeEditorPropertyCreationDesc
+{
+	/** Property Bag Description of the Property to Create */
+	FPropertyBagPropertyDesc PropertyDesc;
+
+	/** Optional: property to copy into the new created property */
+	const FProperty* SourceProperty = nullptr;
+
+	/** Optional: container address of the property to copy */
+	const void* SourceContainerAddress = nullptr;
+};
+
 class STATETREEEDITORMODULE_API IStateTreeEditorPropertyBindingsOwner
 {
 	GENERATED_IINTERFACE_BODY()
@@ -166,6 +179,20 @@ class STATETREEEDITORMODULE_API IStateTreeEditorPropertyBindingsOwner
 	virtual const FStateTreeEditorPropertyBindings* GetPropertyEditorBindings() const PURE_VIRTUAL(IStateTreeEditorPropertyBindingsOwner::GetPropertyEditorBindings, return nullptr; );
 
 	virtual EStateTreeVisitor EnumerateBindablePropertyFunctionNodes(TFunctionRef<EStateTreeVisitor(const UScriptStruct* NodeStruct, const FStateTreeBindableStructDesc& Desc, const FStateTreeDataView Value)> InFunc) const PURE_VIRTUAL(IStateTreeEditorPropertyBindingsOwner::EnumerateBindablePropertyFunctionNodes, return static_cast<EStateTreeVisitor>(0); );
+
+	/**
+	 * Determines whether the struct matching the given struct id is capable of adding new properties
+	 * @param StructID Target struct ID
+	 * @return True if struct supports adding new properties
+	 */
+	virtual bool CanCreateParameter(const FGuid StructID) const PURE_VIRTUAL(IStateTreeEditorPropertyBindingsOwner::CanCreateProperty, return false;);
+
+	/**
+	 * Creates the given properties in the property bag of the struct matching the given struct ID
+	 * @param StructID Target struct ID
+	 * @param InOutCreationDescs the descriptions of the properties to create. This is modified to update the property names that actually got created
+	 */
+	virtual void CreateParameters(const FGuid StructID, TArrayView<FStateTreeEditorPropertyCreationDesc> InOutCreationDescs) PURE_VIRTUAL(IStateTreeEditorPropertyBindingsOwner::CreateProperties);
 };
 
 // TODO: We should merge this with IStateTreeEditorPropertyBindingsOwner and FStateTreeEditorPropertyBindings.
diff --git a/Engine/Plugins/Runtime/StateTree/Source/StateTreeModule/Private/StateTreePropertyRefHelpers.cpp b/Engine/Plugins/Runtime/StateTree/Source/StateTreeModule/Private/StateTreePropertyRefHelpers.cpp
--- a/Engine/Plugins/Runtime/StateTree/Source/StateTreeModule/Private/StateTreePropertyRefHelpers.cpp
+++ b/Engine/Plugins/Runtime/StateTree/Source/StateTreeModule/Private/StateTreePropertyRefHelpers.cpp
@@ -359,87 +359,12 @@ namespace UE::StateTree::PropertyRefHelpers
 
 	FEdGraphPinType GetNativePropertyRefInternalTypeAsPin(const FProperty& RefProperty)
 	{
-		ensure(IsPropertyRef(RefProperty));
-
-		FEdGraphPinType PinType;
-		PinType.PinSubCategory = NAME_None;
-
-		PinType.ContainerType = RefProperty.HasMetaData(IsRefToArrayName) ? EPinContainerType::Array : EPinContainerType::None;
-		const FString& TargetTypeNameStr = RefProperty.GetMetaData(RefTypeName);
-		const FName TargetTypeName = FName(*TargetTypeNameStr);
-
-		if(TargetTypeName == BoolName)
-		{
-			PinType.PinCategory = UEdGraphSchema_K2::PC_Boolean;
-		}
-		else if(TargetTypeName == ByteName)
-		{
-			PinType.PinCategory = UEdGraphSchema_K2::PC_Byte;
-		}
-		else if(TargetTypeName == Int32Name)
-		{
-			PinType.PinCategory = UEdGraphSchema_K2::PC_Int;
-		}
-		else if(TargetTypeName == Int64Name)
-		{
-			PinType.PinCategory = UEdGraphSchema_K2::PC_Int64;
-		}
-		else if(TargetTypeName == FloatName)
-		{
-			PinType.PinCategory = UEdGraphSchema_K2::PC_Real;
-			PinType.PinSubCategory = UEdGraphSchema_K2::PC_Float;
-		}
-		else if(TargetTypeName == DoubleName)
+		TArray<FEdGraphPinType, TInlineAllocator<1>> PinTypes = GetPropertyRefInternalTypesAsPins(RefProperty);
+		if (PinTypes.Num() == 1)
 		{
-			PinType.PinCategory = UEdGraphSchema_K2::PC_Real;
-			PinType.PinSubCategory = UEdGraphSchema_K2::PC_Double;
+			return PinTypes[0];
 		}
-		else if(TargetTypeName == NameName)
-		{
-			PinType.PinCategory = UEdGraphSchema_K2::PC_Name;
-		}
-		else if(TargetTypeName == StringName)
-		{
-			PinType.PinCategory = UEdGraphSchema_K2::PC_String;
-		}
-		else if(TargetTypeName == TextName)
-		{
-			PinType.PinCategory = UEdGraphSchema_K2::PC_Text;
-		}
-		else
-		{
-			UField* TargetRefField = UClass::TryFindTypeSlow<UField>(TargetTypeNameStr);
-			if (!TargetRefField)
-			{
-				TargetRefField = LoadObject<UField>(nullptr, *TargetTypeNameStr);
-			}
-
-			if (UStruct* Struct = Cast<UStruct>(TargetRefField))
-			{
-				PinType.PinCategory = UEdGraphSchema_K2::PC_Struct;
-				PinType.PinSubCategoryObject = Struct;
-			}
-			else if (UObject* Object = Cast<UObject>(TargetRefField))
-			{
-				PinType.PinCategory = UEdGraphSchema_K2::PC_Object;
-				PinType.PinSubCategoryObject = Object;
-			}
-			else if (UEnum* Enum = Cast<UEnum>(TargetRefField))
-			{
-				PinType.PinCategory = UEdGraphSchema_K2::PC_Enum;
-				PinType.PinSubCategoryObject = Enum;
-			}
-			else
-			{
-				checkCode(
-					TArray<FString> AllowedTypes;
-					TargetTypeNameStr.ParseIntoArray(AllowedTypes, TEXT(","), true);
-					check(AllowedTypes.Num() > 1); // We are trying to bind to multiple type so its normal we can't find what the pin type should be.
-				)
-			}
-		}
-
-		return PinType;
+		return FEdGraphPinType();
 	}
 
 	FEdGraphPinType GetPropertyRefInternalTypeAsPin(const FProperty& RefProperty, const void* PropertyRefAddress)
@@ -585,6 +510,101 @@ namespace UE::StateTree::PropertyRefHelpers
 		checkNoEntry();
 		return false;
 	}
+
+	TArray<FEdGraphPinType, TInlineAllocator<1>> GetPropertyRefInternalTypesAsPins(const FProperty& RefProperty)
+	{
+		ensure(IsPropertyRef(RefProperty));
+
+		const EPinContainerType ContainerType = RefProperty.HasMetaData(IsRefToArrayName) ? EPinContainerType::Array : EPinContainerType::None;
+
+		TArray<FEdGraphPinType, TInlineAllocator<1>> PinTypes;
+
+		FString TargetTypesString = RefProperty.GetMetaData(RefTypeName);
+		if (TargetTypesString.IsEmpty())
+		{
+			return PinTypes;
+		}
+
+		TArray<FString> TargetTypes;
+		TargetTypesString.RemoveSpacesInline();
+		TargetTypesString.ParseIntoArray(TargetTypes, TEXT(","), true);
+
+		for (const FString& TargetType : TargetTypes)
+		{
+			const FName TargetTypeName = *TargetType;
+
+			FEdGraphPinType& PinType = PinTypes.AddDefaulted_GetRef();
+			PinType.ContainerType = ContainerType;
+
+			if (TargetTypeName == BoolName)
+			{
+				PinType.PinCategory = UEdGraphSchema_K2::PC_Boolean;
+			}
+			else if (TargetTypeName == ByteName)
+			{
+				PinType.PinCategory = UEdGraphSchema_K2::PC_Byte;
+			}
+			else if (TargetTypeName == Int32Name)
+			{
+				PinType.PinCategory = UEdGraphSchema_K2::PC_Int;
+			}
+			else if (TargetTypeName == Int64Name)
+			{
+				PinType.PinCategory = UEdGraphSchema_K2::PC_Int64;
+			}
+			else if (TargetTypeName == FloatName)
+			{
+				PinType.PinCategory = UEdGraphSchema_K2::PC_Real;
+				PinType.PinSubCategory = UEdGraphSchema_K2::PC_Float;
+			}
+			else if (TargetTypeName == DoubleName)
+			{
+				PinType.PinCategory = UEdGraphSchema_K2::PC_Real;
+				PinType.PinSubCategory = UEdGraphSchema_K2::PC_Double;
+			}
+			else if (TargetTypeName == NameName)
+			{
+				PinType.PinCategory = UEdGraphSchema_K2::PC_Name;
+			}
+			else if (TargetTypeName == StringName)
+			{
+				PinType.PinCategory = UEdGraphSchema_K2::PC_String;
+			}
+			else if (TargetTypeName == TextName)
+			{
+				PinType.PinCategory = UEdGraphSchema_K2::PC_Text;
+			}
+			else
+			{
+				UField* TargetRefField = UClass::TryFindTypeSlow<UField>(TargetType);
+				if (!TargetRefField)
+				{
+					TargetRefField = LoadObject<UField>(nullptr, *TargetType);
+				}
+
+				if (UClass* ObjectClass = Cast<UClass>(TargetRefField))
+				{
+					PinType.PinCategory = UEdGraphSchema_K2::PC_Object;
+					PinType.PinSubCategoryObject = ObjectClass;
+				}
+				else if (UStruct* Struct = Cast<UStruct>(TargetRefField))
+				{
+					PinType.PinCategory = UEdGraphSchema_K2::PC_Struct;
+					PinType.PinSubCategoryObject = Struct;
+				}
+				else if (UEnum* Enum = Cast<UEnum>(TargetRefField))
+				{
+					PinType.PinCategory = UEdGraphSchema_K2::PC_Enum;
+					PinType.PinSubCategoryObject = Enum;
+				}
+				else
+				{
+					checkf(false, TEXT("Typename in meta-data (%s) is invalid"), *TargetType);
+				}
+			}
+		}
+		return PinTypes;
+	}
 #endif
 
 	bool IsBlueprintPropertyRefCompatibleWithProperty(const FProperty& SourceProperty, const void* PropertyRefAddress)
diff --git a/Engine/Plugins/Runtime/StateTree/Source/StateTreeModule/Public/StateTreePropertyRefHelpers.h b/Engine/Plugins/Runtime/StateTree/Source/StateTreeModule/Public/StateTreePropertyRefHelpers.h
--- a/Engine/Plugins/Runtime/StateTree/Source/StateTreeModule/Public/StateTreePropertyRefHelpers.h
+++ b/Engine/Plugins/Runtime/StateTree/Source/StateTreeModule/Public/StateTreePropertyRefHelpers.h
@@ -55,6 +55,12 @@ namespace UE::StateTree::PropertyRefHelpers
 	 */
 	bool STATETREEMODULE_API IsPropertyRefMarkedAsOptional(const FProperty& RefProperty, const void* PropertyRefAddress);
 
+	/**
+	 * @param RefProperty Property of PropertyRef type
+	 * @return PinTypes for PropertyRef's internal types
+	 */
+	TArray<FEdGraphPinType, TInlineAllocator<1>> STATETREEMODULE_API GetPropertyRefInternalTypesAsPins(const FProperty& RefProperty);
+
 	/**
 	 * @param RefProperty Property of PropertyRef type
 	 * @param PropertyRefAddress Address of PropertyRef.
