State Tree: Added comments in GetPropertyRefInternalTypesAsPins

diff --git a/Engine/Plugins/Runtime/StateTree/Source/StateTreeModule/Private/StateTreePropertyRefHelpers.cpp b/Engine/Plugins/Runtime/StateTree/Source/StateTreeModule/Private/StateTreePropertyRefHelpers.cpp
--- a/Engine/Plugins/Runtime/StateTree/Source/StateTreeModule/Private/StateTreePropertyRefHelpers.cpp
+++ b/Engine/Plugins/Runtime/StateTree/Source/StateTreeModule/Private/StateTreePropertyRefHelpers.cpp
@@ -582,11 +582,14 @@ namespace UE::StateTree::PropertyRefHelpers
 					TargetRefField = LoadObject<UField>(nullptr, *TargetType);
 				}
 
+				// First check whether the Target Ref Field is a UClass since UClass derives from UStruct
 				if (UClass* ObjectClass = Cast<UClass>(TargetRefField))
 				{
 					PinType.PinCategory = UEdGraphSchema_K2::PC_Object;
 					PinType.PinSubCategoryObject = ObjectClass;
 				}
+				// If it isn't a UClass, it should either be a UStruct, or UEnum.
+				// Both are not derivatives of each other so order from here on doesn't matter
 				else if (UStruct* Struct = Cast<UStruct>(TargetRefField))
 				{
 					PinType.PinCategory = UEdGraphSchema_K2::PC_Struct;
