State Tree: corrected lambda capture to be by copy not ref + fixed to check UScriptStruct instead of UStruct for struct pin type

diff --git a/Engine/Plugins/Runtime/StateTree/Source/StateTreeEditorModule/Private/Customizations/StateTreeBindingExtension.cpp b/Engine/Plugins/Runtime/StateTree/Source/StateTreeEditorModule/Private/Customizations/StateTreeBindingExtension.cpp
--- a/Engine/Plugins/Runtime/StateTree/Source/StateTreeEditorModule/Private/Customizations/StateTreeBindingExtension.cpp
+++ b/Engine/Plugins/Runtime/StateTree/Source/StateTreeEditorModule/Private/Customizations/StateTreeBindingExtension.cpp
@@ -1665,7 +1665,7 @@ void FStateTreeBindingExtension::ExtendWidgetRow(FDetailWidgetRow& InWidgetRow,
 								else
 								{
 									InMenuBuilder.AddSubMenu(MakeContextStructWidget(ContextStruct),
-										FNewMenuDelegate::CreateLambda([&CachedBindingDataRef, PropertyName, &ContextStruct, RefTypeInfos = MoveTemp(RefTypeInfos)](FMenuBuilder& InSubMenuBuilder)
+										FNewMenuDelegate::CreateLambda([CachedBindingDataRef, PropertyName, &ContextStruct, RefTypeInfos = MoveTemp(RefTypeInfos)](FMenuBuilder& InSubMenuBuilder)
 										{
 											FSectionHelper SectionHelper(InSubMenuBuilder);
 											SectionHelper.SetSection(LOCTEXT("RefTypeParams", "Reference Types"));
diff --git a/Engine/Plugins/Runtime/StateTree/Source/StateTreeModule/Private/StateTreePropertyRefHelpers.cpp b/Engine/Plugins/Runtime/StateTree/Source/StateTreeModule/Private/StateTreePropertyRefHelpers.cpp
--- a/Engine/Plugins/Runtime/StateTree/Source/StateTreeModule/Private/StateTreePropertyRefHelpers.cpp
+++ b/Engine/Plugins/Runtime/StateTree/Source/StateTreeModule/Private/StateTreePropertyRefHelpers.cpp
@@ -582,19 +582,16 @@ namespace UE::StateTree::PropertyRefHelpers
 					TargetRefField = LoadObject<UField>(nullptr, *TargetType);
 				}
 
-				// First check whether the Target Ref Field is a UClass since UClass derives from UStruct
-				if (UClass* ObjectClass = Cast<UClass>(TargetRefField))
-				{
-					PinType.PinCategory = UEdGraphSchema_K2::PC_Object;
-					PinType.PinSubCategoryObject = ObjectClass;
-				}
-				// If it isn't a UClass, it should either be a UStruct, or UEnum.
-				// Both are not derivatives of each other so order from here on doesn't matter
-				else if (UStruct* Struct = Cast<UStruct>(TargetRefField))
+				if (UScriptStruct* Struct = Cast<UScriptStruct>(TargetRefField))
 				{
 					PinType.PinCategory = UEdGraphSchema_K2::PC_Struct;
 					PinType.PinSubCategoryObject = Struct;
 				}
+				else if (UClass* ObjectClass = Cast<UClass>(TargetRefField))
+				{
+					PinType.PinCategory = UEdGraphSchema_K2::PC_Object;
+					PinType.PinSubCategoryObject = ObjectClass;
+				}
 				else if (UEnum* Enum = Cast<UEnum>(TargetRefField))
 				{
 					PinType.PinCategory = UEdGraphSchema_K2::PC_Enum;
