Add verbose logging for cases when a StateTree property is not bindable

diff --git a/Engine/Plugins/Runtime/StateTree/Source/StateTreeEditorModule/Private/Customizations/StateTreeBindingExtension.cpp b/Engine/Plugins/Runtime/StateTree/Source/StateTreeEditorModule/Private/Customizations/StateTreeBindingExtension.cpp
--- a/Engine/Plugins/Runtime/StateTree/Source/StateTreeEditorModule/Private/Customizations/StateTreeBindingExtension.cpp
+++ b/Engine/Plugins/Runtime/StateTree/Source/StateTreeEditorModule/Private/Customizations/StateTreeBindingExtension.cpp
@@ -1349,8 +1349,24 @@ private:
 
 bool IsPropertyBindable(const FProperty& Property)
 {
-	return Property.HasAnyPropertyFlags(CPF_Edit) 
-			&& (!Property.HasAnyPropertyFlags(CPF_NativeAccessSpecifierPrivate | CPF_NativeAccessSpecifierProtected) || Property.GetBoolMetaData(FBlueprintMetadata::MD_AllowPrivateAccess));
+	const bool bIsUserEditable = Property.HasAnyPropertyFlags(CPF_Edit);
+	if (!bIsUserEditable)
+	{
+		UE_LOG(LogStateTreeEditor, Verbose, TEXT("Property %s is not bindable because it's not user-settable in the editor"),
+			*Property.GetName());
+		return false;
+	}
+
+	const bool bPrivateOrProtected = !Property.HasAnyPropertyFlags(CPF_NativeAccessSpecifierPrivate | CPF_NativeAccessSpecifierProtected);
+	const bool bPrivateButBlueprintAccessible = Property.GetBoolMetaData(FBlueprintMetadata::MD_AllowPrivateAccess);
+	if (!bPrivateOrProtected && !bPrivateButBlueprintAccessible)
+	{
+		UE_LOG(LogStateTreeEditor, Verbose, TEXT("Property %s is not bindable because it's either private or protected and not private-accessible to blueprints"),
+			*Property.GetName());
+		return false;
+	}
+
+	return true;
 }
 
 /* Provides PropertyFunctionNode instance for a property node. */
