[State Tree] Utility Selector Runtime First Pass

diff --git a/Engine/Plugins/Runtime/StateTree/Source/StateTreeEditorModule/Private/Customizations/StateTreeStateDetails.cpp b/Engine/Plugins/Runtime/StateTree/Source/StateTreeEditorModule/Private/Customizations/StateTreeStateDetails.cpp
--- a/Engine/Plugins/Runtime/StateTree/Source/StateTreeEditorModule/Private/Customizations/StateTreeStateDetails.cpp
+++ b/Engine/Plugins/Runtime/StateTree/Source/StateTreeEditorModule/Private/Customizations/StateTreeStateDetails.cpp
@@ -89,7 +89,7 @@ void FStateTreeStateDetails::CustomizeDetails(IDetailLayoutBuilder& DetailBuilde
 	
 	// Never show enabled
 	EnabledProperty->MarkHiddenByCustomization();
-
+	
 	// Show ID only for debugging
 	if (UE::StateTree::Editor::GbDisplayItemIds == false)
 	{
@@ -168,7 +168,7 @@ void FStateTreeStateDetails::CustomizeDetails(IDetailLayoutBuilder& DetailBuilde
 	LinkedSubtreeProperty->MarkHiddenByCustomization();
 	LinkedAssetProperty->MarkHiddenByCustomization();
 
-	if (StateType == EStateTreeStateType::State)
+	if (StateType == EStateTreeStateType::State || StateType == EStateTreeStateType::Subtree)
 	{
 		StateCategory.AddProperty(SelectionBehaviorProperty);
 	}
@@ -286,10 +286,10 @@ void FStateTreeStateDetails::CustomizeDetails(IDetailLayoutBuilder& DetailBuilde
 		DetailBuilder.EditCategory(EnterConditionsCategoryName).SetCategoryVisibility(false);
 	}
 
-	// Utility
-	const FName UtilityCategoryName(TEXT("Utility"));
-	if (StateTreeSchemaUtilityCVars::CVarAllowUtilityConsiderations->GetBool() && Schema && Schema->AllowUtilityConsiderations())
+	ConsiderationsProperty->MarkHiddenByCustomization();
+	if (Schema && Schema->AllowUtilityConsiderations())
 	{
+		const FName UtilityCategoryName(TEXT("Utility"));
 		IDetailCategoryBuilder& UtilityConsiderationsCategory = UE::StateTreeEditor::EditorNodeUtils::MakeArrayCategory(
 			DetailBuilder,
 			ConsiderationsProperty,
@@ -300,12 +300,6 @@ void FStateTreeStateDetails::CustomizeDetails(IDetailLayoutBuilder& DetailBuilde
 			UE::StateTree::Colors::Yellow.WithAlpha(192),
 			LOCTEXT("UtilityAddTooltip", "Add new Utility Consideration"),
 			/*SortOrder*/3);
-
-		ConsiderationsProperty->MarkHiddenByCustomization();
-	}
-	else
-	{
-		DetailBuilder.EditCategory(UtilityCategoryName).SetCategoryVisibility(false);
 	}
 
 	// Tasks
@@ -330,7 +324,7 @@ void FStateTreeStateDetails::CustomizeDetails(IDetailLayoutBuilder& DetailBuilde
 		{
 			const FName TaskCategoryName(TEXT("Task"));
 			IDetailCategoryBuilder& Category = DetailBuilder.EditCategory(TaskCategoryName);
-			Category.SetSortOrder(3);
+			Category.SetSortOrder(4);
 			
 			IDetailPropertyRow& Row = Category.AddProperty(SingleTaskProperty);
 			Row.ShouldAutoExpand(true);
diff --git a/Engine/Plugins/Runtime/StateTree/Source/StateTreeEditorModule/Private/SStateTreeViewRow.cpp b/Engine/Plugins/Runtime/StateTree/Source/StateTreeEditorModule/Private/SStateTreeViewRow.cpp
--- a/Engine/Plugins/Runtime/StateTree/Source/StateTreeEditorModule/Private/SStateTreeViewRow.cpp
+++ b/Engine/Plugins/Runtime/StateTree/Source/StateTreeEditorModule/Private/SStateTreeViewRow.cpp
@@ -885,13 +885,42 @@ const FSlateBrush* SStateTreeViewRow::GetSelectorIcon() const
 		}
 		else if (State->SelectionBehavior == EStateTreeStateSelectionBehavior::TrySelectChildrenAtUniformRandom)
 		{
-			if (State->Children.IsEmpty())
+			if (State->Children.IsEmpty()
+				|| State->Type == EStateTreeStateType::Linked
+				|| State->Type == EStateTreeStateType::LinkedAsset)
 			{
+				// Backwards compatible behavior
 				return FStateTreeEditorStyle::Get().GetBrush("StateTreeEditor.TryEnterState");
 			}
 
 			return FStateTreeEditorStyle::Get().GetBrush("StateTreeEditor.TrySelectCHildrenAtRandom");
 		}
+		else if (State->SelectionBehavior == EStateTreeStateSelectionBehavior::TrySelectChildrenWithHighestUtility)
+		{
+			if (State->Children.IsEmpty()
+				|| State->Type == EStateTreeStateType::Linked
+				|| State->Type == EStateTreeStateType::LinkedAsset)
+			{
+				// Backwards compatible behavior
+				return FStateTreeEditorStyle::Get().GetBrush("StateTreeEditor.TryEnterState");
+			}
+
+			//place holder
+			return FStateTreeEditorStyle::Get().GetBrush("StateTreeEditor.TrySelectChildrenInOrder");
+		}
+		else if (State->SelectionBehavior == EStateTreeStateSelectionBehavior::TrySelectChildrenBasedOnRelativeUtility)
+		{
+			if (State->Children.IsEmpty()
+				|| State->Type == EStateTreeStateType::Linked
+				|| State->Type == EStateTreeStateType::LinkedAsset)
+			{
+				// Backwards compatible behavior
+				return FStateTreeEditorStyle::Get().GetBrush("StateTreeEditor.TryEnterState");
+			}
+
+			//place holder
+			return FStateTreeEditorStyle::Get().GetBrush("StateTreeEditor.TrySelectChildrenInOrder");
+		}
 		else if (State->SelectionBehavior == EStateTreeStateSelectionBehavior::TryFollowTransitions)
 		{
 			return FStateTreeEditorStyle::Get().GetBrush("StateTreeEditor.TryFollowTransitions");
@@ -909,32 +938,30 @@ FText SStateTreeViewRow::GetSelectorTooltip() const
 		check(Enum);
 		const int32 Index = Enum->GetIndexByValue((int64)State->SelectionBehavior);
 		
-		if (State->SelectionBehavior == EStateTreeStateSelectionBehavior::None)
-		{
-			return Enum->GetToolTipTextByIndex(Index);
-		}
-		else if (State->SelectionBehavior == EStateTreeStateSelectionBehavior::TryEnterState)
-		{
-			return Enum->GetToolTipTextByIndex(Index);
-		}
-		else if (State->SelectionBehavior == EStateTreeStateSelectionBehavior::TrySelectChildrenInOrder)
+		switch (State->SelectionBehavior)
 		{
-			if (State->Children.IsEmpty()
-				|| State->Type == EStateTreeStateType::Linked
-				|| State->Type == EStateTreeStateType::LinkedAsset)
-			{
-				const int32 EnterStateIndex = Enum->GetIndexByValue((int64)EStateTreeStateSelectionBehavior::TryEnterState);
-				return FText::Format(LOCTEXT("ConvertedToEnterState", "{0}\nAutomatically converted from '{1}' becase the State has no child States."),
-					Enum->GetToolTipTextByIndex(EnterStateIndex), UEnum::GetDisplayValueAsText(State->SelectionBehavior));
-			}
-			else
-			{
+			case EStateTreeStateSelectionBehavior::None:
+			case EStateTreeStateSelectionBehavior::TryEnterState:
+			case EStateTreeStateSelectionBehavior::TryFollowTransitions:
 				return Enum->GetToolTipTextByIndex(Index);
-			}
-		}
-		else if (State->SelectionBehavior == EStateTreeStateSelectionBehavior::TryFollowTransitions)
-		{
-			return Enum->GetToolTipTextByIndex(Index);
+			case EStateTreeStateSelectionBehavior::TrySelectChildrenInOrder:
+			case EStateTreeStateSelectionBehavior::TrySelectChildrenAtUniformRandom:
+			case EStateTreeStateSelectionBehavior::TrySelectChildrenWithHighestUtility:
+			case EStateTreeStateSelectionBehavior::TrySelectChildrenBasedOnRelativeUtility:
+				if (State->Children.IsEmpty()
+					|| State->Type == EStateTreeStateType::Linked
+					|| State->Type == EStateTreeStateType::LinkedAsset)
+				{
+					const int32 EnterStateIndex = Enum->GetIndexByValue((int64)EStateTreeStateSelectionBehavior::TryEnterState);
+					return FText::Format(LOCTEXT("ConvertedToEnterState", "{0}\nAutomatically converted from '{1}' because the State has no child States."),
+						Enum->GetToolTipTextByIndex(EnterStateIndex), UEnum::GetDisplayValueAsText(State->SelectionBehavior));
+				}
+				else
+				{
+					return Enum->GetToolTipTextByIndex(Index);
+				}
+			default:
+				check(false);
 		}
 	}
 
@@ -1074,6 +1101,24 @@ bool SStateTreeViewRow::GetStateWarnings(FText* OutText) const
 		bHasWarnings = true;
 	}
 
+	// Child states should not have any considerations if their parent doesn't use utility
+	if (State->Considerations.Num() != 0)
+	{
+		if (!State->Parent 
+			|| (State->Parent->SelectionBehavior != EStateTreeStateSelectionBehavior::TrySelectChildrenWithHighestUtility
+				&& State->Parent->SelectionBehavior != EStateTreeStateSelectionBehavior::TrySelectChildrenBasedOnRelativeUtility))
+		{
+			if (OutText)
+			{
+				*OutText = LOCTEXT("ChildStateUtilityConsiderationWarning", 
+					"State has Utility Considerations but they don't have effect."
+					"The Utility Considerations are used only when parent State's Selection Behavior is:"
+					"\"Try Select Children with Highest Utility\" or \"Try Select Children Based on Relative Utility.");
+			}
+			bHasWarnings = true;
+		}
+	}
+
 	return bHasWarnings;
 }
 
diff --git a/Engine/Plugins/Runtime/StateTree/Source/StateTreeEditorModule/Private/StateTreeCompiler.cpp b/Engine/Plugins/Runtime/StateTree/Source/StateTreeEditorModule/Private/StateTreeCompiler.cpp
--- a/Engine/Plugins/Runtime/StateTree/Source/StateTreeEditorModule/Private/StateTreeCompiler.cpp
+++ b/Engine/Plugins/Runtime/StateTree/Source/StateTreeEditorModule/Private/StateTreeCompiler.cpp
@@ -8,6 +8,7 @@
 #include "StateTreeEvaluatorBase.h"
 #include "StateTreeTaskBase.h"
 #include "StateTreeConditionBase.h"
+#include "StateTreeConsiderationBase.h"
 #include "Serialization/ArchiveUObject.h"
 #include "GameFramework/Actor.h"
 #include "StateTreePropertyRef.h"
@@ -191,6 +192,7 @@ bool FStateTreeCompiler::Compile(UStateTree& InStateTree)
 	{
 		return false;
 	}
+	
 	// Cleanup existing state
 	StateTree->ResetCompiled();
 
@@ -306,6 +308,12 @@ bool FStateTreeCompiler::Compile(UStateTree& InStateTree)
 		return false;
 	}
 
+	if (!CreateStateConsiderations())
+	{
+		StateTree->ResetCompiled();
+		return false;
+	}
+
 	StateTree->Nodes = Nodes;
 	StateTree->DefaultInstanceData.Init(*StateTree, InstanceStructs);
 	StateTree->SharedInstanceData.Init(*StateTree, SharedInstanceStructs);
@@ -340,6 +348,7 @@ bool FStateTreeCompiler::Compile(UStateTree& InStateTree)
 	UE::StateTree::Compiler::FCheckOutersArchive CheckOuters(*StateTree, *EditorData, Log);
 	StateTree->Serialize(CheckOuters);
 	
+
 	return true;
 }
 
@@ -1057,6 +1066,53 @@ bool FStateTreeCompiler::CreateStateTransitions()
 	return true;
 }
 
+bool FStateTreeCompiler::CreateStateConsiderations()
+{
+	check(StateTree);
+
+	for (int32 i = 0; i < StateTree->States.Num(); i++)
+	{
+		FCompactStateTreeState& CompactState = StateTree->States[i];
+		UStateTreeState* SourceState = SourceStates[i];
+		check(SourceState != nullptr);
+
+		FStateTreeCompilerLogStateScope LogStateScope(SourceState, Log);
+
+		const FString StatePath = SourceState->GetPath();
+
+		const int32 UtilityConsiderationsBegin = Nodes.Num();
+		if (const auto Validation = UE::StateTree::Compiler::IsValidCount16(UtilityConsiderationsBegin); Validation.DidFail())
+		{
+			Validation.Log(Log, TEXT("UtilityConsiderationsBegin"));
+			return false;
+		}
+		CompactState.UtilityConsiderationsBegin = uint16(UtilityConsiderationsBegin);
+
+		const FString StatePathWithConsiderations = StatePath + TEXT("/Considerations");
+		if (!CreateConsiderations(*SourceState, StatePathWithConsiderations, SourceState->Considerations))
+		{
+			Log.Reportf(EMessageSeverity::Error,
+				TEXT("Failed to create state utility considerations."));
+			return false;
+		}
+
+		const int32 UtilityConsiderationsNum = Nodes.Num() - UtilityConsiderationsBegin;
+		if (const auto Validation = UE::StateTree::Compiler::IsValidCount8(UtilityConsiderationsNum); Validation.DidFail())
+		{
+			Validation.Log(Log, TEXT("UtilityConsiderationsNum"));
+			return false;
+		}
+		CompactState.UtilityConsiderationsNum = uint8(UtilityConsiderationsNum);
+
+		if (!CreateBindingsForNodes(SourceState->Considerations, FStateTreeIndex16(UtilityConsiderationsBegin)))
+		{
+			return false;
+		}
+	}
+
+	return true;
+}
+
 bool FStateTreeCompiler::CreateBindingsForNodes(TConstArrayView<FStateTreeEditorNode> EditorNodes, FStateTreeIndex16 NodesBegin)
 {
 	check(NodesBegin.IsValid());
@@ -1355,6 +1411,60 @@ bool FStateTreeCompiler::CreateCondition(UStateTreeState& State, const FString&
 	return false;
 }
 
+bool FStateTreeCompiler::CreateConsiderations(UStateTreeState& State, const FString& StatePath, TConstArrayView<FStateTreeEditorNode> Considerations)
+{
+	if (State.Considerations.Num() != 0)
+	{
+		if (!State.Parent
+			|| (State.Parent->SelectionBehavior != EStateTreeStateSelectionBehavior::TrySelectChildrenWithHighestUtility
+				&& State.Parent->SelectionBehavior != EStateTreeStateSelectionBehavior::TrySelectChildrenBasedOnRelativeUtility))
+		{
+			Log.Reportf(EMessageSeverity::Warning, TEXT("State's Utility Considerations data are compiled but they don't have effect."
+					"The Utility Considerations are used only when parent State's Selection Behavior is:"
+					"\"Try Select Children with Highest Utility\" or \"Try Select Children Based on Relative Utility\"."));
+		}
+	}
+
+	for (int32 Index = 0; Index < Considerations.Num(); Index++)
+	{
+		const bool bIsFirst = Index == 0;
+		const FStateTreeEditorNode& ConsiderationNode = Considerations[Index];
+		// First operand should be copy as we dont have a previous item to operate on.
+		const EStateTreeExpressionOperand Operand = bIsFirst ? EStateTreeExpressionOperand::Copy : ConsiderationNode.ExpressionOperand;
+		// First indent must be 0 to make the parentheses calculation match.
+		const int32 CurrIndent = bIsFirst ? 0 : FMath::Clamp((int32)ConsiderationNode.ExpressionIndent, 0, UE::StateTree::MaxExpressionIndent);
+		// Next indent, or terminate at zero.
+		const int32 NextIndent = Considerations.IsValidIndex(Index + 1) ? FMath::Clamp((int32)Considerations[Index + 1].ExpressionIndent, 0, UE::StateTree::MaxExpressionIndent) : 0;
+
+		const int32 DeltaIndent = NextIndent - CurrIndent;
+
+		if (!CreateConsideration(State, StatePath, ConsiderationNode, Operand, (int8)DeltaIndent))
+		{
+			return false;
+		}
+	}
+
+	return true;
+}
+
+bool FStateTreeCompiler::CreateConsideration(UStateTreeState& State, const FString& StatePath, const FStateTreeEditorNode& ConsiderationNode, const EStateTreeExpressionOperand Operand, const int8 DeltaIndent)
+{
+	if (!ConsiderationNode.Node.IsValid())
+	{
+		// Empty line in conditions array, just silently ignore.
+		return true;
+	}
+
+	if (FStateTreeConsiderationBase* Consideration = CreateNodeWithSharedInstanceData<FStateTreeConsiderationBase>(&State, StatePath, ConsiderationNode, EStateTreeBindableStructSource::Consideration))
+	{
+		Consideration->Operand = Operand;
+		Consideration->DeltaIndent = DeltaIndent;
+		return true;
+	}
+
+	return false;
+}
+
 bool FStateTreeCompiler::CompileAndValidateNode(const UStateTreeState* SourceState, const FStateTreeBindableStructDesc& NodeDesc, FStructView NodeView, const FStateTreeDataView InstanceData)
 {
 	if (!NodeView.IsValid())
diff --git a/Engine/Plugins/Runtime/StateTree/Source/StateTreeEditorModule/Private/StateTreeEditor.cpp b/Engine/Plugins/Runtime/StateTree/Source/StateTreeEditorModule/Private/StateTreeEditor.cpp
--- a/Engine/Plugins/Runtime/StateTree/Source/StateTreeEditorModule/Private/StateTreeEditor.cpp
+++ b/Engine/Plugins/Runtime/StateTree/Source/StateTreeEditorModule/Private/StateTreeEditor.cpp
@@ -970,18 +970,14 @@ void FStateTreeEditor::RegisterToolbar()
 					TAttribute<FSlateIcon>(StateTreeEditorRef, &FStateTreeEditor::GetNewConditionButtonImage)
 				));
 
-				//reached here before Schema was loaded
-				if (StateTreeSchemaUtilityCVars::CVarAllowUtilityConsiderations->GetBool())
-				{
-					FToolMenuEntry& CreateNewConsiderationDropdown = InSection.AddEntry(FToolMenuEntry::InitComboButton(
-						"CreateNewConsiderationComboButton",
-						FUIAction(),
-						FOnGetContent::CreateSP(StateTreeEditorRef, &FStateTreeEditor::GenerateConsiderationBPBaseClassesMenu),
-						LOCTEXT("CreateNewConsideration_Title", "New Consideration"),
-						LOCTEXT("CreateNewConsideration_ToolbarTooltip", "Create a new Blueprint State Tree Utility Consideration"),
-						TAttribute<FSlateIcon>(StateTreeEditorRef, &FStateTreeEditor::GetNewConsiderationButtonImage)
-					));
-				}
+				FToolMenuEntry& CreateNewConsiderationDropdown = InSection.AddEntry(FToolMenuEntry::InitComboButton(
+					"CreateNewConsiderationComboButton",
+					FUIAction(),
+					FOnGetContent::CreateSP(StateTreeEditorRef, &FStateTreeEditor::GenerateConsiderationBPBaseClassesMenu),
+					LOCTEXT("CreateNewConsideration_Title", "New Consideration"),
+					LOCTEXT("CreateNewConsideration_ToolbarTooltip", "Create a new Blueprint State Tree Utility Consideration"),
+					TAttribute<FSlateIcon>(StateTreeEditorRef, &FStateTreeEditor::GetNewConsiderationButtonImage)
+				));
 			}
 		}
 	}));
diff --git a/Engine/Plugins/Runtime/StateTree/Source/StateTreeEditorModule/Private/StateTreeState.cpp b/Engine/Plugins/Runtime/StateTree/Source/StateTreeEditorModule/Private/StateTreeState.cpp
--- a/Engine/Plugins/Runtime/StateTree/Source/StateTreeEditorModule/Private/StateTreeState.cpp
+++ b/Engine/Plugins/Runtime/StateTree/Source/StateTreeEditorModule/Private/StateTreeState.cpp
@@ -157,9 +157,9 @@ void UStateTreeState::PostEditChangeChainProperty(FPropertyChangedChainEvent& Pr
 		}
 	}
 
-	// Broadcast selection type changes so that the UI can update.
 	if (ChangePropertyPath.IsPathExact(SelectionBehaviorPath))
 	{
+		// Broadcast selection type changes so that the UI can update.
 		const UStateTree* StateTree = GetTypedOuter<UStateTree>();
 		if (ensure(StateTree))
 		{
@@ -169,7 +169,13 @@ void UStateTreeState::PostEditChangeChainProperty(FPropertyChangedChainEvent& Pr
 	
 	if (ChangePropertyPath.IsPathExact(StateTypePath))
 	{
-		// Remove any tasks and evaluators when they are not used.
+		// Reset Selection Behavior back to Try Enter State for group and linked types
+		if (Type == EStateTreeStateType::Group || Type == EStateTreeStateType::Linked || Type == EStateTreeStateType::LinkedAsset)
+		{
+			SelectionBehavior = EStateTreeStateSelectionBehavior::TryEnterState;
+		}
+
+		// Remove any tasks when they are not used.
 		if (Type == EStateTreeStateType::Group || Type == EStateTreeStateType::Linked || Type == EStateTreeStateType::LinkedAsset)
 		{
 			Tasks.Reset();
diff --git a/Engine/Plugins/Runtime/StateTree/Source/StateTreeEditorModule/Public/StateTreeCompiler.h b/Engine/Plugins/Runtime/StateTree/Source/StateTreeEditorModule/Public/StateTreeCompiler.h
--- a/Engine/Plugins/Runtime/StateTree/Source/StateTreeEditorModule/Public/StateTreeCompiler.h
+++ b/Engine/Plugins/Runtime/StateTree/Source/StateTreeEditorModule/Public/StateTreeCompiler.h
@@ -48,7 +48,8 @@ private:
 	bool CreateGlobalTasks();
 	bool CreateStateTasksAndParameters();
 	bool CreateStateTransitions();
-
+	bool CreateStateConsiderations();
+	
 	bool CreateBindingsForNodes(TConstArrayView<FStateTreeEditorNode> EditorNodes, FStateTreeIndex16 NodesBegin);
 	bool CreateBindingsForStruct(const FStateTreeBindableStructDesc& TargetStruct, FStateTreeDataView TargetValue, FStateTreeIndex16 PropertyFuncsBegin, FStateTreeIndex16 PropertyFuncsEnd, FStateTreeIndex16& OutBatchIndex);
 
@@ -57,6 +58,8 @@ private:
 	
 	bool CreateConditions(UStateTreeState& State, const FString& StatePath, TConstArrayView<FStateTreeEditorNode> Conditions);
 	bool CreateCondition(UStateTreeState& State, const FString& StatePath, const FStateTreeEditorNode& CondNode, const EStateTreeExpressionOperand Operand, const int8 DeltaIndent);
+	bool CreateConsiderations(UStateTreeState& State, const FString& StatePath, TConstArrayView<FStateTreeEditorNode> Considerations);
+	bool CreateConsideration(UStateTreeState& State, const FString& StatePath, const FStateTreeEditorNode& ConsiderationNode, const EStateTreeExpressionOperand Operand, const int8 DeltaIndent);
 
 	template<class T = FStateTreeNodeBase>
 	T* CreateNodeWithSharedInstanceData(UStateTreeState* State, const FString& StatePath, const FStateTreeEditorNode& EditorNode, EStateTreeBindableStructSource StructSource);
diff --git a/Engine/Plugins/Runtime/StateTree/Source/StateTreeModule/Private/StateTreeExecutionContext.cpp b/Engine/Plugins/Runtime/StateTree/Source/StateTreeModule/Private/StateTreeExecutionContext.cpp
--- a/Engine/Plugins/Runtime/StateTree/Source/StateTreeModule/Private/StateTreeExecutionContext.cpp
+++ b/Engine/Plugins/Runtime/StateTree/Source/StateTreeModule/Private/StateTreeExecutionContext.cpp
@@ -4,6 +4,7 @@
 #include "StateTreeTaskBase.h"
 #include "StateTreeEvaluatorBase.h"
 #include "StateTreeConditionBase.h"
+#include "StateTreeConsiderationBase.h"
 #include "StateTreePropertyFunctionBase.h"
 #include "StateTreeReference.h"
 #include "Containers/StaticArray.h"
@@ -2638,6 +2639,86 @@ bool FStateTreeExecutionContext::TestAllConditions(const FStateTreeExecutionFram
 	return Values[0];
 }
 
+float FStateTreeExecutionContext::EvaluateUtility(const FStateTreeExecutionFrame* CurrentParentFrame, const FStateTreeExecutionFrame& CurrentFrame, const int32 ConsiderationsOffset, const int32 ConsiderationsNum)
+{
+	// @todo: Tracing support
+	CSV_SCOPED_TIMING_STAT_EXCLUSIVE(StateTree_EvaluateUtility);
+
+	if (ConsiderationsNum == 0)
+	{
+		return .0f;
+	}
+
+	TStaticArray<EStateTreeExpressionOperand, UE::StateTree::MaxExpressionIndent + 1> Operands(InPlace, EStateTreeExpressionOperand::Copy);
+	TStaticArray<float, UE::StateTree::MaxExpressionIndent + 1> Values(InPlace, false);
+
+	int32 Level = 0;
+	float Value = .0f;
+	for (int32 Index = 0; Index < ConsiderationsNum; Index++)
+	{
+		const int32 ConsiderationIndex = ConsiderationsOffset + Index;
+		const FStateTreeConsiderationBase& Consideration = CurrentFrame.StateTree->Nodes[ConsiderationIndex].Get<const FStateTreeConsiderationBase>();
+		const FStateTreeDataView ConsiderationInstanceView = GetDataView(CurrentParentFrame, CurrentFrame, Consideration.InstanceDataHandle);
+		FNodeInstanceDataScope DataScope(*this, Consideration.InstanceDataHandle, ConsiderationInstanceView);
+
+		// Copy bound properties.
+		if (Consideration.BindingsBatch.IsValid())
+		{
+			// Use validated copy, since we test in situations where the sources are not always valid (e.g. considerations may try to access inactive parent state). 
+			if (!CopyBatchWithValidation(CurrentParentFrame, CurrentFrame, ConsiderationInstanceView, Consideration.BindingsBatch))
+			{
+				// If the source data cannot be accessed, the whole expression evaluates to zero.
+				Values[0] = .0f;
+				break;
+			}
+		}
+
+		Value = Consideration.ComputeNormalizedScore(*this);
+
+		// Reset copied properties that might contain object references.
+		if (Consideration.BindingsBatch.IsValid())
+		{
+			CurrentFrame.StateTree->PropertyBindings.ResetObjects(Consideration.BindingsBatch, ConsiderationInstanceView);
+		}
+
+		const int32 DeltaIndent = Consideration.DeltaIndent;
+		const int32 OpenParens = FMath::Max(0, DeltaIndent) + 1;	// +1 for the current value that is stored at the empty slot at the top of the value stack.
+		const int32 ClosedParens = FMath::Max(0, -DeltaIndent) + 1;
+
+		// Store the operand to apply when merging higher level down when returning to this level.
+		const EStateTreeExpressionOperand Operand = Index == 0 ? EStateTreeExpressionOperand::Copy : Consideration.Operand;
+		Operands[Level] = Operand;
+
+		// Store current value at the top of the stack.
+		Level += OpenParens;
+		Values[Level] = Value;
+
+		// Evaluate and merge down values based on closed braces.
+		// The current value is placed in parens (see +1 above), which makes merging down and applying the new value consistent.
+		// The default operand is copy, so if the value is needed immediately, it is just copied down, or if we're on the same level,
+		// the operand storing above gives handles with the right logic.
+		for (int32 Paren = 0; Paren < ClosedParens; Paren++)
+		{
+			Level--;
+			switch (Operands[Level])
+			{
+			case EStateTreeExpressionOperand::Copy:
+				Values[Level] = Values[Level + 1];
+				break;
+			case EStateTreeExpressionOperand::And:
+				Values[Level] = FMath::Min(Values[Level], Values[Level + 1]);
+				break;
+			case EStateTreeExpressionOperand::Or:
+				Values[Level] = FMath::Max(Values[Level], Values[Level + 1]);
+				break;
+			}
+			Operands[Level] = EStateTreeExpressionOperand::Copy;
+		}
+	}
+
+	return Values[0];
+}
+
 void FStateTreeExecutionContext::EvaluatePropertyFunctionsOnActiveInstances(const FStateTreeExecutionFrame* CurrentParentFrame, const FStateTreeExecutionFrame& CurrentFrame, FStateTreeIndex16 FuncsBegin, uint16 FuncsNum)
 {
 	for (int32 FuncIndex = FuncsBegin.Get(); FuncIndex < FuncsBegin.Get() + FuncsNum; ++FuncIndex)
@@ -3933,28 +4014,124 @@ bool FStateTreeExecutionContext::SelectStateInternal(
 			if (NextState.HasChildren())
 			{
 				STATETREE_TRACE_SCOPED_STATE_PHASE(NextStateHandle, EStateTreeUpdatePhase::TrySelectBehavior);
-
-				TArray<uint16, TInlineAllocator<8>> RandomChildStates;
+				
+				TArray<uint16, TInlineAllocator<8>> NextLevelChildStates;
 				for (uint16 ChildState = NextState.ChildrenBegin; ChildState < NextState.ChildrenEnd; ChildState = CurrentStateTree->States[ChildState].GetNextSibling())
 				{
-					RandomChildStates.Push(ChildState);
+					NextLevelChildStates.Push(ChildState);
+				}
+
+				while (!NextLevelChildStates.IsEmpty())
+				{
+					const int32 ChildStateIndex = Exec.RandomStream.RandRange(0, NextLevelChildStates.Num() - 1);
+					if (SelectStateInternal(CurrentParentFrame, CurrentFrame, CurrentFrameInActiveFrames, { FStateTreeStateHandle(ChildStateIndex) }, OutSelectionResult))
+					{
+						// Selection succeeded
+						return true;
+					}
+
+					constexpr EAllowShrinking AllowShrinking = EAllowShrinking::No;
+					NextLevelChildStates.RemoveAtSwap(ChildStateIndex, AllowShrinking);
 				}
+			}
+			else
+			{
+				// Select this state (For backwards compatibility)
+				STATETREE_TRACE_STATE_EVENT(NextStateHandle, EStateTreeTraceEventType::OnStateSelected);
+				return true;
+			}
+		}
+		else if (NextState.SelectionBehavior == EStateTreeStateSelectionBehavior::TrySelectChildrenWithHighestUtility)
+		{
+			if (NextState.HasChildren())
+			{
+				STATETREE_TRACE_SCOPED_STATE_PHASE(NextStateHandle, EStateTreeUpdatePhase::TrySelectBehavior);
 
-				const uint32 LastIndex = RandomChildStates.Num() - 1;
-				for (uint32 Index = 0; Index < LastIndex; ++Index)
+				TArray<uint16, TInlineAllocator<8>> NextLevelChildStates;
+				for (uint16 ChildState = NextState.ChildrenBegin; ChildState < NextState.ChildrenEnd; ChildState = CurrentStateTree->States[ChildState].GetNextSibling())
 				{
-					// Get a random integer in [Index, Num)
-					const uint32 IndexToSwap = Exec.RandomStream.RandRange(Index, LastIndex);
-					RandomChildStates.Swap(Index, IndexToSwap);
+					NextLevelChildStates.Push(ChildState);
 				}
 
-				for (uint16 ChildState : RandomChildStates)
+				while (!NextLevelChildStates.IsEmpty())
 				{
-					if (SelectStateInternal(CurrentParentFrame, CurrentFrame, CurrentFrameInActiveFrames, { FStateTreeStateHandle(ChildState) }, OutSelectionResult))
+					//Find one with highest score in the remaining candidates
+					float HighestScore = .0f;
+					uint16 StateIndexWithHighestScore = FStateTreeStateHandle::InvalidIndex;
+					int32 StateArrayIndex = -1;
+					for (int32 Index = 0; Index < NextLevelChildStates.Num(); ++Index)
+					{
+						const uint16 CurrentStateIndex = NextLevelChildStates[Index];
+						const FCompactStateTreeState& CurrentState = CurrentStateTree->States[CurrentStateIndex];
+						const float Score = EvaluateUtility(CurrentParentFrame, CurrentFrame, CurrentState.UtilityConsiderationsBegin, CurrentState.UtilityConsiderationsNum);
+						if (Score > HighestScore)
+						{
+							HighestScore = Score;
+							StateIndexWithHighestScore = CurrentStateIndex;
+							StateArrayIndex = Index;
+						}
+					}
+
+					check(FStateTreeStateHandle::IsValidIndex(StateIndexWithHighestScore));
+					if (SelectStateInternal(CurrentParentFrame, CurrentFrame, CurrentFrameInActiveFrames, { FStateTreeStateHandle(StateIndexWithHighestScore) }, OutSelectionResult))
 					{
 						// Selection succeeded
 						return true;
 					}
+					
+					//Disqualify the state we failed to enter
+					constexpr EAllowShrinking AllowShrinking = EAllowShrinking::No;
+					NextLevelChildStates.RemoveAtSwap(StateArrayIndex);
+				}
+			}
+			else
+			{
+				// Select this state (For backwards compatibility)
+				STATETREE_TRACE_STATE_EVENT(NextStateHandle, EStateTreeTraceEventType::OnStateSelected);
+				return true;
+			}
+		}
+		else if (NextState.SelectionBehavior == EStateTreeStateSelectionBehavior::TrySelectChildrenBasedOnRelativeUtility)
+		{
+			if (NextState.HasChildren())
+			{
+				TArray<TTuple<uint16, float>, TInlineAllocator<8>> NextLevelChildStates;
+				float TotalScore = .0f;
+				for (uint16 CurrentStateIndex = NextState.ChildrenBegin; CurrentStateIndex < NextState.ChildrenEnd; CurrentStateIndex = CurrentStateTree->States[CurrentStateIndex].GetNextSibling())
+				{
+					const FCompactStateTreeState& CurrentState = CurrentStateTree->States[CurrentStateIndex];
+					const float CurrentStateScore = EvaluateUtility(CurrentParentFrame, CurrentFrame, CurrentState.UtilityConsiderationsBegin, CurrentState.UtilityConsiderationsNum);
+					NextLevelChildStates.Emplace(CurrentStateIndex, CurrentStateScore);
+					TotalScore += CurrentStateScore;
+				}
+
+				while (!NextLevelChildStates.IsEmpty())
+				{
+					const float RandomScore = Exec.RandomStream.FRand() * TotalScore;
+					float AccumulatedScore = .0f;
+					for (int32 Index = 0; Index < NextLevelChildStates.Num(); ++Index)
+					{
+						const TTuple<uint16, float>& StateScorePair = NextLevelChildStates[Index];
+						const uint16 StateIndex = StateScorePair.Key;
+						const float StateScore = StateScorePair.Value;
+						AccumulatedScore += StateScore;
+
+						if (RandomScore < AccumulatedScore || (Index == (NextLevelChildStates.Num() - 1)))
+						{
+							if (SelectStateInternal(CurrentParentFrame, CurrentFrame, CurrentFrameInActiveFrames, { FStateTreeStateHandle(StateIndex) }, OutSelectionResult))
+							{
+								// Selection succeeded
+								return true;
+							}
+
+							//Disqualify the state we failed to enter, and restart the loop
+							TotalScore -= StateScore;
+							constexpr EAllowShrinking AllowShrinking = EAllowShrinking::No;
+							NextLevelChildStates.RemoveAtSwap(Index, AllowShrinking);
+
+							break;
+						}
+					}
 				}
 			}
 			else
diff --git a/Engine/Plugins/Runtime/StateTree/Source/StateTreeModule/Private/StateTreeSchema.cpp b/Engine/Plugins/Runtime/StateTree/Source/StateTreeModule/Private/StateTreeSchema.cpp
--- a/Engine/Plugins/Runtime/StateTree/Source/StateTreeModule/Private/StateTreeSchema.cpp
+++ b/Engine/Plugins/Runtime/StateTree/Source/StateTreeModule/Private/StateTreeSchema.cpp
@@ -5,17 +5,6 @@
 
 #include UE_INLINE_GENERATED_CPP_BY_NAME(StateTreeSchema)
 
-namespace StateTreeSchemaUtilityCVars
-{
-	FAutoConsoleVariable CVarAllowUtilityConsiderations
-	(
-		TEXT("StateTree.AllowUtilityConsiderations"),
-		false,
-		TEXT("Reveal Experimental Utility Consideration in the State Tree Editor."),
-		ECVF_Default
-	);
-}
-
 bool UStateTreeSchema::IsChildOfBlueprintBase(const UClass* InClass) const
 {
 	return InClass->IsChildOf<UStateTreeNodeBlueprintBase>();
diff --git a/Engine/Plugins/Runtime/StateTree/Source/StateTreeModule/Public/StateTree.h b/Engine/Plugins/Runtime/StateTree/Source/StateTreeModule/Public/StateTree.h
--- a/Engine/Plugins/Runtime/StateTree/Source/StateTreeModule/Public/StateTree.h
+++ b/Engine/Plugins/Runtime/StateTree/Source/StateTreeModule/Public/StateTree.h
@@ -251,7 +251,7 @@ private:
 	UPROPERTY()
 	TArray<FCompactStateTransition> Transitions;
 
-	/** Evaluators, Tasks, and Condition nodes. */
+	/** Evaluators, Tasks, Condition and Consideration nodes. */
 	UPROPERTY()
 	FInstancedStructContainer Nodes;
 
@@ -259,7 +259,7 @@ private:
 	UPROPERTY()
 	FStateTreeInstanceData DefaultInstanceData;
 
-	/** Shared node instance data (e.g. conditions). */
+	/** Shared node instance data (e.g. conditions, considerations). */
 	UPROPERTY()
 	FStateTreeInstanceData SharedInstanceData;
 
diff --git a/Engine/Plugins/Runtime/StateTree/Source/StateTreeModule/Public/StateTreeExecutionContext.h b/Engine/Plugins/Runtime/StateTree/Source/StateTreeModule/Public/StateTreeExecutionContext.h
--- a/Engine/Plugins/Runtime/StateTree/Source/StateTreeModule/Public/StateTreeExecutionContext.h
+++ b/Engine/Plugins/Runtime/StateTree/Source/StateTreeModule/Public/StateTreeExecutionContext.h
@@ -586,6 +586,12 @@ protected:
 	 */
 	bool TestAllConditions(const FStateTreeExecutionFrame* CurrentParentFrame, const FStateTreeExecutionFrame& CurrentFrame, const int32 ConditionsOffset, const int32 ConditionsNum);
 
+	/**
+	 * Calculate the final score of all considerations at given range
+	 * @return the final score
+	 */
+	float EvaluateUtility(const FStateTreeExecutionFrame* CurrentParentFrame, const FStateTreeExecutionFrame& CurrentFrame, const int32 ConsiderationsOffset, const int32 ConsiderationsNum);
+
 	/* Evaluate all function at given range. Should be used only on active instances, assumes valid handles and does not consider temporary instances. */
 	void EvaluatePropertyFunctionsOnActiveInstances(const FStateTreeExecutionFrame* CurrentParentFrame, const FStateTreeExecutionFrame& CurrentFrame, FStateTreeIndex16 FuncsBegin, uint16 FuncsNum);
 
diff --git a/Engine/Plugins/Runtime/StateTree/Source/StateTreeModule/Public/StateTreeSchema.h b/Engine/Plugins/Runtime/StateTree/Source/StateTreeModule/Public/StateTreeSchema.h
--- a/Engine/Plugins/Runtime/StateTree/Source/StateTreeModule/Public/StateTreeSchema.h
+++ b/Engine/Plugins/Runtime/StateTree/Source/StateTreeModule/Public/StateTreeSchema.h
@@ -7,12 +7,6 @@
 
 struct FStateTreeExternalDataDesc;
 
-namespace StateTreeSchemaUtilityCVars
-{
-	//temporary param to be removed when Utility is fully there.
-	extern STATETREEMODULE_API FAutoConsoleVariable CVarAllowUtilityConsiderations;
-}
-
 /**
  * Schema describing which inputs, evaluators, and tasks a StateTree can contain.
  * Each StateTree asset saves the schema class name in asset data tags, which can be
diff --git a/Engine/Plugins/Runtime/StateTree/Source/StateTreeModule/Public/StateTreeTypes.h b/Engine/Plugins/Runtime/StateTree/Source/StateTreeModule/Public/StateTreeTypes.h
--- a/Engine/Plugins/Runtime/StateTree/Source/StateTreeModule/Public/StateTreeTypes.h
+++ b/Engine/Plugins/Runtime/StateTree/Source/StateTreeModule/Public/StateTreeTypes.h
@@ -133,6 +133,12 @@ enum class EStateTreeStateSelectionBehavior : uint8
 	/** When state is considered for selection, shuffle the order of child states and try to select the first one. If no child states are present, behaves like SelectState. */
 	TrySelectChildrenAtUniformRandom UMETA(DisplayName = "Try Select Children At Uniform Random"),
 
+	/** When state is considered for selection, try to select the child state with highest utility score. If there is a tie, it will try to select in order. */
+	TrySelectChildrenWithHighestUtility UMETA(DisplayName = "Try Select Children With Highest Utility Score"),
+
+	/** When state is considered for selection, randomly pick one of its child states. The probability of selecting each child state is its normalized utility score */
+	TrySelectChildrenBasedOnRelativeUtility UMETA(DisplayName = "Try Select Children Based On Relative Utility"),
+	
 	/** When state is considered for selection, try to trigger the transitions instead. */
 	TryFollowTransitions UMETA(DisplayName = "Try Follow Transitions"),
 
@@ -274,10 +280,10 @@ enum class EStateTreeDataSourceType : uint8
 	/** Active State Tasks */
 	ActiveInstanceDataObject,
 
-	/** Conditions */
+	/** Conditions and Considerations */
 	SharedInstanceData,
 
-	/** Conditions */
+	/** Conditions and Considerations */
 	SharedInstanceDataObject,
 
 	/** Context Data, Tree Parameters */
@@ -682,6 +688,10 @@ struct STATETREEMODULE_API FCompactStateTreeState
 	UPROPERTY()
 	uint16 EnterConditionsBegin = 0;
 
+	/** Index to first state utility consideration */
+	UPROPERTY()
+	uint16 UtilityConsiderationsBegin = 0;
+
 	/** Index to first transition */
 	UPROPERTY()
 	uint16 TransitionsBegin = 0;
@@ -706,6 +716,10 @@ struct STATETREEMODULE_API FCompactStateTreeState
 	/** Number of enter conditions */
 	UPROPERTY()
 	uint8 EnterConditionsNum = 0;
+	
+	/** Number of utility considerations */
+	UPROPERTY()
+	uint8 UtilityConsiderationsNum = 0;
 
 	/** Number of transitions */
 	UPROPERTY()
