StateTree: fixed not calling OnExitState on nested ST's global tasks and evaluators (contributed)

diff --git a/Engine/Plugins/Runtime/StateTree/Source/StateTreeModule/Private/StateTreeExecutionContext.cpp b/Engine/Plugins/Runtime/StateTree/Source/StateTreeModule/Private/StateTreeExecutionContext.cpp
--- a/Engine/Plugins/Runtime/StateTree/Source/StateTreeModule/Private/StateTreeExecutionContext.cpp
+++ b/Engine/Plugins/Runtime/StateTree/Source/StateTreeModule/Private/StateTreeExecutionContext.cpp
@@ -434,17 +434,9 @@ EStateTreeRunStatus FStateTreeExecutionContext::Stop(EStateTreeRunStatus Complet
 		// Transition to Succeeded state.
 		FStateTreeTransitionResult Transition;
 		Transition.TargetState = FStateTreeStateHandle::FromCompletionStatus(CompletionStatus);
-		Transition.CurrentRunStatus = CompletionStatus;
-		FStateTreeExecutionFrame& NewFrame = Transition.NextActiveFrames.AddDefaulted_GetRef();
-		NewFrame.StateTree = &RootStateTree;
-		NewFrame.RootState = FStateTreeStateHandle::Root;
-		NewFrame.ActiveStates = {};
-		
+		Transition.CurrentRunStatus = CompletionStatus;		
 		ExitState(Transition);
 
-		// Stop evaluators and global tasks.
-		StopEvaluatorsAndGlobalTasks(CompletionStatus);
-
 		// No active states or global tasks anymore, reset frames.
 		Exec.ActiveFrames.Reset();
 
@@ -1771,6 +1763,36 @@ void FStateTreeExecutionContext::ExitState(const FStateTreeTransitionResult& Tra
 			NextFrame = &Transition.NextActiveFrames[FrameIndex];
 		}
 
+		const bool bShouldCallOnEvaluatorsAndGlobalTasks = NextFrame == nullptr && CurrentFrame.bIsGlobalFrame;
+		ExitStateCalls.Emplace(EStateTreeStateChangeType::Changed, bShouldCallOnEvaluatorsAndGlobalTasks);
+		
+		if (bShouldCallOnEvaluatorsAndGlobalTasks)
+		{
+			for (int32 EvalIndex = CurrentStateTree->EvaluatorsBegin; EvalIndex < (CurrentStateTree->EvaluatorsBegin + CurrentStateTree->EvaluatorsNum); EvalIndex++)
+			{
+				const FStateTreeEvaluatorBase& Eval = CurrentStateTree->Nodes[EvalIndex].Get<const FStateTreeEvaluatorBase>();
+				const FStateTreeDataView EvalInstanceView = GetDataView(CurrentParentFrame, CurrentFrame, Eval.InstanceDataHandle);
+				FNodeInstanceDataScope DataScope(*this, Eval.InstanceDataHandle, EvalInstanceView);
+
+				if (Eval.BindingsBatch.IsValid())
+				{
+					CopyBatchOnActiveInstances(CurrentParentFrame, CurrentFrame, EvalInstanceView, Eval.BindingsBatch);
+				}
+			}
+
+			for (int32 TaskIndex = CurrentStateTree->GlobalTasksBegin; TaskIndex < (CurrentStateTree->GlobalTasksBegin + CurrentStateTree->GlobalTasksNum); TaskIndex++)
+			{
+				const FStateTreeTaskBase& Task = CurrentStateTree->Nodes[TaskIndex].Get<const FStateTreeTaskBase>();
+				const FStateTreeDataView TaskInstanceView = GetDataView(CurrentParentFrame, CurrentFrame, Task.InstanceDataHandle);
+				FNodeInstanceDataScope DataScope(*this, Task.InstanceDataHandle, TaskInstanceView);
+
+				if (Task.BindingsBatch.IsValid() && Task.bShouldCopyBoundPropertiesOnExitState)
+				{
+					CopyBatchOnActiveInstances(CurrentParentFrame, CurrentFrame, TaskInstanceView, Task.BindingsBatch);
+				}
+			}
+		}
+
 		for (int32 Index = 0; Index < CurrentFrame.ActiveStates.Num(); Index++)
 		{
 			const FStateTreeStateHandle CurrentHandle = CurrentFrame.ActiveStates[Index];
@@ -1894,6 +1916,16 @@ void FStateTreeExecutionContext::ExitState(const FStateTreeTransitionResult& Tra
 
 			STATETREE_TRACE_STATE_EVENT(CurrentHandle, EStateTreeTraceEventType::OnExited);
 		}
+
+		// Frame exit call
+		{
+			const FExitStateCall& ExitCall = ExitStateCalls[CallIndex--];
+			if (ExitCall.bShouldCall)
+			{
+				CurrentTransition.ChangeType = ExitCall.ChangeType;
+				CallStopOnEvaluatorsAndGlobalTasks(CurrentParentFrame, CurrentFrame, CurrentTransition);
+			}
+		}
 	}
 
 }
@@ -2209,55 +2241,65 @@ void FStateTreeExecutionContext::StopEvaluatorsAndGlobalTasks(const EStateTreeRu
 	Transition.TargetState = FStateTreeStateHandle::FromCompletionStatus(CompletionStatus);
 	Transition.CurrentRunStatus = CompletionStatus;
 
+	bool bIsLastGlobalFrame = true;
 	for (int32 FrameIndex = Exec.ActiveFrames.Num() - 1; FrameIndex >= 0; FrameIndex--)
 	{
 		const FStateTreeExecutionFrame* CurrentParentFrame = FrameIndex > 0 ? &Exec.ActiveFrames[FrameIndex - 1] : nullptr;
 		const FStateTreeExecutionFrame& CurrentFrame = Exec.ActiveFrames[FrameIndex];
 		if (CurrentFrame.bIsGlobalFrame)
 		{
-			FCurrentlyProcessedFrameScope FrameScope(*this, CurrentParentFrame, CurrentFrame);
+			// LastInitializedTaskIndex belongs to the last frame.
+			const FStateTreeIndex16 LastTaskToBeStopped = bIsLastGlobalFrame ? LastInitializedTaskIndex : FStateTreeIndex16::Invalid;
+			CallStopOnEvaluatorsAndGlobalTasks(CurrentParentFrame, CurrentFrame, Transition, LastTaskToBeStopped);
+			bIsLastGlobalFrame = false;
+		}
+	}
+}
 
-			const UStateTree* CurrentStateTree = CurrentFrame.StateTree;
+void FStateTreeExecutionContext::CallStopOnEvaluatorsAndGlobalTasks(const FStateTreeExecutionFrame* ParentFrame, const FStateTreeExecutionFrame& Frame, const FStateTreeTransitionResult& Transition, const FStateTreeIndex16 LastInitializedTaskIndex /*= FStateTreeIndex16()*/)
+{
+	check(Frame.bIsGlobalFrame);
 
-			for (int32 TaskIndex = (CurrentStateTree->GlobalTasksBegin + CurrentStateTree->GlobalTasksNum) - 1;  TaskIndex >= CurrentStateTree->GlobalTasksBegin; TaskIndex--)
-			{
-				const FStateTreeTaskBase& Task =  CurrentStateTree->Nodes[TaskIndex].Get<const FStateTreeTaskBase>();
-				const FStateTreeDataView TaskInstanceView = GetDataView(CurrentParentFrame, CurrentFrame, Task.InstanceDataHandle);
-				FNodeInstanceDataScope DataScope(*this, Task.InstanceDataHandle, TaskInstanceView);
+	FCurrentlyProcessedFrameScope FrameScope(*this, ParentFrame, Frame);
+	const UStateTree* CurrentStateTree = Frame.StateTree;
 
-				// Ignore disabled task
-				if (Task.bTaskEnabled == false)
-				{
-					STATETREE_LOG(VeryVerbose, TEXT("%*sSkipped 'ExitState' for disabled Task: '%s'"), UE::StateTree::DebugIndentSize, TEXT(""), *Task.Name.ToString());
-					continue;
-				}
+	for (int32 TaskIndex = (CurrentStateTree->GlobalTasksBegin + CurrentStateTree->GlobalTasksNum) - 1;  TaskIndex >= CurrentStateTree->GlobalTasksBegin; TaskIndex--)
+	{
+		const FStateTreeTaskBase& Task =  CurrentStateTree->Nodes[TaskIndex].Get<const FStateTreeTaskBase>();
+		const FStateTreeDataView TaskInstanceView = GetDataView(ParentFrame, Frame, Task.InstanceDataHandle);
+		FNodeInstanceDataScope DataScope(*this, Task.InstanceDataHandle, TaskInstanceView);
 
-				// Relying here that invalid value of LastInitializedTaskIndex == MAX_uint16.
-				if (TaskIndex <= LastInitializedTaskIndex.Get())
-				{
-					STATETREE_LOG(Verbose, TEXT("  Stop: '%s'"), *Task.Name.ToString());
-					{
-						QUICK_SCOPE_CYCLE_COUNTER(StateTree_Task_TreeStop);
-						Task.ExitState(*this, Transition);
-					}
-					STATETREE_TRACE_TASK_EVENT(TaskIndex, TaskInstanceView, EStateTreeTraceEventType::OnExited, Transition.CurrentRunStatus);
-				}
-			}
+		// Ignore disabled task
+		if (Task.bTaskEnabled == false)
+		{
+			STATETREE_LOG(VeryVerbose, TEXT("%*sSkipped 'ExitState' for disabled Task: '%s'"), UE::StateTree::DebugIndentSize, TEXT(""), *Task.Name.ToString());
+			continue;
+		}
 
-			for (int32 EvalIndex = (CurrentStateTree->EvaluatorsBegin + CurrentStateTree->EvaluatorsNum) - 1; EvalIndex >= CurrentStateTree->EvaluatorsBegin; EvalIndex--)
+		// Relying here that invalid value of LastInitializedTaskIndex == MAX_uint16.
+		if (TaskIndex <= LastInitializedTaskIndex.Get())
+		{
+			STATETREE_LOG(Verbose, TEXT("  Stop: '%s'"), *Task.Name.ToString());
 			{
-				const FStateTreeEvaluatorBase& Eval = CurrentStateTree->Nodes[EvalIndex].Get<const FStateTreeEvaluatorBase>();
-				const FStateTreeDataView EvalInstanceView = GetDataView(CurrentParentFrame, CurrentFrame, Eval.InstanceDataHandle);
-				FNodeInstanceDataScope DataScope(*this, Eval.InstanceDataHandle, EvalInstanceView);
+				QUICK_SCOPE_CYCLE_COUNTER(StateTree_Task_TreeStop);
+				Task.ExitState(*this, Transition);
+			}
+			STATETREE_TRACE_TASK_EVENT(TaskIndex, TaskInstanceView, EStateTreeTraceEventType::OnExited, Transition.CurrentRunStatus);
+		}
+	}
 
-				STATETREE_LOG(Verbose, TEXT("  Stop: '%s'"), *Eval.Name.ToString());
-				{
-					QUICK_SCOPE_CYCLE_COUNTER(StateTree_Eval_TreeStop);
-					Eval.TreeStop(*this);
+	for (int32 EvalIndex = (CurrentStateTree->EvaluatorsBegin + CurrentStateTree->EvaluatorsNum) - 1; EvalIndex >= CurrentStateTree->EvaluatorsBegin; EvalIndex--)
+	{
+		const FStateTreeEvaluatorBase& Eval = CurrentStateTree->Nodes[EvalIndex].Get<const FStateTreeEvaluatorBase>();
+		const FStateTreeDataView EvalInstanceView = GetDataView(ParentFrame, Frame, Eval.InstanceDataHandle);
+		FNodeInstanceDataScope DataScope(*this, Eval.InstanceDataHandle, EvalInstanceView);
 
-					STATETREE_TRACE_EVALUATOR_EVENT(EvalIndex, EvalInstanceView, EStateTreeTraceEventType::OnTreeStopped);
-				}
-			}
+		STATETREE_LOG(Verbose, TEXT("  Stop: '%s'"), *Eval.Name.ToString());
+		{
+			QUICK_SCOPE_CYCLE_COUNTER(StateTree_Eval_TreeStop);
+			Eval.TreeStop(*this);
+
+			STATETREE_TRACE_EVALUATOR_EVENT(EvalIndex, EvalInstanceView, EStateTreeTraceEventType::OnTreeStopped);
 		}
 	}
 }
diff --git a/Engine/Plugins/Runtime/StateTree/Source/StateTreeModule/Public/StateTreeExecutionContext.h b/Engine/Plugins/Runtime/StateTree/Source/StateTreeModule/Public/StateTreeExecutionContext.h
--- a/Engine/Plugins/Runtime/StateTree/Source/StateTreeModule/Public/StateTreeExecutionContext.h
+++ b/Engine/Plugins/Runtime/StateTree/Source/StateTreeModule/Public/StateTreeExecutionContext.h
@@ -558,6 +558,11 @@ protected:
 	 */
 	void StopEvaluatorsAndGlobalTasks(const EStateTreeRunStatus CompletionStatus, const FStateTreeIndex16 LastInitializedTaskIndex = FStateTreeIndex16());
 
+	/**
+	 * Stops evaluators and global tasks of the given frame. Expects nodes data to be already bound.
+	 */
+	void CallStopOnEvaluatorsAndGlobalTasks(const FStateTreeExecutionFrame* ParentFrame, const FStateTreeExecutionFrame& Frame, const FStateTreeTransitionResult& Transition, const FStateTreeIndex16 LastInitializedTaskIndex = FStateTreeIndex16());
+
 	/** Starts temporary instances of global evaluators and tasks for a given frame. */
 	EStateTreeRunStatus StartTemporaryEvaluatorsAndGlobalTasks(const FStateTreeExecutionFrame* CurrentParentFrame, const FStateTreeExecutionFrame& CurrentFrame);
 
diff --git a/Engine/Plugins/Runtime/StateTree/Source/StateTreeTestSuite/Private/StateTreeTest.cpp b/Engine/Plugins/Runtime/StateTree/Source/StateTreeTestSuite/Private/StateTreeTest.cpp
--- a/Engine/Plugins/Runtime/StateTree/Source/StateTreeTestSuite/Private/StateTreeTest.cpp
+++ b/Engine/Plugins/Runtime/StateTree/Source/StateTreeTestSuite/Private/StateTreeTest.cpp
@@ -2971,6 +2971,73 @@ struct FStateTreeTest_DeferredStop_ExitTask : FStateTreeTest_DeferredStop
 };
 IMPLEMENT_AI_INSTANT_TEST(FStateTreeTest_DeferredStop_ExitTask, "System.StateTree.DeferredStop.ExitTask");
 
+struct FStateTreeTest_EnterAndExitLinkedAsset : FAITestBase
+{
+	virtual bool InstantTest() override
+	{
+		FStateTreeCompilerLog Log;
+
+		// Asset 2
+		UStateTree& StateTree2 = UE::StateTree::Tests::NewStateTree(&GetWorld());
+		UStateTreeEditorData& EditorData2 = *Cast<UStateTreeEditorData>(StateTree2.EditorData);
+		UStateTreeState& Root2 = EditorData2.AddSubTree(FName(TEXT("Root2")));
+		TStateTreeEditorNode<FTestTask_Stand>& Task2 = Root2.AddTask<FTestTask_Stand>(FName(TEXT("Task2")));
+		TStateTreeEditorNode<FTestTask_Stand>& GlobalTask2 = EditorData2.AddGlobalTask<FTestTask_Stand>(FName(TEXT("GlobalTask2")));
+		GlobalTask2.GetNode().TicksToCompletion = 2;
+
+		FStateTreeCompiler Compiler2(Log);
+		const bool bResult2 = Compiler2.Compile(StateTree2);
+		AITEST_TRUE("StateTree2 should get compiled", bResult2);
+
+		// Main asset
+		UStateTree& StateTree = UE::StateTree::Tests::NewStateTree(&GetWorld());
+		UStateTreeEditorData& EditorData = *Cast<UStateTreeEditorData>(StateTree.EditorData);
+		
+		UStateTreeState& Root = EditorData.AddSubTree(FName(TEXT("Root1")));
+		UStateTreeState& A1 = Root.AddChildState(FName(TEXT("A1")), EStateTreeStateType::LinkedAsset);
+		A1.AddTransition(EStateTreeTransitionTrigger::OnStateCompleted, EStateTreeTransitionType::NextState);
+		A1.SetLinkedStateAsset(&StateTree2);
+
+		UStateTreeState& B1 = Root.AddChildState(FName(TEXT("B1")), EStateTreeStateType::State);
+		TStateTreeEditorNode<FTestTask_Stand>& Task1 = B1.AddTask<FTestTask_Stand>(FName(TEXT("Task1")));
+
+		FStateTreeCompiler Compiler(Log);
+		const bool bResult = Compiler.Compile(StateTree);
+		AITEST_TRUE("StateTree should get compiled", bResult);
+
+		const FString EnterStateStr(TEXT("EnterState"));
+		const FString ExitStateStr(TEXT("ExitState"));
+
+		{
+			EStateTreeRunStatus Status = EStateTreeRunStatus::Unset;
+			FStateTreeInstanceData InstanceData;
+			FTestStateTreeExecutionContext Exec(StateTree, StateTree, InstanceData);
+			const bool bInitSucceeded = Exec.IsValid();
+			AITEST_TRUE("StateTree should init", bInitSucceeded);
+
+			Status = Exec.Start();
+			AITEST_EQUAL("Start should complete with Running", Status, EStateTreeRunStatus::Running);
+			AITEST_TRUE("StateTree should enter GlobalTask2", Exec.Expect(GlobalTask2.GetName(), EnterStateStr));
+			AITEST_FALSE("StateTree should not exit GlobalTask2", Exec.Expect(GlobalTask2.GetName(), ExitStateStr));
+			AITEST_TRUE("StateTree should enter Task2", Exec.Expect(Task2.GetName(), EnterStateStr));
+			AITEST_FALSE("StateTree should not exit Task2", Exec.Expect(Task2.GetName(), ExitStateStr));
+			AITEST_FALSE("StateTree should not enter Task1", Exec.Expect(Task1.GetName(), EnterStateStr));
+			Exec.LogClear();
+
+			Status = Exec.Tick(0.1f);
+			AITEST_EQUAL("Tick should complete with Running", Status, EStateTreeRunStatus::Running);
+			AITEST_FALSE("StateTree should not enter GlobalTask2", Exec.Expect(GlobalTask2.GetName(), EnterStateStr));
+			AITEST_TRUE("StateTree should exit GlobalTask2", Exec.Expect(GlobalTask2.GetName(), ExitStateStr));
+			AITEST_FALSE("StateTree should not enter Task2", Exec.Expect(Task2.GetName(), EnterStateStr));
+			AITEST_TRUE("StateTree should exit Task2", Exec.Expect(Task2.GetName(), ExitStateStr));
+			AITEST_TRUE("StateTree should enter Task1", Exec.Expect(Task1.GetName(), EnterStateStr));
+			Exec.LogClear();
+		}
+
+		return true;
+	}
+};
+IMPLEMENT_AI_INSTANT_TEST(FStateTreeTest_EnterAndExitLinkedAsset, "System.StateTree.EnterAndExitLinkedAsset");
 
 // Test nested tree overrides
 struct FStateTreeTest_NestedOverride : FAITestBase
