StateTree: fixed stopping temporary nodes with uninitialized data handles (contributed)

diff --git a/Engine/Plugins/Runtime/StateTree/Source/StateTreeModule/Private/StateTreeExecutionContext.cpp b/Engine/Plugins/Runtime/StateTree/Source/StateTreeModule/Private/StateTreeExecutionContext.cpp
--- a/Engine/Plugins/Runtime/StateTree/Source/StateTreeModule/Private/StateTreeExecutionContext.cpp
+++ b/Engine/Plugins/Runtime/StateTree/Source/StateTreeModule/Private/StateTreeExecutionContext.cpp
@@ -993,9 +993,6 @@ void FStateTreeExecutionContext::UpdateInstanceData(TConstArrayView<FStateTreeEx
 		MakeArrayView(InstanceStructs.GetData() + NumCommonInstanceData, InstanceStructs.Num() - NumCommonInstanceData),
 		MakeArrayView(TempInstanceStructs.GetData() + NumCommonInstanceData, TempInstanceStructs.Num() - NumCommonInstanceData));
 
-	// Stop any temporary global tasks or evals that were left over.
-	StopTemporaryEvaluatorsAndGlobalTasks(TempInstances);
-	
 	InstanceData.ResetTemporaryInstances();
 }
 
@@ -2404,10 +2401,10 @@ EStateTreeRunStatus FStateTreeExecutionContext::StartTemporaryEvaluatorsAndGloba
 	return Result;
 }
 
-void FStateTreeExecutionContext::StopTemporaryEvaluatorsAndGlobalTasks(TArrayView<FStateTreeTemporaryInstanceData> TempInstances)
+void FStateTreeExecutionContext::StopTemporaryEvaluatorsAndGlobalTasks(const FStateTreeExecutionFrame* CurrentParentFrame, const FStateTreeExecutionFrame& CurrentFrame)
 {
 	// @todo: figure out debugger phase for temporary stop.
-	STATETREE_CLOG(!TempInstances.IsEmpty(), Verbose, TEXT("Stop Temporary Evaluators & Global tasks left over from previous state selection"));
+	STATETREE_LOG(Verbose, TEXT("Stop Temporary Evaluators & Global tasks"));
 
 	// Create temporary transition to stop the unused global tasks and evaluators.
 	constexpr  EStateTreeRunStatus CompletionStatus = EStateTreeRunStatus::Stopped; 
@@ -2415,16 +2412,18 @@ void FStateTreeExecutionContext::StopTemporaryEvaluatorsAndGlobalTasks(TArrayVie
 	Transition.TargetState = FStateTreeStateHandle::FromCompletionStatus(CompletionStatus);
 	Transition.CurrentRunStatus = CompletionStatus;
 
+	TArrayView<FStateTreeTemporaryInstanceData> TempInstances = InstanceDataStorage->GetMutableTemporaryInstances();
 	for (int32 Index = TempInstances.Num() - 1; Index >= 0; Index--)
 	{
 		FStateTreeTemporaryInstanceData& TempInstance = TempInstances[Index];
+		if (TempInstance.StateTree != CurrentFrame.StateTree || TempInstance.RootState != CurrentFrame.RootState)
+		{
+			continue;
+		}
+
 		if (TempInstance.OwnerNodeIndex.IsValid()
 			&& TempInstance.Instance.IsValid())
 		{
-			FStateTreeExecutionFrame TempFrame;
-			TempFrame.StateTree = TempInstance.StateTree;
-			TempFrame.RootState = TempInstance.RootState;
-
 			FStateTreeDataView NodeInstanceView;
 			if (FStateTreeInstanceObjectWrapper* Wrapper = TempInstance.Instance.GetMutablePtr<FStateTreeInstanceObjectWrapper>())
 			{
@@ -2435,10 +2434,10 @@ void FStateTreeExecutionContext::StopTemporaryEvaluatorsAndGlobalTasks(TArrayVie
 				NodeInstanceView = FStateTreeDataView(TempInstance.Instance);
 			}
 			
-			FCurrentlyProcessedFrameScope FrameScope(*this, nullptr, TempFrame);
+			FCurrentlyProcessedFrameScope FrameScope(*this, CurrentParentFrame, CurrentFrame);
 			FNodeInstanceDataScope DataScope(*this, TempInstance.DataHandle, NodeInstanceView);
 
-			FConstStructView NodeView = TempFrame.StateTree->Nodes[TempInstance.OwnerNodeIndex.Get()];
+			FConstStructView NodeView = CurrentFrame.StateTree->Nodes[TempInstance.OwnerNodeIndex.Get()];
 			if (const FStateTreeTaskBase* Task = NodeView.GetPtr<const FStateTreeTaskBase>())
 			{
 				STATETREE_LOG(Verbose, TEXT("  Stop: '%s'"), *Task->Name.ToString());
@@ -3873,6 +3872,8 @@ bool FStateTreeExecutionContext::SelectStateInternal(
 				{
 					return Frame.StateTree == StateTree && Frame.RootState == RootState;
 				});
+
+			bool bStartedTemporaryEvaluatorsAndGlobalTasks = false;
 			if (ExistingFrame)
 			{
 				NewFrame.ActiveInstanceIndexBase = ExistingFrame->ActiveInstanceIndexBase;
@@ -3904,8 +3905,12 @@ bool FStateTreeExecutionContext::SelectStateInternal(
 				{
 					STATETREE_LOG(VeryVerbose, TEXT("%hs: Cannot select state '%s' because cannot start nested tree's '%s' global tasks and evaluators.  '%s' using StateTree '%s'."),
 						__FUNCTION__, *GetSafeStateName(CurrentFrame, NextStateHandle), *GetFullNameSafe(NewFrame.StateTree), *GetNameSafe(&Owner), *GetFullNameSafe(CurrentFrame.StateTree));
+					
+					StopTemporaryEvaluatorsAndGlobalTasks(nullptr, NewFrame);
 					return false;
 				}
+
+				bStartedTemporaryEvaluatorsAndGlobalTasks = true;
 			}
 				
 			OutSelectionResult.PushFrame(NewFrame);
@@ -3915,6 +3920,11 @@ bool FStateTreeExecutionContext::SelectStateInternal(
 			{
 				return true;
 			}
+
+			if (bStartedTemporaryEvaluatorsAndGlobalTasks)
+			{
+				StopTemporaryEvaluatorsAndGlobalTasks(nullptr, NewFrame);
+			}
 				
 			OutSelectionResult.PopFrame();
 		}
diff --git a/Engine/Plugins/Runtime/StateTree/Source/StateTreeModule/Public/StateTreeExecutionContext.h b/Engine/Plugins/Runtime/StateTree/Source/StateTreeModule/Public/StateTreeExecutionContext.h
--- a/Engine/Plugins/Runtime/StateTree/Source/StateTreeModule/Public/StateTreeExecutionContext.h
+++ b/Engine/Plugins/Runtime/StateTree/Source/StateTreeModule/Public/StateTreeExecutionContext.h
@@ -566,8 +566,8 @@ protected:
 	/** Starts temporary instances of global evaluators and tasks for a given frame. */
 	EStateTreeRunStatus StartTemporaryEvaluatorsAndGlobalTasks(const FStateTreeExecutionFrame* CurrentParentFrame, const FStateTreeExecutionFrame& CurrentFrame);
 
-	/** Stops leftover global evaluators and tasks in the provided temporary instance data. */
-	void StopTemporaryEvaluatorsAndGlobalTasks(TArrayView<FStateTreeTemporaryInstanceData> TempInstances);
+	/** Stops temporary global evaluators and tasks for the provided frame. */
+	void StopTemporaryEvaluatorsAndGlobalTasks(const FStateTreeExecutionFrame* CurrentParentFrame, const FStateTreeExecutionFrame& CurrentFrame);
 
 	/**
 	 * Ticks tasks of all active states starting from current state by delta time.
diff --git a/Engine/Plugins/Runtime/StateTree/Source/StateTreeTestSuite/Private/StateTreeTest.cpp b/Engine/Plugins/Runtime/StateTree/Source/StateTreeTestSuite/Private/StateTreeTest.cpp
--- a/Engine/Plugins/Runtime/StateTree/Source/StateTreeTestSuite/Private/StateTreeTest.cpp
+++ b/Engine/Plugins/Runtime/StateTree/Source/StateTreeTestSuite/Private/StateTreeTest.cpp
@@ -2971,6 +2971,69 @@ struct FStateTreeTest_DeferredStop_ExitTask : FStateTreeTest_DeferredStop
 };
 IMPLEMENT_AI_INSTANT_TEST(FStateTreeTest_DeferredStop_ExitTask, "System.StateTree.DeferredStop.ExitTask");
 
+struct FStateTreeTest_FailEnterLinkedAsset : FAITestBase
+{
+	virtual bool InstantTest() override
+	{
+		FStateTreeCompilerLog Log;
+
+		// Asset 2
+		UStateTree& StateTree2 = UE::StateTree::Tests::NewStateTree(&GetWorld());
+		UStateTreeEditorData& EditorData2 = *Cast<UStateTreeEditorData>(StateTree2.EditorData);
+		UStateTreeState& Root2 = EditorData2.AddSubTree(FName(TEXT("Root2")));
+		TStateTreeEditorNode<FTestTask_Stand>& Task2 = Root2.AddTask<FTestTask_Stand>(FName(TEXT("Task2")));
+		TStateTreeEditorNode<FTestTask_Stand>& GlobalTask2 = EditorData2.AddGlobalTask<FTestTask_Stand>(FName(TEXT("GlobalTask2")));
+		GlobalTask2.GetInstanceData().Value = 123;
+
+		// Always failing enter condition
+		TStateTreeEditorNode<FStateTreeCompareIntCondition>& IntCond2 = Root2.AddEnterCondition<FStateTreeCompareIntCondition>();
+		EditorData2.AddPropertyBinding(GlobalTask2, TEXT("Value"), IntCond2, TEXT("Left"));
+		IntCond2.GetInstanceData().Right = 0;
+
+		FStateTreeCompiler Compiler2(Log);
+		const bool bResult2 = Compiler2.Compile(StateTree2);
+		AITEST_TRUE("StateTree2 should get compiled", bResult2);
+
+		// Main asset
+		UStateTree& StateTree = UE::StateTree::Tests::NewStateTree(&GetWorld());
+		UStateTreeEditorData& EditorData = *Cast<UStateTreeEditorData>(StateTree.EditorData);
+		
+		UStateTreeState& Root = EditorData.AddSubTree(FName(TEXT("Root1")));
+		UStateTreeState& A1 = Root.AddChildState(FName(TEXT("A1")), EStateTreeStateType::LinkedAsset);
+		A1.SetLinkedStateAsset(&StateTree2);
+
+		UStateTreeState& B1 = Root.AddChildState(FName(TEXT("B1")), EStateTreeStateType::State);
+		TStateTreeEditorNode<FTestTask_Stand>& Task1 = B1.AddTask<FTestTask_Stand>(FName(TEXT("Task1")));
+
+		FStateTreeCompiler Compiler(Log);
+		const bool bResult = Compiler.Compile(StateTree);
+		AITEST_TRUE("StateTree should get compiled", bResult);
+
+		const FString EnterStateStr(TEXT("EnterState"));
+		const FString ExitStateStr(TEXT("ExitState"));
+
+		{
+			EStateTreeRunStatus Status = EStateTreeRunStatus::Unset;
+			FStateTreeInstanceData InstanceData;
+			FTestStateTreeExecutionContext Exec(StateTree, StateTree, InstanceData);
+			const bool bInitSucceeded = Exec.IsValid();
+			AITEST_TRUE("StateTree should init", bInitSucceeded);
+
+			Status = Exec.Start();
+			AITEST_EQUAL("Start should complete with Running", Status, EStateTreeRunStatus::Running);
+			AITEST_TRUE("StateTree should enter GlobalTask2", Exec.Expect(GlobalTask2.GetName(), EnterStateStr));
+			AITEST_TRUE("StateTree should exit GlobalTask2", Exec.Expect(GlobalTask2.GetName(), ExitStateStr));
+			AITEST_FALSE("StateTree should not enter Task2", Exec.Expect(Task2.GetName(), EnterStateStr));
+			AITEST_TRUE("StateTree should enter Task1", Exec.Expect(Task1.GetName(), EnterStateStr));
+
+			Exec.LogClear();
+		}
+
+		return true;
+	}
+};
+IMPLEMENT_AI_INSTANT_TEST(FStateTreeTest_FailEnterLinkedAsset, "System.StateTree.FailEnterLinkedAsset");
+
 struct FStateTreeTest_EnterAndExitLinkedAsset : FAITestBase
 {
 	virtual bool InstantTest() override
