StateTree: Fixed not restoring selection states stack when linked state tree failed to be selected. (contributed)

diff --git a/Engine/Plugins/Runtime/StateTree/Source/StateTreeModule/Private/StateTreeExecutionContext.cpp b/Engine/Plugins/Runtime/StateTree/Source/StateTreeModule/Private/StateTreeExecutionContext.cpp
--- a/Engine/Plugins/Runtime/StateTree/Source/StateTreeModule/Private/StateTreeExecutionContext.cpp
+++ b/Engine/Plugins/Runtime/StateTree/Source/StateTreeModule/Private/StateTreeExecutionContext.cpp
@@ -3707,12 +3707,40 @@ bool FStateTreeExecutionContext::SelectStateInternal(
 				StateSelectionEvents.Emplace();
 			}
 		}
+
+		if (StateSelectionEvents.IsEmpty())
+		{
+			return false;
+		}
 	}
 	else
 	{
 		StateSelectionEvents.Emplace();
 	}
 
+	if (!CurrentFrame.ActiveStates.Push(NextStateHandle))
+	{
+		STATETREE_LOG(Error, TEXT("%hs: Reached max execution depth when trying to select state %s from '%s'.  '%s' using StateTree '%s'."),
+			__FUNCTION__, *GetSafeStateName(CurrentFrame, NextStateHandle), *GetStateStatusString(Exec), *GetNameSafe(&Owner), *GetFullNameSafe(CurrentFrame.StateTree));
+		return false;
+	}
+
+	// Check if we're still tracking on the current active frame and state.
+	// If we are, update the NumCurrentlyActiveStates to indicate that this state's instance data can be accessed. 
+	const uint8 PrevNumCurrentlyActiveStates = CurrentFrame.NumCurrentlyActiveStates; 
+	if (CurrentFrame.ActiveInstanceIndexBase.IsValid()
+		&& CurrentFrameInActiveFrames)
+	{
+		const int32 CurrentStateIndex = CurrentFrame.ActiveStates.Num() - 1;
+		const FStateTreeStateHandle MatchingActiveHandle = CurrentFrameInActiveFrames->ActiveStates.GetStateSafe(CurrentStateIndex);
+		if (MatchingActiveHandle == NextStateHandle)
+		{
+			CurrentFrame.NumCurrentlyActiveStates = static_cast<uint8>(CurrentFrame.ActiveStates.Num());
+		}
+	}
+
+	bool bSucceededToSelectState = false;
+
 	for (const FStateTreeSharedEvent* StateSelectionEvent : StateSelectionEvents)
 	{
 		if (StateSelectionEvent)
@@ -3732,34 +3760,14 @@ bool FStateTreeExecutionContext::SelectStateInternal(
 				continue;
 			}
 		}
-
-		if (!CurrentFrame.ActiveStates.Push(NextStateHandle))
-		{
-			STATETREE_LOG(Error, TEXT("%hs: Reached max execution depth when trying to select state %s from '%s'.  '%s' using StateTree '%s'."),
-				__FUNCTION__, *GetSafeStateName(CurrentFrame, NextStateHandle), *GetStateStatusString(Exec), *GetNameSafe(&Owner), *GetFullNameSafe(CurrentFrame.StateTree));
-			return false;
-		}
-
-		// Check if we're still tracking on the current active frame and state.
-		// If we are, update the NumCurrentlyActiveStates to indicate that this state's instance data can be accessed. 
-		const uint8 PrevNumCurrentlyActiveStates = CurrentFrame.NumCurrentlyActiveStates; 
-		if (CurrentFrame.ActiveInstanceIndexBase.IsValid()
-			&& CurrentFrameInActiveFrames)
-		{
-			const int32 CurrentStateIndex = CurrentFrame.ActiveStates.Num() - 1;
-			const FStateTreeStateHandle MatchingActiveHandle = CurrentFrameInActiveFrames->ActiveStates.GetStateSafe(CurrentStateIndex);
-			if (MatchingActiveHandle == NextStateHandle)
-			{
-				CurrentFrame.NumCurrentlyActiveStates = static_cast<uint8>(CurrentFrame.ActiveStates.Num());
-			}
-		}
 		
 		if (!bIsDestinationState)
 		{
 			// Next child state is already known. Passing TransitionEvent further so state selected directly by transition can use it.
 			if (SelectStateInternal(CurrentParentFrame, CurrentFrame, CurrentFrameInActiveFrames, PathToNextState.Mid(1), OutSelectionResult, TransitionEvent))
 			{
-				return true;
+				bSucceededToSelectState = true;
+				break;
 			}
 		}
 		else if (NextState.Type == EStateTreeStateType::Linked)
@@ -3770,7 +3778,7 @@ bool FStateTreeExecutionContext::SelectStateInternal(
 				{
 					STATETREE_LOG(Error, TEXT("%hs: Reached max execution depth when trying to select state %s from '%s'.  '%s' using StateTree '%s'."),
 						__FUNCTION__, *GetSafeStateName(CurrentFrame, NextStateHandle), *GetStateStatusString(Exec), *GetNameSafe(&Owner), *GetFullNameSafe(CurrentFrame.StateTree));
-					return false;
+					break;
 				}
 
 				FStateTreeExecutionFrame NewFrame;
@@ -3787,7 +3795,7 @@ bool FStateTreeExecutionContext::SelectStateInternal(
 				{
 					STATETREE_LOG(Error, TEXT("%hs: Trying to recursively enter subtree '%s' from '%s'.  '%s' using StateTree '%s'."),
 						__FUNCTION__, *GetSafeStateName(NewFrame, NewFrame.RootState), *GetStateStatusString(Exec), *GetNameSafe(&Owner), *GetFullNameSafe(CurrentFrame.StateTree));
-					return false;
+					break;
 				}
 
 				// If the Frame already exists, copy instance indices so that conditions that rely on active states work correctly.
@@ -3816,7 +3824,8 @@ bool FStateTreeExecutionContext::SelectStateInternal(
 				// If State is linked, proceed to the linked state.
 				if (SelectStateInternal(&CurrentFrame, OutSelectionResult.GetSelectedFrames().Last(), ExistingFrame, {NewFrame.RootState}, OutSelectionResult))
 				{
-					return true;
+					bSucceededToSelectState = true;
+					break;
 				}
 				
 				OutSelectionResult.PopFrame();
@@ -3831,14 +3840,14 @@ bool FStateTreeExecutionContext::SelectStateInternal(
 		{
 			if (NextLinkedStateAsset == nullptr)
 			{
-				return false;
+				break;
 			}
 
 			if (OutSelectionResult.IsFull())
 			{
 				STATETREE_LOG(Error, TEXT("%hs: Reached max execution depth when trying to select state %s from '%s'.  '%s' using StateTree '%s'."),
 					__FUNCTION__, *GetSafeStateName(CurrentFrame, NextStateHandle), *GetStateStatusString(Exec), *GetNameSafe(&Owner), *GetFullNameSafe(CurrentFrame.StateTree));
-				return false;
+				break;
 			}
 
 			// The linked state tree should have compatible context requirements.
@@ -3846,7 +3855,7 @@ bool FStateTreeExecutionContext::SelectStateInternal(
 			{
 				STATETREE_LOG(Error, TEXT("%hs: The linked State Tree '%s' does not have compatible schema, trying to select state %s from '%s'.  '%s' using StateTree '%s'."),
 					__FUNCTION__, *GetFullNameSafe(NextLinkedStateAsset), *GetSafeStateName(CurrentFrame, NextStateHandle), *GetStateStatusString(Exec), *GetNameSafe(&Owner), *GetFullNameSafe(CurrentFrame.StateTree));
-				return false;
+				break;
 			}
 				
 			FStateTreeExecutionFrame NewFrame;
@@ -3863,7 +3872,7 @@ bool FStateTreeExecutionContext::SelectStateInternal(
 			{
 				STATETREE_LOG(Error, TEXT("%hs: Trying to recursively enter subtree '%s' from '%s'.  '%s' using StateTree '%s'."),
 					__FUNCTION__, *GetSafeStateName(NewFrame, NewFrame.RootState), *GetStateStatusString(Exec), *GetNameSafe(&Owner), *GetFullNameSafe(CurrentFrame.StateTree));
-				return false;
+				break;
 			}
 
 			// If the Frame already exists, copy instance indices so that conditions that rely on active states work correctly.
@@ -3893,7 +3902,7 @@ bool FStateTreeExecutionContext::SelectStateInternal(
 				{
 					STATETREE_LOG(VeryVerbose, TEXT("%hs: Cannot select state '%s' because failed to collect external data for nested tree '%s'.  '%s' using StateTree '%s'."),
 						__FUNCTION__, *GetSafeStateName(CurrentFrame, NextStateHandle), *GetFullNameSafe(NewFrame.StateTree), *GetNameSafe(&Owner), *GetFullNameSafe(CurrentFrame.StateTree));
-					return false;
+					break;
 				}
 					
 				// The state parameters will be from the root state.
@@ -3907,7 +3916,7 @@ bool FStateTreeExecutionContext::SelectStateInternal(
 						__FUNCTION__, *GetSafeStateName(CurrentFrame, NextStateHandle), *GetFullNameSafe(NewFrame.StateTree), *GetNameSafe(&Owner), *GetFullNameSafe(CurrentFrame.StateTree));
 					
 					StopTemporaryEvaluatorsAndGlobalTasks(nullptr, NewFrame);
-					return false;
+					break;
 				}
 
 				bStartedTemporaryEvaluatorsAndGlobalTasks = true;
@@ -3918,7 +3927,8 @@ bool FStateTreeExecutionContext::SelectStateInternal(
 			// If State is linked, proceed to the linked state.
 			if (SelectStateInternal(&CurrentFrame, OutSelectionResult.GetSelectedFrames().Last(), ExistingFrame, {NewFrame.RootState}, OutSelectionResult))
 			{
-				return true;
+				bSucceededToSelectState = true;
+				break;
 			}
 
 			if (bStartedTemporaryEvaluatorsAndGlobalTasks)
@@ -3932,7 +3942,8 @@ bool FStateTreeExecutionContext::SelectStateInternal(
 		{
 			// Select this state.
 			STATETREE_TRACE_STATE_EVENT(NextStateHandle, EStateTreeTraceEventType::OnStateSelected);
-			return true;
+			bSucceededToSelectState = true;
+			break;
 		}
 		else if (NextState.SelectionBehavior == EStateTreeStateSelectionBehavior::TryFollowTransitions)
 		{
@@ -4035,7 +4046,8 @@ bool FStateTreeExecutionContext::SelectStateInternal(
 
 			if (CurrentPriority != EStateTreeTransitionPriority::None)
 			{
-				return true;
+				bSucceededToSelectState = true;
+				break;
 			}
 		}
 		else if (NextState.SelectionBehavior == EStateTreeStateSelectionBehavior::TrySelectChildrenInOrder)
@@ -4050,15 +4062,22 @@ bool FStateTreeExecutionContext::SelectStateInternal(
 					if (SelectStateInternal(CurrentParentFrame, CurrentFrame, CurrentFrameInActiveFrames, {FStateTreeStateHandle(ChildState)}, OutSelectionResult))
 					{
 						// Selection succeeded
-						return true;
+						bSucceededToSelectState = true;
+						break;
 					}
 				}
+
+				if (bSucceededToSelectState)
+				{
+					break;
+				}
 			}
 			else
 			{
 				// Select this state (For backwards compatibility)
 				STATETREE_TRACE_STATE_EVENT(NextStateHandle, EStateTreeTraceEventType::OnStateSelected);
-				return true;
+				bSucceededToSelectState = true;
+				break;
 			}
 		}
 		else if (NextState.SelectionBehavior == EStateTreeStateSelectionBehavior::TrySelectChildrenAtUniformRandom)
@@ -4079,18 +4098,25 @@ bool FStateTreeExecutionContext::SelectStateInternal(
 					if (SelectStateInternal(CurrentParentFrame, CurrentFrame, CurrentFrameInActiveFrames, { FStateTreeStateHandle(ChildStateIndex) }, OutSelectionResult))
 					{
 						// Selection succeeded
-						return true;
+						bSucceededToSelectState = true;
+						break;
 					}
 
 					constexpr EAllowShrinking AllowShrinking = EAllowShrinking::No;
 					NextLevelChildStates.RemoveAtSwap(ChildStateIndex, AllowShrinking);
 				}
+
+				if (bSucceededToSelectState)
+				{
+					break;
+				}
 			}
 			else
 			{
 				// Select this state (For backwards compatibility)
 				STATETREE_TRACE_STATE_EVENT(NextStateHandle, EStateTreeTraceEventType::OnStateSelected);
-				return true;
+				bSucceededToSelectState = true;
+				break;
 			}
 		}
 		else if (NextState.SelectionBehavior == EStateTreeStateSelectionBehavior::TrySelectChildrenWithHighestUtility)
@@ -4128,19 +4154,26 @@ bool FStateTreeExecutionContext::SelectStateInternal(
 					if (SelectStateInternal(CurrentParentFrame, CurrentFrame, CurrentFrameInActiveFrames, { FStateTreeStateHandle(StateIndexWithHighestScore) }, OutSelectionResult))
 					{
 						// Selection succeeded
-						return true;
+						bSucceededToSelectState = true;
+						break;
 					}
 					
 					//Disqualify the state we failed to enter
 					constexpr EAllowShrinking AllowShrinking = EAllowShrinking::No;
 					NextLevelChildStates.RemoveAtSwap(StateArrayIndex);
 				}
+
+				if (bSucceededToSelectState)
+				{
+					break;
+				}
 			}
 			else
 			{
 				// Select this state (For backwards compatibility)
 				STATETREE_TRACE_STATE_EVENT(NextStateHandle, EStateTreeTraceEventType::OnStateSelected);
-				return true;
+				bSucceededToSelectState = true;
+				break;
 			}
 		}
 		else if (NextState.SelectionBehavior == EStateTreeStateSelectionBehavior::TrySelectChildrenBasedOnRelativeUtility)
@@ -4173,7 +4206,8 @@ bool FStateTreeExecutionContext::SelectStateInternal(
 							if (SelectStateInternal(CurrentParentFrame, CurrentFrame, CurrentFrameInActiveFrames, { FStateTreeStateHandle(StateIndex) }, OutSelectionResult))
 							{
 								// Selection succeeded
-								return true;
+								bSucceededToSelectState = true;
+								break;
 							}
 
 							//Disqualify the state we failed to enter, and restart the loop
@@ -4184,22 +4218,36 @@ bool FStateTreeExecutionContext::SelectStateInternal(
 							break;
 						}
 					}
+
+					if (bSucceededToSelectState)
+					{
+						break;
+					}
+				}
+
+				if (bSucceededToSelectState)
+				{
+					break;
 				}
 			}
 			else
 			{
 				// Select this state (For backwards compatibility)
 				STATETREE_TRACE_STATE_EVENT(NextStateHandle, EStateTreeTraceEventType::OnStateSelected);
-				return true;
+				bSucceededToSelectState = true;
+				break;
 			}
 		}
+	}
+
+	if (!bSucceededToSelectState)
+	{
 		// State could not be selected, restore.
 		CurrentFrame.NumCurrentlyActiveStates = PrevNumCurrentlyActiveStates;
 		CurrentFrame.ActiveStates.Pop();
 	}
 
-	// Nothing got selected.
-	return false;
+	return bSucceededToSelectState;
 }
 
 FString FStateTreeExecutionContext::GetSafeStateName(const FStateTreeExecutionFrame& CurrentFrame, const FStateTreeStateHandle State) const
