StateTrees: Logging an error when an event with empty tag and payload has been sent to ST.
- Fixed event matching logic.
- Fixed validating instance data for property functions

diff --git a/Engine/Plugins/Runtime/StateTree/Source/StateTreeModule/Private/StateTree.cpp b/Engine/Plugins/Runtime/StateTree/Source/StateTreeModule/Private/StateTree.cpp
--- a/Engine/Plugins/Runtime/StateTree/Source/StateTreeModule/Private/StateTree.cpp
+++ b/Engine/Plugins/Runtime/StateTree/Source/StateTreeModule/Private/StateTree.cpp
@@ -470,7 +470,7 @@ bool UStateTree::ValidateInstanceData() const
 		{
 			const UStruct* CurrentInstanceDataType;
 			{
-				const bool bUseSharedInstanceData = NodeView.GetPtr<const FStateTreeConditionBase>() || NodeView.GetPtr<const FStateTreeConsiderationBase>();
+				const bool bUseSharedInstanceData = NodeView.GetPtr<const FStateTreeConditionBase>() || NodeView.GetPtr<const FStateTreeConsiderationBase>() || NodeView.GetPtr<const FStateTreePropertyFunctionBase>();
 
 				const FStateTreeInstanceData& SourceInstanceData = bUseSharedInstanceData ? SharedInstanceData : DefaultInstanceData;
 
diff --git a/Engine/Plugins/Runtime/StateTree/Source/StateTreeModule/Private/StateTreeEvents.cpp b/Engine/Plugins/Runtime/StateTree/Source/StateTreeModule/Private/StateTreeEvents.cpp
--- a/Engine/Plugins/Runtime/StateTree/Source/StateTreeModule/Private/StateTreeEvents.cpp
+++ b/Engine/Plugins/Runtime/StateTree/Source/StateTreeModule/Private/StateTreeEvents.cpp
@@ -22,6 +22,12 @@ void FStateTreeSharedEvent::AddStructReferencedObjects(FReferenceCollector& Coll
 
 void FStateTreeEventQueue::SendEvent(const UObject* Owner, const FGameplayTag& Tag, const FConstStructView Payload, const FName Origin)
 {
+	if (!Tag.IsValid() && !Payload.IsValid())
+	{
+		UE_VLOG_UELOG(Owner, LogStateTree, Error, TEXT("%s: An event with an invalid tag and payload has been sent to '%s'. This is not allowed."), ANSI_TO_TCHAR(__FUNCTION__), *GetNameSafe(Owner));
+		return;
+	}
+
 	if (SharedEvents.Num() >= MaxActiveEvents)
 	{
 		UE_VLOG_UELOG(Owner, LogStateTree, Error, TEXT("%s: Too many events send on '%s'. Dropping event %s"), ANSI_TO_TCHAR(__FUNCTION__), *GetNameSafe(Owner), *Tag.ToString());
diff --git a/Engine/Plugins/Runtime/StateTree/Source/StateTreeModule/Private/StateTreeTypes.cpp b/Engine/Plugins/Runtime/StateTree/Source/StateTreeModule/Private/StateTreeTypes.cpp
--- a/Engine/Plugins/Runtime/StateTree/Source/StateTreeModule/Private/StateTreeTypes.cpp
+++ b/Engine/Plugins/Runtime/StateTree/Source/StateTreeModule/Private/StateTreeTypes.cpp
@@ -116,6 +116,16 @@ PRAGMA_ENABLE_DEPRECATION_WARNINGS
 
 bool FCompactEventDesc::DoesEventMatchDesc(const FStateTreeEvent& Event) const
 {
+	if (Tag.IsValid() && (!Event.Tag.IsValid() || !Event.Tag.MatchesTag(Tag)))
+	{
+		return false;
+	}
+
 	const UScriptStruct* EventPayloadStruct = Event.Payload.GetScriptStruct();
-	return (!Tag.IsValid() || !Event.Tag.IsValid() || Event.Tag.MatchesTag(Tag)) && (EventPayloadStruct == nullptr || PayloadStruct == nullptr || EventPayloadStruct->IsChildOf(PayloadStruct));
+	if (PayloadStruct && (EventPayloadStruct == nullptr || !EventPayloadStruct->IsChildOf(PayloadStruct)))
+	{
+		return false;
+	}
+
+	return true;
 }
