StateTree: Added utility icon

diff --git a/Engine/Plugins/Runtime/StateTree/Source/StateTreeEditorModule/Private/Customizations/StateTreeStateDetails.cpp b/Engine/Plugins/Runtime/StateTree/Source/StateTreeEditorModule/Private/Customizations/StateTreeStateDetails.cpp
--- a/Engine/Plugins/Runtime/StateTree/Source/StateTreeEditorModule/Private/Customizations/StateTreeStateDetails.cpp
+++ b/Engine/Plugins/Runtime/StateTree/Source/StateTreeEditorModule/Private/Customizations/StateTreeStateDetails.cpp
@@ -295,9 +295,9 @@ void FStateTreeStateDetails::CustomizeDetails(IDetailLayoutBuilder& DetailBuilde
 			ConsiderationsProperty,
 			UtilityCategoryName,
 			LOCTEXT("StateDetailsUtility", "Utility"),
-			FName("StateTreeEditor.Considerations"),
-			UE::StateTree::Colors::Yellow,
-			UE::StateTree::Colors::Yellow.WithAlpha(192),
+			FName("StateTreeEditor.Utility"),
+			UE::StateTree::Colors::Orange,
+			UE::StateTree::Colors::Orange.WithAlpha(192),
 			LOCTEXT("UtilityAddTooltip", "Add new Utility Consideration"),
 			/*SortOrder*/3);
 	}
diff --git a/Engine/Plugins/Runtime/StateTree/Source/StateTreeEditorModule/Private/StateTreeEditorStyle.cpp b/Engine/Plugins/Runtime/StateTree/Source/StateTreeEditorModule/Private/StateTreeEditorStyle.cpp
--- a/Engine/Plugins/Runtime/StateTree/Source/StateTreeEditorModule/Private/StateTreeEditorStyle.cpp
+++ b/Engine/Plugins/Runtime/StateTree/Source/StateTreeEditorModule/Private/StateTreeEditorStyle.cpp
@@ -337,6 +337,7 @@ FStateTreeEditorStyle::FStateTreeEditorStyle()
 		Set("StateTreeEditor.StateConditions", new IMAGE_BRUSH_SVG("Icons/State_Conditions", CoreStyleConstants::Icon16x16));
 
 		Set("StateTreeEditor.Conditions", new IMAGE_BRUSH_SVG("Icons/Conditions", CoreStyleConstants::Icon16x16));
+		Set("StateTreeEditor.Utility", new IMAGE_BRUSH_SVG("Icons/Utility", CoreStyleConstants::Icon16x16));
 		Set("StateTreeEditor.Tasks", new IMAGE_BRUSH_SVG("Icons/Tasks", CoreStyleConstants::Icon16x16));
 		Set("StateTreeEditor.Transitions", new IMAGE_BRUSH_SVG("Icons/Transitions", CoreStyleConstants::Icon16x16));
 
