StateTree: Fixed binding to transition event payload from property functions. (contributed)

diff --git a/Engine/Plugins/Runtime/StateTree/Source/StateTreeEditorModule/Private/StateTreeEditorData.cpp b/Engine/Plugins/Runtime/StateTree/Source/StateTreeEditorModule/Private/StateTreeEditorData.cpp
--- a/Engine/Plugins/Runtime/StateTree/Source/StateTreeEditorModule/Private/StateTreeEditorData.cpp
+++ b/Engine/Plugins/Runtime/StateTree/Source/StateTreeEditorModule/Private/StateTreeEditorData.cpp
@@ -318,6 +318,12 @@ void UStateTreeEditorData::GetAccessibleStructs(const TConstArrayView<const USta
 	const UStateTree* StateTree = GetTypedOuter<UStateTree>();
 	checkf(StateTree, TEXT("UStateTreeEditorData should only be allocated within a UStateTree"));
 
+	FStateTreeBindableStructDesc TargetStructDesc;
+	bool bIsTargetPropertyFunction = false;
+	if (GetStructByID(TargetStructID, TargetStructDesc))
+	{
+		bIsTargetPropertyFunction = TargetStructDesc.DataSource == EStateTreeBindableStructSource::PropertyFunction;
+	}
 
 	EStateTreeVisitor BaseProgress = VisitGlobalNodes([&OutStructDescs, TargetStructID]
 		(const UStateTreeState* State, const FStateTreeBindableStructDesc& Desc, const FStateTreeDataView Value)
@@ -343,7 +349,7 @@ void UStateTreeEditorData::GetAccessibleStructs(const TConstArrayView<const USta
 				continue;
 			}
 			
-			const EStateTreeVisitor StateProgress = VisitStateNodes(*State, [&OutStructDescs, &BindableDescs, &Path, TargetStructID, this]
+			const EStateTreeVisitor StateProgress = VisitStateNodes(*State, [&OutStructDescs, &BindableDescs, &Path, TargetStructID, bIsTargetPropertyFunction, this]
 				(const UStateTreeState* State, const FStateTreeBindableStructDesc& Desc, const FStateTreeDataView Value)
 				{
 					// Stop iterating as soon as we find the target node.
@@ -372,9 +378,10 @@ void UStateTreeEditorData::GetAccessibleStructs(const TConstArrayView<const USta
 									bool bFoundOwningTransition = false;
 									for (const FStateTreeEditorNode& ConditionNode : Transition.Conditions)
 									{
-										if (ConditionNode.ID == TargetStructID)
+										if (ConditionNode.ID == TargetStructID
+											|| (bIsTargetPropertyFunction && UE::StateTree::Editor::IsPropertyFunctionOwnedByNode(ConditionNode.ID, TargetStructID, EditorBindings)))
 										{
-											if(Transition.GetEventID() == Desc.ID)
+											if (Transition.GetEventID() == Desc.ID)
 											{
 												BindableDescs.Add(Desc);
 											}
