StateTree: Fixed calculating hash for DataHandle (contributed)

diff --git a/Engine/Plugins/Runtime/StateTree/Source/StateTreeModule/Public/StateTreeTypes.h b/Engine/Plugins/Runtime/StateTree/Source/StateTreeModule/Public/StateTreeTypes.h
--- a/Engine/Plugins/Runtime/StateTree/Source/StateTreeModule/Public/StateTreeTypes.h
+++ b/Engine/Plugins/Runtime/StateTree/Source/StateTreeModule/Public/StateTreeTypes.h
@@ -325,7 +325,7 @@ struct STATETREEMODULE_API FStateTreeDataHandle
 
 	friend FORCEINLINE uint32 GetTypeHash(const FStateTreeDataHandle& Handle)
 	{
-		uint32 Hash = GetTypeHash(Handle.Source);
+		uint32 Hash = GetTypeHash(Handle.Index);
 		Hash = HashCombineFast(Hash, GetTypeHash(Handle.Source));
 		Hash = HashCombineFast(Hash, GetTypeHash(Handle.StateHandle));
 		return Hash;
