StateTree: Fixed providing incorrect instance data during bindings compilation. (contributed)

diff --git a/Engine/Plugins/Runtime/StateTree/Source/StateTreeEditorModule/Private/StateTreeCompiler.cpp b/Engine/Plugins/Runtime/StateTree/Source/StateTreeEditorModule/Private/StateTreeCompiler.cpp
--- a/Engine/Plugins/Runtime/StateTree/Source/StateTreeEditorModule/Private/StateTreeCompiler.cpp
+++ b/Engine/Plugins/Runtime/StateTree/Source/StateTreeEditorModule/Private/StateTreeCompiler.cpp
@@ -529,7 +529,7 @@ bool FStateTreeCompiler::CreateEvaluators()
 	}
 	StateTree->EvaluatorsNum = uint16(EvaluatorsNum);
 
-	return CreateBindingsForNodes(EditorData->Evaluators, FStateTreeIndex16(EvaluatorsBegin));
+	return CreateBindingsForNodes(EditorData->Evaluators, FStateTreeIndex16(EvaluatorsBegin), InstanceStructs);
 }
 
 bool FStateTreeCompiler::CreateGlobalTasks()
@@ -574,7 +574,7 @@ bool FStateTreeCompiler::CreateGlobalTasks()
 	}
 	StateTree->GlobalTasksNum = uint16(GlobalTasksNum);
 
-	return CreateBindingsForNodes(EditorData->GlobalTasks, FStateTreeIndex16(GlobalTasksBegin));
+	return CreateBindingsForNodes(EditorData->GlobalTasks, FStateTreeIndex16(GlobalTasksBegin), InstanceStructs);
 }
 
 bool FStateTreeCompiler::CreateStateTasksAndParameters()
@@ -767,7 +767,7 @@ bool FStateTreeCompiler::CreateStateTasksAndParameters()
 		CompactState.TasksNum = uint8(TasksNum);
 		CompactState.InstanceDataNum = uint8(InstanceDataNum);
 
-		if (!CreateBindingsForNodes(Tasks, FStateTreeIndex16(TasksBegin)))
+		if (!CreateBindingsForNodes(Tasks, FStateTreeIndex16(TasksBegin), InstanceStructs))
 		{
 			return false;
 		}
@@ -815,7 +815,7 @@ bool FStateTreeCompiler::CreateStateTransitions()
 		}
 		CompactState.EnterConditionsNum = uint8(EnterConditionsNum);
 
-		if (!CreateBindingsForNodes(SourceState->EnterConditions, FStateTreeIndex16(EnterConditionsBegin)))
+		if (!CreateBindingsForNodes(SourceState->EnterConditions, FStateTreeIndex16(EnterConditionsBegin), SharedInstanceStructs))
 		{
 			return false;
 		}
@@ -1046,7 +1046,7 @@ bool FStateTreeCompiler::CreateStateTransitions()
 			}
 			CompactTransition.ConditionsNum = uint8(ConditionsNum);
 
-			if (!CreateBindingsForNodes(Transition.Conditions, FStateTreeIndex16(ConditionsBegin)))
+			if (!CreateBindingsForNodes(Transition.Conditions, FStateTreeIndex16(ConditionsBegin), SharedInstanceStructs))
 			{
 				return false;
 			}
@@ -1104,7 +1104,7 @@ bool FStateTreeCompiler::CreateStateConsiderations()
 		}
 		CompactState.UtilityConsiderationsNum = uint8(UtilityConsiderationsNum);
 
-		if (!CreateBindingsForNodes(SourceState->Considerations, FStateTreeIndex16(UtilityConsiderationsBegin)))
+		if (!CreateBindingsForNodes(SourceState->Considerations, FStateTreeIndex16(UtilityConsiderationsBegin), SharedInstanceStructs))
 		{
 			return false;
 		}
@@ -1113,7 +1113,7 @@ bool FStateTreeCompiler::CreateStateConsiderations()
 	return true;
 }
 
-bool FStateTreeCompiler::CreateBindingsForNodes(TConstArrayView<FStateTreeEditorNode> EditorNodes, FStateTreeIndex16 NodesBegin)
+bool FStateTreeCompiler::CreateBindingsForNodes(TConstArrayView<FStateTreeEditorNode> EditorNodes, FStateTreeIndex16 NodesBegin, TArrayView<FInstancedStruct> Instances)
 {
 	check(NodesBegin.IsValid());
 
@@ -1144,7 +1144,22 @@ bool FStateTreeCompiler::CreateBindingsForNodes(TConstArrayView<FStateTreeEditor
 			PropertyFunctionsEnd = FStateTreeIndex16::Invalid;
 		}
 
-		if (!CreateBindingsForStruct(*BindableStruct, EditorNode.GetInstance(), PropertyFunctionsBegin, PropertyFunctionsEnd, Node.BindingsBatch))
+		FStateTreeDataView InstanceView;
+		check(Instances.IsValidIndex(Node.InstanceTemplateIndex.Get()));
+
+		FInstancedStruct& Instance = Instances[Node.InstanceTemplateIndex.Get()];
+		if (FStateTreeInstanceObjectWrapper* ObjectWrapper = Instance.GetMutablePtr<FStateTreeInstanceObjectWrapper>())
+		{
+			check(EditorNode.InstanceObject->GetClass() == ObjectWrapper->InstanceObject->GetClass());
+			InstanceView = FStateTreeDataView(ObjectWrapper->InstanceObject);
+		}
+		else
+		{
+			check(EditorNode.Instance.GetScriptStruct() == Instance.GetScriptStruct());
+			InstanceView = FStateTreeDataView(Instance);
+		}
+
+		if (!CreateBindingsForStruct(*BindableStruct, InstanceView, PropertyFunctionsBegin, PropertyFunctionsEnd, Node.BindingsBatch))
 		{
 			return false;
 		}
@@ -1229,7 +1244,22 @@ bool FStateTreeCompiler::CreatePropertyFunction(const FStateTreeEditorNode& Func
 	const FStateTreeBindableStructDesc* BindableStruct = BindingsCompiler.GetSourceStructDescByID(FuncEditorNode.ID);
 	check(BindableStruct);
 
-	return CreateBindingsForStruct(*BindableStruct, FuncEditorNode.GetInstance(), FStateTreeIndex16::Invalid, FStateTreeIndex16::Invalid, Node->BindingsBatch);
+	FStateTreeDataView InstanceView;
+	check(SharedInstanceStructs.IsValidIndex(Node->InstanceTemplateIndex.Get()));
+
+	FInstancedStruct& Instance = SharedInstanceStructs[Node->InstanceTemplateIndex.Get()];
+	if (FStateTreeInstanceObjectWrapper* ObjectWrapper = Instance.GetMutablePtr<FStateTreeInstanceObjectWrapper>())
+	{
+		check(FuncEditorNode.InstanceObject->GetClass() == ObjectWrapper->InstanceObject->GetClass());
+		InstanceView = FStateTreeDataView(ObjectWrapper->InstanceObject);
+	}
+	else
+	{
+		check(FuncEditorNode.Instance.GetScriptStruct() == Instance.GetScriptStruct());
+		InstanceView = FStateTreeDataView(Instance);
+	}
+
+	return CreateBindingsForStruct(*BindableStruct, InstanceView, FStateTreeIndex16::Invalid, FStateTreeIndex16::Invalid, Node->BindingsBatch);
 }
 
 template<class T>
diff --git a/Engine/Plugins/Runtime/StateTree/Source/StateTreeEditorModule/Public/StateTreeCompiler.h b/Engine/Plugins/Runtime/StateTree/Source/StateTreeEditorModule/Public/StateTreeCompiler.h
--- a/Engine/Plugins/Runtime/StateTree/Source/StateTreeEditorModule/Public/StateTreeCompiler.h
+++ b/Engine/Plugins/Runtime/StateTree/Source/StateTreeEditorModule/Public/StateTreeCompiler.h
@@ -50,7 +50,7 @@ private:
 	bool CreateStateTransitions();
 	bool CreateStateConsiderations();
 	
-	bool CreateBindingsForNodes(TConstArrayView<FStateTreeEditorNode> EditorNodes, FStateTreeIndex16 NodesBegin);
+	bool CreateBindingsForNodes(TConstArrayView<FStateTreeEditorNode> EditorNodes, FStateTreeIndex16 NodesBegin, TArrayView<FInstancedStruct> Instances);
 	bool CreateBindingsForStruct(const FStateTreeBindableStructDesc& TargetStruct, FStateTreeDataView TargetValue, FStateTreeIndex16 PropertyFuncsBegin, FStateTreeIndex16 PropertyFuncsEnd, FStateTreeIndex16& OutBatchIndex);
 
 	bool CreatePropertyFunctionsForStruct(FGuid StructID);
