StateTree: Fixed crashes when delayed tick transition has been used (contributed)

diff --git a/Engine/Plugins/Runtime/StateTree/Source/StateTreeModule/Private/StateTreeEvents.cpp b/Engine/Plugins/Runtime/StateTree/Source/StateTreeModule/Private/StateTreeEvents.cpp
--- a/Engine/Plugins/Runtime/StateTree/Source/StateTreeModule/Private/StateTreeEvents.cpp
+++ b/Engine/Plugins/Runtime/StateTree/Source/StateTreeModule/Private/StateTreeEvents.cpp
@@ -12,7 +12,10 @@
 
 void FStateTreeSharedEvent::AddStructReferencedObjects(FReferenceCollector& Collector)
 {
-	Collector.AddPropertyReferencesWithStructARO(FStateTreeEvent::StaticStruct(), Event.Get());
+	if (IsValid())
+	{
+		Collector.AddPropertyReferencesWithStructARO(FStateTreeEvent::StaticStruct(), Event.Get());
+	}
 }
 
 
diff --git a/Engine/Plugins/Runtime/StateTree/Source/StateTreeModule/Private/StateTreeExecutionContext.cpp b/Engine/Plugins/Runtime/StateTree/Source/StateTreeModule/Private/StateTreeExecutionContext.cpp
--- a/Engine/Plugins/Runtime/StateTree/Source/StateTreeModule/Private/StateTreeExecutionContext.cpp
+++ b/Engine/Plugins/Runtime/StateTree/Source/StateTreeModule/Private/StateTreeExecutionContext.cpp
@@ -3420,8 +3420,6 @@ bool FStateTreeExecutionContext::SelectState(const FStateTreeExecutionFrame& Cur
 											const FStateTreeSharedEvent* TransitionEvent,
 											const EStateTreeSelectionFallback Fallback)
 {
-	check(TransitionEvent == nullptr || TransitionEvent->IsValid());
-
 	const FStateTreeExecutionState& Exec = GetExecState();
 
 	if (Exec.ActiveFrames.IsEmpty())
