StateTree: Added StateScope when checking state's enter conditions (contributed)

diff --git a/Engine/Plugins/Runtime/StateTree/Source/StateTreeModule/Private/StateTreeExecutionContext.cpp b/Engine/Plugins/Runtime/StateTree/Source/StateTreeModule/Private/StateTreeExecutionContext.cpp
--- a/Engine/Plugins/Runtime/StateTree/Source/StateTreeModule/Private/StateTreeExecutionContext.cpp
+++ b/Engine/Plugins/Runtime/StateTree/Source/StateTreeModule/Private/StateTreeExecutionContext.cpp
@@ -3593,6 +3593,7 @@ bool FStateTreeExecutionContext::SelectStateInternal(
 	}
 
 	FCurrentlyProcessedFrameScope FrameScope(*this, CurrentParentFrame, CurrentFrame);
+	FCurrentlyProcessedStateScope NextStateScope(*this, NextStateHandle);
 	FCurrentFrameStateSelectionEventsScope CapturedEventsScope(*this, OutSelectionResult.GetFramesStateSelectionEvents().Last());
 
 	const UStateTree* CurrentStateTree = CurrentFrame.StateTree;
