StateTree: Fixed reentering states in parent frames of the performed transition. (contributed)

diff --git a/Engine/Plugins/Runtime/StateTree/Source/StateTreeModule/Private/StateTreeExecutionContext.cpp b/Engine/Plugins/Runtime/StateTree/Source/StateTreeModule/Private/StateTreeExecutionContext.cpp
--- a/Engine/Plugins/Runtime/StateTree/Source/StateTreeModule/Private/StateTreeExecutionContext.cpp
+++ b/Engine/Plugins/Runtime/StateTree/Source/StateTreeModule/Private/StateTreeExecutionContext.cpp
@@ -1624,7 +1624,7 @@ EStateTreeRunStatus FStateTreeExecutionContext::EnterState(FStateTreeTransitionR
 				}
 			}
 
-			bOnTargetBranch = bOnTargetBranch || CurrentHandle == Transition.TargetState;
+			bOnTargetBranch |= CurrentFrame.StateTree == Transition.SourceStateTree && CurrentFrame.RootState == Transition.SourceRootState && CurrentHandle == Transition.TargetState;
 			const bool bWasActive = PreviousHandle == CurrentHandle;
 
 			// Do not enter a disabled State tasks but maintain property bindings
