StateTree: Added State outliner and made the main tree view more detailed
- Simplied the task icon
- Created common treeview for States, used for state selection and outliner view
- Added common functions to describe state links and transitions
- Made custom expander arrow the main tree view
- Move the task lits on separate row on the main tree view
- Added enter condition list for on main tree view
- Clicking on tasks or enter conditions highlight them in the details panel
- Added state outliner which displays the state list compactly
- Added feedback for the drag and drop, and moved the code to separate file

diff --git a/Engine/Plugins/Runtime/StateTree/Source/StateTreeEditorModule/Private/Customizations/StateTreeStateLinkDetails.cpp b/Engine/Plugins/Runtime/StateTree/Source/StateTreeEditorModule/Private/Customizations/StateTreeStateLinkDetails.cpp
--- a/Engine/Plugins/Runtime/StateTree/Source/StateTreeEditorModule/Private/Customizations/StateTreeStateLinkDetails.cpp
+++ b/Engine/Plugins/Runtime/StateTree/Source/StateTreeEditorModule/Private/Customizations/StateTreeStateLinkDetails.cpp
@@ -5,22 +5,19 @@
 #include "DetailLayoutBuilder.h"
 #include "StateTree.h"
 #include "StateTreeEditorData.h"
-#include "StateTreeDelegates.h"
 #include "Widgets/Input/SComboButton.h"
 #include "Widgets/Images/SImage.h"
 #include "Framework/MultiBox/MultiBoxBuilder.h"
 #include "ScopedTransaction.h"
+#include "StateTreeDescriptionHelpers.h"
 #include "StateTreeEditorStyle.h"
 #include "StateTreePropertyHelpers.h"
 #include "TextStyleDecorator.h"
-#include "Widgets/Input/SSearchBox.h"
+#include "Widgets/SCompactStateTreeView.h"
 #include "Widgets/Text/SRichTextBlock.h"
 
 #define LOCTEXT_NAMESPACE "StateTreeEditor"
 
-TMap<FObjectKey, FStateTreeStateLinkDetails::FStateExpansionState> FStateTreeStateLinkDetails::StateExpansionStates;
-
-
 TSharedRef<IPropertyTypeCustomization> FStateTreeStateLinkDetails::MakeInstance()
 {
 	return MakeShareable(new FStateTreeStateLinkDetails);
@@ -43,9 +40,25 @@ void FStateTreeStateLinkDetails::CustomizeHeader(TSharedRef<class IPropertyHandl
 		bDirectStatesOnly = MetaDataProperty->HasMetaData(NAME_DirectStatesOnly);
 		bSubtreesOnly = MetaDataProperty->HasMetaData(NAME_SubtreesOnly);
 	}
-	
-	CacheStates();
 
+	// Store pointer to editor data.
+	TArray<UObject*> OuterObjects;
+	StructProperty->GetOuterObjects(OuterObjects);
+	for (const UObject* Object : OuterObjects)
+	{
+		if (Object)
+		{
+			if (const UStateTree* OuterStateTree = Object->GetTypedOuter<UStateTree>())
+			{
+				WeakEditorData = Cast<UStateTreeEditorData>(OuterStateTree->EditorData);
+				if (WeakEditorData.IsValid())
+				{
+					break;
+				}
+			}
+		}
+	}
+	
 	HeaderRow
 		.NameContent()
 		[
@@ -101,560 +114,252 @@ void FStateTreeStateLinkDetails::CustomizeHeader(TSharedRef<class IPropertyHandl
 				]
 			]
 		];
-
-	UE::StateTree::Delegates::OnIdentifierChanged.AddSP(this, &FStateTreeStateLinkDetails::OnIdentifierChanged);
 }
 
 void FStateTreeStateLinkDetails::CustomizeChildren(TSharedRef<class IPropertyHandle> StructPropertyHandle, class IDetailChildrenBuilder& StructBuilder, IPropertyTypeCustomizationUtils& StructCustomizationUtils)
 {
 }
 
-void FStateTreeStateLinkDetails::OnIdentifierChanged(const UStateTree& StateTree)
-{
-	CacheStates();
-}
-
-void FStateTreeStateLinkDetails::CacheStates(TSharedPtr<FStateTreeStateItem> ParentNode, const UStateTreeState* State)
+TSharedRef<SWidget> FStateTreeStateLinkDetails::GenerateStatePicker()
 {
-	if (State == nullptr)
-	{
-		return;
-	}
-
-	bool bShouldAdd = true;
-	if (bSubtreesOnly && State->Type != EStateTreeStateType::Subtree)
-	{
-		bShouldAdd = false;
-	}
-
-	if (State->SelectionBehavior == EStateTreeStateSelectionBehavior::None)
-	{
-		bShouldAdd = false;
-	}
+	check(ComboButton);
 	
-	if (bShouldAdd)
-	{
-		TSharedRef<FStateTreeStateItem> StateItem = MakeShared<FStateTreeStateItem>();
-		StateItem->Desc = FText::FromName(State->Name);
-		StateItem->TransitionType = EStateTreeTransitionType::GotoState;
-		StateItem->StateID = State->ID;
-		StateItem->Color = FLinearColor(1.f, 1.f, 1.f, 0.25f);
-		StateItem->bIsSubTree = State->Type == EStateTreeStateType::Subtree;
-
-		// Figure out icon.
-		if (State->SelectionBehavior == EStateTreeStateSelectionBehavior::None)
-		{
-			StateItem->Icon = FStateTreeEditorStyle::Get().GetBrush("StateTreeEditor.SelectNone");
-		}
-		else if (State->SelectionBehavior == EStateTreeStateSelectionBehavior::TryEnterState)
-		{
-			StateItem->Icon = FStateTreeEditorStyle::Get().GetBrush("StateTreeEditor.TryEnterState");			
-		}
-		else if (State->SelectionBehavior == EStateTreeStateSelectionBehavior::TrySelectChildrenInOrder)
-		{
-			if (State->Children.IsEmpty()
-				|| State->Type == EStateTreeStateType::Linked
-				|| State->Type == EStateTreeStateType::LinkedAsset)
-			{
-				StateItem->Icon = FStateTreeEditorStyle::Get().GetBrush("StateTreeEditor.TryEnterState");			
-			}
-			else
-			{
-				StateItem->Icon = FStateTreeEditorStyle::Get().GetBrush("StateTreeEditor.TrySelectChildrenInOrder");
-			}
-		}
-		else if (State->SelectionBehavior == EStateTreeStateSelectionBehavior::TryFollowTransitions)
-		{
-			StateItem->Icon = FStateTreeEditorStyle::Get().GetBrush("StateTreeEditor.TryFollowTransitions");
-		}
-
-		// Linked states
-		if (State->Type == EStateTreeStateType::Linked)
-		{
-			StateItem->bIsLinked = true;
-			StateItem->LinkedDesc = FText::FromName(State->LinkedSubtree.Name);
-		}
-		else if (State->Type == EStateTreeStateType::LinkedAsset)
-		{
-			StateItem->bIsLinked = true;
-			StateItem->LinkedDesc = FText::FromString(GetNameSafe(State->LinkedAsset.Get()));
-		}
-		
-		ParentNode->Children.Add(StateItem);
-
-		ParentNode = StateItem;
-	}
-
-	for (UStateTreeState* ChildState : State->Children)
-	{
-		CacheStates(ParentNode, ChildState);
-	}
-}
-
-void FStateTreeStateLinkDetails::CacheStates()
-{
-	RootItem = MakeShared<FStateTreeStateItem>();
+	constexpr bool bCloseMenuAfterSelection = true;
+	FMenuBuilder MenuBuilder(bCloseMenuAfterSelection, nullptr);
 
 	if (!bDirectStatesOnly)
 	{
-		RootItem->Children.Add(MakeShared<FStateTreeStateItem>(
-			LOCTEXT("TransitionNoneRich", "<i>None</>"),
+		auto MakeMetaStateWidget = [WeakEditorData = WeakEditorData](EStateTreeTransitionType Type)
+		{
+			const FStateTreeStateLink Link(Type);
+			
+			return SNew(SHorizontalBox)
+				+SHorizontalBox::Slot()
+				.HAlign(HAlign_Left)
+				.VAlign(VAlign_Center)
+				.Padding(0, 2.0f, 4.0f, 2.0f)
+				.AutoWidth()
+				[
+					SNew(SImage)
+					.DesiredSizeOverride(FVector2D(16.0f, 16.0f))
+					.Image(UE::StateTree::Editor::GetStateLinkIcon(WeakEditorData.Get(), Link))
+					.ColorAndOpacity(UE::StateTree::Editor::GetStateLinkColor(WeakEditorData.Get(), Link))
+				]
+
+				+SHorizontalBox::Slot()
+				.AutoWidth()
+				.VAlign(VAlign_Center)
+				[
+					SNew(SRichTextBlock)
+					.Text(UE::StateTree::Editor::GetStateLinkDesc(WeakEditorData.Get(), Link, EStateTreeNodeFormatting::RichText))
+					.TextStyle(&FStateTreeEditorStyle::Get().GetWidgetStyle<FTextBlockStyle>("Normal.Normal"))
+					.OverflowPolicy(ETextOverflowPolicy::Ellipsis)
+					+SRichTextBlock::Decorator(FTextStyleDecorator::Create(TEXT(""), FStateTreeEditorStyle::Get().GetWidgetStyle<FTextBlockStyle>("Normal.Normal")))
+					+SRichTextBlock::Decorator(FTextStyleDecorator::Create(TEXT("b"), FStateTreeEditorStyle::Get().GetWidgetStyle<FTextBlockStyle>("Normal.Bold")))
+					+SRichTextBlock::Decorator(FTextStyleDecorator::Create(TEXT("i"), FStateTreeEditorStyle::Get().GetWidgetStyle<FTextBlockStyle>("Normal.Italic")))
+					+SRichTextBlock::Decorator(FTextStyleDecorator::Create(TEXT("s"), FStateTreeEditorStyle::Get().GetWidgetStyle<FTextBlockStyle>("Normal.Subdued")))
+				];
+		};
+
+		MenuBuilder.AddMenuEntry(
+			FUIAction(
+				FExecuteAction::CreateSP(this, &FStateTreeStateLinkDetails::SetTransitionByType, EStateTreeTransitionType::None),
+				FCanExecuteAction(),
+				FIsActionChecked::CreateSPLambda(this, [this]
+				{
+					return GetTransitionType() == EStateTreeTransitionType::None; 
+				})),
+			MakeMetaStateWidget(EStateTreeTransitionType::None),
+			FName(),
 			LOCTEXT("TransitionNoneTooltip", "No transition."),
-			FStateTreeEditorStyle::Get().GetBrush("StateTreeEditor.Transition.None"),
-			EStateTreeTransitionType::None));
+			EUserInterfaceActionType::Check);
 
-		RootItem->Children.Add(MakeShared<FStateTreeStateItem>(
-			LOCTEXT("TransitionNextStateRich", "<i>Next State</>"),
+		MenuBuilder.AddMenuEntry(
+			FUIAction(
+				FExecuteAction::CreateSP(this, &FStateTreeStateLinkDetails::SetTransitionByType, EStateTreeTransitionType::NextState),
+				FCanExecuteAction(),
+				FIsActionChecked::CreateSPLambda(this, [this]
+				{
+					return GetTransitionType() == EStateTreeTransitionType::NextState; 
+				})),
+			MakeMetaStateWidget(EStateTreeTransitionType::NextState),
+			FName(),
 			LOCTEXT("TransitionNextTooltip", "Goto next sibling State."),
-			FStateTreeEditorStyle::Get().GetBrush("StateTreeEditor.Transition.Next"),
-			EStateTreeTransitionType::NextState));
+			EUserInterfaceActionType::Check);
 
-		RootItem->Children.Add(MakeShared<FStateTreeStateItem>(
-			LOCTEXT("TransitionNextSelectableStateRich", "<i>Next Selectable State</>"),
+		MenuBuilder.AddMenuEntry(
+			FUIAction(
+				FExecuteAction::CreateSP(this, &FStateTreeStateLinkDetails::SetTransitionByType, EStateTreeTransitionType::NextSelectableState),
+				FCanExecuteAction(),
+				FIsActionChecked::CreateSPLambda(this, [this]
+				{
+					return GetTransitionType() == EStateTreeTransitionType::NextSelectableState; 
+			})),
+			MakeMetaStateWidget(EStateTreeTransitionType::NextSelectableState),
+			FName(),
 			LOCTEXT("TransitionNextSelectableTooltip", "Goto next sibling state, whose enter conditions pass."),
-			FStateTreeEditorStyle::Get().GetBrush("StateTreeEditor.Transition.Next"),
-			EStateTreeTransitionType::NextSelectableState));
+			EUserInterfaceActionType::Check);
 
-		RootItem->Children.Add(MakeShared<FStateTreeStateItem>(
-			LOCTEXT("TransitionTreeSucceededRich", "<i>Tree Succeeded</>"),
+		MenuBuilder.AddMenuEntry(
+			FUIAction(
+				FExecuteAction::CreateSP(this, &FStateTreeStateLinkDetails::SetTransitionByType, EStateTreeTransitionType::Succeeded),
+				FCanExecuteAction(),
+				FIsActionChecked::CreateSPLambda(this, [this]
+				{
+					return GetTransitionType() == EStateTreeTransitionType::Succeeded; 
+				})),
+			MakeMetaStateWidget(EStateTreeTransitionType::Succeeded),
+			FName(),
 			LOCTEXT("TransitionTreeSuccessTooltip", "Complete tree with success."),
-			FStateTreeEditorStyle::Get().GetBrush("StateTreeEditor.Transition.Succeeded"),
-			EStateTreeTransitionType::Succeeded));
+			EUserInterfaceActionType::Check);
 
-		RootItem->Children.Add(MakeShared<FStateTreeStateItem>(
-			LOCTEXT("TransitionTreeFailedRich", "<i>Tree Failed</>"),
-			LOCTEXT("TransitionTreeFailedTooltip", "Complete tree with failure."),
-			FStateTreeEditorStyle::Get().GetBrush("StateTreeEditor.Transition.Failed"),
-			EStateTreeTransitionType::Failed));
-	}
-	
-	TArray<UObject*> OuterObjects;
-	StructProperty->GetOuterObjects(OuterObjects);
-	for (int32 ObjectIdx = 0; ObjectIdx < OuterObjects.Num(); ObjectIdx++)
-	{
-		if (const UStateTree* OuterStateTree = OuterObjects[ObjectIdx]->GetTypedOuter<UStateTree>())
-		{
-			if (const UStateTreeEditorData* TreeData = Cast<UStateTreeEditorData>(OuterStateTree->EditorData))
-			{
-				WeakStateTree = OuterStateTree;
-				for (const UStateTreeState* SubTree : TreeData->SubTrees)
+		MenuBuilder.AddMenuEntry(
+			FUIAction(
+				FExecuteAction::CreateSP(this, &FStateTreeStateLinkDetails::SetTransitionByType, EStateTreeTransitionType::Failed),
+				FCanExecuteAction(),
+				FIsActionChecked::CreateSPLambda(this, [this]
 				{
-					CacheStates(RootItem, SubTree);
-				}
-			}
-			break;
-		}
+					return GetTransitionType() == EStateTreeTransitionType::Failed; 
+				})),
+			MakeMetaStateWidget(EStateTreeTransitionType::Failed),
+			FName(),
+			LOCTEXT("TransitionTreeFailedTooltip", "Complete tree with failure."),
+			EUserInterfaceActionType::Check);
 	}
 
-	FilteredRootItem = RootItem;
-}
+	MenuBuilder.BeginSection("States", LOCTEXT("States", "States"));
 
-TSharedRef<SWidget> FStateTreeStateLinkDetails::GenerateStatePicker()
-{
-	check(ComboButton);
+	TSharedPtr<SCompactStateTreeView> StateView;
 	
-	CacheStates();
-
 	TSharedRef<SWidget> MenuWidget = 
 		SNew(SBox)
 		.MinDesiredWidth(300.f)
 		.MaxDesiredHeight(400.f)
 		.Padding(2)	
 		[
-			SNew(SVerticalBox)
-			+ SVerticalBox::Slot()
-			.HAlign(HAlign_Fill)
-			.VAlign(VAlign_Top)
-			.Padding(4, 2, 4, 2)
-			.AutoHeight()
-			[
-				SAssignNew(SearchBox, SSearchBox)
-				.OnTextChanged(this, &FStateTreeStateLinkDetails::OnSearchBoxTextChanged)
-			]
-			+ SVerticalBox::Slot()
-			[
-				SAssignNew(StateItemTree, STreeView<TSharedPtr<FStateTreeStateItem>>)
-				.SelectionMode(ESelectionMode::Single)
-				.TreeItemsSource(&FilteredRootItem->Children)
-				.OnGenerateRow(this, &FStateTreeStateLinkDetails::GenerateStateItemRow)
-				.OnGetChildren(this, &FStateTreeStateLinkDetails::GetStateItemChildren)
-				.OnSelectionChanged(this, &FStateTreeStateLinkDetails::OnStateItemSelected)
-				.OnExpansionChanged(this, &FStateTreeStateLinkDetails::OnStateItemExpansionChanged)
-			]
+			SAssignNew(StateView, SCompactStateTreeView)
+			.StateTreeEditorData(WeakEditorData.Get())
+			.SelectionMode(ESelectionMode::Single)
+			.SelectableStatesOnly(true)
+			.SubtreesOnly(bSubtreesOnly)
+			.OnSelectionChanged(this, &FStateTreeStateLinkDetails::OnStateSelected)
 		];
 
-	// Restore category expansion state from previous use.
-	RestoreExpansionState();
+	check(StateView);
 	
-	// Expand and select currently selected item.
-	TArray<TSharedPtr<FStateTreeStateItem>> Path;
-	if (GetCurrentStateItem(Path))
+	const EStateTreeTransitionType TransitionType = GetTransitionType().Get(EStateTreeTransitionType::Failed);
+	if (TransitionType == EStateTreeTransitionType::GotoState)
 	{
-		// Expand all categories up to the selected item.
-		bIsRestoringExpansion = true;
-		for (int32 Index = 0; Index < Path.Num() - 1; Index++)
+		if (const UStateTreeState* State = GetState())
 		{
-			StateItemTree->SetItemExpansion(Path[Index], true);
+			StateView->SetSelection({ State->ID });
 		}
-		bIsRestoringExpansion = false;
-		
-		StateItemTree->SetItemSelection(Path.Last(), true);
-		StateItemTree->RequestScrollIntoView(Path.Last());
 	}
-	
-	ComboButton->SetMenuContentWidgetToFocus(SearchBox);
 
-	return MenuWidget;
-}
-
-TSharedRef<ITableRow> FStateTreeStateLinkDetails::GenerateStateItemRow(TSharedPtr<FStateTreeStateItem> Item, const TSharedRef<STableViewBase>& OwnerTable)
-{
-	TSharedRef<SHorizontalBox> Container = SNew(SHorizontalBox);
+	ComboButton->SetMenuContentWidgetToFocus(StateView->GetWidgetToFocusOnOpen());
 
-	// Icon
-	Container->AddSlot()
-		.HAlign(HAlign_Left)
-		.VAlign(VAlign_Center)
-		.Padding(0, 2.0f, 4.0f, 2.0f)
-		.AutoWidth()
-		[
-			SNew(SImage)
-			.Visibility(Item->Icon ? EVisibility::Visible : EVisibility::Collapsed)
-			.DesiredSizeOverride(FVector2D(16.0f, 16.0f))
-			.Image(Item->Icon)
-			.ColorAndOpacity(Item->Color)
-		];
-
-	// Name
-	Container->AddSlot()
-		.HAlign(HAlign_Fill)
-		.VAlign(VAlign_Center)
-		.AutoWidth()
-		[
-			SNew(SRichTextBlock)
-			.Text(Item->Desc)
-			.TextStyle(&FStateTreeEditorStyle::Get().GetWidgetStyle<FTextBlockStyle>("Normal.Normal"))
-			.OverflowPolicy(ETextOverflowPolicy::Ellipsis)
-			.ToolTipText(Item->TooltipText)
-			.HighlightText_Lambda([this]() { return SearchBox.IsValid() ? SearchBox->GetText() : FText::GetEmpty(); })
-			+SRichTextBlock::Decorator(FTextStyleDecorator::Create(TEXT(""), FStateTreeEditorStyle::Get().GetWidgetStyle<FTextBlockStyle>("Normal.Normal")))
-			+SRichTextBlock::Decorator(FTextStyleDecorator::Create(TEXT("b"), FStateTreeEditorStyle::Get().GetWidgetStyle<FTextBlockStyle>("Normal.Bold")))
-			+SRichTextBlock::Decorator(FTextStyleDecorator::Create(TEXT("i"), FStateTreeEditorStyle::Get().GetWidgetStyle<FTextBlockStyle>("Normal.Italic")))
-			+SRichTextBlock::Decorator(FTextStyleDecorator::Create(TEXT("s"), FStateTreeEditorStyle::Get().GetWidgetStyle<FTextBlockStyle>("Normal.Subdued")))
-		];
-
-	// Link
-	if (Item->bIsLinked)
-	{
-		// Link icon
-		Container->AddSlot()
-			.HAlign(HAlign_Left)
-			.VAlign(VAlign_Center)
-			.Padding(4.f, 0.f)
-			.AutoWidth()
-			[
-				SNew(SImage)
-				.ColorAndOpacity(FSlateColor::UseSubduedForeground())
-				.Image(FStateTreeEditorStyle::Get().GetBrush("StateTreeEditor.StateLinked"))
-		];
-
-		// Linked name
-		Container->AddSlot()
-			.HAlign(HAlign_Fill)
-			.VAlign(VAlign_Center)
-			.AutoWidth()
-			[
-				SNew(STextBlock)
-				.Font(FAppStyle::Get().GetFontStyle("PropertyWindow.NormalFont"))
-				.ColorAndOpacity(FSlateColor::UseSubduedForeground())
-				.Text(Item->LinkedDesc)
-			];
-	}
+	MenuBuilder.AddWidget(MenuWidget, FText::GetEmpty(), /*bInNoIndent*/true);
 	
-	return SNew(STableRow<TSharedPtr<FStateTreeStateItem>>, OwnerTable)
-		[
-			Container
-		];
-}
+	MenuBuilder.EndSection();
 
-void FStateTreeStateLinkDetails::GetStateItemChildren(TSharedPtr<FStateTreeStateItem> Item, TArray<TSharedPtr<FStateTreeStateItem>>& OutItems) const
-{
-	if (Item.IsValid())
-	{
-		OutItems = Item->Children;
-	}
+	return MenuBuilder.MakeWidget();
 }
 
-void FStateTreeStateLinkDetails::OnStateItemSelected(TSharedPtr<FStateTreeStateItem> SelectedItem, ESelectInfo::Type Type)
+void FStateTreeStateLinkDetails::OnStateSelected(TConstArrayView<FGuid> SelectedStateIDs)
 {
-	// Skip selection from code
-	if (Type == ESelectInfo::Direct)
+	const UStateTreeState* State = nullptr;
+	if (!SelectedStateIDs.IsEmpty())
 	{
-		return;
+		if (const UStateTreeEditorData* EditorData = WeakEditorData.Get())
+		{
+			State = EditorData->GetStateByID(SelectedStateIDs[0]);
+		}
 	}
-	
-	if (SelectedItem
+
+	if (State
 		&& NameProperty
 		&& IDProperty)
 	{
 		FScopedTransaction Transaction(FText::Format(LOCTEXT("SetPropertyValue", "Set {0}"), StructProperty->GetPropertyDisplayName()));
 
-		LinkTypeProperty->SetValue((uint8)SelectedItem->TransitionType);
+		LinkTypeProperty->SetValue((uint8)EStateTreeTransitionType::GotoState);
 
-		if (SelectedItem->TransitionType == EStateTreeTransitionType::GotoState)
-		{
-			NameProperty->SetValue(FName(SelectedItem->Desc.ToString()), EPropertyValueSetFlags::NotTransactable);
-			UE::StateTree::PropertyHelpers::SetStructValue<FGuid>(IDProperty, SelectedItem->StateID, EPropertyValueSetFlags::NotTransactable);
-		}
-		else
-		{
-			// Clear name and id.
-			NameProperty->SetValue(FName(), EPropertyValueSetFlags::NotTransactable);
-			UE::StateTree::PropertyHelpers::SetStructValue<FGuid>(IDProperty, FGuid(), EPropertyValueSetFlags::NotTransactable);
-		}
+		NameProperty->SetValue(State->Name, EPropertyValueSetFlags::NotTransactable);
+		UE::StateTree::PropertyHelpers::SetStructValue<FGuid>(IDProperty, State->ID, EPropertyValueSetFlags::NotTransactable);
 	}
 
 	check(ComboButton);
 	ComboButton->SetIsOpen(false);
 }
 
-void FStateTreeStateLinkDetails::OnStateItemExpansionChanged(TSharedPtr<FStateTreeStateItem> ExpandedItem, bool bInExpanded) const
-{
-	// Do not save expansion state we're restoring expansion state, or when showing filtered results. 
-	if (bIsRestoringExpansion || FilteredRootItem != RootItem)
-	{
-		return;
-	}
-
-	if (ExpandedItem.IsValid() && ExpandedItem->StateID.IsValid())
-	{
-		FStateExpansionState& ExpansionState = StateExpansionStates.FindOrAdd(FObjectKey(WeakStateTree.Get()));
-		if (bInExpanded)
-		{
-			ExpansionState.CollapsedStates.Remove(ExpandedItem->StateID);
-		}
-		else
-		{
-			ExpansionState.CollapsedStates.Add(ExpandedItem->StateID);
-		}
-	}
-}
-
-void FStateTreeStateLinkDetails::OnSearchBoxTextChanged(const FText& NewText)
-{
-	if (!StateItemTree.IsValid())
-	{
-		return;
-	}
-	
-	FilteredRootItem.Reset();
-
-	TArray<FString> FilterStrings;
-	NewText.ToString().ParseIntoArrayWS(FilterStrings);
-	FilterStrings.RemoveAll([](const FString& String) { return String.IsEmpty(); });
-	
-	if (FilterStrings.IsEmpty())
-	{
-		// Show all when there's no filter string.
-		FilteredRootItem = RootItem;
-		StateItemTree->SetTreeItemsSource(&FilteredRootItem->Children);
-		RestoreExpansionState();
-		StateItemTree->RequestTreeRefresh();
-		return;
-	}
-
-	FilteredRootItem = MakeShared<FStateTreeStateItem>();
-	FilterStateItemChildren(FilterStrings, /*bParentMatches*/false, RootItem->Children, FilteredRootItem->Children);
-
-	StateItemTree->SetTreeItemsSource(&FilteredRootItem->Children);
-	ExpandAll(FilteredRootItem->Children);
-
-	// Update selection based on filtered items.
-	TArray<TSharedPtr<FStateTreeStateItem>> Path;
-	const EStateTreeTransitionType TransitionType = GetTransitionType().Get(EStateTreeTransitionType::Failed);
-	FGuid StateID;
-	if (IDProperty
-		&& UE::StateTree::PropertyHelpers::GetStructValue<FGuid>(IDProperty, StateID) == FPropertyAccess::Success)
-	{
-		// Use RootItem (not filtered) to get any result even during search.
-		if (FindStateByIDRecursive(FilteredRootItem, TransitionType, StateID, Path))
-		{
-			StateItemTree->SetItemSelection(Path.Last(), true);
-		}
-	}
-
-	StateItemTree->RequestTreeRefresh();
-}
-
-
-int32 FStateTreeStateLinkDetails::FilterStateItemChildren(const TArray<FString>& FilterStrings, const bool bParentMatches,
-															const TArray<TSharedPtr<FStateTreeStateItem>>& SourceArray,
-															TArray<TSharedPtr<FStateTreeStateItem>>& OutDestArray)
+void FStateTreeStateLinkDetails::SetTransitionByType(const EStateTreeTransitionType TransitionType)
 {
-	int32 NumFound = 0;
-
-	auto MatchFilter = [&FilterStrings](const TSharedPtr<FStateTreeStateItem>& SourceItem)
-	{
-		const FString ItemName = SourceItem->Desc.ToString();
-		for (const FString& Filter : FilterStrings)
-		{
-			if (ItemName.Contains(Filter))
-			{
-				return true;
-			}
-		}
-		return false;
-	};
-
-	for (const TSharedPtr<FStateTreeStateItem>& SourceItem : SourceArray)
+	if (NameProperty
+		&& IDProperty)
 	{
-		// Check if our name matches the filters
-		// If bParentMatches is true, the search matched a parent category.
-		const bool bMatchesFilters = bParentMatches || MatchFilter(SourceItem);
-
-		int32 NumChildren = 0;
-		if (bMatchesFilters)
-		{
-			NumChildren++;
-		}
-
-		// if we don't match, then we still want to check all our children
-		TArray<TSharedPtr<FStateTreeStateItem>> FilteredChildren;
-		NumChildren += FilterStateItemChildren(FilterStrings, bMatchesFilters, SourceItem->Children, FilteredChildren);
-
-		// then add this item to the destination array
-		if (NumChildren > 0)
-		{
-			TSharedPtr<FStateTreeStateItem>& NewItem = OutDestArray.Add_GetRef(MakeShared<FStateTreeStateItem>());
-			*NewItem = *SourceItem;
-			NewItem->Children = FilteredChildren;
-
-			NumFound += NumChildren;
-		}
-	}
-
-	return NumFound;
-}
-
-void FStateTreeStateLinkDetails::RestoreExpansionState()
-{
-	check(StateItemTree.IsValid());
-	
-	FStateExpansionState& ExpansionState = StateExpansionStates.FindOrAdd(FObjectKey(WeakStateTree.Get()));
-
-	bIsRestoringExpansion = true;
+		FScopedTransaction Transaction(FText::Format(LOCTEXT("SetPropertyValue", "Set {0}"), StructProperty->GetPropertyDisplayName()));
 
-	// Default state is expanded.
-	ExpandAll(FilteredRootItem->Children);
+		LinkTypeProperty->SetValue((uint8)TransitionType);
 
-	// Collapse the ones that are specifically collapsed.
-	for (const FGuid& StateID : ExpansionState.CollapsedStates)
-	{
-		TArray<TSharedPtr<FStateTreeStateItem>> Path;
-		if (FindStateByIDRecursive(FilteredRootItem, EStateTreeTransitionType::GotoState, StateID, Path))
-		{
-			StateItemTree->SetItemExpansion(Path.Last(), false);
-		}
+		// Clear name and id.
+		NameProperty->SetValue(FName(), EPropertyValueSetFlags::NotTransactable);
+		UE::StateTree::PropertyHelpers::SetStructValue<FGuid>(IDProperty, FGuid(), EPropertyValueSetFlags::NotTransactable);
 	}
 
-	bIsRestoringExpansion = false;
-}
-
-void FStateTreeStateLinkDetails::ExpandAll(const TArray<TSharedPtr<FStateTreeStateItem>>& Items)
-{
-	for (const TSharedPtr<FStateTreeStateItem>& Item : Items)
-	{
-		StateItemTree->SetItemExpansion(Item, true);
-		ExpandAll(Item->Children);
-	}
+	check(ComboButton);
+	ComboButton->SetIsOpen(false);
 }
 
-bool FStateTreeStateLinkDetails::FindStateByIDRecursive(const TSharedPtr<FStateTreeStateItem>& Item, const EStateTreeTransitionType TransitionType, const FGuid StateID, TArray<TSharedPtr<FStateTreeStateItem>>& OutPath)
+const UStateTreeState* FStateTreeStateLinkDetails::GetState() const
 {
-	OutPath.Push(Item);
-
-	if (!Item->Desc.IsEmpty()
-		&& Item->TransitionType == TransitionType
-		&& Item->StateID == StateID)
-	{
-		return true;
-	}
-
-	for (const TSharedPtr<FStateTreeStateItem>& ChildItem : Item->Children)
+	if (const UStateTreeEditorData* EditorData = WeakEditorData.Get())
 	{
-		if (FindStateByIDRecursive(ChildItem, TransitionType, StateID, OutPath))
+		FGuid StateID;
+		if (UE::StateTree::PropertyHelpers::GetStructValue<FGuid>(IDProperty, StateID) == FPropertyAccess::Success)
 		{
-			return true;
+			return EditorData->GetStateByID(StateID);
 		}
 	}
-
-	OutPath.Pop();
-
-	return false;
-}
-
-bool FStateTreeStateLinkDetails::GetCurrentStateItem(TArray<TSharedPtr<FStateTreeStateItem>>& OutPath) const
-{
-	OutPath.Reset();
-	const EStateTreeTransitionType TransitionType = GetTransitionType().Get(EStateTreeTransitionType::Failed);
-	FGuid StateID;
-	if (IDProperty
-		&& UE::StateTree::PropertyHelpers::GetStructValue<FGuid>(IDProperty, StateID) == FPropertyAccess::Success)
-	{
-		// Use RootItem (not filtered) to get any result even during search.
-		return FindStateByIDRecursive(RootItem, TransitionType, StateID, OutPath);
-	}
-	return false;
+	return nullptr;
 }
 
 FText FStateTreeStateLinkDetails::GetCurrentStateDesc() const
 {
-	TArray<TSharedPtr<FStateTreeStateItem>> Path;
-	if (GetCurrentStateItem(Path))
+	if (const FStateTreeStateLink* Link = UE::StateTree::PropertyHelpers::GetStructPtr<FStateTreeStateLink>(StructProperty))
 	{
-		return Path.Last()->Desc;
+		return UE::StateTree::Editor::GetStateLinkDesc(WeakEditorData.Get(), *Link, EStateTreeNodeFormatting::RichText);
 	}
-	// Could not find item, try to display old name as a hint what is missing.
-	if (NameProperty)
-	{
-		FName OldName;
-		if (NameProperty->GetValue(OldName) == FPropertyAccess::Success)
-		{
-			return FText::FromName(OldName);
-		}
-	}
-	return LOCTEXT("TransitionInvalid", "Invalid");
+	return LOCTEXT("MultipleSelected", "Multiple Selected");
 }
 
 const FSlateBrush* FStateTreeStateLinkDetails::GetCurrentStateIcon() const
 {
-	TArray<TSharedPtr<FStateTreeStateItem>> Path;
-	if (GetCurrentStateItem(Path))
+	if (const FStateTreeStateLink* Link = UE::StateTree::PropertyHelpers::GetStructPtr<FStateTreeStateLink>(StructProperty))
 	{
-		return Path.Last()->Icon;
+		return UE::StateTree::Editor::GetStateLinkIcon(WeakEditorData.Get(), *Link);
 	}
 	return nullptr;
 }
 
 FSlateColor FStateTreeStateLinkDetails::GetCurrentStateColor() const
 {
-	TArray<TSharedPtr<FStateTreeStateItem>> Path;
-	if (GetCurrentStateItem(Path))
+	if (const FStateTreeStateLink* Link = UE::StateTree::PropertyHelpers::GetStructPtr<FStateTreeStateLink>(StructProperty))
 	{
-		return Path.Last()->Color;
+		return UE::StateTree::Editor::GetStateLinkColor(WeakEditorData.Get(), *Link);
 	}
-	return {};
+	return FSlateColor::UseForeground();
 }
 
 bool FStateTreeStateLinkDetails::IsValidLink() const
 {
-	TArray<TSharedPtr<FStateTreeStateItem>> Path;
-	if (GetCurrentStateItem(Path))
+	const EStateTreeTransitionType TransitionType = GetTransitionType().Get(EStateTreeTransitionType::Failed);
+
+	if (TransitionType == EStateTreeTransitionType::GotoState)
 	{
-		return true;
+		return GetState() != nullptr;
 	}
-	// The state is missing.
-	return false;
+
+	return true;
 }
 
 TOptional<EStateTreeTransitionType> FStateTreeStateLinkDetails::GetTransitionType() const
diff --git a/Engine/Plugins/Runtime/StateTree/Source/StateTreeEditorModule/Private/Customizations/StateTreeStateLinkDetails.h b/Engine/Plugins/Runtime/StateTree/Source/StateTreeEditorModule/Private/Customizations/StateTreeStateLinkDetails.h
--- a/Engine/Plugins/Runtime/StateTree/Source/StateTreeEditorModule/Private/Customizations/StateTreeStateLinkDetails.h
+++ b/Engine/Plugins/Runtime/StateTree/Source/StateTreeEditorModule/Private/Customizations/StateTreeStateLinkDetails.h
@@ -14,7 +14,7 @@ class IPropertyHandle;
 class SWidget;
 class SSearchBox;
 class SComboButton;
-class UStateTree;
+class UStateTreeEditorData;
 class UStateTreeState;
 
 /**
@@ -32,76 +32,23 @@ public:
 	virtual void CustomizeChildren(TSharedRef<class IPropertyHandle> StructPropertyHandle, class IDetailChildrenBuilder& StructBuilder, IPropertyTypeCustomizationUtils& StructCustomizationUtils) override;
 
 private:
-	// Stores info about state
-	struct FStateTreeStateItem : TSharedFromThis<FStateTreeStateItem>
-	{
-		FStateTreeStateItem() = default;
-		
-		FStateTreeStateItem(const FText& InDesc, const FText& InTooltipText, const FSlateBrush* InIcon, const EStateTreeTransitionType InType)
-			: Desc(InDesc)
-			, TooltipText(InTooltipText)
-			, TransitionType(InType)
-			, Icon(InIcon)
-		{
-		}
-
-		FText Desc;
-		FText TooltipText;
-		EStateTreeTransitionType TransitionType = EStateTreeTransitionType::None;
-		FGuid StateID = {};
-		FSlateColor Color = FSlateColor(FLinearColor::White);
-		const FSlateBrush* Icon = nullptr;
-		bool bIsSubTree = false;
-		bool bIsLinked = false;
-		FText LinkedDesc;
-		TArray<TSharedPtr<FStateTreeStateItem>> Children;
-	};
-
-	// Stores per session node expansion state for a node type.
-	struct FStateExpansionState
-	{
-		TSet<FGuid> CollapsedStates;
-	};
-
-	void CacheStates();
-	void CacheStates(TSharedPtr<FStateTreeStateItem> ParentNode, const UStateTreeState* State);
-	void OnIdentifierChanged(const UStateTree& StateTree);
-
 	TSharedRef<SWidget> GenerateStatePicker();
 	FText GetCurrentStateDesc() const;
 	const FSlateBrush* GetCurrentStateIcon() const;
 	FSlateColor GetCurrentStateColor() const;
 	bool IsValidLink() const;
 	TOptional<EStateTreeTransitionType> GetTransitionType() const;
-	TSharedRef<ITableRow> GenerateStateItemRow(TSharedPtr<FStateTreeStateItem> Item, const TSharedRef<STableViewBase>& OwnerTable);
-	void GetStateItemChildren(TSharedPtr<FStateTreeStateItem> Item, TArray<TSharedPtr<FStateTreeStateItem>>& OutItems) const;
-	void OnStateItemSelected(TSharedPtr<FStateTreeStateItem> SelectedItem, ESelectInfo::Type);
-	void OnStateItemExpansionChanged(TSharedPtr<FStateTreeStateItem> ExpandedItem, bool bInExpanded) const;
-	void OnSearchBoxTextChanged(const FText& NewText);
-	static int32 FilterStateItemChildren(const TArray<FString>& FilterStrings, const bool bParentMatches, const TArray<TSharedPtr<FStateTreeStateItem>>& SourceArray, TArray<TSharedPtr<FStateTreeStateItem>>& OutDestArray);
-	void ExpandAll(const TArray<TSharedPtr<FStateTreeStateItem>>& Items);
+	void OnStateSelected(TConstArrayView<FGuid> SelectedStateIDs);
+	void SetTransitionByType(const EStateTreeTransitionType TransitionType);
+	const UStateTreeState* GetState() const;
 	
-	static bool FindStateByIDRecursive(const TSharedPtr<FStateTreeStateItem>& Item, EStateTreeTransitionType TransitionType, const FGuid StateID, TArray<TSharedPtr<FStateTreeStateItem>>& OutPath);
-	bool GetCurrentStateItem(TArray<TSharedPtr<FStateTreeStateItem>>& OutPath) const;
-	
-	void RestoreExpansionState();
-
-	TSharedPtr<FStateTreeStateItem> RootItem;
-	TSharedPtr<FStateTreeStateItem> FilteredRootItem;
-
 	TSharedPtr<IPropertyHandle> NameProperty;
 	TSharedPtr<IPropertyHandle> IDProperty;
 	TSharedPtr<IPropertyHandle> LinkTypeProperty;
 
 	TSharedPtr<SComboButton> ComboButton;
-	TSharedPtr<SSearchBox> SearchBox;
-	
-	TSharedPtr<STreeView<TSharedPtr<FStateTreeStateItem>>> StateItemTree;
-	bool bIsRestoringExpansion = false;
-	TWeakObjectPtr<const UStateTree> WeakStateTree = nullptr;
 
-	// Save expansion state for each base node type. The expansion state does not persist between editor sessions. 
-	static TMap<FObjectKey, FStateExpansionState> StateExpansionStates;
+	TWeakObjectPtr<const UStateTreeEditorData> WeakEditorData = nullptr;
 
 	// If set, hide selecting meta states like Next or (tree) Succeeded.
 	bool bDirectStatesOnly = false;
diff --git a/Engine/Plugins/Runtime/StateTree/Source/StateTreeEditorModule/Private/Customizations/StateTreeTransitionDetails.cpp b/Engine/Plugins/Runtime/StateTree/Source/StateTreeEditorModule/Private/Customizations/StateTreeTransitionDetails.cpp
--- a/Engine/Plugins/Runtime/StateTree/Source/StateTreeEditorModule/Private/Customizations/StateTreeTransitionDetails.cpp
+++ b/Engine/Plugins/Runtime/StateTree/Source/StateTreeEditorModule/Private/Customizations/StateTreeTransitionDetails.cpp
@@ -8,6 +8,7 @@
 #include "IDetailChildrenBuilder.h"
 #include "IPropertyUtilities.h"
 #include "ScopedTransaction.h"
+#include "StateTreeDescriptionHelpers.h"
 #include "StateTreeEditor.h"
 #include "StateTreeEditorData.h"
 #include "StateTreeEditorNodeUtils.h"
@@ -389,92 +390,14 @@ bool FStateTreeTransitionDetails::GetDelayTransition() const
 FText FStateTreeTransitionDetails::GetDescription() const
 {
 	check(StateProperty);
-	if (StateProperty->GetNumPerObjectValues() != 1)
-	{
-		return LOCTEXT("MultipleSelected", "Multiple Selected");
-	}
-
-	const EStateTreeTransitionTrigger Trigger = GetTrigger();
-
-	FText TriggerText;
-	if (Trigger == EStateTreeTransitionTrigger::OnStateCompleted)
-	{
-		TriggerText = LOCTEXT("TransitionOnStateCompleted", "<b>On State Completed</>");
-	}
-	else if (Trigger == EStateTreeTransitionTrigger::OnStateSucceeded)
-	{
-		TriggerText = LOCTEXT("TransitionOnStateSucceeded", "<b>On State Succeeded</>");
-	}
-	else if (Trigger == EStateTreeTransitionTrigger::OnStateFailed)
-	{
-		TriggerText = LOCTEXT("TransitionOnStateFailed", "<b>On State Failed</>");
-	}
-	else if (Trigger == EStateTreeTransitionTrigger::OnTick)
-	{
-		TriggerText = LOCTEXT("TransitionOnTick", "<b>On Tick</>");
-	}
-	else if (Trigger == EStateTreeTransitionTrigger::OnEvent)
-	{
-		FStateTreeEventDesc RequiredEvent;
-		UE::StateTree::PropertyHelpers::GetStructValue<FStateTreeEventDesc>(RequiredEventProperty, RequiredEvent);
-
-		TArray<FText> PayloadItems;
-		
-		if (RequiredEvent.IsValid())
-		{
-			if (RequiredEvent.Tag.IsValid())
-			{
-				PayloadItems.Add(FText::Format(LOCTEXT("TransitionEventTag", "<s>Tag:</> '{0}'"), FText::FromName(RequiredEvent.Tag.GetTagName())));
-			}
-			
-			if (RequiredEvent.PayloadStruct)
-			{
-				PayloadItems.Add(FText::Format(LOCTEXT("TransitionEventPayload", "<s>Payload:</> '{0}'"), RequiredEvent.PayloadStruct->GetDisplayNameText()));
-			}
-		}
-		else
-		{
-			PayloadItems.Add(LOCTEXT("TransitionInvalidEvent", "Invalid"));
-		}
-		
-		TriggerText = FText::Format(LOCTEXT("TransitionOnEvent", "<b>On Event</> ({0})"), FText::Join(INVTEXT(", "), PayloadItems));
-	}
 
-	
-	FText TargetText;
-	TArray<void*> RawData;
-	StateProperty->AccessRawData(RawData);
-	check(RawData.Num() > 0);
-	
-	const FStateTreeStateLink* State = static_cast<FStateTreeStateLink*>(RawData[0]);
-	if (State != nullptr)
+	const FStateTreeTransition* Transition = UE::StateTree::PropertyHelpers::GetStructPtr<FStateTreeTransition>(StructProperty);
+	if (!Transition)
 	{
-		switch (State->LinkType)
-		{
-		case EStateTreeTransitionType::None:
-			TargetText = LOCTEXT("TransitionNone", "<i>None</>");
-			break;
-		case EStateTreeTransitionType::Succeeded:
-			TargetText = LOCTEXT("TransitionTreeSucceeded", "<s>return</> <i>Tree Succeeded</>");
-			break;
-		case EStateTreeTransitionType::Failed:
-			TargetText = LOCTEXT("TransitionTreeFailed", "<s>return</> <i>Tree Failed</>");
-			break;
-		case EStateTreeTransitionType::NextState:
-			TargetText = LOCTEXT("TransitionNextState", "<s>go to</> <i>Next State</>");
-			break;
-		case EStateTreeTransitionType::NextSelectableState:
-			TargetText = LOCTEXT("TransitionNextSelectableState", "<s>go to</> <i>Next Selectable State</>");
-			break;
-		case EStateTreeTransitionType::GotoState:
-			{
-				TargetText = FText::Format(LOCTEXT("TransitionGotoState", "<s>go to</> {0}"), FText::FromName(State->Name));
-			}
-			break;
-		}
+		return LOCTEXT("MultipleSelected", "Multiple Selected");
 	}
 
-	return FText::Format(LOCTEXT("TransitionDesc", "{0} {1}"), TriggerText, TargetText);
+	return UE::StateTree::Editor::GetTransitionDesc(GetEditorData(), *Transition, EStateTreeNodeFormatting::RichText);
 }
 
 void FStateTreeTransitionDetails::OnCopyTransition() const
diff --git a/Engine/Plugins/Runtime/StateTree/Source/StateTreeEditorModule/Private/Customizations/Widgets/SCompactStateTreeView.cpp b/Engine/Plugins/Runtime/StateTree/Source/StateTreeEditorModule/Private/Customizations/Widgets/SCompactStateTreeView.cpp
new file mode 100644
--- /dev/null
+++ b/Engine/Plugins/Runtime/StateTree/Source/StateTreeEditorModule/Private/Customizations/Widgets/SCompactStateTreeView.cpp
@@ -0,0 +1,621 @@
+﻿// Copyright Epic Games, Inc. All Rights Reserved.
+
+#include "SCompactStateTreeView.h"
+
+#include "StateTreeDelegates.h"
+#include "StateTreeDragDrop.h"
+#include "StateTreeEditorData.h"
+#include "StateTreeEditorStyle.h"
+#include "StateTreeViewModel.h"
+#include "Widgets/Input/SSearchBox.h"
+#include "Widgets/Text/SRichTextBlock.h"
+#include "TextStyleDecorator.h"
+
+#define LOCTEXT_NAMESPACE "StateTreeEditor"
+
+TMap<FObjectKey, SCompactStateTreeView::FStateExpansionState> SCompactStateTreeView::StateExpansionStates;
+
+SCompactStateTreeView::SCompactStateTreeView()
+{
+}
+
+SCompactStateTreeView::~SCompactStateTreeView()
+{
+}
+
+void SCompactStateTreeView::Construct(const SCompactStateTreeView::FArguments& InArgs, TSharedPtr<FStateTreeViewModel> InViewModel)
+{
+	StateTreeViewModel = InViewModel;
+	WeakStateTreeEditorData = InArgs._StateTreeEditorData;
+	bSubtreesOnly = InArgs._SubtreesOnly;
+	bSelectableStatesOnly = InArgs._SelectableStatesOnly;
+	OnSelectionChanged = InArgs._OnSelectionChanged;
+	OnContextMenuOpening = InArgs._OnContextMenuOpening;
+	
+	CacheStates();
+
+	ChildSlot
+	[
+		SNew(SVerticalBox)
+		+ SVerticalBox::Slot()
+		.HAlign(HAlign_Fill)
+		.VAlign(VAlign_Top)
+		.Padding(4, 2, 4, 2)
+		.AutoHeight()
+		[
+			SAssignNew(SearchBox, SSearchBox)
+			.OnTextChanged(this, &SCompactStateTreeView::OnSearchBoxTextChanged)
+		]
+		+ SVerticalBox::Slot()
+		[
+			SAssignNew(StateItemTree, STreeView<TSharedPtr<FStateTreeStateItem>>)
+			.SelectionMode(InArgs._SelectionMode)
+			.TreeItemsSource(&FilteredRootItem->Children)
+			.OnGenerateRow(this, &SCompactStateTreeView::GenerateStateItemRow)
+			.OnGetChildren(this, &SCompactStateTreeView::GetStateItemChildren)
+			.OnSelectionChanged(this, &SCompactStateTreeView::OnStateItemSelected)
+			.OnExpansionChanged(this, &SCompactStateTreeView::OnStateItemExpansionChanged)
+			.OnContextMenuOpening_Lambda([this]()
+			{
+				if (OnContextMenuOpening.IsBound())
+				{
+					return OnContextMenuOpening.Execute();
+				}
+				return SNullWidget::NullWidget.ToSharedPtr();
+			})
+		]
+	];
+
+	// Restore category expansion state from previous use.
+	RestoreExpansionState();
+}
+
+void SCompactStateTreeView::Refresh(const UStateTreeEditorData* NewStateTreeEditorData)
+{
+	if (!StateItemTree)
+	{
+		return;
+	}
+
+	if (NewStateTreeEditorData)
+	{
+		WeakStateTreeEditorData = NewStateTreeEditorData;
+	}
+
+	TArray<FGuid> SelectedItemIDs = GetSelection();
+	
+	CacheStates();
+	UpdateFilteredRoot(/*bRestoreSelection*/false);
+
+	SetSelection(SelectedItemIDs);
+}
+
+void SCompactStateTreeView::SetSelection(TConstArrayView<FGuid> Selection)
+{
+	if (!StateItemTree)
+	{
+		return;
+	}
+	
+	if (bIsSettingSelection)
+	{
+		return;
+	}
+	
+	TArray<TSharedPtr<FStateTreeStateItem>> SelectedStates;
+	FindStatesByIDRecursive(FilteredRootItem, Selection, SelectedStates);
+
+	bIsSettingSelection = true;
+	
+	StateItemTree->ClearSelection();
+	StateItemTree->SetItemSelection(SelectedStates, true);
+
+	if (SelectedStates.Num() == 1)
+	{
+		StateItemTree->RequestScrollIntoView(SelectedStates[0]);	
+	}
+
+	bIsSettingSelection = false;
+}
+
+TArray<FGuid> SCompactStateTreeView::GetSelection() const
+{
+	TArray<FGuid> SelectedItemIDs;
+	
+	TArray<TSharedPtr<FStateTreeStateItem>> SelectedItems = StateItemTree->GetSelectedItems();
+	for (const TSharedPtr<FStateTreeStateItem>& Item : SelectedItems)
+	{
+		if (Item)
+		{
+			SelectedItemIDs.Add(Item->StateID);
+		}
+	}
+
+	return SelectedItemIDs;
+}
+
+TSharedPtr<SWidget> SCompactStateTreeView::GetWidgetToFocusOnOpen()
+{
+	return SearchBox;
+}
+
+void SCompactStateTreeView::CacheStates()
+{
+	RootItem = MakeShared<FStateTreeStateItem>();
+
+	if (const UStateTreeEditorData* StateTreeEditorData = WeakStateTreeEditorData.Get())
+	{
+		for (const UStateTreeState* SubTree : StateTreeEditorData->SubTrees)
+		{
+			CacheState(RootItem, SubTree);
+		}
+	}
+
+	FilteredRootItem = RootItem;
+}
+
+void SCompactStateTreeView::CacheState(TSharedPtr<FStateTreeStateItem> ParentNode, const UStateTreeState* State)
+{
+	if (!State)
+	{
+		return;
+	}
+	const UStateTreeEditorData* StateTreeEditorData = WeakStateTreeEditorData.Get();
+	if (!StateTreeEditorData)
+	{
+		return;
+	}
+
+	bool bShouldAdd = true;
+	if (bSubtreesOnly
+		&& State->Type != EStateTreeStateType::Subtree)
+	{
+		bShouldAdd = false;
+	}
+
+	if (bSelectableStatesOnly
+		&& State->SelectionBehavior == EStateTreeStateSelectionBehavior::None)
+	{
+		bShouldAdd = false;
+	}
+	
+	if (bShouldAdd)
+	{
+		TSharedRef<FStateTreeStateItem> StateItem = MakeShared<FStateTreeStateItem>();
+		StateItem->Desc = FText::FromName(State->Name);
+		StateItem->StateID = State->ID;
+		StateItem->bIsSubTree = State->Type == EStateTreeStateType::Subtree;
+		StateItem->bIsEnabled = State->bEnabled;
+
+		StateItem->Color = FColor(31, 151, 167);
+		if (const FStateTreeEditorColor* FoundColor = StateTreeEditorData->FindColor(State->ColorRef))
+		{
+			StateItem->Color = FoundColor->Color;
+		}
+		
+		// Figure out icon.
+		if (State->SelectionBehavior == EStateTreeStateSelectionBehavior::None)
+		{
+			StateItem->Icon = FStateTreeEditorStyle::Get().GetBrush("StateTreeEditor.SelectNone");
+		}
+		else if (State->SelectionBehavior == EStateTreeStateSelectionBehavior::TryEnterState)
+		{
+			StateItem->Icon = FStateTreeEditorStyle::Get().GetBrush("StateTreeEditor.TryEnterState");			
+		}
+		else if (State->SelectionBehavior == EStateTreeStateSelectionBehavior::TrySelectChildrenInOrder)
+		{
+			if (State->Children.IsEmpty()
+				|| State->Type == EStateTreeStateType::Linked
+				|| State->Type == EStateTreeStateType::LinkedAsset)
+			{
+				StateItem->Icon = FStateTreeEditorStyle::Get().GetBrush("StateTreeEditor.TryEnterState");			
+			}
+			else
+			{
+				StateItem->Icon = FStateTreeEditorStyle::Get().GetBrush("StateTreeEditor.TrySelectChildrenInOrder");
+			}
+		}
+		else if (State->SelectionBehavior == EStateTreeStateSelectionBehavior::TrySelectChildrenAtUniformRandom)
+		{
+			StateItem->Icon = FStateTreeEditorStyle::Get().GetBrush("StateTreeEditor.TrySelectChildrenAtRandom");
+		}
+		else if (State->SelectionBehavior == EStateTreeStateSelectionBehavior::TryFollowTransitions)
+		{
+			StateItem->Icon = FStateTreeEditorStyle::Get().GetBrush("StateTreeEditor.TryFollowTransitions");
+		}
+
+		// Linked states
+		if (State->Type == EStateTreeStateType::Linked)
+		{
+			StateItem->bIsLinked = true;
+			StateItem->LinkedDesc = FText::FromName(State->LinkedSubtree.Name);
+		}
+		else if (State->Type == EStateTreeStateType::LinkedAsset)
+		{
+			StateItem->bIsLinked = true;
+			StateItem->LinkedDesc = FText::FromString(GetNameSafe(State->LinkedAsset.Get()));
+		}
+		
+		ParentNode->Children.Add(StateItem);
+
+		ParentNode = StateItem;
+	}
+
+	for (const UStateTreeState* ChildState : State->Children)
+	{
+		CacheState(ParentNode, ChildState);
+	}
+}
+
+TSharedRef<ITableRow> SCompactStateTreeView::GenerateStateItemRow(TSharedPtr<FStateTreeStateItem> Item, const TSharedRef<STableViewBase>& OwnerTable)
+{
+	TSharedRef<SHorizontalBox> Container = SNew(SHorizontalBox);
+
+	// Icon
+	Container->AddSlot()
+		.HAlign(HAlign_Left)
+		.VAlign(VAlign_Center)
+		.Padding(0, 2.0f, 4.0f, 2.0f)
+		.AutoWidth()
+		[
+			SNew(SImage)
+			.Visibility(Item->Icon ? EVisibility::Visible : EVisibility::Collapsed)
+			.DesiredSizeOverride(FVector2D(16.0f, 16.0f))
+			.Image(Item->Icon)
+			.ColorAndOpacity(Item->Color)
+			.IsEnabled(Item->bIsEnabled)
+		];
+
+	// Name
+	Container->AddSlot()
+		.HAlign(HAlign_Fill)
+		.VAlign(VAlign_Center)
+		.AutoWidth()
+		[
+			SNew(SRichTextBlock)
+			.Text(Item->Desc)
+			.TextStyle(&FStateTreeEditorStyle::Get().GetWidgetStyle<FTextBlockStyle>("Normal.Normal"))
+			.OverflowPolicy(ETextOverflowPolicy::Ellipsis)
+			.ToolTipText(Item->TooltipText)
+			.IsEnabled(Item->bIsEnabled)
+			.HighlightText_Lambda([this]() { return SearchBox.IsValid() ? SearchBox->GetText() : FText::GetEmpty(); })
+			+SRichTextBlock::Decorator(FTextStyleDecorator::Create(TEXT(""), FStateTreeEditorStyle::Get().GetWidgetStyle<FTextBlockStyle>("Normal.Normal")))
+			+SRichTextBlock::Decorator(FTextStyleDecorator::Create(TEXT("b"), FStateTreeEditorStyle::Get().GetWidgetStyle<FTextBlockStyle>("Normal.Bold")))
+			+SRichTextBlock::Decorator(FTextStyleDecorator::Create(TEXT("i"), FStateTreeEditorStyle::Get().GetWidgetStyle<FTextBlockStyle>("Normal.Italic")))
+			+SRichTextBlock::Decorator(FTextStyleDecorator::Create(TEXT("s"), FStateTreeEditorStyle::Get().GetWidgetStyle<FTextBlockStyle>("Normal.Subdued")))
+		];
+
+	// Link
+	if (Item->bIsLinked)
+	{
+		// Link icon
+		Container->AddSlot()
+			.HAlign(HAlign_Left)
+			.VAlign(VAlign_Center)
+			.Padding(4.f, 0.f)
+			.AutoWidth()
+			[
+				SNew(SImage)
+				.ColorAndOpacity(FSlateColor::UseSubduedForeground())
+				.Image(FStateTreeEditorStyle::Get().GetBrush("StateTreeEditor.StateLinked"))
+		];
+
+		// Linked name
+		Container->AddSlot()
+			.HAlign(HAlign_Fill)
+			.VAlign(VAlign_Center)
+			.AutoWidth()
+			[
+				SNew(STextBlock)
+				.Font(FAppStyle::Get().GetFontStyle("PropertyWindow.NormalFont"))
+				.ColorAndOpacity(FSlateColor::UseSubduedForeground())
+				.Text(Item->LinkedDesc)
+			];
+	}
+	
+	return SNew(STableRow<TSharedPtr<FStateTreeStateItem>>, OwnerTable)
+		.OnDragDetected(this, &SCompactStateTreeView::HandleDragDetected)
+		.OnDragLeave(this, &SCompactStateTreeView::HandleDragLeave)
+		.OnCanAcceptDrop(this, &SCompactStateTreeView::HandleCanAcceptDrop)
+		.OnAcceptDrop(this, &SCompactStateTreeView::HandleAcceptDrop)
+		[
+			Container
+		];
+}
+
+void SCompactStateTreeView::GetStateItemChildren(TSharedPtr<FStateTreeStateItem> Item, TArray<TSharedPtr<FStateTreeStateItem>>& OutItems) const
+{
+	if (Item.IsValid())
+	{
+		OutItems = Item->Children;
+	}
+}
+
+void SCompactStateTreeView::OnStateItemSelected(TSharedPtr<FStateTreeStateItem> SelectedItem, ESelectInfo::Type Type)
+{
+	// Skip selection from code
+	if (Type == ESelectInfo::Direct)
+	{
+		return;
+	}
+
+	if (bIsSettingSelection)
+	{
+		return;
+	}
+
+	if (OnSelectionChanged.IsBound())
+	{
+		TArray<FGuid> SelectedStateIDs;
+		
+		TArray<TSharedPtr<FStateTreeStateItem>> Selection = StateItemTree->GetSelectedItems();
+		for (const TSharedPtr<FStateTreeStateItem>& Item : Selection)
+		{
+			if (Item)
+			{
+				SelectedStateIDs.Add(Item->StateID);
+			}
+		}
+
+		OnSelectionChanged.Execute(SelectedStateIDs);
+	}
+}
+
+void SCompactStateTreeView::OnStateItemExpansionChanged(TSharedPtr<FStateTreeStateItem> ExpandedItem, bool bInExpanded) const
+{
+	// Do not save expansion state when restoring expansion state, or when showing filtered results. 
+	if (bIsRestoringExpansion || FilteredRootItem != RootItem)
+	{
+		return;
+	}
+
+	if (ExpandedItem.IsValid() && ExpandedItem->StateID.IsValid())
+	{
+		FStateExpansionState& ExpansionState = StateExpansionStates.FindOrAdd(FObjectKey(WeakStateTreeEditorData.Get()));
+		if (bInExpanded)
+		{
+			ExpansionState.CollapsedStates.Remove(ExpandedItem->StateID);
+		}
+		else
+		{
+			ExpansionState.CollapsedStates.Add(ExpandedItem->StateID);
+		}
+	}
+}
+
+void SCompactStateTreeView::OnSearchBoxTextChanged(const FText& NewText)
+{
+	if (!StateItemTree.IsValid())
+	{
+		return;
+	}
+	
+
+	NewText.ToString().ParseIntoArrayWS(FilterStrings);
+	FilterStrings.RemoveAll([](const FString& String) { return String.IsEmpty(); });
+
+	UpdateFilteredRoot();
+}
+
+void SCompactStateTreeView::UpdateFilteredRoot(bool bRestoreSelection)
+{
+	FilteredRootItem.Reset();
+
+	TArray<FGuid> Selection;
+	if (bRestoreSelection)
+	{
+		Selection = GetSelection();
+	}
+
+	if (FilterStrings.IsEmpty())
+	{
+		// Show all when there's no filter string.
+		FilteredRootItem = RootItem;
+		StateItemTree->SetTreeItemsSource(&FilteredRootItem->Children);
+		RestoreExpansionState();
+	}
+	else
+	{
+		FilteredRootItem = MakeShared<FStateTreeStateItem>();
+		FilterStateItemChildren(FilterStrings, /*bParentMatches*/false, RootItem->Children, FilteredRootItem->Children);
+
+		StateItemTree->SetTreeItemsSource(&FilteredRootItem->Children);
+		ExpandAll(FilteredRootItem->Children);
+	}
+
+	if (bRestoreSelection)
+	{
+		SetSelection(Selection);
+	}
+
+	StateItemTree->RequestTreeRefresh();
+}
+
+int32 SCompactStateTreeView::FilterStateItemChildren(const TArray<FString>& FilterStrings, const bool bParentMatches, const TArray<TSharedPtr<FStateTreeStateItem>>& SourceArray, TArray<TSharedPtr<FStateTreeStateItem>>& OutDestArray)
+{
+	int32 NumFound = 0;
+
+	auto MatchFilter = [&FilterStrings](const TSharedPtr<FStateTreeStateItem>& SourceItem)
+	{
+		const FString ItemName = SourceItem->Desc.ToString();
+		for (const FString& Filter : FilterStrings)
+		{
+			if (ItemName.Contains(Filter))
+			{
+				return true;
+			}
+		}
+		return false;
+	};
+
+	for (const TSharedPtr<FStateTreeStateItem>& SourceItem : SourceArray)
+	{
+		// Check if our name matches the filters
+		// If bParentMatches is true, the search matched a parent category.
+		const bool bMatchesFilters = bParentMatches || MatchFilter(SourceItem);
+
+		int32 NumChildren = 0;
+		if (bMatchesFilters)
+		{
+			NumChildren++;
+		}
+
+		// if we don't match, then we still want to check all our children
+		TArray<TSharedPtr<FStateTreeStateItem>> FilteredChildren;
+		NumChildren += FilterStateItemChildren(FilterStrings, bMatchesFilters, SourceItem->Children, FilteredChildren);
+
+		// then add this item to the destination array
+		if (NumChildren > 0)
+		{
+			TSharedPtr<FStateTreeStateItem>& NewItem = OutDestArray.Add_GetRef(MakeShared<FStateTreeStateItem>());
+			*NewItem = *SourceItem;
+			NewItem->Children = FilteredChildren;
+
+			NumFound += NumChildren;
+		}
+	}
+
+	return NumFound;
+}
+
+void SCompactStateTreeView::ExpandAll(const TArray<TSharedPtr<FStateTreeStateItem>>& Items)
+{
+	for (const TSharedPtr<FStateTreeStateItem>& Item : Items)
+	{
+		StateItemTree->SetItemExpansion(Item, true);
+		ExpandAll(Item->Children);
+	}
+}
+	
+bool SCompactStateTreeView::FindStateByIDRecursive(const TSharedPtr<FStateTreeStateItem>& Item, const FGuid StateID, TArray<TSharedPtr<FStateTreeStateItem>>& OutPath)
+{
+	OutPath.Push(Item);
+
+	if (Item->StateID == StateID)
+	{
+		return true;
+	}
+
+	for (const TSharedPtr<FStateTreeStateItem>& ChildItem : Item->Children)
+	{
+		if (FindStateByIDRecursive(ChildItem, StateID, OutPath))
+		{
+			return true;
+		}
+	}
+
+	OutPath.Pop();
+
+	return false;
+}
+
+void SCompactStateTreeView::FindStatesByIDRecursive(const TSharedPtr<FStateTreeStateItem>& Item, TConstArrayView<FGuid> StateIDs, TArray<TSharedPtr<FStateTreeStateItem>>& OutStates)
+{
+	if (StateIDs.Contains(Item->StateID))
+	{
+		OutStates.Add(Item);
+	}
+
+	for (const TSharedPtr<FStateTreeStateItem>& ChildItem : Item->Children)
+	{
+		FindStatesByIDRecursive(ChildItem, StateIDs, OutStates);
+	}
+}
+
+
+void SCompactStateTreeView::RestoreExpansionState()
+{
+	if (!StateItemTree.IsValid())
+	{
+		return;
+	}
+	
+	FStateExpansionState& ExpansionState = StateExpansionStates.FindOrAdd(FObjectKey(WeakStateTreeEditorData.Get()));
+
+	bIsRestoringExpansion = true;
+
+	// Default state is expanded.
+	ExpandAll(FilteredRootItem->Children);
+
+	// Collapse the ones that are specifically collapsed.
+	for (const FGuid& StateID : ExpansionState.CollapsedStates)
+	{
+		TArray<TSharedPtr<FStateTreeStateItem>> Path;
+		if (FindStateByIDRecursive(FilteredRootItem, StateID, Path))
+		{
+			StateItemTree->SetItemExpansion(Path.Last(), false);
+		}
+	}
+
+	bIsRestoringExpansion = false;
+}
+
+FReply SCompactStateTreeView::HandleDragDetected(const FGeometry&, const FPointerEvent&) const
+{
+	return FReply::Handled().BeginDragDrop(FStateTreeSelectedDragDrop::New(StateTreeViewModel));
+}
+
+void SCompactStateTreeView::HandleDragLeave(const FDragDropEvent& DragDropEvent) const
+{
+	const TSharedPtr<FStateTreeSelectedDragDrop> DragDropOperation = DragDropEvent.GetOperationAs<FStateTreeSelectedDragDrop>();
+	if (DragDropOperation.IsValid())
+	{
+		DragDropOperation->SetCanDrop(false);
+	}
+}
+
+TOptional<EItemDropZone> SCompactStateTreeView::HandleCanAcceptDrop(const FDragDropEvent& DragDropEvent, EItemDropZone DropZone, TSharedPtr<FStateTreeStateItem> TargetState) const
+{
+	if (StateTreeViewModel)
+	{
+		const TSharedPtr<FStateTreeSelectedDragDrop> DragDropOperation = DragDropEvent.GetOperationAs<FStateTreeSelectedDragDrop>();
+		if (DragDropOperation.IsValid())
+		{
+			DragDropOperation->SetCanDrop(true);
+
+			// Cannot drop on selection or child of selection.
+			if (StateTreeViewModel && StateTreeViewModel->IsChildOfSelection(StateTreeViewModel->GetMutableStateByID(TargetState->StateID)))
+			{
+				DragDropOperation->SetCanDrop(false);
+				return TOptional<EItemDropZone>();
+			}
+
+			return DropZone;
+		}
+	}
+
+	return TOptional<EItemDropZone>();
+}
+
+FReply SCompactStateTreeView::HandleAcceptDrop(const FDragDropEvent& DragDropEvent, EItemDropZone DropZone, TSharedPtr<FStateTreeStateItem> TargetState) const
+{
+	if (StateTreeViewModel)
+	{
+		const TSharedPtr<FStateTreeSelectedDragDrop> DragDropOperation = DragDropEvent.GetOperationAs<FStateTreeSelectedDragDrop>();
+		if (DragDropOperation.IsValid())
+		{
+			if (StateTreeViewModel)
+			{
+				if (DropZone == EItemDropZone::AboveItem)
+				{
+					StateTreeViewModel->MoveSelectedStatesBefore(StateTreeViewModel->GetMutableStateByID(TargetState->StateID));
+				}
+				else if (DropZone == EItemDropZone::BelowItem)
+				{
+					StateTreeViewModel->MoveSelectedStatesAfter(StateTreeViewModel->GetMutableStateByID(TargetState->StateID));
+				}
+				else
+				{
+					StateTreeViewModel->MoveSelectedStatesInto(StateTreeViewModel->GetMutableStateByID(TargetState->StateID));
+				}
+
+				return FReply::Handled();
+			}
+		}
+	}
+
+	return FReply::Unhandled();
+}
+
+
+#undef LOCTEXT_NAMESPACE
\ No newline at end of file
diff --git a/Engine/Plugins/Runtime/StateTree/Source/StateTreeEditorModule/Private/Customizations/Widgets/SCompactStateTreeView.h b/Engine/Plugins/Runtime/StateTree/Source/StateTreeEditorModule/Private/Customizations/Widgets/SCompactStateTreeView.h
new file mode 100644
--- /dev/null
+++ b/Engine/Plugins/Runtime/StateTree/Source/StateTreeEditorModule/Private/Customizations/Widgets/SCompactStateTreeView.h
@@ -0,0 +1,133 @@
+﻿// Copyright Epic Games, Inc. All Rights Reserved.
+
+#pragma once
+
+#include "Widgets/SCompoundWidget.h"
+#include "Widgets/Views/STreeView.h"
+#include "Textures/SlateIcon.h"
+#include "UObject/ObjectKey.h"
+#include "Misc/Guid.h"
+
+class UStateTree;
+class UStateTreeState;
+class UStateTreeEditorData;
+class SSearchBox;
+class FStateTreeViewModel;
+
+/**
+ * Widget that displays a list of State Tree nodes which match base types and specified schema.
+ * Can be used e.g. in popup menus to select node types.
+ */
+class SCompactStateTreeView : public SCompoundWidget
+{
+public:
+
+	DECLARE_DELEGATE_OneParam(FOnSelectionChanged, TConstArrayView<FGuid> /*SelectedStateIDs*/);
+
+	SLATE_BEGIN_ARGS(SCompactStateTreeView)
+		: _StateTreeEditorData(nullptr)
+		, _SelectionMode(ESelectionMode::Single)
+		, _SelectableStatesOnly(false)
+		, _SubtreesOnly(false)
+	{}
+		/** Currently selected struct, initially highlighted. */
+		SLATE_ARGUMENT(const UStateTreeEditorData*, StateTreeEditorData)
+		SLATE_ARGUMENT(ESelectionMode::Type, SelectionMode)
+		SLATE_ARGUMENT(bool, SelectableStatesOnly)
+		SLATE_ARGUMENT(bool, SubtreesOnly)
+		SLATE_EVENT(FOnSelectionChanged, OnSelectionChanged)
+		SLATE_EVENT(FOnContextMenuOpening, OnContextMenuOpening)
+	SLATE_END_ARGS()
+
+	SCompactStateTreeView();
+	virtual ~SCompactStateTreeView() override;
+
+	void Construct(const FArguments& InArgs, TSharedPtr<FStateTreeViewModel> InViewModel = nullptr);
+
+	/** @returns widget to focus (search box) when the picker is opened. */
+	TSharedPtr<SWidget> GetWidgetToFocusOnOpen();
+
+	void SetSelection(TConstArrayView<FGuid> Selection);
+	TArray<FGuid> GetSelection() const;
+
+	void Refresh(const UStateTreeEditorData* NewStateTreeEditorData = nullptr);
+	
+private:
+	// Stores info about state
+	struct FStateTreeStateItem : TSharedFromThis<FStateTreeStateItem>
+	{
+		FStateTreeStateItem() = default;
+		
+		FStateTreeStateItem(const FText& InDesc, const FText& InTooltipText, const FSlateBrush* InIcon)
+			: Desc(InDesc)
+			, TooltipText(InTooltipText)
+			, Icon(InIcon)
+		{
+		}
+
+		FText Desc;
+		FText TooltipText;
+		FGuid StateID = {};
+		FSlateColor Color = FSlateColor(FLinearColor::White);
+		const FSlateBrush* Icon = nullptr;
+		bool bIsSubTree = false;
+		bool bIsLinked = false;
+		bool bIsEnabled = true;
+		FText LinkedDesc;
+		TArray<TSharedPtr<FStateTreeStateItem>> Children;
+	};
+
+	// Stores per session node expansion state for a node type.
+	struct FStateExpansionState
+	{
+		TSet<FGuid> CollapsedStates;
+	};
+
+	void CacheStates();
+	void CacheState(TSharedPtr<FStateTreeStateItem> ParentNode, const UStateTreeState* State);
+
+	TSharedRef<ITableRow> GenerateStateItemRow(TSharedPtr<FStateTreeStateItem> Item, const TSharedRef<STableViewBase>& OwnerTable);
+	void GetStateItemChildren(TSharedPtr<FStateTreeStateItem> Item, TArray<TSharedPtr<FStateTreeStateItem>>& OutItems) const;
+	void OnStateItemSelected(TSharedPtr<FStateTreeStateItem> SelectedItem, ESelectInfo::Type);
+	void OnStateItemExpansionChanged(TSharedPtr<FStateTreeStateItem> ExpandedItem, bool bInExpanded) const;
+	void OnSearchBoxTextChanged(const FText& NewText);
+	void UpdateFilteredRoot(bool bRestoreSelection = true);
+	static int32 FilterStateItemChildren(const TArray<FString>& FilterStrings, const bool bParentMatches, const TArray<TSharedPtr<FStateTreeStateItem>>& SourceArray, TArray<TSharedPtr<FStateTreeStateItem>>& OutDestArray);
+	void ExpandAll(const TArray<TSharedPtr<FStateTreeStateItem>>& Items);
+	
+	static bool FindStateByIDRecursive(const TSharedPtr<FStateTreeStateItem>& Item, const FGuid StateID, TArray<TSharedPtr<FStateTreeStateItem>>& OutPath);
+	static void FindStatesByIDRecursive(const TSharedPtr<FStateTreeStateItem>& Item, TConstArrayView<FGuid> StateIDs, TArray<TSharedPtr<FStateTreeStateItem>>& OutStates);
+	
+	void RestoreExpansionState();
+
+	FReply HandleDragDetected(const FGeometry&, const FPointerEvent&) const;
+	void HandleDragLeave(const FDragDropEvent& DragDropEvent) const;
+	TOptional<EItemDropZone> HandleCanAcceptDrop(const FDragDropEvent& DragDropEvent, EItemDropZone DropZone, TSharedPtr<FStateTreeStateItem> TargetState) const;
+	FReply HandleAcceptDrop(const FDragDropEvent& DragDropEvent, EItemDropZone DropZone, TSharedPtr<FStateTreeStateItem> TargetState) const;
+
+	TSharedPtr<FStateTreeStateItem> RootItem;
+	TSharedPtr<FStateTreeStateItem> FilteredRootItem;
+
+	TSharedPtr<SSearchBox> SearchBox;
+	
+	TSharedPtr<STreeView<TSharedPtr<FStateTreeStateItem>>> StateItemTree;
+	bool bIsRestoringExpansion = false;
+	TWeakObjectPtr<const UStateTreeEditorData> WeakStateTreeEditorData = nullptr;
+
+	TArray<FString> FilterStrings;
+	
+	// If set, allow to select only states marked as subtrees.
+	bool bSubtreesOnly = false;
+
+	bool bSelectableStatesOnly = false;
+
+	bool bIsSettingSelection = false;
+
+	TSharedPtr<FStateTreeViewModel> StateTreeViewModel;
+	
+	FOnSelectionChanged OnSelectionChanged;
+	FOnContextMenuOpening OnContextMenuOpening;
+
+	// Save expansion state for each base node type. The expansion state does not persist between editor sessions. 
+	static TMap<FObjectKey, FStateExpansionState> StateExpansionStates;
+};
\ No newline at end of file
diff --git a/Engine/Plugins/Runtime/StateTree/Source/StateTreeEditorModule/Private/SStateTreeExpanderArrow.cpp b/Engine/Plugins/Runtime/StateTree/Source/StateTreeEditorModule/Private/SStateTreeExpanderArrow.cpp
new file mode 100644
--- /dev/null
+++ b/Engine/Plugins/Runtime/StateTree/Source/StateTreeEditorModule/Private/SStateTreeExpanderArrow.cpp
@@ -0,0 +1,216 @@
+﻿// Copyright Epic Games, Inc. All Rights Reserved.
+
+#include "SStateTreeExpanderArrow.h"
+#include "Framework/Application/SlateApplication.h"
+#include "Widgets/Images/SImage.h"
+#include "Widgets/Input/SButton.h"
+#include "Widgets/Views/STableRow.h"
+
+SStateTreeExpanderArrow::SStateTreeExpanderArrow()
+{
+}
+
+void SStateTreeExpanderArrow::Construct( const FArguments& InArgs, const TSharedPtr<class ITableRow>& TableRow  )
+{
+	OwnerRowPtr = TableRow;
+	IndentAmount = InArgs._IndentAmount;
+	BaseIndentLevel = InArgs._BaseIndentLevel;
+	StyleSet = InArgs._StyleSet;
+	
+	WireColor = InArgs._WireColorAndOpacity;
+	ImageSize = InArgs._ImageSize;
+	ImagePadding = InArgs._ImagePadding;
+
+	this->ChildSlot
+	[
+		SAssignNew(ExpanderArrow, SButton)
+		.ButtonStyle(FCoreStyle::Get(), "NoBorder")
+		.Visibility(this, &SStateTreeExpanderArrow::GetExpanderVisibility)
+		.ClickMethod(EButtonClickMethod::MouseDown )
+		.OnClicked(this, &SStateTreeExpanderArrow::OnArrowClicked)
+		.ForegroundColor(FLinearColor(1,1,1,0.75f))
+		.IsFocusable(false)
+		.ContentPadding(TAttribute<FMargin>( this, &SStateTreeExpanderArrow::GetExpanderPadding )) //FMargin(0))
+		.HAlign(HAlign_Left)
+		.VAlign(VAlign_Top)
+		[
+			SNew(SImage)
+			.DesiredSizeOverride(FVector2D(InArgs._ImageSize))
+			.Image(this, &SStateTreeExpanderArrow::GetExpanderImage)
+			.ColorAndOpacity(FLinearColor(1,1,1,0.5f))
+		]
+	];
+}
+
+int32 SStateTreeExpanderArrow::OnPaint(const FPaintArgs& Args, const FGeometry& AllottedGeometry, const FSlateRect& MyCullingRect, FSlateWindowElementList& OutDrawElements, int32 LayerId, const FWidgetStyle& InWidgetStyle, bool bParentEnabled) const
+{
+	static constexpr float WireThickness = 2.0f;
+	static constexpr float HalfWireThickness = WireThickness / 2.0f;
+
+	// We want to support drawing wires for the tree
+	//                 Needs Wire Array
+	//   v-[A]         {}
+	//   |-v[B]        {1}
+	//   | '-v[B]      {1,1}
+	//   |   |--[C]    {1,0,1}
+	//   |   |--[D]    {1,0,1}
+	//   |   '--[E]    {1,0,1} 
+	//   |>-[F]        {}
+	//   '--[G]        {}
+	//   
+	//
+
+	const float Indent = IndentAmount;
+	const FSlateBrush* VerticalBarBrush = FAppStyle::GetBrush("WhiteBrush");;
+
+	const float OffsetX = ImageSize.Y * 0.5f + ImagePadding.Left;;
+	const float VerticalWireLoc = ImageSize.Y * 0.5f + ImagePadding.Top;
+
+	if (VerticalBarBrush != nullptr)
+	{
+		const TSharedPtr<ITableRow> OwnerRow = OwnerRowPtr.Pin();
+		const FLinearColor WireTint = WireColor.GetSpecifiedColor();
+
+		// Draw vertical wires to indicate paths to parent nodes.
+		const TBitArray<>& NeedsWireByLevel = OwnerRow->GetWiresNeededByDepth();
+		const int32 NumLevels = NeedsWireByLevel.Num();
+		for (int32 Level = 1; Level < NumLevels; Level++)
+		{
+			const float CurrentIndent = Indent * (Level - 1);
+
+			if (NeedsWireByLevel[Level])
+			{
+				const FVector2f Offset(CurrentIndent + OffsetX, 0);
+				const FVector2f Size(WireThickness, AllottedGeometry.Size.Y);
+				FSlateDrawElement::MakeBox(
+					OutDrawElements,
+					LayerId,
+					AllottedGeometry.ToPaintGeometry(Size, FSlateLayoutTransform(Offset)),
+					VerticalBarBrush,
+					ESlateDrawEffect::None,
+					WireTint
+				);
+			}
+		}
+		
+		// For items that are the last expanded child in a list, we need to draw a special angle connector wire.
+		if (OwnerRow->IsLastChild())
+		{
+			const float CurrentIndent = Indent * (NumLevels - 2);
+			const FVector2f Offset(CurrentIndent + OffsetX, 0);
+			const FVector2f Size(WireThickness, VerticalWireLoc + HalfWireThickness);
+			FSlateDrawElement::MakeBox(
+				OutDrawElements,
+				LayerId,
+				AllottedGeometry.ToPaintGeometry(Size, FSlateLayoutTransform(Offset)),
+				VerticalBarBrush,
+				ESlateDrawEffect::None,
+				WireTint
+			);
+		}
+
+		// If this item is expanded, we need to draw a 1/2-height the line down to its first child cell.
+		if (OwnerRow->IsItemExpanded() && OwnerRow->DoesItemHaveChildren())
+		{
+			const float CurrentIndent = Indent * (NumLevels - 1);
+			const FVector2f Offset(CurrentIndent + OffsetX, ImageSize.Y + ImagePadding.Top);
+			const FVector2f Size(WireThickness,  (AllottedGeometry.Size.Y - (ImageSize.Y + ImagePadding.Top)));
+			FSlateDrawElement::MakeBox(
+				OutDrawElements,
+				LayerId,
+				AllottedGeometry.ToPaintGeometry(Size, FSlateLayoutTransform(Offset)),
+				VerticalBarBrush,
+				ESlateDrawEffect::None,
+				WireTint
+			);
+		}
+
+		// Draw horizontal connector from parent wire to child.
+		if (NumLevels > 1)
+		{
+			const float HorizontalWireStart = (NumLevels - 2) * Indent + OffsetX;
+			const float HorizontalWireEnd = (NumLevels - 1) * Indent + ImagePadding.Left - WireThickness + (OwnerRow->DoesItemHaveChildren() ? 0 : ImageSize.X);
+			const FVector2f Offset(HorizontalWireStart + WireThickness, VerticalWireLoc - WireThickness * 0.5f);
+			const FVector2f Size(HorizontalWireEnd - HorizontalWireStart, WireThickness);
+			FSlateDrawElement::MakeBox(
+				OutDrawElements,
+				LayerId,
+				AllottedGeometry.ToPaintGeometry(Size, FSlateLayoutTransform(Offset)),
+				VerticalBarBrush,
+				ESlateDrawEffect::None,
+				WireTint
+			);
+		}	
+	}
+
+	return SCompoundWidget::OnPaint(Args, AllottedGeometry, MyCullingRect, OutDrawElements, LayerId + 1, InWidgetStyle, bParentEnabled);
+}
+
+/** Invoked when the expanded button is clicked (toggle item expansion) */
+FReply SStateTreeExpanderArrow::OnArrowClicked()
+{
+	// Recurse the expansion if "shift" is being pressed
+	const FModifierKeysState ModKeyState = FSlateApplication::Get().GetModifierKeys();
+	if(ModKeyState.IsShiftDown())
+	{
+		OwnerRowPtr.Pin()->Private_OnExpanderArrowShiftClicked();
+	}
+	else
+	{
+		OwnerRowPtr.Pin()->ToggleExpansion();
+	}
+
+	return FReply::Handled();
+}
+
+/** @return Visible when has children; invisible otherwise */
+EVisibility SStateTreeExpanderArrow::GetExpanderVisibility() const
+{
+	return OwnerRowPtr.Pin()->DoesItemHaveChildren() ? EVisibility::Visible : EVisibility::Hidden;
+}
+
+/** @return the name of an image that should be shown as the expander arrow */
+const FSlateBrush* SStateTreeExpanderArrow::GetExpanderImage() const
+{
+	const bool bIsItemExpanded = OwnerRowPtr.Pin()->IsItemExpanded();
+
+	FName ResourceName;
+	if (bIsItemExpanded)
+	{
+		if ( ExpanderArrow->IsHovered() )
+		{
+			static FName ExpandedHoveredName = "TreeArrow_Expanded_Hovered";
+			ResourceName = ExpandedHoveredName;
+		}
+		else
+		{
+			static FName ExpandedName = "TreeArrow_Expanded";
+			ResourceName = ExpandedName;
+		}
+	}
+	else
+	{
+		if ( ExpanderArrow->IsHovered() )
+		{
+			static FName CollapsedHoveredName = "TreeArrow_Collapsed_Hovered";
+			ResourceName = CollapsedHoveredName;
+		}
+		else
+		{
+			static FName CollapsedName = "TreeArrow_Collapsed";
+			ResourceName = CollapsedName;
+		}
+	}
+
+	return StyleSet->GetBrush(ResourceName);
+}
+
+/** @return the margin corresponding to how far this item is indented */
+FMargin SStateTreeExpanderArrow::GetExpanderPadding() const
+{
+	const int32 NestingDepth = FMath::Max(0, OwnerRowPtr.Pin()->GetIndentLevel() - BaseIndentLevel);
+	const float Indent = IndentAmount;
+	FMargin Padding = ImagePadding;
+	Padding.Left += NestingDepth * Indent;
+	return Padding;
+}
diff --git a/Engine/Plugins/Runtime/StateTree/Source/StateTreeEditorModule/Private/SStateTreeExpanderArrow.h b/Engine/Plugins/Runtime/StateTree/Source/StateTreeEditorModule/Private/SStateTreeExpanderArrow.h
new file mode 100644
--- /dev/null
+++ b/Engine/Plugins/Runtime/StateTree/Source/StateTreeEditorModule/Private/SStateTreeExpanderArrow.h
@@ -0,0 +1,85 @@
+﻿// Copyright Epic Games, Inc. All Rights Reserved.
+
+#pragma once
+
+#include "Misc/Attribute.h"
+#include "Input/Reply.h"
+#include "Layout/Visibility.h"
+#include "Widgets/DeclarativeSyntaxSupport.h"
+#include "Layout/Margin.h"
+#include "Widgets/SCompoundWidget.h"
+#include "Styling/CoreStyle.h"
+
+class ITableRow;
+class SButton;
+class SImage;
+
+/**
+ * Bespoke implementation of expander arrow for State Tree.
+ */
+class SStateTreeExpanderArrow : public SCompoundWidget
+{
+public:
+
+	SLATE_BEGIN_ARGS( SStateTreeExpanderArrow )
+		: _StyleSet(&FCoreStyle::Get())
+		, _IndentAmount(10)
+		, _BaseIndentLevel(0)
+		, _Image(nullptr)
+		, _ImageSize(16.f, 16.f)
+		, _ImagePadding(FMargin())
+	{ }
+		SLATE_ARGUMENT(const ISlateStyle*, StyleSet)
+		SLATE_ARGUMENT(float, IndentAmount)
+		SLATE_ARGUMENT(int32, BaseIndentLevel)
+		SLATE_ATTRIBUTE(const FSlateBrush*, Image)
+		SLATE_ATTRIBUTE(FSlateColor, ColorAndOpacity)
+		SLATE_ARGUMENT(FVector2f, ImageSize)
+		SLATE_ARGUMENT(FSlateColor, WireColorAndOpacity)
+		SLATE_ARGUMENT(FMargin, ImagePadding)
+
+	SLATE_END_ARGS()
+
+	SStateTreeExpanderArrow();
+	void Construct( const FArguments& InArgs, const TSharedPtr<class ITableRow>& TableRow );
+
+protected:
+
+	virtual int32 OnPaint(const FPaintArgs& Args, const FGeometry& AllottedGeometry, const FSlateRect& MyCullingRect, FSlateWindowElementList& OutDrawElements, int32 LayerId, const FWidgetStyle& InWidgetStyle, bool bParentEnabled) const override;
+
+	/** Invoked when the expanded button is clicked (toggle item expansion) */
+	FReply OnArrowClicked();
+
+	/** @return Visible when has children; invisible otherwise */
+	EVisibility GetExpanderVisibility() const;
+
+	/** @return the name of an image that should be shown as the expander arrow */
+	const FSlateBrush* GetExpanderImage() const;
+
+	/** @return the margin corresponding to how far this item is indented */
+	FMargin GetExpanderPadding() const;
+
+	/** Pointer to the owning row. */
+	TWeakPtr<class ITableRow> OwnerRowPtr;
+
+	/** The amount of space to indent at each level */
+	float IndentAmount = 10.0f;
+
+	/** The level in the tree that begins the indention amount */
+	int32 BaseIndentLevel = 0;
+
+	/** Color for the wires */
+	FSlateColor WireColor;
+
+	/** Size of the expander image. */
+	FVector2f ImageSize = FVector2f(16, 16);
+
+	/** Padding for the expander image. */
+	FMargin ImagePadding;
+	
+	/** A reference to the expander button */
+	TSharedPtr<SButton> ExpanderArrow;
+
+	/** The slate style to use */
+	const ISlateStyle* StyleSet = nullptr;
+};
diff --git a/Engine/Plugins/Runtime/StateTree/Source/StateTreeEditorModule/Private/SStateTreeOutliner.cpp b/Engine/Plugins/Runtime/StateTree/Source/StateTreeEditorModule/Private/SStateTreeOutliner.cpp
new file mode 100644
--- /dev/null
+++ b/Engine/Plugins/Runtime/StateTree/Source/StateTreeEditorModule/Private/SStateTreeOutliner.cpp
@@ -0,0 +1,491 @@
+﻿// Copyright Epic Games, Inc. All Rights Reserved.
+
+#include "SStateTreeOutliner.h"
+
+#include "StateTreeDelegates.h"
+#include "StateTreeEditorCommands.h"
+#include "StateTreeEditorData.h"
+#include "StateTreeState.h"
+#include "StateTreeViewModel.h"
+#include "Customizations/Widgets/SCompactStateTreeView.h"
+#include "Debugger/StateTreeDebuggerTypes.h"
+#include "Framework/Commands/UICommandList.h"
+#include "Framework/MultiBox/MultiBoxBuilder.h"
+
+#define LOCTEXT_NAMESPACE "StateTreeEditor"
+
+SStateTreeOutliner::SStateTreeOutliner()
+{
+}
+
+SStateTreeOutliner::~SStateTreeOutliner()
+{
+	if (StateTreeViewModel)
+	{
+		StateTreeViewModel->GetOnAssetChanged().RemoveAll(this);
+		StateTreeViewModel->GetOnStatesRemoved().RemoveAll(this);
+		StateTreeViewModel->GetOnStatesMoved().RemoveAll(this);
+		StateTreeViewModel->GetOnStateAdded().RemoveAll(this);
+		StateTreeViewModel->GetOnStatesChanged().RemoveAll(this);
+		StateTreeViewModel->GetOnSelectionChanged().RemoveAll(this);
+	}
+	UE::StateTree::Delegates::OnVisualThemeChanged.RemoveAll(this);
+}
+
+void SStateTreeOutliner::Construct(const FArguments& InArgs, TSharedRef<FStateTreeViewModel> InStateTreeViewModel, const TSharedRef<FUICommandList>& InCommandList)
+{
+	StateTreeViewModel = InStateTreeViewModel;
+
+	StateTreeViewModel->GetOnAssetChanged().AddSP(this, &SStateTreeOutliner::HandleModelAssetChanged);
+	StateTreeViewModel->GetOnStatesRemoved().AddSP(this, &SStateTreeOutliner::HandleModelStatesRemoved);
+	StateTreeViewModel->GetOnStatesMoved().AddSP(this, &SStateTreeOutliner::HandleModelStatesMoved);
+	StateTreeViewModel->GetOnStateAdded().AddSP(this, &SStateTreeOutliner::HandleModelStateAdded);
+	StateTreeViewModel->GetOnStatesChanged().AddSP(this, &SStateTreeOutliner::HandleModelStatesChanged);
+	StateTreeViewModel->GetOnSelectionChanged().AddSP(this, &SStateTreeOutliner::HandleModelSelectionChanged);
+
+	UE::StateTree::Delegates::OnVisualThemeChanged.AddSP(this, &SStateTreeOutliner::HandleVisualThemeChanged);
+	
+	bUpdatingSelection = false;
+
+	ChildSlot
+	[
+		SAssignNew(CompactStateTreeView, SCompactStateTreeView, StateTreeViewModel)
+		.SelectionMode(ESelectionMode::Multi)
+		.StateTreeEditorData(StateTreeViewModel->GetStateTreeEditorData())
+		.OnSelectionChanged(this, &SStateTreeOutliner::HandleTreeViewSelectionChanged)
+		.OnContextMenuOpening(this, &SStateTreeOutliner::HandleContextMenuOpening)
+	];
+	
+	CommandList = InCommandList;
+	BindCommands();
+}
+
+void SStateTreeOutliner::BindCommands()
+{
+	const FStateTreeEditorCommands& Commands = FStateTreeEditorCommands::Get();
+
+	CommandList->MapAction(
+		Commands.AddSiblingState,
+		FExecuteAction::CreateSP(this, &SStateTreeOutliner::HandleAddSiblingState),
+		FCanExecuteAction());
+
+	CommandList->MapAction(
+		Commands.AddChildState,
+		FExecuteAction::CreateSP(this, &SStateTreeOutliner::HandleAddChildState),
+		FCanExecuteAction::CreateSP(this, &SStateTreeOutliner::HasSelection));
+
+	CommandList->MapAction(
+		Commands.CutStates,
+		FExecuteAction::CreateSP(this, &SStateTreeOutliner::HandleCutSelectedStates),
+		FCanExecuteAction::CreateSP(this, &SStateTreeOutliner::HasSelection));
+
+	CommandList->MapAction(
+		Commands.CopyStates,
+		FExecuteAction::CreateSP(this, &SStateTreeOutliner::HandleCopySelectedStates),
+		FCanExecuteAction::CreateSP(this, &SStateTreeOutliner::HasSelection));
+
+	CommandList->MapAction(
+		Commands.DeleteStates,
+		FExecuteAction::CreateSP(this, &SStateTreeOutliner::HandleDeleteStates),
+		FCanExecuteAction::CreateSP(this, &SStateTreeOutliner::HasSelection));
+
+	CommandList->MapAction(
+		Commands.PasteStatesAsSiblings,
+		FExecuteAction::CreateSP(this, &SStateTreeOutliner::HandlePasteStatesAsSiblings),
+		FCanExecuteAction::CreateSP(this, &SStateTreeOutliner::CanPaste));
+
+	CommandList->MapAction(
+		Commands.PasteStatesAsChildren,
+		FExecuteAction::CreateSP(this, &SStateTreeOutliner::HandlePasteStatesAsChildren),
+		FCanExecuteAction::CreateSP(this, &SStateTreeOutliner::CanPaste));
+
+	CommandList->MapAction(
+		Commands.DuplicateStates,
+		FExecuteAction::CreateSP(this, &SStateTreeOutliner::HandleDuplicateSelectedStates),
+		FCanExecuteAction::CreateSP(this, &SStateTreeOutliner::HasSelection));
+
+	CommandList->MapAction(
+		Commands.EnableStates,
+		FExecuteAction::CreateSP(this, &SStateTreeOutliner::HandleEnableSelectedStates),
+		FCanExecuteAction(),
+		FGetActionCheckState::CreateSPLambda(this, [this]
+			{
+				const bool bCanEnable = CanEnableStates();
+				const bool bCanDisable = CanDisableStates();
+				if (bCanEnable && bCanDisable)
+				{
+					return ECheckBoxState::Undetermined;
+				}
+				
+				if (bCanDisable)
+				{
+					return ECheckBoxState::Checked;
+				}
+
+				if (bCanEnable)
+				{
+					return ECheckBoxState::Unchecked;
+				}
+
+				// Should not happen since action is not visible in this case
+				return ECheckBoxState::Undetermined;
+			}),
+		FIsActionButtonVisible::CreateSPLambda(this, [this]
+		{
+			return CanEnableStates() || CanDisableStates();
+		}));
+
+#if WITH_STATETREE_TRACE_DEBUGGER
+	CommandList->MapAction(
+		Commands.EnableOnEnterStateBreakpoint,
+		FExecuteAction::CreateSPLambda(this, [this]
+		{
+			if (StateTreeViewModel)
+			{
+				StateTreeViewModel->HandleEnableStateBreakpoint(EStateTreeBreakpointType::OnEnter);
+			}
+		}),
+		FCanExecuteAction(),
+		FGetActionCheckState::CreateSPLambda(this, [this]
+		{
+			return StateTreeViewModel ? StateTreeViewModel->GetStateBreakpointCheckState(EStateTreeBreakpointType::OnEnter) : ECheckBoxState::Unchecked;
+		}),
+		FIsActionButtonVisible::CreateSPLambda(this, [this]
+		{
+			return (StateTreeViewModel)
+				&& (StateTreeViewModel->CanAddStateBreakpoint(EStateTreeBreakpointType::OnEnter)
+					|| StateTreeViewModel->CanRemoveStateBreakpoint(EStateTreeBreakpointType::OnEnter));
+		}));
+
+	CommandList->MapAction(
+		Commands.EnableOnExitStateBreakpoint,
+		FExecuteAction::CreateSPLambda(this, [this]
+		{
+			if (StateTreeViewModel)
+			{
+				StateTreeViewModel->HandleEnableStateBreakpoint(EStateTreeBreakpointType::OnExit);
+			}
+		}),
+		FCanExecuteAction(),
+		FGetActionCheckState::CreateSPLambda(this, [this]
+		{
+			return StateTreeViewModel ? StateTreeViewModel->GetStateBreakpointCheckState(EStateTreeBreakpointType::OnExit) : ECheckBoxState::Unchecked;
+		}),
+		FIsActionButtonVisible::CreateSPLambda(this, [this]
+		{
+			return (StateTreeViewModel)
+				&& (StateTreeViewModel->CanAddStateBreakpoint(EStateTreeBreakpointType::OnExit)
+					|| StateTreeViewModel->CanRemoveStateBreakpoint(EStateTreeBreakpointType::OnExit));
+		}));
+#endif // WITH_STATETREE_TRACE_DEBUGGER
+}
+
+
+void SStateTreeOutliner::HandleModelAssetChanged()
+{
+	bItemsDirty = true;
+
+	if (CompactStateTreeView && StateTreeViewModel)
+	{
+		CompactStateTreeView->Refresh(StateTreeViewModel->GetStateTreeEditorData());
+	}
+}
+
+void SStateTreeOutliner::HandleModelStatesRemoved(const TSet<UStateTreeState*>& AffectedParents)
+{
+	bItemsDirty = true;
+	
+	if (CompactStateTreeView && StateTreeViewModel)
+	{
+		CompactStateTreeView->Refresh(StateTreeViewModel->GetStateTreeEditorData());
+	}
+}
+
+void SStateTreeOutliner::HandleModelStatesMoved(const TSet<UStateTreeState*>& AffectedParents, const TSet<UStateTreeState*>& MovedStates)
+{
+	bItemsDirty = true;
+	
+	if (CompactStateTreeView && StateTreeViewModel)
+	{
+		CompactStateTreeView->Refresh(StateTreeViewModel->GetStateTreeEditorData());
+	}
+}
+
+void SStateTreeOutliner::HandleModelStateAdded(UStateTreeState* ParentState, UStateTreeState* NewState)
+{
+	bItemsDirty = true;
+	
+	if (CompactStateTreeView && StateTreeViewModel)
+	{
+		CompactStateTreeView->Refresh(StateTreeViewModel->GetStateTreeEditorData());
+	}
+}
+
+void SStateTreeOutliner::HandleModelStatesChanged(const TSet<UStateTreeState*>& AffectedStates, const FPropertyChangedEvent& PropertyChangedEvent)
+{
+	bool bArraysChanged = false;
+
+	// The purpose of the rebuild below is to update the task visualization (number of widgets change).
+	// This method is called when anything in a state changes, make sure to only rebuild when needed.
+	if (PropertyChangedEvent.MemberProperty->GetFName() == GET_MEMBER_NAME_CHECKED(UStateTreeState, Tasks))
+	{
+		bArraysChanged = true;
+	}
+		
+	if (bArraysChanged)
+	{
+	}
+
+	if (CompactStateTreeView && StateTreeViewModel)
+	{
+		CompactStateTreeView->Refresh(StateTreeViewModel->GetStateTreeEditorData());
+	}
+}
+
+void SStateTreeOutliner::HandleModelSelectionChanged(const TArray<TWeakObjectPtr<UStateTreeState>>& SelectedStates)
+{
+	if (bUpdatingSelection)
+	{
+		return;
+	}
+
+	if (CompactStateTreeView && StateTreeViewModel)
+	{
+		TArray<FGuid> StateIDs;
+		for (const TWeakObjectPtr<UStateTreeState>& WeakState : SelectedStates)
+		{
+			if (const UStateTreeState* State = WeakState.Get())
+			{
+				StateIDs.Add(State->ID);
+			}
+		}
+	
+		CompactStateTreeView->SetSelection(StateIDs);
+	}
+}
+
+void SStateTreeOutliner::HandleTreeViewSelectionChanged(TConstArrayView<FGuid> SelectedStateIDs)
+{
+	if (StateTreeViewModel)
+	{
+		TArray<TWeakObjectPtr<UStateTreeState>> Selection;
+
+		if (const UStateTreeEditorData* StateTreeEditorData = StateTreeViewModel->GetStateTreeEditorData())
+		{
+			for (const FGuid& StateID : SelectedStateIDs)
+			{
+				if (const UStateTreeState* State = StateTreeEditorData->GetStateByID(StateID))
+				{
+					Selection.Add(const_cast<UStateTreeState*>(State));
+				}
+			}
+		}
+		
+		StateTreeViewModel->SetSelection(Selection);
+	}
+}
+
+void SStateTreeOutliner::HandleVisualThemeChanged(const UStateTree& StateTree)
+{
+	if (StateTreeViewModel
+		&& StateTreeViewModel->GetStateTree() == &StateTree)
+	{
+		CompactStateTreeView->Refresh(StateTreeViewModel->GetStateTreeEditorData());
+	}
+}
+
+FReply SStateTreeOutliner::OnKeyDown(const FGeometry& MyGeometry, const FKeyEvent& InKeyEvent)
+{
+	if(CommandList->ProcessCommandBindings(InKeyEvent))
+	{
+		return FReply::Handled();
+	}
+	else
+	{
+		return SCompoundWidget::OnKeyDown(MyGeometry, InKeyEvent);
+	}
+}
+
+TSharedPtr<SWidget> SStateTreeOutliner::HandleContextMenuOpening()
+{
+	if (!StateTreeViewModel)
+	{
+		return nullptr;
+	}
+
+	FMenuBuilder MenuBuilder(true, CommandList);
+
+	MenuBuilder.AddSubMenu(
+		LOCTEXT("AddState", "Add State"),
+		FText(),
+		FNewMenuDelegate::CreateLambda([this](FMenuBuilder& MenuBuilder)
+		{
+			MenuBuilder.AddMenuEntry(FStateTreeEditorCommands::Get().AddSiblingState);
+			MenuBuilder.AddMenuEntry(FStateTreeEditorCommands::Get().AddChildState);
+		}),
+		/*bInOpenSubMenuOnClick =*/false,
+		FSlateIcon(FAppStyle::GetAppStyleSetName(), "Icons.Plus")
+	);
+
+	MenuBuilder.AddSeparator();
+
+	MenuBuilder.AddMenuEntry(FStateTreeEditorCommands::Get().CutStates);
+	MenuBuilder.AddMenuEntry(FStateTreeEditorCommands::Get().CopyStates);
+
+	MenuBuilder.AddSubMenu(
+		LOCTEXT("Paste", "Paste"),
+		FText(),
+		FNewMenuDelegate::CreateLambda([this](FMenuBuilder& MenuBuilder)
+		{
+			MenuBuilder.AddMenuEntry(FStateTreeEditorCommands::Get().PasteStatesAsSiblings);
+			MenuBuilder.AddMenuEntry(FStateTreeEditorCommands::Get().PasteStatesAsChildren);
+		}),
+		/*bInOpenSubMenuOnClick =*/false,
+		FSlateIcon(FAppStyle::GetAppStyleSetName(), "GenericCommands.Paste")
+	);
+	
+	MenuBuilder.AddMenuEntry(FStateTreeEditorCommands::Get().DuplicateStates);
+	MenuBuilder.AddMenuEntry(FStateTreeEditorCommands::Get().DeleteStates);
+	MenuBuilder.AddSeparator();
+	MenuBuilder.AddMenuEntry(FStateTreeEditorCommands::Get().EnableStates);
+
+#if WITH_STATETREE_TRACE_DEBUGGER
+	MenuBuilder.AddSeparator();
+	MenuBuilder.AddMenuEntry(FStateTreeEditorCommands::Get().EnableOnEnterStateBreakpoint);
+	MenuBuilder.AddMenuEntry(FStateTreeEditorCommands::Get().EnableOnExitStateBreakpoint);
+#endif // WITH_STATETREE_TRACE_DEBUGGER
+	
+	return MenuBuilder.MakeWidget();
+}
+
+UStateTreeState* SStateTreeOutliner::GetFirstSelectedState() const
+{
+	TArray<UStateTreeState*> SelectedStates;
+	if (StateTreeViewModel)
+	{
+		StateTreeViewModel->GetSelectedStates(SelectedStates);
+	}
+	return SelectedStates.IsEmpty() ? nullptr : SelectedStates[0];
+}
+
+void SStateTreeOutliner::HandleAddSiblingState()
+{
+	if (StateTreeViewModel)
+	{
+		StateTreeViewModel->AddState(GetFirstSelectedState());
+	}
+}
+
+void SStateTreeOutliner::HandleAddChildState()
+{
+	if (StateTreeViewModel)
+	{
+		UStateTreeState* ParentState = GetFirstSelectedState();
+		if (ParentState)
+		{
+			StateTreeViewModel->AddChildState(ParentState);
+		}
+	}
+}
+
+void SStateTreeOutliner::HandleCutSelectedStates()
+{
+	if (StateTreeViewModel)
+	{
+		StateTreeViewModel->CopySelectedStates();
+		StateTreeViewModel->RemoveSelectedStates();
+	}
+}
+
+void SStateTreeOutliner::HandleCopySelectedStates()
+{
+	if (StateTreeViewModel)
+	{
+		StateTreeViewModel->CopySelectedStates();
+	}
+}
+
+void SStateTreeOutliner::HandlePasteStatesAsSiblings()
+{
+	if (StateTreeViewModel)
+	{
+		StateTreeViewModel->PasteStatesFromClipboard(GetFirstSelectedState());
+	}
+}
+
+void SStateTreeOutliner::HandlePasteStatesAsChildren()
+{
+	if (StateTreeViewModel)
+	{
+		StateTreeViewModel->PasteStatesAsChildrenFromClipboard(GetFirstSelectedState());
+	}
+}
+
+void SStateTreeOutliner::HandleDuplicateSelectedStates()
+{
+	if (StateTreeViewModel)
+	{
+		StateTreeViewModel->DuplicateSelectedStates();
+	}
+}
+
+void SStateTreeOutliner::HandleDeleteStates()
+{
+	if (StateTreeViewModel)
+	{
+		StateTreeViewModel->RemoveSelectedStates();
+	}
+}
+
+void SStateTreeOutliner::HandleEnableSelectedStates()
+{
+	if (StateTreeViewModel)
+	{
+		// Process CanEnable first so in case of undetermined state (mixed selection) we Enable by default. 
+		if (CanEnableStates())
+		{
+			StateTreeViewModel->SetSelectedStatesEnabled(true);	
+		}
+		else if (CanDisableStates())
+		{
+			StateTreeViewModel->SetSelectedStatesEnabled(false);
+		}
+	}
+}
+
+void SStateTreeOutliner::HandleDisableSelectedStates()
+{
+	if (StateTreeViewModel)
+	{
+		StateTreeViewModel->SetSelectedStatesEnabled(false);
+	}
+}
+
+bool SStateTreeOutliner::HasSelection() const
+{
+	return StateTreeViewModel && StateTreeViewModel->HasSelection();
+}
+
+bool SStateTreeOutliner::CanPaste() const
+{
+	return StateTreeViewModel
+			&& StateTreeViewModel->HasSelection()
+			&& StateTreeViewModel->CanPasteStatesFromClipboard();
+}
+
+bool SStateTreeOutliner::CanEnableStates() const
+{
+	return StateTreeViewModel
+			&& StateTreeViewModel->HasSelection()
+			&& StateTreeViewModel->CanEnableStates();
+}
+
+bool SStateTreeOutliner::CanDisableStates() const
+{
+	return StateTreeViewModel
+			&& StateTreeViewModel->HasSelection()
+			&& StateTreeViewModel->CanDisableStates();
+}
+
+
+#undef LOCTEXT_NAMESPACE
diff --git a/Engine/Plugins/Runtime/StateTree/Source/StateTreeEditorModule/Private/SStateTreeOutliner.h b/Engine/Plugins/Runtime/StateTree/Source/StateTreeEditorModule/Private/SStateTreeOutliner.h
new file mode 100644
--- /dev/null
+++ b/Engine/Plugins/Runtime/StateTree/Source/StateTreeEditorModule/Private/SStateTreeOutliner.h
@@ -0,0 +1,70 @@
+﻿// Copyright Epic Games, Inc. All Rights Reserved.
+
+#pragma once
+
+#include "Widgets/SCompoundWidget.h"
+#include "Widgets/Views/STreeView.h"
+
+class UStateTree;
+class FStateTreeViewModel;
+class SCompactStateTreeView;
+struct FPropertyChangedEvent;
+class UStateTreeState;
+class FUICommandList;
+
+class SStateTreeOutliner : public SCompoundWidget
+{
+public:
+	SLATE_BEGIN_ARGS(SStateTreeOutliner) {}
+	SLATE_END_ARGS()
+
+	SStateTreeOutliner();
+	virtual ~SStateTreeOutliner() override;
+
+	void Construct(const FArguments& InArgs, TSharedRef<FStateTreeViewModel> StateTreeViewModel, const TSharedRef<FUICommandList>& InCommandList);
+
+private:
+	// ViewModel handlers
+	void HandleModelAssetChanged();
+	void HandleModelStatesRemoved(const TSet<UStateTreeState*>& AffectedParents);
+	void HandleModelStatesMoved(const TSet<UStateTreeState*>& AffectedParents, const TSet<UStateTreeState*>& MovedStates);
+	void HandleModelStateAdded(UStateTreeState* ParentState, UStateTreeState* NewState);
+	void HandleModelStatesChanged(const TSet<UStateTreeState*>& AffectedStates, const FPropertyChangedEvent& PropertyChangedEvent);
+	void HandleModelSelectionChanged(const TArray<TWeakObjectPtr<UStateTreeState>>& SelectedStates);
+	void HandleTreeViewSelectionChanged(TConstArrayView<FGuid> SelectedStateIDs);
+	void HandleVisualThemeChanged(const UStateTree& StateTree);
+
+	virtual FReply OnKeyDown(const FGeometry& MyGeometry, const FKeyEvent& InKeyEvent) override;
+	
+	TSharedPtr<SWidget> HandleContextMenuOpening();
+
+	// Action handlers
+	// @todo: these are also defined in the SStateTreeView, figure out how to share code.
+	UStateTreeState* GetFirstSelectedState() const;
+	void HandleAddSiblingState();
+	void HandleAddChildState();
+	void HandleCutSelectedStates();
+	void HandleCopySelectedStates();
+	void HandlePasteStatesAsSiblings();
+	void HandlePasteStatesAsChildren();
+	void HandleDuplicateSelectedStates();
+	void HandleDeleteStates();
+	void HandleEnableSelectedStates();
+	void HandleDisableSelectedStates();
+
+	bool HasSelection() const;
+	bool CanPaste() const;
+	bool CanEnableStates() const;
+	bool CanDisableStates() const;
+
+	void BindCommands();
+
+	TSharedPtr<FStateTreeViewModel> StateTreeViewModel;
+
+	TSharedPtr<SCompactStateTreeView> CompactStateTreeView;
+	
+	TSharedPtr<FUICommandList> CommandList;
+
+	bool bItemsDirty = false;
+	bool bUpdatingSelection = false;
+};
diff --git a/Engine/Plugins/Runtime/StateTree/Source/StateTreeEditorModule/Private/SStateTreeView.cpp b/Engine/Plugins/Runtime/StateTree/Source/StateTreeEditorModule/Private/SStateTreeView.cpp
--- a/Engine/Plugins/Runtime/StateTree/Source/StateTreeEditorModule/Private/SStateTreeView.cpp
+++ b/Engine/Plugins/Runtime/StateTree/Source/StateTreeEditorModule/Private/SStateTreeView.cpp
@@ -64,7 +64,7 @@ void SStateTreeView::Construct(const FArguments& InArgs, TSharedRef<FStateTreeVi
 		.OnSelectionChanged(this, &SStateTreeView::HandleTreeSelectionChanged)
 		.OnExpansionChanged(this, &SStateTreeView::HandleTreeExpansionChanged)
 		.OnContextMenuOpening(this, &SStateTreeView::HandleContextMenuOpening)
-		.AllowOverscroll(EAllowOverscroll::No)
+		.AllowOverscroll(EAllowOverscroll::Yes)
 		.ExternalScrollbar(VerticalScrollBar);
 
 	ChildSlot
@@ -407,6 +407,14 @@ void SStateTreeView::HandleModelStatesChanged(const TSet<UStateTreeState*>& Affe
 	{
 		bArraysChanged = true;
 	}
+	if (PropertyChangedEvent.MemberProperty->GetFName() == GET_MEMBER_NAME_CHECKED(UStateTreeState, Transitions))
+	{
+		bArraysChanged = true;
+	}
+	if (PropertyChangedEvent.MemberProperty->GetFName() == GET_MEMBER_NAME_CHECKED(UStateTreeState, EnterConditions))
+	{
+		bArraysChanged = true;
+	}
 		
 	if (bArraysChanged)
 	{
diff --git a/Engine/Plugins/Runtime/StateTree/Source/StateTreeEditorModule/Private/SStateTreeView.h b/Engine/Plugins/Runtime/StateTree/Source/StateTreeEditorModule/Private/SStateTreeView.h
--- a/Engine/Plugins/Runtime/StateTree/Source/StateTreeEditorModule/Private/SStateTreeView.h
+++ b/Engine/Plugins/Runtime/StateTree/Source/StateTreeEditorModule/Private/SStateTreeView.h
@@ -15,27 +15,12 @@ class UStateTreeState;
 class SScrollBox;
 class FUICommandList;
 
-class FActionTreeViewDragDrop : public FDragDropOperation
+namespace UE::StateTree::Editor
 {
-public:
-	DRAG_DROP_OPERATOR_TYPE(FActionTreeViewDragDrop, FDragDropOperation);
-
-	static TSharedRef<FActionTreeViewDragDrop> New(const UStateTreeState* InState)
-	{
-		return MakeShareable(new FActionTreeViewDragDrop(InState));
-	}
-
-	const UStateTreeState* GetDraggedState() const { return State; }
-
-private:
-	FActionTreeViewDragDrop(const UStateTreeState* InState)
-		: State(InState)
-	{
-	}
-
-	const UStateTreeState* State;
-};
+static constexpr float StateRowHeight = 32.0f;
+static constexpr float TaskRowHeight = 16.0f;
 
+} // UE::StateTree::Editor
 
 class SStateTreeView : public SCompoundWidget
 {
@@ -73,6 +58,7 @@ private:
 	TSharedPtr<SWidget> HandleContextMenuOpening();
 
 	// Action handlers
+	// @todo: these are also defined in the outliner, figure out how to share code.
 	UStateTreeState* GetFirstSelectedState() const;
 	FReply HandleAddStateButton();
 	void HandleAddSiblingState();
diff --git a/Engine/Plugins/Runtime/StateTree/Source/StateTreeEditorModule/Private/SStateTreeViewRow.cpp b/Engine/Plugins/Runtime/StateTree/Source/StateTreeEditorModule/Private/SStateTreeViewRow.cpp
--- a/Engine/Plugins/Runtime/StateTree/Source/StateTreeEditorModule/Private/SStateTreeViewRow.cpp
+++ b/Engine/Plugins/Runtime/StateTree/Source/StateTreeEditorModule/Private/SStateTreeViewRow.cpp
@@ -2,7 +2,7 @@
 
 #include "SStateTreeViewRow.h"
 #include "SStateTreeView.h"
-#include "EditorFontGlyphs.h"
+#include "SStateTreeExpanderArrow.h"
 #include "StateTreeEditor.h"
 #include "StateTreeEditorData.h"
 #include "StateTreeEditorStyle.h"
@@ -10,13 +10,18 @@
 #include "Widgets/Layout/SScrollBox.h"
 #include "Widgets/Layout/SSpacer.h"
 #include "StateTree.h"
+#include "StateTreeConditionBase.h"
+#include "StateTreeDescriptionHelpers.h"
+#include "StateTreeDragDrop.h"
 #include "StateTreeState.h"
 #include "StateTreeTaskBase.h"
 #include "StateTreeViewModel.h"
 #include "Widgets/Views/SListView.h"
 #include "TextStyleDecorator.h"
 #include "Customizations/StateTreeEditorNodeUtils.h"
+#include "Widgets/Layout/SWrapBox.h"
 #include "Widgets/Text/SRichTextBlock.h"
+#include "Widgets/Input/SButton.h"
 
 #define LOCTEXT_NAMESPACE "StateTreeEditor"
 
@@ -42,16 +47,94 @@ void SStateTreeViewRow::Construct(const FArguments& InArgs, const TSharedRef<STa
 	WeakEditorData = State != nullptr ? State->GetTypedOuter<UStateTreeEditorData>() : nullptr;
 
 	ConstructInternal(STableRow::FArguments()
-		.Padding(5.f)
 		.OnDragDetected(this, &SStateTreeViewRow::HandleDragDetected)
+		.OnDragLeave(this, &SStateTreeViewRow::HandleDragLeave)
 		.OnCanAcceptDrop(this, &SStateTreeViewRow::HandleCanAcceptDrop)
 		.OnAcceptDrop(this, &SStateTreeViewRow::HandleAcceptDrop)
 		.Style(&FStateTreeEditorStyle::Get().GetWidgetStyle<FTableRowStyle>("StateTree.Selection"))
 		, InOwnerTableView);
 
-	static const FLinearColor LinkBackground = FLinearColor(FColor(84, 84, 84));
 	static constexpr FLinearColor IconTint = FLinearColor(1, 1, 1, 0.5f);
 
+	const float ExpanderArrowTopPadding = FMath::FloorToFloat((UE::StateTree::Editor::StateRowHeight - 16.f) / 2.f);
+
+	
+	auto MakeTransitionWidget = [this](const EStateTreeTransitionTrigger Trigger, const FSlateBrush* Icon)
+	{
+		FTransitionDescFilterOptions FilterOptions;
+		FilterOptions.bUseMask = EnumHasAnyFlags(Trigger, EStateTreeTransitionTrigger::OnTick | EStateTreeTransitionTrigger::OnEvent);
+		
+			return SNew(SBox)
+			.HeightOverride(UE::StateTree::Editor::StateRowHeight)
+			.Visibility(this, &SStateTreeViewRow::GetTransitionsVisibility, Trigger)
+			[
+				SNew(SHorizontalBox)
+				+ SHorizontalBox::Slot()
+				.VAlign(VAlign_Center)
+				.AutoWidth()
+				.Padding(FMargin(0.f, 0.f, 0.f, 0.f))
+				[
+					SNew(SHorizontalBox)
+					
+					+ SHorizontalBox::Slot()
+					.VAlign(VAlign_Center)
+					.AutoWidth()
+					[
+						SNew(SImage)
+						.Image(Icon)
+					]
+					
+					+ SHorizontalBox::Slot()
+					.VAlign(VAlign_Center)
+					.AutoWidth()
+					[
+						SNew(SImage)
+						.Image(this, &SStateTreeViewRow::GetTransitionsIcon, Trigger)
+						.ColorAndOpacity(IconTint)
+					]
+				]
+				+ SHorizontalBox::Slot()
+				.VAlign(VAlign_Center)
+				.AutoWidth()
+				.Padding(FMargin(4.f, 0.f, 12.f, 0.f))
+				[
+					SNew(SOverlay)
+					+ SOverlay::Slot()
+					[
+						MakeTransitionWidgets(Trigger, FilterOptions)
+					]
+					+ SOverlay::Slot()
+					[
+						// Breakpoint box
+						SNew(SHorizontalBox)
+						+ SHorizontalBox::Slot()
+						.VAlign(VAlign_Top)
+						.HAlign(HAlign_Left)
+						.AutoWidth()
+						[
+							SNew(SBox)
+							.Padding(FMargin(-4.f, -4.f, 0.f, 0.f))
+							[
+								SNew(SImage)
+								.DesiredSizeOverride(FVector2D(10.f, 10.f))
+								.Image(FStateTreeEditorStyle::Get().GetBrush(TEXT("StateTreeEditor.Debugger.Breakpoint.EnabledAndValid")))
+								.Visibility(this, &SStateTreeViewRow::GetTransitionsBreakpointVisibility, Trigger)
+								.ToolTipText_Lambda([this, Trigger, InFilterOptions = FilterOptions]
+									{
+										FTransitionDescFilterOptions FilterOptions = InFilterOptions;
+										FilterOptions.WithBreakpoint = ETransitionDescRequirement::RequiredTrue;
+
+										return FText::Format(LOCTEXT("TransitionBreakpointTooltip","Break when executing transition: {0}"),
+											GetTransitionsDesc(Trigger, FilterOptions));
+									})
+							]
+						]
+					]
+
+				]
+			];
+	};
+	
 	this->ChildSlot
 	.HAlign(HAlign_Fill)
 	[
@@ -67,6 +150,7 @@ void SStateTreeViewRow::Construct(const FArguments& InArgs, const TSharedRef<STa
 				}
 				return 0.f;
 			})
+		.Padding(FMargin(0, 0, 0, 0))
 		[
 			SNew(SHorizontalBox)
 			+ SHorizontalBox::Slot()
@@ -74,163 +158,213 @@ void SStateTreeViewRow::Construct(const FArguments& InArgs, const TSharedRef<STa
 			.HAlign(HAlign_Left)
 			.AutoWidth()
 			[
-				SNew(SExpanderArrow, SharedThis(this))
-				.ShouldDrawWires(true)
-				.IndentAmount(32.f)
+				SNew(SStateTreeExpanderArrow, SharedThis(this))
+				.IndentAmount(24.f)
 				.BaseIndentLevel(0)
+				.ImageSize(FVector2f(16,16))
+				.ImagePadding(FMargin(9,14,0,0))
+				.Image(this, &SStateTreeViewRow::GetSelectorIcon)
+				.ColorAndOpacity(FLinearColor(1, 1, 1, 0.2f))
+				.WireColorAndOpacity(FLinearColor(1, 1, 1, 0.2f))
 			]
 
 			+ SHorizontalBox::Slot()
 			.VAlign(VAlign_Fill)
 			.HAlign(HAlign_Left)
-			.Padding(FMargin(0.f, 4.f))
 			.AutoWidth()
+			.Padding(FMargin(0, 6, 0, 6))
 			[
-				SNew(SBox)
-				.HeightOverride(28.f)
-				.VAlign(VAlign_Fill)
+				// State and tasks
+				SNew(SVerticalBox)
+				+ SVerticalBox::Slot()
+				.AutoHeight()
 				[
-					SNew(SBorder)
-					.BorderImage(FStateTreeEditorStyle::Get().GetBrush("StateTree.State.Border"))
-					.BorderBackgroundColor(this, &SStateTreeViewRow::GetActiveStateColor)
+					// State
+					SNew(SBox)
+					.HeightOverride(UE::StateTree::Editor::StateRowHeight)
+					.HAlign(HAlign_Left)
 					[
 						SNew(SHorizontalBox)
 
-						// Sub tree marker
-						+ SHorizontalBox::Slot()
-						.VAlign(VAlign_Center)
-						.AutoWidth()
-						[
-							SNew(SBox)
-							.WidthOverride(4.f)
-							.HeightOverride(28.f)
-							[
-								SNew(SBorder)
-								.BorderImage(FAppStyle::GetBrush("WhiteBrush"))
-								.BorderBackgroundColor(this, &SStateTreeViewRow::GetSubTreeMarkerColor)
-							]
-						]
-
 						// State Box
 						+ SHorizontalBox::Slot()
 						.VAlign(VAlign_Center)
+						//.FillWidth(1.f)
 						.AutoWidth()
 						[
 							SNew(SBox)
-							.HeightOverride(28.f)
+							.HeightOverride(UE::StateTree::Editor::StateRowHeight)
 							.VAlign(VAlign_Fill)
 							[
 								SNew(SBorder)
-								.BorderImage(FAppStyle::GetBrush("WhiteBrush"))
-								.BorderBackgroundColor(this, &SStateTreeViewRow::GetTitleColor, 1.0f, 0.0f)
-								.Padding(FMargin(4.f, 0.f, 12.f, 0.f))
-								.IsEnabled_Lambda([InState]
+								.BorderImage(FStateTreeEditorStyle::Get().GetBrush("StateTree.State.Border"))
+								.BorderBackgroundColor(this, &SStateTreeViewRow::GetActiveStateColor)
+								[
+									SNew(SBorder)
+									.BorderImage(FStateTreeEditorStyle::Get().GetBrush("StateTree.State"))
+									.BorderBackgroundColor(this, &SStateTreeViewRow::GetTitleColor, 1.0f, 0.0f)
+									.Padding(FMargin(0.f, 0.f, 12.f, 0.f))
+									.IsEnabled_Lambda([InState]
 									{
 										const UStateTreeState* State = InState.Get();
 										return State != nullptr && State->bEnabled;
 									})
-								[
-									SNew(SOverlay)
-									+ SOverlay::Slot()
 									[
-										SNew(SHorizontalBox)
-
-										// Warnings
-										+SHorizontalBox::Slot()
-										.VAlign(VAlign_Center)
-										.AutoWidth()
+										SNew(SOverlay)
+										+ SOverlay::Slot()
 										[
-											SNew(SBox)
-											.Padding(FMargin(0.f, 0.f, 4.f, 0.f))
-											.Visibility(this, &SStateTreeViewRow::GetWarningsVisibility)
+											SNew(SHorizontalBox)
+
+											// Sub tree marker
+											+ SHorizontalBox::Slot()
+											.VAlign(VAlign_Center)
+											.AutoWidth()
+											.Padding(FMargin(0,0,0,0))
 											[
-												SNew(SImage)
-												.Image(FAppStyle::Get().GetBrush("Icons.Warning"))
-												.ToolTipText(this, &SStateTreeViewRow::GetWarningsTooltipText)
+												SNew(SBox)
+												.WidthOverride(4.f)
+												.HeightOverride(UE::StateTree::Editor::StateRowHeight)
+												.Visibility(this, &SStateTreeViewRow::GetSubTreeVisibility)
+												.VAlign(VAlign_Fill)
+												.HAlign(HAlign_Fill)
+												[
+													SNew(SBorder)
+													.BorderImage(FAppStyle::GetBrush("WhiteBrush"))
+													.BorderBackgroundColor(FLinearColor(1,1,1,0.25f))
+												]
 											]
-										]
 
-										// Conditions icon
-										+SHorizontalBox::Slot()
-										.VAlign(VAlign_Center)
-										.AutoWidth()
-										[
-											SNew(SBox)
-											.Padding(FMargin(0.f, 0.f, 4.f, 0.f))
-											.Visibility(this, &SStateTreeViewRow::GetConditionVisibility)
+											// Conditions icon
+											+SHorizontalBox::Slot()
+											.VAlign(VAlign_Center)
+											.AutoWidth()
 											[
-												SNew(SImage)
-												.ColorAndOpacity(IconTint)
-												.Image(FStateTreeEditorStyle::Get().GetBrush("StateTreeEditor.StateConditions"))
-												.ToolTipText(LOCTEXT("StateHasEnterConditions", "State selection is guarded with enter conditions."))
+												SNew(SBox)
+												.Padding(FMargin(4.f, 0.f, -4.f, 0.f))
+												.Visibility(this, &SStateTreeViewRow::GetConditionVisibility)
+												[
+													SNew(SImage)
+													.ColorAndOpacity(IconTint)
+													.Image(FStateTreeEditorStyle::Get().GetBrush("StateTreeEditor.StateConditions"))
+													.ToolTipText(LOCTEXT("StateHasEnterConditions", "State selection is guarded with enter conditions."))
+												]
 											]
-										]
 
-										// Selector icon
-										+ SHorizontalBox::Slot()
-										.VAlign(VAlign_Center)
-										.AutoWidth()
-										[
-											SNew(SBox)
-											.Padding(FMargin(0.f, 0.f, 4.f, 0.f))
+											// Selector icon
+											+ SHorizontalBox::Slot()
+											.VAlign(VAlign_Center)
+											.AutoWidth()
 											[
-												SNew(SImage)
-												.Image(this, &SStateTreeViewRow::GetSelectorIcon)
-												.ColorAndOpacity(IconTint)
-												.ToolTipText(this, &SStateTreeViewRow::GetSelectorTooltip)
+												SNew(SBox)
+												.Padding(FMargin(4.f, 0.f, 0.f, 0.f))
+												[
+													SNew(SImage)
+													.Image(this, &SStateTreeViewRow::GetSelectorIcon)
+													.ColorAndOpacity(IconTint)
+													.ToolTipText(this, &SStateTreeViewRow::GetSelectorTooltip)
+												]
 											]
-										]
 
-										// State Name
-										+ SHorizontalBox::Slot()
-										.VAlign(VAlign_Center)
-										.AutoWidth()
-										[
-											SAssignNew(NameTextBlock, SInlineEditableTextBlock)
-											.Style(FStateTreeEditorStyle::Get(), "StateTree.State.TitleInlineEditableText")
-											.OnVerifyTextChanged_Lambda([](const FText& NewLabel, FText& OutErrorMessage)
+											// Warnings
+											+SHorizontalBox::Slot()
+											.VAlign(VAlign_Center)
+											.AutoWidth()
+											[
+												SNew(SBox)
+												.Padding(FMargin(2.f, 0.f, 2.f, 1.f))
+												.Visibility(this, &SStateTreeViewRow::GetWarningsVisibility)
+												[
+													SNew(SImage)
+													.Image(FAppStyle::Get().GetBrush("Icons.Warning"))
+													.ToolTipText(this, &SStateTreeViewRow::GetWarningsTooltipText)
+												]
+											]
+											
+											// State Name
+											+ SHorizontalBox::Slot()
+											.VAlign(VAlign_Center)
+											.AutoWidth()
+											[
+												SAssignNew(NameTextBlock, SInlineEditableTextBlock)
+												.Style(FStateTreeEditorStyle::Get(), "StateTree.State.TitleInlineEditableText")
+												.OnVerifyTextChanged_Lambda([](const FText& NewLabel, FText& OutErrorMessage)
+													{
+														return !NewLabel.IsEmptyOrWhitespace();
+													})
+												.OnTextCommitted(this, &SStateTreeViewRow::HandleNodeLabelTextCommitted)
+												.Text(this, &SStateTreeViewRow::GetStateDesc)
+												.ToolTipText(this, &SStateTreeViewRow::GetStateTypeTooltip)
+												.Clipping(EWidgetClipping::ClipToBounds)
+												.IsSelected(this, &SStateTreeViewRow::IsStateSelected)
+											]
+
+											// Linked State
+											+ SHorizontalBox::Slot()
+											.VAlign(VAlign_Center)
+											.AutoWidth()
+											[
+												SNew(SBox)
+												.HeightOverride(UE::StateTree::Editor::StateRowHeight)
+												.VAlign(VAlign_Fill)
+												.Visibility(this, &SStateTreeViewRow::GetLinkedStateVisibility)
+												[
+													// Link icon
+													SNew(SHorizontalBox)
+													+ SHorizontalBox::Slot()
+													.VAlign(VAlign_Center)
+													.AutoWidth()
+													.Padding(FMargin(4.f, 0.f, 4.f, 0.f))
+													[
+														SNew(SImage)
+														.ColorAndOpacity(IconTint)
+														.Image(FStateTreeEditorStyle::Get().GetBrush("StateTreeEditor.StateLinked"))
+													]
+
+													// Linked State
+													+ SHorizontalBox::Slot()
+													.VAlign(VAlign_Center)
+													.AutoWidth()
+													[
+														SNew(STextBlock)
+														.Text(this, &SStateTreeViewRow::GetLinkedStateDesc)
+														.TextStyle(FStateTreeEditorStyle::Get(), "StateTree.Details")
+													]
+												]
+											]
+											
+											// State ID
+											+ SHorizontalBox::Slot()
+											.VAlign(VAlign_Center)
+											.AutoWidth()
+											[
+												SNew(STextBlock)
+												.Visibility_Lambda([]()
 												{
-													return !NewLabel.IsEmptyOrWhitespace();
+													return UE::StateTree::Editor::GbDisplayItemIds ? EVisibility::Visible : EVisibility::Collapsed;
 												})
-											.OnTextCommitted(this, &SStateTreeViewRow::HandleNodeLabelTextCommitted)
-											.Text(this, &SStateTreeViewRow::GetStateDesc)
-											.ToolTipText(this, &SStateTreeViewRow::GetStateTypeTooltip)
-											.Clipping(EWidgetClipping::ClipToBounds)
-											.IsSelected(this, &SStateTreeViewRow::IsStateSelected)
+												.Text(this, &SStateTreeViewRow::GetStateIDDesc)
+												.TextStyle(FStateTreeEditorStyle::Get(), "StateTree.Details")
+											]
 										]
-
-										// State ID
-										+ SHorizontalBox::Slot()
-										.VAlign(VAlign_Center)
-										.AutoWidth()
+										+ SOverlay::Slot()
 										[
-											SNew(STextBlock)
-											.Visibility_Lambda([]()
-											{
-												return UE::StateTree::Editor::GbDisplayItemIds ? EVisibility::Visible : EVisibility::Collapsed;
-											})
-											.Text(this, &SStateTreeViewRow::GetStateIDDesc)
-											.TextStyle(FStateTreeEditorStyle::Get(), "StateTree.Details")
-										]
-									]
-									+ SOverlay::Slot()
-									[
-										SNew(SHorizontalBox)
+											SNew(SHorizontalBox)
 
-										// State breakpoint box
-										+ SHorizontalBox::Slot()
-										.VAlign(VAlign_Top)
-										.HAlign(HAlign_Left)
-										.AutoWidth()
-										[
-											SNew(SBox)
-											.Padding(FMargin(-12.f, -6.f, 0.f, 0.f))
+											// State breakpoint box
+											+ SHorizontalBox::Slot()
+											.VAlign(VAlign_Top)
+											.HAlign(HAlign_Left)
+											.AutoWidth()
 											[
-												SNew(SImage)
-												.DesiredSizeOverride(FVector2D(12.f, 12.f))
-												.Image(FStateTreeEditorStyle::Get().GetBrush(TEXT("StateTreeEditor.Debugger.Breakpoint.EnabledAndValid")))
-												.Visibility(this, &SStateTreeViewRow::GetStateBreakpointVisibility)
-												.ToolTipText(this, &SStateTreeViewRow::GetStateBreakpointTooltipText)
+												SNew(SBox)
+												.Padding(FMargin(-12.f, -6.f, 0.f, 0.f))
+												[
+													SNew(SImage)
+													.DesiredSizeOverride(FVector2D(12.f, 12.f))
+													.Image(FStateTreeEditorStyle::Get().GetBrush(TEXT("StateTreeEditor.Debugger.Breakpoint.EnabledAndValid")))
+													.Visibility(this, &SStateTreeViewRow::GetStateBreakpointVisibility)
+													.ToolTipText(this, &SStateTreeViewRow::GetStateBreakpointTooltipText)
+												]
 											]
 										]
 									]
@@ -238,361 +372,114 @@ void SStateTreeViewRow::Construct(const FArguments& InArgs, const TSharedRef<STa
 							]
 						]
 						
-						// Linked State box
 						+ SHorizontalBox::Slot()
 						.VAlign(VAlign_Fill)
-						.AutoWidth()
-						[
-							SNew(SBox)
-							.HeightOverride(28.f)
-							.VAlign(VAlign_Fill)
-							.Visibility(this, &SStateTreeViewRow::GetLinkedStateVisibility)
-							[
-								SNew(SBorder)
-								.BorderImage(FAppStyle::GetBrush("WhiteBrush"))
-								.BorderBackgroundColor(LinkBackground)
-								.Padding(FMargin(6.f, 0.f, 12.f, 0.f))
-								[
-									// Link icon
-									SNew(SHorizontalBox)
-									+ SHorizontalBox::Slot()
-									.VAlign(VAlign_Center)
-									.AutoWidth()
-									.Padding(FMargin(0.f, 0.f, 4.f, 0.f))
-									[
-										SNew(SImage)
-										.ColorAndOpacity(IconTint)
-										.Image(FStateTreeEditorStyle::Get().GetBrush("StateTreeEditor.StateLinked"))
-									]
-
-									// Linked State
-									+ SHorizontalBox::Slot()
-									.VAlign(VAlign_Center)
-									.AutoWidth()
-									[
-										SNew(STextBlock)
-										.Text(this, &SStateTreeViewRow::GetLinkedStateDesc)
-										.TextStyle(FStateTreeEditorStyle::Get(), "StateTree.Details")
-									]
-								]
-							]
-						]
-						// Tasks
-						+ SHorizontalBox::Slot()
-						.VAlign(VAlign_Fill)
-						.AutoWidth()
+						.HAlign(HAlign_Left)
 						[
-							SNew(SBox)
-							.VAlign(VAlign_Fill)
-							.Visibility(this, &SStateTreeViewRow::GetTasksVisibility)
-							[
-								CreateTasksWidget()
-							]
-						]
-					]
-				]
-			]
-
-			// Completed transitions
-			+ SHorizontalBox::Slot()
-			.VAlign(VAlign_Center)
-			.AutoWidth()
-			[
-				SNew(SBox)
-				.Visibility(this, &SStateTreeViewRow::GetCompletedTransitionVisibility)
-				[
-					SNew(SHorizontalBox)
-					+ SHorizontalBox::Slot()
-					.VAlign(VAlign_Center)
-					.AutoWidth()
-					.Padding(FMargin(8.f, 0.f, 0.f, 0.f))
-					[
-						SNew(SOverlay)
-						+ SOverlay::Slot()
-						[
-							SNew(SImage)
-							.Image(this, &SStateTreeViewRow::GetCompletedTransitionsIcon)
-							.ColorAndOpacity(IconTint)
-						]
-						+ SOverlay::Slot()
-						[
-							// Breakpoint box
+							// Transitions
 							SNew(SHorizontalBox)
+
 							+ SHorizontalBox::Slot()
 							.VAlign(VAlign_Top)
-							.HAlign(HAlign_Left)
 							.AutoWidth()
 							[
 								SNew(SBox)
-								.Padding(FMargin(0.f, -10.f, 0.f, 0.f))
+								.HeightOverride(UE::StateTree::Editor::StateRowHeight)
+								.Visibility(this, &SStateTreeViewRow::GetTransitionDashVisibility)
+								.VAlign(VAlign_Center)
 								[
 									SNew(SImage)
-									.DesiredSizeOverride(FVector2D(10.f, 10.f))
-									.Image(FStateTreeEditorStyle::Get().GetBrush(TEXT("StateTreeEditor.Debugger.Breakpoint.EnabledAndValid")))
-									.Visibility(this, &SStateTreeViewRow::GetCompletedTransitionBreakpointVisibility)
-									.ToolTipText_Lambda([this]
-										{
-											return FText::Format(LOCTEXT("TransitionBreakpointTooltip","Break when executing transition: {0}"),
-												GetCompletedTransitionWithBreakpointDesc());
-										})
+									.Image(FStateTreeEditorStyle::Get().GetBrush("StateTreeEditor.Transition.Dash"))
+									.ColorAndOpacity(IconTint)
 								]
 							]
-						]
-					]
-					+ SHorizontalBox::Slot()
-					.VAlign(VAlign_Center)
-					.AutoWidth()
-					.Padding(FMargin(4.f, 0.f, 0.f, 0.f))
-					[
-						SNew(SRichTextBlock)
-						.Text(this, &SStateTreeViewRow::GetCompletedTransitionsDesc)
-						.TextStyle(&FStateTreeEditorStyle::Get().GetWidgetStyle<FTextBlockStyle>("Normal.Normal"))
-						.OverflowPolicy(ETextOverflowPolicy::Ellipsis)
-						+SRichTextBlock::Decorator(FTextStyleDecorator::Create(TEXT(""), FStateTreeEditorStyle::Get().GetWidgetStyle<FTextBlockStyle>("Normal.Normal")))
-						+SRichTextBlock::Decorator(FTextStyleDecorator::Create(TEXT("b"), FStateTreeEditorStyle::Get().GetWidgetStyle<FTextBlockStyle>("Normal.Bold")))
-						+SRichTextBlock::Decorator(FTextStyleDecorator::Create(TEXT("i"), FStateTreeEditorStyle::Get().GetWidgetStyle<FTextBlockStyle>("Normal.Italic")))
-						+SRichTextBlock::Decorator(FTextStyleDecorator::Create(TEXT("s"), FStateTreeEditorStyle::Get().GetWidgetStyle<FTextBlockStyle>("Normal.Subdued")))
-					]
-				]
-			]
 
-			// Succeeded transitions
-			+ SHorizontalBox::Slot()
-			.VAlign(VAlign_Center)
-			.AutoWidth()
-			[
-				SNew(SBox)
-				.Visibility(this, &SStateTreeViewRow::GetSucceededTransitionVisibility)
-				[
-					SNew(SHorizontalBox)
-					+ SHorizontalBox::Slot()
-					.VAlign(VAlign_Center)
-					.Padding(FMargin(8.f, 0.f, 0.f, 0.f))
-					.AutoWidth()
-					[
-						SNew(SImage)
-						.Image(FStateTreeEditorStyle::Get().GetBrush("StateTreeEditor.Transition.Succeeded"))
-					]
-					+ SHorizontalBox::Slot()
-					.VAlign(VAlign_Center)
-					.AutoWidth()
-					[
-						SNew(SOverlay)
-						+ SOverlay::Slot()
-						[
-							SNew(SImage)
-							.Image(this, &SStateTreeViewRow::GetSucceededTransitionIcon)
-							.ColorAndOpacity(IconTint)
-						]
-						+ SOverlay::Slot()
-						[
-							// Breakpoint box
-							SNew(SHorizontalBox)
+							// Completed transitions
 							+ SHorizontalBox::Slot()
 							.VAlign(VAlign_Top)
-							.HAlign(HAlign_Left)
 							.AutoWidth()
 							[
-								SNew(SBox)
-								.Padding(FMargin(0.f, -10.f, 0.f, 0.f))
-								[
-									SNew(SImage)
-									.DesiredSizeOverride(FVector2D(10.f, 10.f))
-									.Image(FStateTreeEditorStyle::Get().GetBrush(TEXT("StateTreeEditor.Debugger.Breakpoint.EnabledAndValid")))
-									.Visibility(this, &SStateTreeViewRow::GetSucceededTransitionBreakpointVisibility)
-									.ToolTipText_Lambda([this]
-										{
-											return FText::Format(LOCTEXT("TransitionBreakpointTooltip", "Break when executing transition: {0}"),
-												GetSucceededTransitionWithBreakpointDesc());
-										})
-								]
+								MakeTransitionWidget(EStateTreeTransitionTrigger::OnStateCompleted, nullptr)
 							]
-						]
-					]
-					+ SHorizontalBox::Slot()
-					.VAlign(VAlign_Center)
-					.AutoWidth()
-					.Padding(FMargin(4.f, 0.f, 0.f, 0.f))
-					[
-						SNew(SRichTextBlock)
-						.Text(this, &SStateTreeViewRow::GetSucceededTransitionDesc)
-						.TextStyle(&FStateTreeEditorStyle::Get().GetWidgetStyle<FTextBlockStyle>("Normal.Normal"))
-						.OverflowPolicy(ETextOverflowPolicy::Ellipsis)
-						+SRichTextBlock::Decorator(FTextStyleDecorator::Create(TEXT(""), FStateTreeEditorStyle::Get().GetWidgetStyle<FTextBlockStyle>("Normal.Normal")))
-						+SRichTextBlock::Decorator(FTextStyleDecorator::Create(TEXT("b"), FStateTreeEditorStyle::Get().GetWidgetStyle<FTextBlockStyle>("Normal.Bold")))
-						+SRichTextBlock::Decorator(FTextStyleDecorator::Create(TEXT("i"), FStateTreeEditorStyle::Get().GetWidgetStyle<FTextBlockStyle>("Normal.Italic")))
-						+SRichTextBlock::Decorator(FTextStyleDecorator::Create(TEXT("s"), FStateTreeEditorStyle::Get().GetWidgetStyle<FTextBlockStyle>("Normal.Subdued")))
-					]
-				]
-			]
 
-			// Failed transitions
-			+ SHorizontalBox::Slot()
-			.VAlign(VAlign_Center)
-			.AutoWidth()
-			[
-				SNew(SBox)
-				.Visibility(this, &SStateTreeViewRow::GetFailedTransitionVisibility)
-				[
-					SNew(SHorizontalBox)
-					+ SHorizontalBox::Slot()
-					.VAlign(VAlign_Center)
-					.Padding(FMargin(8.f, 0.f, 0.f, 0.f))
-					.AutoWidth()
-					[
-						SNew(SImage)
-						.Image(FStateTreeEditorStyle::Get().GetBrush("StateTreeEditor.Transition.Failed"))
-					]
-					+ SHorizontalBox::Slot()
-					.VAlign(VAlign_Center)
-					.AutoWidth()
-					[
-						SNew(SOverlay)
-						+ SOverlay::Slot()
-						[
-							SNew(SImage)
-							.Image(this, &SStateTreeViewRow::GetFailedTransitionIcon)
-							.ColorAndOpacity(IconTint)
-						]
-						+ SOverlay::Slot()
-						[
-							// Breakpoint box
-							SNew(SHorizontalBox)
+							// Succeeded transitions
 							+ SHorizontalBox::Slot()
 							.VAlign(VAlign_Top)
-							.HAlign(HAlign_Left)
 							.AutoWidth()
 							[
-								SNew(SBox)
-								.Padding(FMargin(0.f, -10.f, 0.f, 0.f))
-								[
-									SNew(SImage)
-									.DesiredSizeOverride(FVector2D(10.f, 10.f))
-									.Image(FStateTreeEditorStyle::Get().GetBrush(TEXT("StateTreeEditor.Debugger.Breakpoint.EnabledAndValid")))
-									.Visibility(this, &SStateTreeViewRow::GetFailedTransitionBreakpointVisibility)
-									.ToolTipText_Lambda([this]
-										{
-											return FText::Format(LOCTEXT("TransitionBreakpointTooltip", "Break when executing transition: {0}"),
-												GetFailedTransitionWithBreakpointDesc());
-										})
-								]
+								MakeTransitionWidget(EStateTreeTransitionTrigger::OnStateSucceeded, FStateTreeEditorStyle::Get().GetBrush("StateTreeEditor.Transition.Succeeded"))
 							]
-						]
-					]
-					+ SHorizontalBox::Slot()
-					.VAlign(VAlign_Center)
-					.AutoWidth()
-					.Padding(FMargin(4.f, 0.f, 0.f, 0.f))
-					[
-						SNew(SRichTextBlock)
-						.Text(this, &SStateTreeViewRow::GetFailedTransitionDesc)
-						.TextStyle(&FStateTreeEditorStyle::Get().GetWidgetStyle<FTextBlockStyle>("Normal.Normal"))
-						.OverflowPolicy(ETextOverflowPolicy::Ellipsis)
-						+SRichTextBlock::Decorator(FTextStyleDecorator::Create(TEXT(""), FStateTreeEditorStyle::Get().GetWidgetStyle<FTextBlockStyle>("Normal.Normal")))
-						+SRichTextBlock::Decorator(FTextStyleDecorator::Create(TEXT("b"), FStateTreeEditorStyle::Get().GetWidgetStyle<FTextBlockStyle>("Normal.Bold")))
-						+SRichTextBlock::Decorator(FTextStyleDecorator::Create(TEXT("i"), FStateTreeEditorStyle::Get().GetWidgetStyle<FTextBlockStyle>("Normal.Italic")))
-						+SRichTextBlock::Decorator(FTextStyleDecorator::Create(TEXT("s"), FStateTreeEditorStyle::Get().GetWidgetStyle<FTextBlockStyle>("Normal.Subdued")))
-					]
-				]
-			]
 
-			// Transitions
-			+ SHorizontalBox::Slot()
-			.VAlign(VAlign_Center)
-			.AutoWidth()
-			[
-				SNew(SBox)
-				.Visibility(this, &SStateTreeViewRow::GetConditionalTransitionsVisibility)
-				[
-					SNew(SHorizontalBox)
-					+ SHorizontalBox::Slot()
-					.VAlign(VAlign_Center)
-					.Padding(FMargin(8.f, 0.f, 0.f, 0.f))
-					.AutoWidth()
-					[
-						SNew(SImage)
-						.Image(FStateTreeEditorStyle::Get().GetBrush("StateTreeEditor.Transition.Condition"))
-					]
-					+ SHorizontalBox::Slot()
-					.VAlign(VAlign_Center)
-					.AutoWidth()
-					[
-						SNew(SOverlay)
-						+ SOverlay::Slot()
-						[
-							SNew(SImage)
-							.Image(FStateTreeEditorStyle::Get().GetBrush("StateTreeEditor.Transition.Goto"))
-							.ColorAndOpacity(IconTint)
-						]
-						+ SOverlay::Slot()
-						[
-							// Breakpoint box
-							SNew(SHorizontalBox)
+							// Failed transitions
 							+ SHorizontalBox::Slot()
 							.VAlign(VAlign_Top)
-							.HAlign(HAlign_Left)
 							.AutoWidth()
 							[
-								SNew(SBox)
-								.Padding(FMargin(0.f, -10.f, 0.f, 0.f))
-								[
-									SNew(SImage)
-									.DesiredSizeOverride(FVector2D(10.f, 10.f))
-									.Image(FStateTreeEditorStyle::Get().GetBrush(TEXT("StateTreeEditor.Debugger.Breakpoint.EnabledAndValid")))
-									.Visibility(this, &SStateTreeViewRow::GetConditionalTransitionsBreakpointVisibility)
-									.ToolTipText_Lambda([this]
-										{
-											return FText::Format(LOCTEXT("TransitionBreakpointTooltip", "Break when executing transition: {0}"),
-												GetConditionalTransitionsWithBreakpointDesc());
-										})
-								]
+								MakeTransitionWidget(EStateTreeTransitionTrigger::OnStateFailed, FStateTreeEditorStyle::Get().GetBrush("StateTreeEditor.Transition.Failed"))
+							]
+
+							// Transitions
+							+ SHorizontalBox::Slot()
+							.VAlign(VAlign_Top)
+							.AutoWidth()
+							[
+								MakeTransitionWidget(EStateTreeTransitionTrigger::OnTick | EStateTreeTransitionTrigger::OnEvent, FStateTreeEditorStyle::Get().GetBrush("StateTreeEditor.Transition.Condition"))
 							]
 						]
 					]
-					+ SHorizontalBox::Slot()
-					.VAlign(VAlign_Center)
-					.AutoWidth()
-					.Padding(FMargin(4.f, 0.f, 0.f, 0.f))
-					[
-						SNew(SRichTextBlock)
-						.Text(this, &SStateTreeViewRow::GetConditionalTransitionsDesc)
-						.TextStyle(&FStateTreeEditorStyle::Get().GetWidgetStyle<FTextBlockStyle>("Normal.Normal"))
-						.OverflowPolicy(ETextOverflowPolicy::Ellipsis)
-						+SRichTextBlock::Decorator(FTextStyleDecorator::Create(TEXT(""), FStateTreeEditorStyle::Get().GetWidgetStyle<FTextBlockStyle>("Normal.Normal")))
-						+SRichTextBlock::Decorator(FTextStyleDecorator::Create(TEXT("b"), FStateTreeEditorStyle::Get().GetWidgetStyle<FTextBlockStyle>("Normal.Bold")))
-						+SRichTextBlock::Decorator(FTextStyleDecorator::Create(TEXT("i"), FStateTreeEditorStyle::Get().GetWidgetStyle<FTextBlockStyle>("Normal.Italic")))
-						+SRichTextBlock::Decorator(FTextStyleDecorator::Create(TEXT("s"), FStateTreeEditorStyle::Get().GetWidgetStyle<FTextBlockStyle>("Normal.Subdued")))
-					]
+				]
+				+ SVerticalBox::Slot()
+				.AutoHeight()
+				.Padding(FMargin(0, 2, 0, 0))
+				[
+					MakeConditionsWidget(ViewBox)
+				]
+				+ SVerticalBox::Slot()
+				.AutoHeight()
+				.Padding(FMargin(0, 2, 0, 0))
+				[
+					MakeTasksWidget(ViewBox)
 				]
 			]
 		]
 	];
 }
 
-TSharedRef<SHorizontalBox> SStateTreeViewRow::CreateTasksWidget()
+TSharedRef<SWidget> SStateTreeViewRow::MakeTasksWidget(const TSharedPtr<SScrollBox>& ViewBox)
 {
-	const TSharedRef<SHorizontalBox> TasksBox = SNew(SHorizontalBox);
 	const UStateTreeEditorData* EditorData = WeakEditorData.Get();
 	const UStateTreeState* State = WeakState.Get();
-	
-	if (!EditorData || !State || State->Tasks.IsEmpty())
+	if (!EditorData || !State)
 	{
-		return TasksBox;
+		return SNullWidget::NullWidget;
 	}
 
-	const int32 NumTasks = State->Tasks.Num();
+	const TSharedRef<SWrapBox> TasksBox = SNew(SWrapBox)
+	.PreferredSize_Lambda([WeakOwnerViewBox = ViewBox.ToWeakPtr()]()
+	{
+		// Captured as weak ptr so we don't prevent our parent widget from being destroyed (circular pointer reference).
+		if (const TSharedPtr<SScrollBox> OwnerViewBox = WeakOwnerViewBox.Pin())
+		{
+			return FMath::Max(300, OwnerViewBox->GetTickSpaceGeometry().GetLocalSize().X - 200);
+		}
+		return 0.f;
+	});
 
-	// The task descriptions can get long. Make some effort to limit how long they can get. 
-	constexpr float ReferenceWidth = 1000.0f;
-	const float MaxTaskWidth = FMath::RoundToFloat(FMath::Max(150.0f, ReferenceWidth / static_cast<float>(FMath::Max(1, NumTasks))));
+	if (State->Tasks.IsEmpty())
+	{
+		return SNullWidget::NullWidget;
+	}
 
+	const int32 NumTasks = State->Tasks.Num();
+
+	// The task descriptions can get long. Make some effort to limit how long they can get.
 	for (int32 TaskIndex = 0; TaskIndex < NumTasks; TaskIndex++)
 	{
 		const FStateTreeEditorNode& TaskNode = State->Tasks[TaskIndex];
 		if (const FStateTreeTaskBase* Task = TaskNode.Node.GetPtr<FStateTreeTaskBase>())
 		{
-			FGuid TaskId = State->Tasks[TaskIndex].ID;
-			auto IsTaskEnabledFunc = [WeakState=WeakState, TaskIndex]
+			const FGuid TaskId = State->Tasks[TaskIndex].ID;
+			auto IsTaskEnabledFunc = [WeakState = WeakState, TaskIndex]
 				{
 					const UStateTreeState* State = WeakState.Get();
 					if (State != nullptr && State->Tasks.IsValidIndex(TaskIndex))
@@ -644,72 +531,312 @@ TSharedRef<SHorizontalBox> SStateTreeViewRow::CreateTasksWidget()
 				};
 
 			TasksBox->AddSlot()
-				.AutoWidth()
-				.VAlign(VAlign_Fill)
+				.Padding(FMargin(0, 0, 6, 0))
 				[
-					SNew(SBorder)
-					.VAlign(VAlign_Center)
-					.BorderImage(FAppStyle::GetBrush("WhiteBrush"))
-					.BorderBackgroundColor(this, &SStateTreeViewRow::GetTitleColor, 0.25f, 0.25f)
-					.Padding(0)
-					.IsEnabled_Lambda(IsTaskEnabledFunc)
+					SNew(SButton)
+					.ButtonStyle(FAppStyle::Get(), "SimpleButton")
+					.ContentPadding(FMargin(0, 0))
+					.OnClicked_Lambda([StateTreeViewModel = StateTreeViewModel, WeakState = WeakState, TaskId]()
+					{
+						StateTreeViewModel->BringNodeToFocus(WeakState.Get(), TaskId);
+						return FReply::Handled();
+					})
 					[
-						SNew(SOverlay)
-						+ SOverlay::Slot()
+						SNew(SBorder)
+						.VAlign(VAlign_Center)
+						.BorderImage(FAppStyle::GetNoBrush())
+						.Padding(0)
+						.IsEnabled_Lambda(IsTaskEnabledFunc)
 						[
-							SNew(SBox)
-							.MaxDesiredWidth(MaxTaskWidth)
-							.Padding(FMargin(6.f, 0.f))
+							SNew(SOverlay)
+							+ SOverlay::Slot()
+							[
+								SNew(SBox)
+								.HeightOverride(UE::StateTree::Editor::TaskRowHeight)
+								.Padding(FMargin(0.f, 0.f))
+								[
+									SNew(SHorizontalBox)
+									+ SHorizontalBox::Slot()
+									.VAlign(VAlign_Center)
+									.HAlign(HAlign_Left)
+									.FillContentWidth(0.f, 0.f)
+									[
+										SNew(SBox)
+										.Padding(FMargin(0.f, 0.f, 2.f, 0.f))
+										.Visibility(this, &SStateTreeViewRow::GetTaskIconVisibility, TaskId)
+										[
+											SNew(SImage)
+											.Image(this, &SStateTreeViewRow::GetTaskIcon, TaskId)
+											.ColorAndOpacity(this, &SStateTreeViewRow::GetTaskIconColor, TaskId)
+										]
+									]
+
+									+ SHorizontalBox::Slot()
+									.VAlign(VAlign_Center)
+									.HAlign(HAlign_Left)
+									.FillContentWidth(0.f, 1.f)
+									[
+										SNew(SRichTextBlock)
+										.Text(this, &SStateTreeViewRow::GetTaskDesc, TaskId, EStateTreeNodeFormatting::RichText)
+										.ToolTipText(this, &SStateTreeViewRow::GetTaskDesc, TaskId, EStateTreeNodeFormatting::Text)
+										.TextStyle(&FStateTreeEditorStyle::Get().GetWidgetStyle<FTextBlockStyle>("StateTree.Task.Title"))
+										.OverflowPolicy(ETextOverflowPolicy::Ellipsis)
+										.Clipping(EWidgetClipping::OnDemand)
+										+SRichTextBlock::Decorator(FTextStyleDecorator::Create(TEXT(""), FStateTreeEditorStyle::Get().GetWidgetStyle<FTextBlockStyle>("StateTree.Task.Title")))
+										+SRichTextBlock::Decorator(FTextStyleDecorator::Create(TEXT("b"), FStateTreeEditorStyle::Get().GetWidgetStyle<FTextBlockStyle>("StateTree.Task.Title.Bold")))
+										+SRichTextBlock::Decorator(FTextStyleDecorator::Create(TEXT("s"), FStateTreeEditorStyle::Get().GetWidgetStyle<FTextBlockStyle>("StateTree.Task.Title.Subdued")))
+									]
+								]
+							]
+							+ SOverlay::Slot()
 							[
+								// Task Breakpoint box
 								SNew(SHorizontalBox)
 								+ SHorizontalBox::Slot()
-								.VAlign(VAlign_Center)
+								.VAlign(VAlign_Top)
 								.HAlign(HAlign_Left)
 								.AutoWidth()
 								[
 									SNew(SBox)
-									.Padding(FMargin(0.f, 0.f, 2.f, 0.f))
-									.Visibility(this, &SStateTreeViewRow::GetTaskIconVisibility, TaskId)
+									.Padding(FMargin(-2.0f, -2.0f, 0.0f, 0.0f))
 									[
 										SNew(SImage)
-										.Image(this, &SStateTreeViewRow::GetTaskIcon, TaskId)
-										.ColorAndOpacity(this, &SStateTreeViewRow::GetTaskIconColor, TaskId)
+										.DesiredSizeOverride(FVector2D(10.f, 10.f))
+										.Image(FStateTreeEditorStyle::Get().GetBrush(TEXT("StateTreeEditor.Debugger.Breakpoint.EnabledAndValid")))
+										.Visibility_Lambda(IsTaskBreakpointEnabledFunc)
+										.ToolTipText_Lambda(GetTaskBreakpointTooltipFunc)
 									]
 								]
+							]
+						]
+					]
+				];
+		}
+	}
 
-								+ SHorizontalBox::Slot()
+	return TasksBox;
+}
+
+TSharedRef<SWidget> SStateTreeViewRow::MakeConditionsWidget(const TSharedPtr<SScrollBox>& ViewBox)
+{
+	const UStateTreeEditorData* EditorData = WeakEditorData.Get();
+	const UStateTreeState* State = WeakState.Get();
+	if (!EditorData || !State)
+	{
+		return SNullWidget::NullWidget;
+	}
+
+	const TSharedRef<SWrapBox> ConditionsBox = SNew(SWrapBox)
+	.PreferredSize_Lambda([WeakOwnerViewBox = ViewBox.ToWeakPtr()]()
+	{
+		// Captured as weak ptr so we don't prevent our parent widget from being destroyed (circular pointer reference).
+		if (const TSharedPtr<SScrollBox> OwnerViewBox = WeakOwnerViewBox.Pin())
+		{
+			return FMath::Max(300, OwnerViewBox->GetTickSpaceGeometry().GetLocalSize().X - 200);
+		}
+		return 0.f;
+	});
+
+	if (State->EnterConditions.IsEmpty())
+	{
+		return SNullWidget::NullWidget;
+	}
+
+	const int32 NumConditions = State->EnterConditions.Num();
+
+	for (int32 ConditionIndex = 0; ConditionIndex < NumConditions; ConditionIndex++)
+	{
+		const FStateTreeEditorNode& ConditionNode = State->EnterConditions[ConditionIndex];
+		if (const FStateTreeConditionBase* Condition = ConditionNode.Node.GetPtr<FStateTreeConditionBase>())
+		{
+			const FGuid ConditionId = ConditionNode.ID;
+
+			auto IsConditionEnabledFunc = [WeakState = WeakState]
+			{
+				const UStateTreeState* State = WeakState.Get();
+				return State && State->bEnabled;
+			};
+
+			auto IsForcedConditionVisibleFunc = [WeakState = WeakState, ConditionIndex]()
+				{
+					const UStateTreeState* State = WeakState.Get();
+					if (State != nullptr && State->EnterConditions.IsValidIndex(ConditionIndex))
+					{
+						if (const FStateTreeConditionBase* Condition = State->EnterConditions[ConditionIndex].Node.GetPtr<FStateTreeConditionBase>())
+						{
+							return Condition->EvaluationMode != EStateTreeConditionEvaluationMode::Evaluated ? EVisibility::Visible : EVisibility::Hidden;
+						}
+					}
+					return EVisibility::Hidden;
+				};
+
+			auto GetForcedConditionTooltipFunc = [WeakState = WeakState, ConditionIndex]()
+				{
+					const UStateTreeState* State = WeakState.Get();
+					if (State != nullptr && State->EnterConditions.IsValidIndex(ConditionIndex))
+					{
+						if (const FStateTreeConditionBase* Condition = State->EnterConditions[ConditionIndex].Node.GetPtr<FStateTreeConditionBase>())
+						{
+							if (Condition->EvaluationMode == EStateTreeConditionEvaluationMode::ForcedTrue)
+							{
+								return LOCTEXT("ForcedTrueConditionTooltip", "This condition is not evaluated and result forced to 'true'.");
+							}
+							if (Condition->EvaluationMode == EStateTreeConditionEvaluationMode::ForcedFalse)
+							{
+								return LOCTEXT("ForcedFalseConditionTooltip", "This condition is not evaluated and result forced to 'false'.");
+							}
+						}
+					}
+					return FText::GetEmpty();
+				};
+
+			auto GetForcedConditionImageFunc = [WeakState = WeakState, ConditionIndex]() -> const FSlateBrush*
+				{
+					const UStateTreeState* State = WeakState.Get();
+					if (State != nullptr && State->EnterConditions.IsValidIndex(ConditionIndex))
+					{
+						if (const FStateTreeConditionBase* Condition = State->EnterConditions[ConditionIndex].Node.GetPtr<FStateTreeConditionBase>())
+						{
+							if (Condition->EvaluationMode == EStateTreeConditionEvaluationMode::ForcedTrue)
+							{
+								return FStateTreeEditorStyle::Get().GetBrush("StateTreeEditor.Debugger.Condition.Passed");
+							}
+							if (Condition->EvaluationMode == EStateTreeConditionEvaluationMode::ForcedFalse)
+							{
+								return FStateTreeEditorStyle::Get().GetBrush("StateTreeEditor.Debugger.Condition.Failed");
+							}
+						}
+					}
+					return nullptr;
+				};
+
+			ConditionsBox->AddSlot()
+				[
+					SNew(SBorder)
+					.VAlign(VAlign_Center)
+					.BorderImage(FAppStyle::GetNoBrush())
+					.IsEnabled_Lambda(IsConditionEnabledFunc)
+					.Padding(0)
+					[
+						SNew(SBox)
+						.HeightOverride(UE::StateTree::Editor::TaskRowHeight)
+						[
+							SNew(SHorizontalBox)
+
+							// Operand
+							+ SHorizontalBox::Slot()
+							.AutoWidth()
+							.VAlign(VAlign_Center)
+							[
+								SNew(SBox)
+								.Padding(FMargin(4, 2, 4, 0))
 								.VAlign(VAlign_Center)
-								.HAlign(HAlign_Left)
-								.AutoWidth()
 								[
-									SNew(SRichTextBlock)
-									.Text(this, &SStateTreeViewRow::GetTaskDesc, TaskId, EStateTreeNodeFormatting::RichText)
-									.ToolTipText(this, &SStateTreeViewRow::GetTaskDesc, TaskId, EStateTreeNodeFormatting::Text)
-									.TextStyle(&FStateTreeEditorStyle::Get().GetWidgetStyle<FTextBlockStyle>("StateTree.Task.Title"))
-									.OverflowPolicy(ETextOverflowPolicy::Ellipsis)
-									+SRichTextBlock::Decorator(FTextStyleDecorator::Create(TEXT(""), FStateTreeEditorStyle::Get().GetWidgetStyle<FTextBlockStyle>("StateTree.Task.Title")))
-									+SRichTextBlock::Decorator(FTextStyleDecorator::Create(TEXT("b"), FStateTreeEditorStyle::Get().GetWidgetStyle<FTextBlockStyle>("StateTree.Task.Title.Bold")))
-									+SRichTextBlock::Decorator(FTextStyleDecorator::Create(TEXT("s"), FStateTreeEditorStyle::Get().GetWidgetStyle<FTextBlockStyle>("StateTree.Task.Title.Subdued")))
+									SNew(STextBlock)
+									.TextStyle(FStateTreeEditorStyle::Get(), "StateTree.Node.Operand")
+									.Text(this, &SStateTreeViewRow::GetOperandText, ConditionIndex)
+								]
+							]
+							// Open parens
+							+ SHorizontalBox::Slot()
+							.AutoWidth()
+							.VAlign(VAlign_Center)
+							[
+								SNew(SBox)
+								.Padding(FMargin(FMargin(0.0f, 1.0f, 0.0f, 0.0f)))
+								[
+									SNew(STextBlock)
+									.TextStyle(FStateTreeEditorStyle::Get(), "StateTree.Task.Title")
+									.Text(this, &SStateTreeViewRow::GetOpenParens, ConditionIndex)
+								]
+							]
+							// Open parens
+							+ SHorizontalBox::Slot() 
+							.AutoWidth()
+							.VAlign(VAlign_Center)
+							[
+								SNew(SOverlay)
+								+ SOverlay::Slot()
+								[
+
+									SNew(SButton)
+									.ButtonStyle(FAppStyle::Get(), "SimpleButton")
+									.ContentPadding(FMargin(2.f, 0.f))
+									.OnClicked_Lambda([StateTreeViewModel = StateTreeViewModel, WeakState = WeakState, ConditionId]()
+									{
+										StateTreeViewModel->BringNodeToFocus(WeakState.Get(), ConditionId);
+										return FReply::Handled();
+										})
+									[
+										SNew(SHorizontalBox)
+										
+										// Icon
+										+ SHorizontalBox::Slot()
+										.VAlign(VAlign_Center)
+										.HAlign(HAlign_Left)
+										.AutoWidth()
+										[
+											SNew(SBox)
+											.Padding(FMargin(0.f, 0.f, 2.f, 0.f))
+											.Visibility(this, &SStateTreeViewRow::GetConditionIconVisibility, ConditionId)
+											[
+												SNew(SImage)
+												.Image(this, &SStateTreeViewRow::GetConditionIcon, ConditionId)
+												.ColorAndOpacity(this, &SStateTreeViewRow::GetConditionIconColor, ConditionId)
+											]
+										]
+										// Desc
+										+ SHorizontalBox::Slot()
+										.VAlign(VAlign_Center)
+										.HAlign(HAlign_Left)
+										.AutoWidth()
+										[
+											SNew(SRichTextBlock)
+											.Text(this, &SStateTreeViewRow::GetConditionDesc, ConditionId, EStateTreeNodeFormatting::RichText)
+											.ToolTipText(this, &SStateTreeViewRow::GetConditionDesc, ConditionId, EStateTreeNodeFormatting::Text)
+											.TextStyle(&FStateTreeEditorStyle::Get().GetWidgetStyle<FTextBlockStyle>("StateTree.Task.Title"))
+											.OverflowPolicy(ETextOverflowPolicy::Ellipsis)
+											.Clipping(EWidgetClipping::OnDemand)
+											+SRichTextBlock::Decorator(FTextStyleDecorator::Create(TEXT(""), FStateTreeEditorStyle::Get().GetWidgetStyle<FTextBlockStyle>("StateTree.Task.Title")))
+											+SRichTextBlock::Decorator(FTextStyleDecorator::Create(TEXT("b"), FStateTreeEditorStyle::Get().GetWidgetStyle<FTextBlockStyle>("StateTree.Task.Title.Bold")))
+											+SRichTextBlock::Decorator(FTextStyleDecorator::Create(TEXT("s"), FStateTreeEditorStyle::Get().GetWidgetStyle<FTextBlockStyle>("StateTree.Task.Title.Subdued")))
+										]
+									]
+								]
+
+								+ SOverlay::Slot()
+								[
+									// Condition override box
+									SNew(SHorizontalBox)
+									+ SHorizontalBox::Slot()
+									.VAlign(VAlign_Top)
+									.HAlign(HAlign_Left)
+									.AutoWidth()
+									[
+										SNew(SBox)
+										.Padding(FMargin(-2.0f, -2.0f, 0.0f, 0.0f))
+										[
+											SNew(SImage)
+											.DesiredSizeOverride(FVector2D(16.f, 16.f))
+											.Image_Lambda(GetForcedConditionImageFunc)
+											.Visibility_Lambda(IsForcedConditionVisibleFunc)
+											.ToolTipText_Lambda(GetForcedConditionTooltipFunc)
+										]
+									]
 								]
 							]
-						]
-						+ SOverlay::Slot()
-						[
-							// Task Breakpoint box
-							SNew(SHorizontalBox)
+							
+							// Close parens
 							+ SHorizontalBox::Slot()
-							.VAlign(VAlign_Top)
-							.HAlign(HAlign_Left)
 							.AutoWidth()
+							.VAlign(VAlign_Center)
 							[
 								SNew(SBox)
-								.Padding(FMargin(0.0f, -10.0f, 0.0f, 0.0f))
+								.Padding(FMargin(0.0f, 1.0f, 0.0f, 0.0f))
 								[
-									SNew(SImage)
-									.DesiredSizeOverride(FVector2D(10.f, 10.f))
-									.Image(FStateTreeEditorStyle::Get().GetBrush(TEXT("StateTreeEditor.Debugger.Breakpoint.EnabledAndValid")))
-									.Visibility_Lambda(IsTaskBreakpointEnabledFunc)
-									.ToolTipText_Lambda(GetTaskBreakpointTooltipFunc)
+									SNew(STextBlock)
+									.TextStyle(FStateTreeEditorStyle::Get(), "StateTree.Task.Title")
+									.Text(this, &SStateTreeViewRow::GetCloseParens, ConditionIndex)
 								]
 							]
 						]
@@ -718,7 +845,7 @@ TSharedRef<SHorizontalBox> SStateTreeViewRow::CreateTasksWidget()
 		}
 	}
 
-	return TasksBox;
+	return ConditionsBox;
 }
 
 void SStateTreeViewRow::RequestRename() const
@@ -744,7 +871,10 @@ FSlateColor SStateTreeViewRow::GetTitleColor(const float Alpha, const float Ligh
 			{
 				Color = UE::StateTree::Editor::LerpColorSRGB(FoundColor->Color, FColor::Black, 0.25f);
 			}
-			Color = FoundColor->Color;
+			else
+			{
+				Color = FoundColor->Color;
+			}
 		}
 	}
 
@@ -766,6 +896,7 @@ FSlateColor SStateTreeViewRow::GetActiveStateColor() const
 		}
 		if (StateTreeViewModel && StateTreeViewModel->IsSelected(State))
 		{
+			// @todo: change to the common selection color.
 			return FLinearColor(FColor(236, 134, 39));
 		}
 	}
@@ -788,6 +919,18 @@ FSlateColor SStateTreeViewRow::GetSubTreeMarkerColor() const
 	return GetTitleColor();
 }
 
+EVisibility SStateTreeViewRow::GetSubTreeVisibility() const
+{
+	if (const UStateTreeState* State = WeakState.Get())
+	{
+		if (IsRootState() || State->Type == EStateTreeStateType::Subtree)
+		{
+			return EVisibility::Visible;
+		}
+	}
+	return EVisibility::Collapsed;	
+}
+
 FText SStateTreeViewRow::GetStateDesc() const
 {
 	if (const UStateTreeState* State = WeakState.Get())
@@ -1042,7 +1185,7 @@ FText SStateTreeViewRow::GetTaskDesc(FGuid TaskID, EStateTreeNodeFormatting Form
 		{
 			if (UE::StateTree::Editor::GbDisplayItemIds)
 			{
-				TaskName = FText::Format(LOCTEXT("TaskNameWithID", "{0} ({1})"), EditorData->GetNodeDescription(*TaskNode, Formatting), FText::AsCultureInvariant(*LexToString(TaskID)));
+				TaskName = FText::Format(LOCTEXT("NodeNameWithID", "{0} ({1})"), EditorData->GetNodeDescription(*TaskNode, Formatting), FText::AsCultureInvariant(*LexToString(TaskID)));
 			}
 			else
 			{
@@ -1053,23 +1196,167 @@ FText SStateTreeViewRow::GetTaskDesc(FGuid TaskID, EStateTreeNodeFormatting Form
 	return TaskName;
 }
 
-EVisibility SStateTreeViewRow::GetTasksVisibility() const
+const FStateTreeEditorNode* SStateTreeViewRow::GetConditionNodeByID(FGuid ConditionID) const
 {
-	if (const UStateTreeState* State = WeakState.Get())
+	const UStateTreeState* State = WeakState.Get();
+	const UStateTreeEditorData* EditorData = WeakEditorData.Get();
+	if (EditorData != nullptr
+		&& State != nullptr)
+	{
+		return State->EnterConditions.FindByPredicate([&ConditionID](const FStateTreeEditorNode& Node)
+		{
+			return Node.ID == ConditionID;
+		});
+	}
+	return nullptr;
+}
+
+EVisibility SStateTreeViewRow::GetConditionIconVisibility(FGuid ConditionID) const
+{
+	bool bHasIcon = false;
+	if (const FStateTreeEditorNode* Node = GetConditionNodeByID(ConditionID))
+	{
+		if (const FStateTreeNodeBase* BaseNode = Node->Node.GetPtr<const FStateTreeNodeBase>())
+		{
+			bHasIcon = !BaseNode->GetIconName().IsNone();
+		}
+	}
+	return bHasIcon ? EVisibility::Visible : EVisibility::Collapsed;  	
+}
+
+const FSlateBrush* SStateTreeViewRow::GetConditionIcon(FGuid ConditionID) const
+{
+	if (const FStateTreeEditorNode* Node = GetConditionNodeByID(ConditionID))
+	{
+		if (const FStateTreeNodeBase* BaseNode = Node->Node.GetPtr<const FStateTreeNodeBase>())
+		{
+			return UE::StateTreeEditor::EditorNodeUtils::ParseIcon(BaseNode->GetIconName()).GetIcon();
+		}
+	}
+	return nullptr;	
+}
+
+FSlateColor SStateTreeViewRow::GetConditionIconColor(FGuid ConditionID) const
+{
+	if (const FStateTreeEditorNode* Node = GetConditionNodeByID(ConditionID))
+	{
+		if (const FStateTreeNodeBase* BaseNode = Node->Node.GetPtr<const FStateTreeNodeBase>())
+		{
+			return FLinearColor(BaseNode->GetIconColor());
+		}
+	}
+	return FSlateColor::UseForeground();
+}
+
+FText SStateTreeViewRow::GetConditionDesc(FGuid ConditionID, EStateTreeNodeFormatting Formatting) const
+{
+	FText Description;
+	if (const UStateTreeEditorData* EditorData = WeakEditorData.Get())
 	{
-		int32 ValidCount = 0;
-		for (int32 i = 0; i < State->Tasks.Num(); i++)
+		if (const FStateTreeEditorNode* Node = GetConditionNodeByID(ConditionID))
 		{
-			if (State->Tasks[i].Node.GetPtr<FStateTreeTaskBase>())
+			if (UE::StateTree::Editor::GbDisplayItemIds)
+			{
+				Description = FText::Format(LOCTEXT("NodeNameWithID", "{0} ({1})"), EditorData->GetNodeDescription(*Node, Formatting), FText::AsCultureInvariant(*LexToString(ConditionID)));
+			}
+			else
 			{
-				ValidCount++;
+				Description = EditorData->GetNodeDescription(*Node, Formatting);
 			}
 		}
-		return ValidCount > 0 ? EVisibility::Visible : EVisibility::Collapsed;
 	}
-	return EVisibility::Collapsed;
+	return Description;
+}
+
+FText SStateTreeViewRow::GetOperandText(const int32 ConditionIndex) const
+{
+	const UStateTreeState* State = WeakState.Get();
+	if (!State
+		|| !State->EnterConditions.IsValidIndex(ConditionIndex))
+	{
+		return FText::GetEmpty();
+	}
+
+	// First item does not relate to anything existing, it could be empty. 
+	// return IF to indicate that we're building condition and IS for consideration.
+	if (ConditionIndex == 0)
+	{
+		return LOCTEXT("IfOperand", "IF");
+	}
+
+	const EStateTreeExpressionOperand Operand = State->EnterConditions[ConditionIndex].ExpressionOperand;
+
+	if (Operand == EStateTreeExpressionOperand::And)
+	{
+		return LOCTEXT("AndOperand", "AND");
+	}
+	else if (Operand == EStateTreeExpressionOperand::Or)
+	{
+		return LOCTEXT("OrOperand", "OR");
+	}
+	else
+	{
+		ensureMsgf(false, TEXT("Unhandled operand %s"), *UEnum::GetValueAsString(Operand));
+	}
+
+	return FText::GetEmpty();
+}
+
+FText SStateTreeViewRow::GetOpenParens(const int32 ConditionIndex) const
+{
+	const UStateTreeState* State = WeakState.Get();
+	if (!State
+		|| !State->EnterConditions.IsValidIndex(ConditionIndex))
+	{
+		return FText::GetEmpty();
+	}
+
+	const int32 NumConditions = State->EnterConditions.Num();
+	const int32 CurrIndent = ConditionIndex == 0 ? 0 : (State->EnterConditions[ConditionIndex].ExpressionIndent + 1);
+	const int32 NextIndent = (ConditionIndex + 1) >= NumConditions ? 0 : (State->EnterConditions[ConditionIndex + 1].ExpressionIndent + 1);
+	const int32 DeltaIndent = NextIndent - CurrIndent;
+	const int32 OpenParens = FMath::Max(0, DeltaIndent);
+
+	static_assert(UE::StateTree::MaxExpressionIndent == 4);
+	switch (OpenParens)
+	{
+	case 1: return FText::FromString(TEXT("("));
+	case 2: return FText::FromString(TEXT("(("));
+	case 3: return FText::FromString(TEXT("((("));
+	case 4: return FText::FromString(TEXT("(((("));
+	case 5: return FText::FromString(TEXT("((((("));
+	}
+	return FText::GetEmpty();
+}
+
+FText SStateTreeViewRow::GetCloseParens(const int32 ConditionIndex) const
+{
+	const UStateTreeState* State = WeakState.Get();
+	if (!State
+		|| !State->EnterConditions.IsValidIndex(ConditionIndex))
+	{
+		return FText::GetEmpty();
+	}
+
+	const int32 NumConditions = State->EnterConditions.Num();
+	const int32 CurrIndent = ConditionIndex == 0 ? 0 : (State->EnterConditions[ConditionIndex].ExpressionIndent + 1);
+	const int32 NextIndent = (ConditionIndex + 1) >= NumConditions ? 0 : (State->EnterConditions[ConditionIndex + 1].ExpressionIndent + 1);
+	const int32 DeltaIndent = NextIndent - CurrIndent;
+	const int32 CloseParens = FMath::Max(0, -DeltaIndent);
+
+	static_assert(UE::StateTree::MaxExpressionIndent == 4);
+	switch (CloseParens)
+	{
+	case 1: return FText::FromString(TEXT(")"));
+	case 2: return FText::FromString(TEXT("))"));
+	case 3: return FText::FromString(TEXT(")))"));
+	case 4: return FText::FromString(TEXT("))))"));
+	case 5: return FText::FromString(TEXT(")))))"));
+	}
+	return FText::GetEmpty();
 }
 
+
 EVisibility SStateTreeViewRow::GetLinkedStateVisibility() const
 {
 	if (const UStateTreeState* State = WeakState.Get())
@@ -1115,106 +1402,303 @@ bool SStateTreeViewRow::GetStateWarnings(FText* OutText) const
 					"The Utility Considerations are used only when parent State's Selection Behavior is:"
 					"\"Try Select Children with Highest Utility\" or \"Try Select Children Based on Relative Utility.");
 			}
-			bHasWarnings = true;
+			bHasWarnings = true;
+		}
+	}
+
+	return bHasWarnings;
+}
+
+FText SStateTreeViewRow::GetLinkedStateDesc() const
+{
+	const UStateTreeState* State = WeakState.Get();
+	if (!State)
+	{
+		return FText::GetEmpty();
+	}
+
+	if (State->Type == EStateTreeStateType::Linked)
+	{
+		return FText::FromName(State->LinkedSubtree.Name);
+	}
+	else if (State->Type == EStateTreeStateType::LinkedAsset)
+	{
+		return FText::FromString(GetNameSafe(State->LinkedAsset.Get()));
+	}
+	
+	return FText::GetEmpty();
+}
+
+EVisibility SStateTreeViewRow::GetWarningsVisibility() const
+{
+	return GetStateWarnings(nullptr) ? EVisibility::Visible : EVisibility::Collapsed;
+}
+
+FText SStateTreeViewRow::GetWarningsTooltipText() const
+{
+	FText Warnings = FText::GetEmpty();
+	GetStateWarnings(&Warnings);
+	return Warnings;
+}
+
+bool SStateTreeViewRow::HasParentTransitionForTrigger(const UStateTreeState& State, const EStateTreeTransitionTrigger Trigger) const
+{
+	EStateTreeTransitionTrigger CombinedTrigger = EStateTreeTransitionTrigger::None;
+	for (const UStateTreeState* ParentState = State.Parent; ParentState != nullptr; ParentState = ParentState->Parent)
+	{
+		for (const FStateTreeTransition& Transition : ParentState->Transitions)
+		{
+			CombinedTrigger |= Transition.Trigger;
+		}
+	}
+	return EnumHasAllFlags(CombinedTrigger, Trigger);
+}
+
+FText SStateTreeViewRow::GetLinkTooltip(const FStateTreeStateLink& Link, const FGuid NodeID) const
+{
+	if (const UStateTreeState* State = WeakState.Get())
+	{
+		const int32 TaskIndex = State->Tasks.IndexOfByPredicate([&NodeID](const FStateTreeEditorNode& Node)
+		{
+			return Node.ID == NodeID;
+		});
+		if (TaskIndex != INDEX_NONE)
+		{
+			return FText::Format(LOCTEXT("TaskTransitionDesc", "Task {0} transitions to {1}"),
+				FText::FromName(State->Tasks[TaskIndex].GetName()),
+				UE::StateTree::Editor::GetStateLinkDesc(WeakEditorData.Get(), Link, EStateTreeNodeFormatting::Text, /*bShowStatePath*/true));
+		}
+
+		if (State->SingleTask.ID == NodeID)
+		{
+			return FText::Format(LOCTEXT("TaskTransitionDesc", "Task {0} transitions to {1}"),
+				FText::FromName(State->SingleTask.GetName()),
+				UE::StateTree::Editor::GetStateLinkDesc(WeakEditorData.Get(), Link, EStateTreeNodeFormatting::Text, /*bShowStatePath*/true));
+		}
+
+		const int32 TransitionIndex = State->Transitions.IndexOfByPredicate([&NodeID](const FStateTreeTransition& Transition)
+		{
+			return Transition.ID == NodeID;
+		});
+		if (TransitionIndex != INDEX_NONE)
+		{
+			return UE::StateTree::Editor::GetTransitionDesc(WeakEditorData.Get(), State->Transitions[TransitionIndex], EStateTreeNodeFormatting::Text, /*bShowStatePath*/true);
+		}
+	}
+	
+	return FText::GetEmpty();
+};
+
+bool SStateTreeViewRow::IsLeafState() const
+{
+	const UStateTreeState* State = WeakState.Get();
+	return State
+		&& State->Children.Num() == 0
+		&& !IsRootState()
+		&& (State->Type == EStateTreeStateType::State
+			|| State->Type == EStateTreeStateType::Linked
+			|| State->Type == EStateTreeStateType::LinkedAsset);
+}
+
+TSharedRef<SWidget> SStateTreeViewRow::MakeTransitionWidgets(const EStateTreeTransitionTrigger Trigger, const FTransitionDescFilterOptions FilterOptions)
+{
+	const UStateTreeEditorData* TreeEditorData = WeakEditorData.Get();
+	const UStateTreeState* State = WeakState.Get();
+
+	if (!TreeEditorData || !State)
+	{
+		return SNullWidget::NullWidget;
+	}
+	
+	struct FItem
+	{
+		FItem() = default;
+		FItem(const FStateTreeStateLink& InLink, const FGuid InNodeID)
+			: Link(InLink)
+			, NodeID(InNodeID)
+		{
+		}
+		FItem(const FText& InDesc, const FText& InTooltip)
+			: Desc(InDesc)
+			, Tooltip(InTooltip)
+		{
+		}
+		
+		FText Desc;
+		FText Tooltip;
+		FStateTreeStateLink Link;
+		FGuid NodeID;
+	};
+
+	TArray<FItem> DescItems;
+
+	for (const FStateTreeTransition& Transition : State->Transitions)
+	{
+		// Apply filter for enabled/disabled transitions
+		if ((FilterOptions.Enabled == ETransitionDescRequirement::RequiredTrue && Transition.bTransitionEnabled == false)
+			|| (FilterOptions.Enabled == ETransitionDescRequirement::RequiredFalse && Transition.bTransitionEnabled))
+		{
+			continue;
+		}
+
+#if WITH_STATETREE_TRACE_DEBUGGER
+		// Apply filter for transitions with/without breakpoint
+		const bool bHasBreakpoint = TreeEditorData != nullptr && TreeEditorData->HasBreakpoint(Transition.ID, EStateTreeBreakpointType::OnTransition);
+		if ((FilterOptions.WithBreakpoint == ETransitionDescRequirement::RequiredTrue && bHasBreakpoint == false)
+			|| (FilterOptions.WithBreakpoint == ETransitionDescRequirement::RequiredFalse && bHasBreakpoint))
+		{
+			continue;
+		}
+#endif // WITH_STATETREE_TRACE_DEBUGGER
+
+		const bool bMatch = FilterOptions.bUseMask ? EnumHasAnyFlags(Transition.Trigger, Trigger) : Transition.Trigger == Trigger;
+		if (bMatch)
+		{
+			DescItems.Emplace(Transition.State, Transition.ID);
+		}
+	}
+
+	// Find states from transition tasks
+	if (EnumHasAnyFlags(Trigger, EStateTreeTransitionTrigger::OnTick | EStateTreeTransitionTrigger::OnEvent))
+	{
+		auto AddLinksFromStruct = [&DescItems, TreeEditorData](FStateTreeDataView Struct, const FGuid NodeID)
+		{
+			if (!Struct.IsValid())
+			{
+				return;
+			}
+			for (TPropertyValueIterator<FStructProperty> It(Struct.GetStruct(), Struct.GetMemory()); It; ++It)
+			{
+				const UScriptStruct* StructType = It.Key()->Struct;
+				if (StructType == TBaseStructure<FStateTreeStateLink>::Get())
+				{
+					const FStateTreeStateLink& Link = *static_cast<const FStateTreeStateLink*>(It.Value());
+					if (Link.LinkType != EStateTreeTransitionType::None)
+					{
+						DescItems.Emplace(Link, NodeID);
+					}
+				}
+			}
+		};
+		
+		for (const FStateTreeEditorNode& Task : State->Tasks)
+		{
+			AddLinksFromStruct(FStateTreeDataView(Task.Node.GetScriptStruct(), const_cast<uint8*>(Task.Node.GetMemory())), Task.ID);
+			AddLinksFromStruct(Task.GetInstance(), Task.ID);
 		}
-	}
-
-	return bHasWarnings;
-}
 
-FText SStateTreeViewRow::GetLinkedStateDesc() const
-{
-	const UStateTreeState* State = WeakState.Get();
-	if (!State)
-	{
-		return FText::GetEmpty();
+		AddLinksFromStruct(FStateTreeDataView(State->SingleTask.Node.GetScriptStruct(), const_cast<uint8*>(State->SingleTask.Node.GetMemory())), State->SingleTask.ID);
+		AddLinksFromStruct(State->SingleTask.GetInstance(), State->SingleTask.ID);
 	}
 
-	if (State->Type == EStateTreeStateType::Linked)
-	{
-		return FText::FromName(State->LinkedSubtree.Name);
-	}
-	else if (State->Type == EStateTreeStateType::LinkedAsset)
+	if (IsLeafState()
+		&& DescItems.Num() == 0
+		&& EnumHasAnyFlags(Trigger, EStateTreeTransitionTrigger::OnStateCompleted))
 	{
-		return FText::FromString(GetNameSafe(State->LinkedAsset.Get()));
+		if (HasParentTransitionForTrigger(*State, Trigger))
+		{
+			DescItems.Emplace(
+				LOCTEXT("TransitionActionHandleInParentRich", "<i>Parent</>"),
+				LOCTEXT("TransitionActionHandleInParent", "Handle transition in parent State")
+				);
+		}
+		else
+		{
+			DescItems.Emplace(
+				LOCTEXT("TransitionActionRoot", "<i>Root</>"),
+				LOCTEXT("TransitionActionRoot", "Transition to Root State.")
+				);
+		}
 	}
-	
-	return FText::GetEmpty();
-}
 
-EVisibility SStateTreeViewRow::GetWarningsVisibility() const
-{
-	return GetStateWarnings(nullptr) ? EVisibility::Visible : EVisibility::Collapsed;
-}
+	TSharedRef<SHorizontalBox> TransitionContainer = SNew(SHorizontalBox);
 
-FText SStateTreeViewRow::GetWarningsTooltipText() const
-{
-	FText Warnings = FText::GetEmpty();
-	GetStateWarnings(&Warnings);
-	return Warnings;
-}
+	auto IsTransitionEnabledFunc = [WeakState = WeakState]
+	{
+		const UStateTreeState* State = WeakState.Get();
+		return State && State->bEnabled;
+	};
 
-bool SStateTreeViewRow::HasParentTransitionForTrigger(const UStateTreeState& State, const EStateTreeTransitionTrigger Trigger) const
-{
-	EStateTreeTransitionTrigger CombinedTrigger = EStateTreeTransitionTrigger::None;
-	for (const UStateTreeState* ParentState = State.Parent; ParentState != nullptr; ParentState = ParentState->Parent)
+	for (int32 Index = 0; Index < DescItems.Num(); Index++)
 	{
-		for (const FStateTreeTransition& Transition : ParentState->Transitions)
+		const FItem& Item = DescItems[Index];
+
+		if (Index > 0)
 		{
-			CombinedTrigger |= Transition.Trigger;
+			TransitionContainer->AddSlot()
+			.AutoWidth()
+			.VAlign(VAlign_Center)
+			[
+				SNew(STextBlock)
+				.Text(FText::FromString(TEXT(", ")))
+				.TextStyle(&FStateTreeEditorStyle::Get().GetWidgetStyle<FTextBlockStyle>("Transition.Subdued"))
+			];
 		}
+		
+		TransitionContainer->AddSlot()
+		.AutoWidth()
+		.VAlign(VAlign_Center)
+		[
+			SNew(SButton)
+			.ButtonStyle(FAppStyle::Get(), "SimpleButton")
+			.ContentPadding(FMargin(0, 0))
+			.OnClicked_Lambda([StateTreeViewModel = StateTreeViewModel, WeakState = WeakState, Item]()
+			{
+				if (StateTreeViewModel)
+				{
+					StateTreeViewModel->BringNodeToFocus(WeakState.Get(), Item.NodeID);
+				}
+				return FReply::Handled();
+			})
+			[
+				SNew(SBorder)
+				.VAlign(VAlign_Center)
+				.BorderImage(FAppStyle::GetNoBrush())
+				.Padding(0)
+				.IsEnabled_Lambda(IsTransitionEnabledFunc)
+				[
+					SNew(SRichTextBlock)
+					.Text_Lambda([WeakEditorData = WeakEditorData, Item]()
+					{
+						if (!Item.Desc.IsEmpty())
+						{
+							return Item.Desc;
+						}
+						return UE::StateTree::Editor::GetStateLinkDesc(WeakEditorData.Get(), Item.Link, EStateTreeNodeFormatting::RichText);
+					})
+					.ToolTipText_Lambda([this, Item]()
+					{
+						if (!Item.Tooltip.IsEmpty())
+						{
+							return Item.Tooltip;
+						}
+						return GetLinkTooltip(Item.Link, Item.NodeID);
+					})
+					.TextStyle(&FStateTreeEditorStyle::Get().GetWidgetStyle<FTextBlockStyle>("Transition.Normal"))
+					+SRichTextBlock::Decorator(FTextStyleDecorator::Create(TEXT(""), FStateTreeEditorStyle::Get().GetWidgetStyle<FTextBlockStyle>("Transition.Normal")))
+					+SRichTextBlock::Decorator(FTextStyleDecorator::Create(TEXT("b"), FStateTreeEditorStyle::Get().GetWidgetStyle<FTextBlockStyle>("Transition.Bold")))
+					+SRichTextBlock::Decorator(FTextStyleDecorator::Create(TEXT("i"), FStateTreeEditorStyle::Get().GetWidgetStyle<FTextBlockStyle>("Transition.Italic")))
+					+SRichTextBlock::Decorator(FTextStyleDecorator::Create(TEXT("s"), FStateTreeEditorStyle::Get().GetWidgetStyle<FTextBlockStyle>("Transition.Subdued")))
+				]
+			]
+		];
 	}
-	return EnumHasAllFlags(CombinedTrigger, Trigger);
+	
+	return TransitionContainer;
 }
 
-
-FText SStateTreeViewRow::GetLinkDescription(const FStateTreeStateLink& Link)
-{
-	switch (Link.LinkType)
-	{
-	case EStateTreeTransitionType::None:
-		return LOCTEXT("TransitionNoneStyled", "<i>None</>");
-		break;
-	case EStateTreeTransitionType::Succeeded:
-		return LOCTEXT("TransitionTreeSucceededStyled", "<i>Succeeded</>");
-		break;
-	case EStateTreeTransitionType::Failed:
-		return LOCTEXT("TransitionTreeFailedStyled", "<i>Failed</>");
-		break;
-	case EStateTreeTransitionType::NextState:
-		return LOCTEXT("TransitionNextStateStyled", "<i>Next</>");
-		break;
-	case EStateTreeTransitionType::NextSelectableState:
-		return LOCTEXT("TransitionNextSelectableStateStyled", "<i>Next Selectable</s>");
-		break;
-	case EStateTreeTransitionType::GotoState:
-		return FText::FromName(Link.Name);
-		break;
-	default:
-		ensureMsgf(false, TEXT("Unhandled transition type."));
-		break;
-	}
-
-	return FText::GetEmpty();
-};
-
-bool SStateTreeViewRow::IsLeafState() const
+FText SStateTreeViewRow::GetTransitionsDesc(const EStateTreeTransitionTrigger Trigger, const FTransitionDescFilterOptions FilterOptions) const
 {
 	const UStateTreeState* State = WeakState.Get();
-	return State
-		&& State->Children.Num() == 0
-		&& !IsRootState()
-		&& (State->Type == EStateTreeStateType::State
-			|| State->Type == EStateTreeStateType::Linked
-			|| State->Type == EStateTreeStateType::LinkedAsset);
-}
+	const UStateTreeEditorData* EditorData = WeakEditorData.Get();
+	if (!State || !EditorData)
+	{
+		return FText::GetEmpty();
+	}
 
-FText SStateTreeViewRow::GetTransitionsDesc(const UStateTreeState& State, const EStateTreeTransitionTrigger Trigger, const FTransitionDescFilterOptions FilterOptions) const
-{
 	TArray<FText> DescItems;
-	const UStateTreeEditorData* TreeEditorData = WeakEditorData.Get();
 
-	for (const FStateTreeTransition& Transition : State.Transitions)
+	for (const FStateTreeTransition& Transition : State->Transitions)
 	{
 		// Apply filter for enabled/disabled transitions
 		if ((FilterOptions.Enabled == ETransitionDescRequirement::RequiredTrue && Transition.bTransitionEnabled == false)
@@ -1225,7 +1709,7 @@ FText SStateTreeViewRow::GetTransitionsDesc(const UStateTreeState& State, const
 
 #if WITH_STATETREE_TRACE_DEBUGGER
 		// Apply filter for transitions with/without breakpoint
-		const bool bHasBreakpoint = TreeEditorData != nullptr && TreeEditorData->HasBreakpoint(Transition.ID, EStateTreeBreakpointType::OnTransition);
+		const bool bHasBreakpoint = EditorData->HasBreakpoint(Transition.ID, EStateTreeBreakpointType::OnTransition);
 		if ((FilterOptions.WithBreakpoint == ETransitionDescRequirement::RequiredTrue && bHasBreakpoint == false)
 			|| (FilterOptions.WithBreakpoint == ETransitionDescRequirement::RequiredFalse && bHasBreakpoint))
 		{
@@ -1236,14 +1720,14 @@ FText SStateTreeViewRow::GetTransitionsDesc(const UStateTreeState& State, const
 		const bool bMatch = FilterOptions.bUseMask ? EnumHasAnyFlags(Transition.Trigger, Trigger) : Transition.Trigger == Trigger;
 		if (bMatch)
 		{
-			DescItems.Add(GetLinkDescription(Transition.State));
+			DescItems.Add(UE::StateTree::Editor::GetStateLinkDesc(EditorData, Transition.State, EStateTreeNodeFormatting::RichText));
 		}
 	}
 
 	// Find states from transition tasks
 	if (EnumHasAnyFlags(Trigger, EStateTreeTransitionTrigger::OnTick | EStateTreeTransitionTrigger::OnEvent))
 	{
-		auto AddLinksFromStruct = [&DescItems](FStateTreeDataView Struct)
+		auto AddLinksFromStruct = [EditorData, &DescItems](FStateTreeDataView Struct)
 		{
 			if (!Struct.IsValid())
 			{
@@ -1257,61 +1741,71 @@ FText SStateTreeViewRow::GetTransitionsDesc(const UStateTreeState& State, const
 					const FStateTreeStateLink& Link = *static_cast<const FStateTreeStateLink*>(It.Value());
 					if (Link.LinkType != EStateTreeTransitionType::None)
 					{
-						DescItems.Add(GetLinkDescription(Link));
+						DescItems.Add(UE::StateTree::Editor::GetStateLinkDesc(EditorData, Link, EStateTreeNodeFormatting::RichText));
 					}
 				}
 			}
 		};
 		
-		for (const FStateTreeEditorNode& Task : State.Tasks)
+		for (const FStateTreeEditorNode& Task : State->Tasks)
 		{
 			AddLinksFromStruct(FStateTreeDataView(Task.Node.GetScriptStruct(), const_cast<uint8*>(Task.Node.GetMemory())));
 			AddLinksFromStruct(Task.GetInstance());
 		}
 
-		AddLinksFromStruct(FStateTreeDataView(State.SingleTask.Node.GetScriptStruct(), const_cast<uint8*>(State.SingleTask.Node.GetMemory())));
-		AddLinksFromStruct(State.SingleTask.GetInstance());
+		AddLinksFromStruct(FStateTreeDataView(State->SingleTask.Node.GetScriptStruct(), const_cast<uint8*>(State->SingleTask.Node.GetMemory())));
+		AddLinksFromStruct(State->SingleTask.GetInstance());
 	}
 
 	if (IsLeafState()
 		&& DescItems.Num() == 0
 		&& EnumHasAnyFlags(Trigger, EStateTreeTransitionTrigger::OnStateCompleted))
 	{
-		if (HasParentTransitionForTrigger(State, Trigger))
+		if (HasParentTransitionForTrigger(*State, Trigger))
 		{
-			DescItems.Add(LOCTEXT("TransitionActionHandleInParentStyled", "<i>Parent</>"));
+			DescItems.Add(LOCTEXT("TransitionActionHandleInParentRich", "<i>Parent</>"));
 		}
 		else
 		{
-			DescItems.Add(LOCTEXT("TransitionActionRoot", "<i>Root</>"));
+			DescItems.Add(LOCTEXT("TransitionActionRootRich", "<i>Root</>"));
 		}
 	}
 	
 	return FText::Join(FText::FromString(TEXT(", ")), DescItems);
 }
 
-const FSlateBrush* SStateTreeViewRow::GetTransitionsIcon(const UStateTreeState& State, const EStateTreeTransitionTrigger Trigger, const FTransitionDescFilterOptions FilterOptions) const
+const FSlateBrush* SStateTreeViewRow::GetTransitionsIcon(const EStateTreeTransitionTrigger Trigger) const
 {
+	const UStateTreeState* State = WeakState.Get();
+	if (!State)
+	{
+		return nullptr;
+	}
+
+	if (EnumHasAnyFlags(Trigger, EStateTreeTransitionTrigger::OnTick | EStateTreeTransitionTrigger::OnEvent))
+	{
+		return FStateTreeEditorStyle::Get().GetBrush("StateTreeEditor.Transition.Goto");
+	}
+
 	enum EIconType
 	{
 		IconNone = 0,
-		IconGoto =	1 << 0,
-		IconNext =		1 << 1,
-		IconParent =		1 << 2,
+		IconGoto = 1 << 0,
+		IconNext = 1 << 1,
+		IconParent = 1 << 2,
 	};
 	uint8 IconType = IconNone;
 	
-	const UStateTreeEditorData* EditorData = WeakEditorData.Get();
-	
-	for (const FStateTreeTransition& Transition : State.Transitions)
+	for (const FStateTreeTransition& Transition : State->Transitions)
 	{
 		// Apply filter for enabled/disabled transitions
-		if ((FilterOptions.Enabled == ETransitionDescRequirement::RequiredTrue && Transition.bTransitionEnabled == false)
+/*		if ((FilterOptions.Enabled == ETransitionDescRequirement::RequiredTrue && Transition.bTransitionEnabled == false)
 			|| (FilterOptions.Enabled == ETransitionDescRequirement::RequiredFalse && Transition.bTransitionEnabled))
 		{
 			continue;
-		}
+		}*/
 
+/*		
 #if WITH_STATETREE_TRACE_DEBUGGER
 		// Apply filter for transitions with/without breakpoint
 		const bool bHasBreakpoint = EditorData != nullptr && EditorData->HasBreakpoint(Transition.ID, EStateTreeBreakpointType::OnTransition);
@@ -1321,9 +1815,10 @@ const FSlateBrush* SStateTreeViewRow::GetTransitionsIcon(const UStateTreeState&
 			continue;
 		}
 #endif // WITH_STATETREE_TRACE_DEBUGGER
+*/
 		
 		// The icons here depict "transition direction", not the type specifically.
-		const bool bMatch = FilterOptions.bUseMask ? EnumHasAnyFlags(Transition.Trigger, Trigger) : Transition.Trigger == Trigger;
+		const bool bMatch = /*FilterOptions.bUseMask ? EnumHasAnyFlags(Transition.Trigger, Trigger) :*/ Transition.Trigger == Trigger;
 		if (bMatch)
 		{
 			switch (Transition.State.LinkType)
@@ -1380,15 +1875,21 @@ const FSlateBrush* SStateTreeViewRow::GetTransitionsIcon(const UStateTreeState&
 	return nullptr;
 }
 
-EVisibility SStateTreeViewRow::GetTransitionsVisibility(const UStateTreeState& State, const EStateTreeTransitionTrigger Trigger) const
+EVisibility SStateTreeViewRow::GetTransitionsVisibility(const EStateTreeTransitionTrigger Trigger) const
 {
+	const UStateTreeState* State = WeakState.Get();
+	if (!State)
+	{
+		return EVisibility::Collapsed;
+	}
+	
 	// Handle completed, succeeded and failed transitions.
 	if (EnumHasAnyFlags(Trigger, EStateTreeTransitionTrigger::OnStateCompleted))
 	{
 		EStateTreeTransitionTrigger HandledTriggers = EStateTreeTransitionTrigger::None;
 		bool bExactMatch = false;
 
-		for (const FStateTreeTransition& Transition : State.Transitions)
+		for (const FStateTreeTransition& Transition : State->Transitions)
 		{
 			// Skip disabled transitions
 			if (Transition.bTransitionEnabled == false)
@@ -1440,7 +1941,7 @@ EVisibility SStateTreeViewRow::GetTransitionsVisibility(const UStateTreeState& S
 			return false;
 		};
 		
-		for (const FStateTreeEditorNode& Task : State.Tasks)
+		for (const FStateTreeEditorNode& Task : State->Tasks)
 		{
 			if (HasAnyLinksInStruct(FStateTreeDataView(Task.Node.GetScriptStruct(), const_cast<uint8*>(Task.Node.GetMemory())))
 				|| HasAnyLinksInStruct(Task.GetInstance()))
@@ -1449,15 +1950,15 @@ EVisibility SStateTreeViewRow::GetTransitionsVisibility(const UStateTreeState& S
 			}
 		}
 
-		if (HasAnyLinksInStruct(FStateTreeDataView(State.SingleTask.Node.GetScriptStruct(), const_cast<uint8*>(State.SingleTask.Node.GetMemory())))
-			|| HasAnyLinksInStruct(State.SingleTask.GetInstance()))
+		if (HasAnyLinksInStruct(FStateTreeDataView(State->SingleTask.Node.GetScriptStruct(), const_cast<uint8*>(State->SingleTask.Node.GetMemory())))
+			|| HasAnyLinksInStruct(State->SingleTask.GetInstance()))
 		{
 			return EVisibility::Visible;
 		}
 	}
 	
 	// Handle the test
-	for (const FStateTreeTransition& Transition : State.Transitions)
+	for (const FStateTreeTransition& Transition : State->Transitions)
 	{
 		// Skip disabled transitions
 		if (Transition.bTransitionEnabled == false)
@@ -1473,207 +1974,37 @@ EVisibility SStateTreeViewRow::GetTransitionsVisibility(const UStateTreeState& S
 	return EVisibility::Collapsed;
 }
 
-EVisibility SStateTreeViewRow::GetTransitionsBreakpointVisibility(const UStateTreeState& State, const EStateTreeTransitionTrigger Trigger) const
+EVisibility SStateTreeViewRow::GetTransitionsBreakpointVisibility(const EStateTreeTransitionTrigger Trigger) const
 {
-#if WITH_STATETREE_TRACE_DEBUGGER
-	if (const UStateTreeEditorData* EditorData = WeakEditorData.Get())
+	if (const UStateTreeState* State = WeakState.Get())
 	{
-		for (const FStateTreeTransition& Transition : State.Transitions)
+#if WITH_STATETREE_TRACE_DEBUGGER
+		if (const UStateTreeEditorData* EditorData = WeakEditorData.Get())
 		{
-			if (Transition.bTransitionEnabled && EnumHasAnyFlags(Trigger, Transition.Trigger))
+			for (const FStateTreeTransition& Transition : State->Transitions)
 			{
-				if (EditorData->HasBreakpoint(Transition.ID, EStateTreeBreakpointType::OnTransition))
+				if (Transition.bTransitionEnabled && EnumHasAnyFlags(Trigger, Transition.Trigger))
 				{
-					return GetTransitionsVisibility(State, Trigger);
+					if (EditorData->HasBreakpoint(Transition.ID, EStateTreeBreakpointType::OnTransition))
+					{
+						return GetTransitionsVisibility(Trigger);
+					}
 				}
 			}
 		}
-	}
 #endif // WITH_STATETREE_TRACE_DEBUGGER
-	
-	return EVisibility::Collapsed;
-}
-
-EVisibility SStateTreeViewRow::GetCompletedTransitionVisibility() const
-{
-	if (const UStateTreeState* State = WeakState.Get())
-	{
-		return GetTransitionsVisibility(*State, EStateTreeTransitionTrigger::OnStateCompleted);
-	}
-	return EVisibility::Visible;
-}
-
-EVisibility SStateTreeViewRow::GetCompletedTransitionBreakpointVisibility() const
-{
-	if (const UStateTreeState* State = WeakState.Get())
-	{
-		return GetTransitionsBreakpointVisibility(*State, EStateTreeTransitionTrigger::OnStateCompleted);
-	}
-	return EVisibility::Visible;
-}
-
-FText SStateTreeViewRow::GetCompletedTransitionsDesc() const
-{
-	if (const UStateTreeState* State = WeakState.Get())
-	{
-		return GetTransitionsDesc(*State, EStateTreeTransitionTrigger::OnStateCompleted);
-	}
-	return LOCTEXT("Invalid", "Invalid");
-}
-
-FText SStateTreeViewRow::GetCompletedTransitionWithBreakpointDesc() const
-{
-	if (const UStateTreeState* State = WeakState.Get())
-	{
-		FTransitionDescFilterOptions FilterOptions;
-		FilterOptions.WithBreakpoint = ETransitionDescRequirement::RequiredTrue;
-		return GetTransitionsDesc(*State, EStateTreeTransitionTrigger::OnStateCompleted, FilterOptions);
-	}
-	return FText::GetEmpty();
-}
-
-const FSlateBrush* SStateTreeViewRow::GetCompletedTransitionsIcon() const
-{
-	if (const UStateTreeState* State = WeakState.Get())
-	{
-		return GetTransitionsIcon(*State, EStateTreeTransitionTrigger::OnStateCompleted);
-	}
-	return nullptr;
-}
-
-EVisibility SStateTreeViewRow::GetSucceededTransitionVisibility() const
-{
-	if (const UStateTreeState* State = WeakState.Get())
-	{
-		return GetTransitionsVisibility(*State, EStateTreeTransitionTrigger::OnStateSucceeded);
-	}
-	return EVisibility::Collapsed;
-}
-
-EVisibility SStateTreeViewRow::GetSucceededTransitionBreakpointVisibility() const
-{
-	if (const UStateTreeState* State = WeakState.Get())
-	{
-		return GetTransitionsBreakpointVisibility(*State, EStateTreeTransitionTrigger::OnStateSucceeded);
-	}
-	return EVisibility::Collapsed;
-}
-
-FText SStateTreeViewRow::GetSucceededTransitionDesc() const
-{
-	if (const UStateTreeState* State = WeakState.Get())
-	{
-		return GetTransitionsDesc(*State, EStateTreeTransitionTrigger::OnStateSucceeded);
-	}
-	return FText::GetEmpty();
-}
-
-FText SStateTreeViewRow::GetSucceededTransitionWithBreakpointDesc() const
-{
-	if (const UStateTreeState* State = WeakState.Get())
-	{
-		FTransitionDescFilterOptions FilterOptions;
-		FilterOptions.WithBreakpoint = ETransitionDescRequirement::RequiredTrue;
-		return GetTransitionsDesc(*State, EStateTreeTransitionTrigger::OnStateSucceeded, FilterOptions);
-	}
-	return FText::GetEmpty();
-}
-
-const FSlateBrush* SStateTreeViewRow::GetSucceededTransitionIcon() const
-{
-	if (const UStateTreeState* State = WeakState.Get())
-	{
-		return GetTransitionsIcon(*State, EStateTreeTransitionTrigger::OnStateSucceeded);
-	}
-	return nullptr;
-}
-
-EVisibility SStateTreeViewRow::GetFailedTransitionVisibility() const
-{
-	if (const UStateTreeState* State = WeakState.Get())
-	{
-		return GetTransitionsVisibility(*State, EStateTreeTransitionTrigger::OnStateFailed);
-	}
-	return EVisibility::Collapsed;
-}
-
-EVisibility SStateTreeViewRow::GetFailedTransitionBreakpointVisibility() const
-{
-	if (const UStateTreeState* State = WeakState.Get())
-	{
-		return GetTransitionsBreakpointVisibility(*State, EStateTreeTransitionTrigger::OnStateFailed);
-	}
-	return EVisibility::Collapsed;
-}
-
-FText SStateTreeViewRow::GetFailedTransitionDesc() const
-{
-	if (const UStateTreeState* State = WeakState.Get())
-	{
-		return GetTransitionsDesc(*State, EStateTreeTransitionTrigger::OnStateFailed);
-	}
-	return LOCTEXT("Invalid", "Invalid");
-}
-
-FText SStateTreeViewRow::GetFailedTransitionWithBreakpointDesc() const
-{
-	if (const UStateTreeState* State = WeakState.Get())
-	{
-		FTransitionDescFilterOptions FilterOptions;
-		FilterOptions.WithBreakpoint = ETransitionDescRequirement::RequiredTrue;
-		return GetTransitionsDesc(*State, EStateTreeTransitionTrigger::OnStateFailed, FilterOptions);
-	}
-	return FText::GetEmpty();
-}
-
-const FSlateBrush* SStateTreeViewRow::GetFailedTransitionIcon() const
-{
-	if (const UStateTreeState* State = WeakState.Get())
-	{
-		return GetTransitionsIcon(*State, EStateTreeTransitionTrigger::OnStateFailed);
-	}
-	return nullptr;
-}
-
-EVisibility SStateTreeViewRow::GetConditionalTransitionsVisibility() const
-{
-	if (const UStateTreeState* State = WeakState.Get())
-	{
-		return GetTransitionsVisibility(*State, EStateTreeTransitionTrigger::OnTick | EStateTreeTransitionTrigger::OnEvent);
-	}
-	return EVisibility::Collapsed;
-}
-
-EVisibility SStateTreeViewRow::GetConditionalTransitionsBreakpointVisibility() const
-{
-	if (const UStateTreeState* State = WeakState.Get())
-	{
-		return GetTransitionsBreakpointVisibility(*State, EStateTreeTransitionTrigger::OnTick | EStateTreeTransitionTrigger::OnEvent);
 	}
+	
 	return EVisibility::Collapsed;
 }
 
-FText SStateTreeViewRow::GetConditionalTransitionsDesc() const
+EVisibility SStateTreeViewRow::GetTransitionDashVisibility() const
 {
 	if (const UStateTreeState* State = WeakState.Get())
 	{
-		FTransitionDescFilterOptions FilterOptions;
-		FilterOptions.bUseMask = true;
-		return GetTransitionsDesc(*State, EStateTreeTransitionTrigger::OnTick | EStateTreeTransitionTrigger::OnEvent, FilterOptions);
-	}
-	return FText::GetEmpty();
-}
-
-FText SStateTreeViewRow::GetConditionalTransitionsWithBreakpointDesc() const
-{
-	if (const UStateTreeState* State = WeakState.Get())
-	{
-		FTransitionDescFilterOptions FilterOptions;
-		FilterOptions.WithBreakpoint = ETransitionDescRequirement::RequiredTrue;
-		FilterOptions.bUseMask = true;
-		return GetTransitionsDesc(*State, EStateTreeTransitionTrigger::OnTick | EStateTreeTransitionTrigger::OnEvent, FilterOptions);
+		return State->Transitions.IsEmpty() ? EVisibility::Collapsed : EVisibility::Visible;
 	}
-	return FText::GetEmpty();
+	return EVisibility::Collapsed;	
 }
 
 bool SStateTreeViewRow::IsRootState() const
@@ -1708,17 +2039,29 @@ void SStateTreeViewRow::HandleNodeLabelTextCommitted(const FText& NewLabel, ETex
 
 FReply SStateTreeViewRow::HandleDragDetected(const FGeometry&, const FPointerEvent&) const
 {
-	return FReply::Handled().BeginDragDrop(FActionTreeViewDragDrop::New(WeakState.Get()));
+	return FReply::Handled().BeginDragDrop(FStateTreeSelectedDragDrop::New(StateTreeViewModel));
+}
+
+void SStateTreeViewRow::HandleDragLeave(const FDragDropEvent& DragDropEvent) const
+{
+	const TSharedPtr<FStateTreeSelectedDragDrop> DragDropOperation = DragDropEvent.GetOperationAs<FStateTreeSelectedDragDrop>();
+	if (DragDropOperation.IsValid())
+	{
+		DragDropOperation->SetCanDrop(false);
+	}	
 }
 
 TOptional<EItemDropZone> SStateTreeViewRow::HandleCanAcceptDrop(const FDragDropEvent& DragDropEvent, EItemDropZone DropZone, TWeakObjectPtr<UStateTreeState> TargetState) const
 {
-	const TSharedPtr<FActionTreeViewDragDrop> DragDropOperation = DragDropEvent.GetOperationAs<FActionTreeViewDragDrop>();
+	const TSharedPtr<FStateTreeSelectedDragDrop> DragDropOperation = DragDropEvent.GetOperationAs<FStateTreeSelectedDragDrop>();
 	if (DragDropOperation.IsValid())
 	{
+		DragDropOperation->SetCanDrop(true);
+
 		// Cannot drop on selection or child of selection.
 		if (StateTreeViewModel && StateTreeViewModel->IsChildOfSelection(TargetState.Get()))
 		{
+			DragDropOperation->SetCanDrop(false);
 			return TOptional<EItemDropZone>();
 		}
 
@@ -1730,7 +2073,7 @@ TOptional<EItemDropZone> SStateTreeViewRow::HandleCanAcceptDrop(const FDragDropE
 
 FReply SStateTreeViewRow::HandleAcceptDrop(const FDragDropEvent& DragDropEvent, EItemDropZone DropZone, TWeakObjectPtr<UStateTreeState> TargetState) const
 {
-	const TSharedPtr<FActionTreeViewDragDrop> DragDropOperation = DragDropEvent.GetOperationAs<FActionTreeViewDragDrop>();
+	const TSharedPtr<FStateTreeSelectedDragDrop> DragDropOperation = DragDropEvent.GetOperationAs<FStateTreeSelectedDragDrop>();
 	if (DragDropOperation.IsValid())
 	{
 		if (StateTreeViewModel)
diff --git a/Engine/Plugins/Runtime/StateTree/Source/StateTreeEditorModule/Private/SStateTreeViewRow.h b/Engine/Plugins/Runtime/StateTree/Source/StateTreeEditorModule/Private/SStateTreeViewRow.h
--- a/Engine/Plugins/Runtime/StateTree/Source/StateTreeEditorModule/Private/SStateTreeViewRow.h
+++ b/Engine/Plugins/Runtime/StateTree/Source/StateTreeEditorModule/Private/SStateTreeViewRow.h
@@ -31,14 +31,17 @@ public:
 	void RequestRename() const;
 
 private:
-	TSharedRef<SHorizontalBox> CreateTasksWidget();
+	TSharedRef<SWidget> MakeTasksWidget(const TSharedPtr<SScrollBox>& ViewBox);
+	TSharedRef<SWidget> MakeConditionsWidget(const TSharedPtr<SScrollBox>& ViewBox);
 
 	FSlateColor GetTitleColor(const float Alpha = 1.0f, const float Lighten = 0.0f) const;
 	FSlateColor GetActiveStateColor() const;
-	FSlateColor GetSubTreeMarkerColor() const;
 	FText GetStateDesc() const;
 	FText GetStateIDDesc() const;
 
+	FSlateColor GetSubTreeMarkerColor() const;
+	EVisibility GetSubTreeVisibility() const;
+	
 	EVisibility GetConditionVisibility() const;
 	EVisibility GetStateBreakpointVisibility() const;
 	FText GetStateBreakpointTooltipText() const;
@@ -53,8 +56,16 @@ private:
 	FSlateColor GetTaskIconColor(FGuid TaskID) const;
 	FText GetTaskDesc(FGuid TaskID, EStateTreeNodeFormatting Formatting) const;
 
-	EVisibility GetTasksVisibility() const;
+	const FStateTreeEditorNode* GetConditionNodeByID(FGuid ConditionID) const;
+	EVisibility GetConditionIconVisibility(FGuid ConditionID) const;
+	const FSlateBrush* GetConditionIcon(FGuid ConditionID) const;
+	FSlateColor GetConditionIconColor(FGuid ConditionID) const;
+	FText GetConditionDesc(FGuid ConditionID, EStateTreeNodeFormatting Formatting) const;
 
+	FText GetOperandText(const int32 ConditionIndex) const;
+	FText GetOpenParens(const int32 ConditionIndex) const;
+	FText GetCloseParens(const int32 ConditionIndex) const;
+	
 	EVisibility GetLinkedStateVisibility() const;
 	FText GetLinkedStateDesc() const;
 
@@ -62,29 +73,8 @@ private:
 	EVisibility GetWarningsVisibility() const;
 	FText GetWarningsTooltipText() const;
 
-	EVisibility GetCompletedTransitionVisibility() const;
-	EVisibility GetCompletedTransitionBreakpointVisibility() const;
-	FText GetCompletedTransitionsDesc() const;
-	FText GetCompletedTransitionWithBreakpointDesc() const;
-	const FSlateBrush* GetCompletedTransitionsIcon() const;
-
-	EVisibility GetSucceededTransitionVisibility() const;
-	EVisibility GetSucceededTransitionBreakpointVisibility() const;
-	FText GetSucceededTransitionDesc() const;
-	FText GetSucceededTransitionWithBreakpointDesc() const;
-	const FSlateBrush* GetSucceededTransitionIcon() const;
-
-	EVisibility GetFailedTransitionVisibility() const;
-	EVisibility GetFailedTransitionBreakpointVisibility() const;
-	FText GetFailedTransitionDesc() const;
-	FText GetFailedTransitionWithBreakpointDesc() const;
-	const FSlateBrush* GetFailedTransitionIcon() const;
-
-	EVisibility GetConditionalTransitionsVisibility() const;
-	EVisibility GetConditionalTransitionsBreakpointVisibility() const;
-	FText GetConditionalTransitionsDesc() const;
-	FText GetConditionalTransitionsWithBreakpointDesc() const;
-
+	EVisibility GetTransitionDashVisibility() const;
+	
 	enum class ETransitionDescRequirement : uint8
 	{
 		Any,
@@ -105,11 +95,13 @@ private:
 		bool bUseMask = false;
 	};
 
-	static FText GetLinkDescription(const FStateTreeStateLink& Link);
-	FText GetTransitionsDesc(const UStateTreeState& State, const EStateTreeTransitionTrigger Trigger, const FTransitionDescFilterOptions FilterOptions = {}) const;
-	const FSlateBrush* GetTransitionsIcon(const UStateTreeState& State, const EStateTreeTransitionTrigger Trigger, const FTransitionDescFilterOptions FilterOptions = {}) const;
-	EVisibility GetTransitionsVisibility(const UStateTreeState& State, const EStateTreeTransitionTrigger Trigger) const;
-	EVisibility GetTransitionsBreakpointVisibility(const UStateTreeState& State, const EStateTreeTransitionTrigger Trigger) const;
+	TSharedRef<SWidget> MakeTransitionWidgets(const EStateTreeTransitionTrigger Trigger, const FTransitionDescFilterOptions FilterOptions = {});
+
+	FText GetLinkTooltip(const FStateTreeStateLink& Link, const FGuid NodeID) const;
+	FText GetTransitionsDesc(const EStateTreeTransitionTrigger Trigger, const FTransitionDescFilterOptions FilterOptions = {}) const;
+	const FSlateBrush* GetTransitionsIcon(const EStateTreeTransitionTrigger Trigger) const;
+	EVisibility GetTransitionsVisibility(const EStateTreeTransitionTrigger Trigger) const;
+	EVisibility GetTransitionsBreakpointVisibility(const EStateTreeTransitionTrigger Trigger) const;
 
 	bool HasParentTransitionForTrigger(const UStateTreeState& State, const EStateTreeTransitionTrigger Trigger) const;
 
@@ -120,9 +112,10 @@ private:
 	void HandleNodeLabelTextCommitted(const FText& NewLabel, ETextCommit::Type CommitType) const;
 
 	FReply HandleDragDetected(const FGeometry&, const FPointerEvent&) const;
+	void HandleDragLeave(const FDragDropEvent& DragDropEvent) const;
 	TOptional<EItemDropZone> HandleCanAcceptDrop(const FDragDropEvent& DragDropEvent, EItemDropZone DropZone, TWeakObjectPtr<UStateTreeState> TargetState) const;
 	FReply HandleAcceptDrop(const FDragDropEvent& DragDropEvent, EItemDropZone DropZone, TWeakObjectPtr<UStateTreeState> TargetState) const;
-	
+
 	TSharedPtr<FStateTreeViewModel> StateTreeViewModel;
 	TWeakObjectPtr<UStateTreeState> WeakState;
 	TWeakObjectPtr<UStateTreeEditorData> WeakEditorData;
diff --git a/Engine/Plugins/Runtime/StateTree/Source/StateTreeEditorModule/Private/StateTreeDescriptionHelpers.cpp b/Engine/Plugins/Runtime/StateTree/Source/StateTreeEditorModule/Private/StateTreeDescriptionHelpers.cpp
new file mode 100644
--- /dev/null
+++ b/Engine/Plugins/Runtime/StateTree/Source/StateTreeEditorModule/Private/StateTreeDescriptionHelpers.cpp
@@ -0,0 +1,243 @@
+﻿// Copyright Epic Games, Inc. All Rights Reserved.
+
+#include "StateTreeDescriptionHelpers.h"
+
+#include "StateTreeEditorData.h"
+#include "StateTreeEditorStyle.h"
+#include "StateTreeState.h"
+#include "StateTreeTypes.h"
+#include "StateTreeNodeBase.h"
+
+#define LOCTEXT_NAMESPACE "StateTreeEditor"
+
+namespace UE::StateTree::Editor
+{
+
+FText GetStateLinkDesc(const UStateTreeEditorData* EditorData, const FStateTreeStateLink& Link, EStateTreeNodeFormatting Formatting, bool bShowStatePath)
+{
+	if (!EditorData)
+	{
+		return FText::GetEmpty();
+	}
+
+	if (Link.LinkType == EStateTreeTransitionType::None)
+	{
+		return Formatting == EStateTreeNodeFormatting::RichText
+			? LOCTEXT("TransitionNoneRich", "<i>None</>")
+			: LOCTEXT("TransitionNone", "None");
+	}
+	if (Link.LinkType == EStateTreeTransitionType::NextState)
+	{
+		return Formatting == EStateTreeNodeFormatting::RichText
+			? LOCTEXT("TransitionNextStateRich", "<i>Next State</>")
+			: LOCTEXT("TransitionNextState", "Next State");
+	}
+	if (Link.LinkType == EStateTreeTransitionType::NextSelectableState)
+	{
+		return Formatting == EStateTreeNodeFormatting::RichText
+			? LOCTEXT("TransitionNextSelectableStateRich", "<i>Next Selectable State</>")
+			: LOCTEXT("TransitionNextSelectableState", "Next Selectable State");
+	}
+	if (Link.LinkType == EStateTreeTransitionType::Succeeded)
+	{
+		return Formatting == EStateTreeNodeFormatting::RichText
+			? LOCTEXT("TransitionTreeSucceededRich", "<i>Tree Succeeded</>")
+			: LOCTEXT("TransitionTreeSucceeded", "Tree Succeeded");
+	}
+	if (Link.LinkType == EStateTreeTransitionType::Failed)
+	{
+		return Formatting == EStateTreeNodeFormatting::RichText
+			? LOCTEXT("TransitionTreeFailedRich", "<i>Tree Failed</>")
+			: LOCTEXT("TransitionTreeFailed", "Tree Failed");
+	}
+	if (Link.LinkType == EStateTreeTransitionType::GotoState)
+	{
+		if (const UStateTreeState* State = EditorData->GetStateByID(Link.ID))
+		{
+			if (bShowStatePath)
+			{
+				TArray<FText> Path;
+				while (State)
+				{
+					Path.Insert(FText::FromName(State->Name), 0);
+					State = State->Parent;
+				}
+				return FText::Join(FText::FromString(TEXT("/")), Path);
+			}
+			else
+			{
+				return FText::FromName(State->Name);
+			}
+		}
+		return FText::FromName(Link.Name);
+	}
+
+	return LOCTEXT("TransitionInvalid", "Invalid");
+}
+
+const FSlateBrush* GetStateLinkIcon(const UStateTreeEditorData* EditorData, const FStateTreeStateLink& Link)
+{
+	if (!EditorData)
+	{
+		return nullptr;
+	}
+
+	if (Link.LinkType == EStateTreeTransitionType::None)
+	{
+		return FStateTreeEditorStyle::Get().GetBrush("StateTreeEditor.Transition.None");
+	}
+	if (Link.LinkType == EStateTreeTransitionType::NextState)
+	{
+		return FStateTreeEditorStyle::Get().GetBrush("StateTreeEditor.Transition.Next");
+	}
+	if (Link.LinkType == EStateTreeTransitionType::NextSelectableState)
+	{
+		return FStateTreeEditorStyle::Get().GetBrush("StateTreeEditor.Transition.Next");
+	}
+	if (Link.LinkType == EStateTreeTransitionType::Succeeded)
+	{
+		return FStateTreeEditorStyle::Get().GetBrush("StateTreeEditor.Transition.Succeeded");
+	}
+	if (Link.LinkType == EStateTreeTransitionType::Failed)
+	{
+		return FStateTreeEditorStyle::Get().GetBrush("StateTreeEditor.Transition.Failed");
+	}
+	if (Link.LinkType == EStateTreeTransitionType::GotoState)
+	{
+		if (const UStateTreeState* State = EditorData->GetStateByID(Link.ID))
+		{
+			// Figure out icon.
+			if (State->SelectionBehavior == EStateTreeStateSelectionBehavior::None)
+			{
+				return FStateTreeEditorStyle::Get().GetBrush("StateTreeEditor.SelectNone");
+			}
+			else if (State->SelectionBehavior == EStateTreeStateSelectionBehavior::TryEnterState)
+			{
+				return FStateTreeEditorStyle::Get().GetBrush("StateTreeEditor.TryEnterState");			
+			}
+			else if (State->SelectionBehavior == EStateTreeStateSelectionBehavior::TrySelectChildrenInOrder)
+			{
+				if (State->Children.IsEmpty()
+					|| State->Type == EStateTreeStateType::Linked
+					|| State->Type == EStateTreeStateType::LinkedAsset)
+				{
+					return FStateTreeEditorStyle::Get().GetBrush("StateTreeEditor.TryEnterState");			
+				}
+				else
+				{
+					return FStateTreeEditorStyle::Get().GetBrush("StateTreeEditor.TrySelectChildrenInOrder");
+				}
+			}
+			else if (State->SelectionBehavior == EStateTreeStateSelectionBehavior::TryFollowTransitions)
+			{
+				return FStateTreeEditorStyle::Get().GetBrush("StateTreeEditor.TryFollowTransitions");
+			}
+		}
+	}
+	
+	return nullptr;
+}
+
+FSlateColor GetStateLinkColor(const UStateTreeEditorData* EditorData, const FStateTreeStateLink& Link)
+{
+	if (Link.LinkType == EStateTreeTransitionType::GotoState)
+	{
+		if (const UStateTreeState* State = EditorData->GetStateByID(Link.ID))
+		{
+			FLinearColor Color = FColor(31, 151, 167);
+			if (const FStateTreeEditorColor* FoundColor = EditorData->FindColor(State->ColorRef))
+			{
+				Color = FoundColor->Color;
+			}
+			return Color;
+		}
+		
+		return FLinearColor(1.f, 1.f, 1.f, 0.25f);
+	}
+	return FLinearColor::White;
+}
+
+FText GetTransitionDesc(const UStateTreeEditorData* EditorData, const FStateTreeTransition& Transition, EStateTreeNodeFormatting Formatting, bool bShowStatePath)
+{
+	if (!EditorData)
+	{
+		return FText::GetEmpty();
+	}
+	
+	FText TriggerText;
+	if (Transition.Trigger == EStateTreeTransitionTrigger::OnStateCompleted)
+	{
+		TriggerText = Formatting == EStateTreeNodeFormatting::RichText
+			? LOCTEXT("TransitionOnStateCompleted", "<b>On State Completed</>")
+			: LOCTEXT("TransitionOnStateCompleted", "On State Completed");
+	}
+	else if (Transition.Trigger == EStateTreeTransitionTrigger::OnStateSucceeded)
+	{
+		TriggerText = Formatting == EStateTreeNodeFormatting::RichText
+			? LOCTEXT("TransitionOnStateSucceeded", "<b>On State Succeeded</b>")
+			: LOCTEXT("TransitionOnStateSucceeded", "On State Succeeded");
+	}
+	else if (Transition.Trigger == EStateTreeTransitionTrigger::OnStateFailed)
+	{
+		TriggerText = Formatting == EStateTreeNodeFormatting::RichText
+			? LOCTEXT("TransitionOnStateFailed", "<b>On State Failed</>")
+			: LOCTEXT("TransitionOnStateFailed", "On State Failed");
+	}
+	else if (Transition.Trigger == EStateTreeTransitionTrigger::OnTick)
+	{
+		TriggerText = Formatting == EStateTreeNodeFormatting::RichText
+			? LOCTEXT("TransitionOnTick", "<b>On Tick</b>")
+			: LOCTEXT("TransitionOnTick", "On Tick");
+	}
+	else if (Transition.Trigger == EStateTreeTransitionTrigger::OnEvent)
+	{
+		TArray<FText> PayloadItems;
+		
+		if (Transition.RequiredEvent.IsValid())
+		{
+			if (Transition.RequiredEvent.Tag.IsValid())
+			{
+				const FText TagFormat = Formatting == EStateTreeNodeFormatting::RichText
+					? LOCTEXT("TransitionEventTag", "<s>Tag:</> '{0}'")
+					: LOCTEXT("TransitionEventTag", "Tag: '{0}'");
+				PayloadItems.Add(FText::Format(TagFormat, FText::FromName(Transition.RequiredEvent.Tag.GetTagName())));
+			}
+			
+			if (Transition.RequiredEvent.PayloadStruct)
+			{
+				const FText PayloadFormat = Formatting == EStateTreeNodeFormatting::RichText
+					? LOCTEXT("TransitionEventPayload", "<s>Payload:</> '{0}'")
+					: LOCTEXT("TransitionEventPayload", "Payload: '{0}'");
+				PayloadItems.Add(FText::Format(PayloadFormat, Transition.RequiredEvent.PayloadStruct->GetDisplayNameText()));
+			}
+		}
+		else
+		{
+			PayloadItems.Add(LOCTEXT("TransitionInvalidEvent", "Invalid"));
+		}
+
+		const FText TransitionFormat = Formatting == EStateTreeNodeFormatting::RichText
+			? LOCTEXT("TransitionOnEvent", "<b>On Event</> ({0})")
+			: LOCTEXT("TransitionOnEvent", "On Event ({0})");
+		
+		TriggerText = FText::Format(TransitionFormat, FText::Join(INVTEXT(", "), PayloadItems));
+	}
+
+	FText ActionText = Formatting == EStateTreeNodeFormatting::RichText
+		? LOCTEXT("ActionGotoRich", "<s>go to</>")
+		: LOCTEXT("ActionGoto", "go to");
+	
+	if (Transition.State.LinkType == EStateTreeTransitionType::Succeeded
+		|| Transition.State.LinkType == EStateTreeTransitionType::Failed)
+	{
+		ActionText = Formatting == EStateTreeNodeFormatting::RichText
+			? LOCTEXT("ActionReturnRich", "<s>return</>")
+			: LOCTEXT("ActionReturn", "return");
+	}
+	
+	return FText::Format(LOCTEXT("TransitionDesc", "{0} {1} {2}"), TriggerText, ActionText, GetStateLinkDesc(EditorData, Transition.State, Formatting, bShowStatePath));
+}
+
+} // UE::StateTree::Editor
+
+#undef LOCTEXT_NAMESPACE
diff --git a/Engine/Plugins/Runtime/StateTree/Source/StateTreeEditorModule/Private/StateTreeDescriptionHelpers.h b/Engine/Plugins/Runtime/StateTree/Source/StateTreeEditorModule/Private/StateTreeDescriptionHelpers.h
new file mode 100644
--- /dev/null
+++ b/Engine/Plugins/Runtime/StateTree/Source/StateTreeEditorModule/Private/StateTreeDescriptionHelpers.h
@@ -0,0 +1,24 @@
+﻿// Copyright Epic Games, Inc. All Rights Reserved.
+
+#pragma once
+
+#include "HAL/Platform.h"
+
+class UStateTreeEditorData;
+struct FStateTreeStateLink;
+struct FStateTreeTransition;
+enum EStateTreeNodeFormatting : uint8;
+struct FSlateBrush;
+struct FSlateColor;
+class FText;
+
+namespace UE::StateTree::Editor
+{
+
+FText GetStateLinkDesc(const UStateTreeEditorData* EditorData, const FStateTreeStateLink& Link, EStateTreeNodeFormatting Formatting, bool bShowStatePath = false);
+const FSlateBrush* GetStateLinkIcon(const UStateTreeEditorData* EditorData, const FStateTreeStateLink& Link);
+FSlateColor GetStateLinkColor(const UStateTreeEditorData* EditorData, const FStateTreeStateLink& Link);
+
+FText GetTransitionDesc(const UStateTreeEditorData* EditorData, const FStateTreeTransition& Transition, EStateTreeNodeFormatting Formatting, bool bShowStatePath = false);
+
+} // UE::StateTree::Editor
\ No newline at end of file
diff --git a/Engine/Plugins/Runtime/StateTree/Source/StateTreeEditorModule/Private/StateTreeDragDrop.cpp b/Engine/Plugins/Runtime/StateTree/Source/StateTreeEditorModule/Private/StateTreeDragDrop.cpp
new file mode 100644
--- /dev/null
+++ b/Engine/Plugins/Runtime/StateTree/Source/StateTreeEditorModule/Private/StateTreeDragDrop.cpp
@@ -0,0 +1,41 @@
+﻿// Copyright Epic Games, Inc. All Rights Reserved.
+
+#include "StateTreeDragDrop.h"
+
+#include "StateTreeState.h"
+#include "StateTreeViewModel.h"
+#include "Widgets/SBoxPanel.h"
+#include "Widgets/Layout/SBorder.h"
+#include "Widgets/Text/STextBlock.h"
+
+TSharedPtr<SWidget> FStateTreeSelectedDragDrop::GetDefaultDecorator() const
+{
+	// Display all dragged nodes, enable the ones that can be moved into the current target node. 
+	const TSharedRef<SVerticalBox> Box = SNew(SVerticalBox);
+
+	if (ViewModel)
+	{
+		TArray<UStateTreeState*> SelectedStates;
+		ViewModel->GetSelectedStates(SelectedStates);
+
+		for (UStateTreeState* State : SelectedStates)
+		{
+			Box->AddSlot()
+			.Padding(FMargin(4, 2))
+			[
+				SNew(STextBlock)
+				.Text(FText::FromName(State->Name))
+				.IsEnabled_Lambda([this]()
+				{
+					return bCanDrop;
+				})
+			];
+		}
+	}
+
+	return SNew(SBorder)
+		.BorderImage(FAppStyle::GetBrush("Menu.Background"))
+		[
+			Box
+		];
+}
diff --git a/Engine/Plugins/Runtime/StateTree/Source/StateTreeEditorModule/Private/StateTreeDragDrop.h b/Engine/Plugins/Runtime/StateTree/Source/StateTreeEditorModule/Private/StateTreeDragDrop.h
new file mode 100644
--- /dev/null
+++ b/Engine/Plugins/Runtime/StateTree/Source/StateTreeEditorModule/Private/StateTreeDragDrop.h
@@ -0,0 +1,32 @@
+﻿// Copyright Epic Games, Inc. All Rights Reserved.
+
+#pragma once
+
+#include "DragAndDrop/DecoratedDragDropOp.h"
+
+class FStateTreeViewModel;
+
+class FStateTreeSelectedDragDrop : public FDecoratedDragDropOp
+{
+public:
+	DRAG_DROP_OPERATOR_TYPE(FActionTreeViewDragDrop, FDecoratedDragDropOp);
+
+	static TSharedRef<FStateTreeSelectedDragDrop> New(TSharedPtr<FStateTreeViewModel> InViewModel)
+	{
+		TSharedRef<FStateTreeSelectedDragDrop> Operation = MakeShared<FStateTreeSelectedDragDrop>();
+		Operation->ViewModel = InViewModel;
+		Operation->Construct();
+
+		return Operation;
+	}
+
+	virtual TSharedPtr<SWidget> GetDefaultDecorator() const override;
+
+	void SetCanDrop(const bool bState)
+	{
+		bCanDrop = bState;
+	}
+
+	TSharedPtr<FStateTreeViewModel> ViewModel;
+	bool bCanDrop = false;
+};
\ No newline at end of file
diff --git a/Engine/Plugins/Runtime/StateTree/Source/StateTreeEditorModule/Private/StateTreeEditor.cpp b/Engine/Plugins/Runtime/StateTree/Source/StateTreeEditorModule/Private/StateTreeEditor.cpp
--- a/Engine/Plugins/Runtime/StateTree/Source/StateTreeEditorModule/Private/StateTreeEditor.cpp
+++ b/Engine/Plugins/Runtime/StateTree/Source/StateTreeEditorModule/Private/StateTreeEditor.cpp
@@ -21,7 +21,6 @@
 #include "Misc/UObjectToken.h"
 #include "SStateTreeView.h"
 #include "StateTree.h"
-//#include "StateTreeSchema.h"
 #include "StateTreeCompiler.h"
 #include "StateTreeCompilerLog.h"
 #include "StateTreeDelegates.h"
@@ -39,6 +38,8 @@
 #include "ToolMenuEntry.h"
 #include "Widgets/Input/SMultiLineEditableTextBox.h"
 #include "FileHelpers.h"
+#include "PropertyPath.h"
+#include "SStateTreeOutliner.h"
 #include "Debugger/SStateTreeDebuggerView.h"
 #include "StateTreeSettings.h"
 #include "ToolMenuSection.h"
@@ -50,6 +51,7 @@ const FName StateTreeEditorAppName(TEXT("StateTreeEditorApp"));
 const FName FStateTreeEditor::SelectionDetailsTabId(TEXT("StateTreeEditor_SelectionDetails"));
 const FName FStateTreeEditor::AssetDetailsTabId(TEXT("StateTreeEditor_AssetDetails"));
 const FName FStateTreeEditor::StateTreeViewTabId(TEXT("StateTreeEditor_StateTreeView"));
+const FName FStateTreeEditor::StateTreeOutlinerTabId(TEXT("StateTreeEditor_StateTreeOutliner"));
 const FName FStateTreeEditor::StateTreeStatisticsTabId(TEXT("StateTreeEditor_StateTreeStatistics"));
 const FName FStateTreeEditor::CompilerResultsTabId(TEXT("StateTreeEditor_CompilerResults"));
 #if WITH_STATETREE_TRACE_DEBUGGER
@@ -104,7 +106,12 @@ void FStateTreeEditor::RegisterTabSpawners(const TSharedRef<class FTabManager>&
 		.SetDisplayName(NSLOCTEXT("StateTreeEditor", "StateTreeViewTab", "States"))
 		.SetGroup(WorkspaceMenuCategoryRef)
 		.SetIcon(FSlateIcon(FAppStyle::GetAppStyleSetName(), "LevelEditor.Tabs.Outliner"));
-	
+
+	InTabManager->RegisterTabSpawner(StateTreeOutlinerTabId, FOnSpawnTab::CreateSP(this, &FStateTreeEditor::SpawnTab_StateTreeOutliner))
+		.SetDisplayName(NSLOCTEXT("StateTreeEditor", "StateTreeOutlinerTab", "Outliner"))
+		.SetGroup(WorkspaceMenuCategoryRef)
+		.SetIcon(FSlateIcon(FAppStyle::GetAppStyleSetName(), "LevelEditor.Tabs.Outliner"));
+
 	InTabManager->RegisterTabSpawner(StateTreeStatisticsTabId, FOnSpawnTab::CreateSP(this, &FStateTreeEditor::SpawnTab_StateTreeStatistics))
 		.SetDisplayName(NSLOCTEXT("StateTreeEditor", "StatisticsTab", "StateTree Statistics"))
 		.SetGroup(WorkspaceMenuCategoryRef)
@@ -131,6 +138,7 @@ void FStateTreeEditor::UnregisterTabSpawners(const TSharedRef<class FTabManager>
 	InTabManager->UnregisterTabSpawner(SelectionDetailsTabId);
 	InTabManager->UnregisterTabSpawner(AssetDetailsTabId);
 	InTabManager->UnregisterTabSpawner(StateTreeViewTabId);
+	InTabManager->UnregisterTabSpawner(StateTreeOutlinerTabId);
 	InTabManager->UnregisterTabSpawner(StateTreeStatisticsTabId);
 	InTabManager->UnregisterTabSpawner(CompilerResultsTabId);
 #if WITH_STATETREE_TRACE_DEBUGGER
@@ -184,6 +192,7 @@ void FStateTreeEditor::InitEditor( const EToolkitMode::Type Mode, const TSharedP
 	StateTreeViewModel->GetOnStatesRemoved().AddSPLambda(this, [this](const TSet<UStateTreeState*>&){ UpdateAsset(); });
 	StateTreeViewModel->GetOnStatesMoved().AddSPLambda(this, [this](const TSet<UStateTreeState*>&, const TSet<UStateTreeState*>&){ UpdateAsset(); });
 	StateTreeViewModel->GetOnSelectionChanged().AddSP(this, &FStateTreeEditor::HandleModelSelectionChanged);
+	StateTreeViewModel->GetOnBringNodeToFocus().AddSP(this, &FStateTreeEditor::HandleModelBringNodeToFocus);
 
 	FMessageLogModule& MessageLogModule = FModuleManager::LoadModuleChecked<FMessageLogModule>("MessageLog");
 	FMessageLogInitializationOptions LogOptions;
@@ -198,7 +207,7 @@ void FStateTreeEditor::InitEditor( const EToolkitMode::Type Mode, const TSharedP
 	CompilerResultsListing->OnMessageTokenClicked().AddSP(this, &FStateTreeEditor::HandleMessageTokenClicked);
 
 	
-	TSharedRef<FTabManager::FLayout> StandaloneDefaultLayout = FTabManager::NewLayout("Standalone_StateTree_Layout_v3")
+	TSharedRef<FTabManager::FLayout> StandaloneDefaultLayout = FTabManager::NewLayout("Standalone_StateTree_Layout_v4")
 	->AddArea
 	(
 		FTabManager::NewPrimaryArea() ->SetOrientation(Orient_Vertical)
@@ -207,20 +216,12 @@ void FStateTreeEditor::InitEditor( const EToolkitMode::Type Mode, const TSharedP
 			FTabManager::NewSplitter() ->SetOrientation(Orient_Horizontal)
 			->Split
 			(
-				FTabManager::NewSplitter()->SetOrientation(Orient_Vertical)
+				FTabManager::NewStack()
 				->SetSizeCoefficient(0.2f)
-				->Split
-				(
-					FTabManager::NewStack()
-					->SetSizeCoefficient(0.65f)
-					->AddTab(AssetDetailsTabId, ETabState::OpenedTab)
-				)
-				->Split
-				(
-					FTabManager::NewStack()
-					->SetSizeCoefficient(0.35f)
-					->AddTab(StateTreeStatisticsTabId, ETabState::OpenedTab)
-				)
+				->AddTab(AssetDetailsTabId, ETabState::OpenedTab)
+				->AddTab(StateTreeStatisticsTabId, ETabState::OpenedTab)
+				->AddTab(StateTreeOutlinerTabId, ETabState::OpenedTab)
+				->SetForegroundTab(AssetDetailsTabId)
 			)
 			->Split
 			(
@@ -324,6 +325,17 @@ TSharedRef<SDockTab> FStateTreeEditor::SpawnTab_StateTreeView(const FSpawnTabArg
 		];
 }
 
+TSharedRef<SDockTab> FStateTreeEditor::SpawnTab_StateTreeOutliner(const FSpawnTabArgs& Args)
+{
+	check(Args.GetTabId() == StateTreeOutlinerTabId);
+
+	return SNew(SDockTab)
+		.Label(NSLOCTEXT("StateTreeEditor", "StateTreeOutlinerTab", "Outliner"))
+		.TabColorScale(GetTabColorScale())
+		[
+			SAssignNew(StateTreeOutliner, SStateTreeOutliner, StateTreeViewModel.ToSharedRef(), TreeViewCommandList)
+		];
+}
 
 TSharedRef<SDockTab> FStateTreeEditor::SpawnTab_SelectionDetails(const FSpawnTabArgs& Args)
 {
@@ -474,6 +486,80 @@ void FStateTreeEditor::HandleModelSelectionChanged(const TArray<TWeakObjectPtr<U
 	}
 }
 
+void FStateTreeEditor::HandleModelBringNodeToFocus(const UStateTreeState* State, const FGuid NodeID)
+{
+	if (SelectionDetailsView && State)
+	{
+		FPropertyPath HighlightPath;
+
+		if (!HighlightPath.IsValid())
+		{
+			FArrayProperty* TasksProperty = CastFieldChecked<FArrayProperty>(UStateTreeState::StaticClass()->FindPropertyByName(GET_MEMBER_NAME_CHECKED(UStateTreeState, Tasks)));
+			const int32 TaskIndex = State->Tasks.IndexOfByPredicate([&NodeID](const FStateTreeEditorNode& Node)
+			{
+				return Node.ID == NodeID;
+			});
+			if (TaskIndex != INDEX_NONE)
+			{
+				HighlightPath.AddProperty(FPropertyInfo(TasksProperty));
+				HighlightPath.AddProperty(FPropertyInfo(TasksProperty->Inner, TaskIndex));
+			}
+		}
+
+		if (!HighlightPath.IsValid())
+		{
+			FProperty* SingleTaskProperty = CastFieldChecked<FProperty>(UStateTreeState::StaticClass()->FindPropertyByName(GET_MEMBER_NAME_CHECKED(UStateTreeState, SingleTask)));
+			if (State->SingleTask.ID == NodeID)
+			{
+				HighlightPath.AddProperty(FPropertyInfo(SingleTaskProperty));
+			}
+		}
+
+		if (!HighlightPath.IsValid())
+		{
+			FArrayProperty* TransitionsProperty = CastFieldChecked<FArrayProperty>(UStateTreeState::StaticClass()->FindPropertyByName(GET_MEMBER_NAME_CHECKED(UStateTreeState, Transitions)));
+			const int32 TransitionIndex = State->Transitions.IndexOfByPredicate([&NodeID](const FStateTreeTransition& Transition)
+			{
+				return Transition.ID == NodeID;
+			});
+			if (TransitionIndex != INDEX_NONE)
+			{
+				HighlightPath.AddProperty(FPropertyInfo(TransitionsProperty));
+				HighlightPath.AddProperty(FPropertyInfo(TransitionsProperty->Inner, TransitionIndex));
+			}
+		}
+
+		if (!HighlightPath.IsValid())
+		{
+			FArrayProperty* EnterConditionsProperty = CastFieldChecked<FArrayProperty>(UStateTreeState::StaticClass()->FindPropertyByName(GET_MEMBER_NAME_CHECKED(UStateTreeState, EnterConditions)));
+			const int32 EnterConditionIndex = State->EnterConditions.IndexOfByPredicate([&NodeID](const FStateTreeEditorNode& Node)
+			{
+				return Node.ID == NodeID;
+			});
+			if (EnterConditionIndex != INDEX_NONE)
+			{
+				HighlightPath.AddProperty(FPropertyInfo(EnterConditionsProperty));
+				HighlightPath.AddProperty(FPropertyInfo(EnterConditionsProperty->Inner, EnterConditionIndex));
+			}
+		}
+
+		if (HighlightPath.IsValid())
+		{
+			SelectionDetailsView->ScrollPropertyIntoView(HighlightPath, /*bExpandProperty*/true);
+			SelectionDetailsView->HighlightProperty(HighlightPath);
+			
+			GEditor->GetTimerManager()->SetTimer(
+				HighlighTimerHandle,
+				FTimerDelegate::CreateLambda([SelectionDetailsView = SelectionDetailsView]()
+				{
+					SelectionDetailsView->HighlightProperty({});
+				}),
+				1.0f,
+				/*Loop*/false);
+		}
+	}	
+}
+
 void FStateTreeEditor::SaveAsset_Execute()
 {
 	// Remember the treeview expansion state
diff --git a/Engine/Plugins/Runtime/StateTree/Source/StateTreeEditorModule/Private/StateTreeEditor.h b/Engine/Plugins/Runtime/StateTree/Source/StateTreeEditorModule/Private/StateTreeEditor.h
--- a/Engine/Plugins/Runtime/StateTree/Source/StateTreeEditorModule/Private/StateTreeEditor.h
+++ b/Engine/Plugins/Runtime/StateTree/Source/StateTreeEditorModule/Private/StateTreeEditor.h
@@ -68,9 +68,11 @@ private:
 
 	void HandleMessageTokenClicked(const TSharedRef<IMessageToken>& InMessageToken);
 	void HandleModelSelectionChanged(const TArray<TWeakObjectPtr<UStateTreeState>>& SelectedStates);
+	void HandleModelBringNodeToFocus(const UStateTreeState* State, const FGuid NodeID);
 
 	/** Spawns the tab with the update graph inside */
 	TSharedRef<SDockTab> SpawnTab_StateTreeView(const FSpawnTabArgs& Args);
+	TSharedRef<SDockTab> SpawnTab_StateTreeOutliner(const FSpawnTabArgs& Args);
 	TSharedRef<SDockTab> SpawnTab_SelectionDetails(const FSpawnTabArgs& Args);
 	TSharedRef<SDockTab> SpawnTab_AssetDetails(const FSpawnTabArgs& Args);
 	TSharedRef<SDockTab> SpawnTab_StateTreeStatistics(const FSpawnTabArgs& Args);
@@ -122,12 +124,17 @@ private:
 	/** Tree View */
 	TSharedPtr<class SStateTreeView> StateTreeView;
 
+	/** Tree Outliner */
+	TSharedPtr<class SStateTreeOutliner> StateTreeOutliner;
+
 	/** Compiler Results log */
 	TSharedPtr<class SWidget> CompilerResults;
 	TSharedPtr<class IMessageLogListing> CompilerResultsListing;
 	
 	TSharedPtr<FStateTreeViewModel> StateTreeViewModel;
 
+	FTimerHandle HighlighTimerHandle;
+	
 #if WITH_STATETREE_TRACE_DEBUGGER
 	TSharedRef<SDockTab> SpawnTab_Debugger(const FSpawnTabArgs& Args);
 
@@ -136,6 +143,7 @@ private:
 #endif // WITH_STATETREE_TRACE_DEBUGGER
 
 	static const FName StateTreeViewTabId;
+	static const FName StateTreeOutlinerTabId;
 	static const FName SelectionDetailsTabId;
 	static const FName AssetDetailsTabId;
 	static const FName StateTreeStatisticsTabId;
diff --git a/Engine/Plugins/Runtime/StateTree/Source/StateTreeEditorModule/Private/StateTreeEditorData.cpp b/Engine/Plugins/Runtime/StateTree/Source/StateTreeEditorModule/Private/StateTreeEditorData.cpp
--- a/Engine/Plugins/Runtime/StateTree/Source/StateTreeEditorModule/Private/StateTreeEditorData.cpp
+++ b/Engine/Plugins/Runtime/StateTree/Source/StateTreeEditorModule/Private/StateTreeEditorData.cpp
@@ -286,6 +286,11 @@ void UStateTreeEditorData::PostEditChangeChainProperty(FPropertyChangedChainEven
 			UE::StateTree::Delegates::OnGlobalDataChanged.Broadcast(*StateTree);
 		}
 
+		// Notify that the color data has changed
+		if (MemberName == GET_MEMBER_NAME_CHECKED(UStateTreeEditorData, Colors))
+		{
+			UE::StateTree::Delegates::OnVisualThemeChanged.Broadcast(*StateTree);
+		}
 	}
 
 	UE::StateTree::PropertyHelpers::DispatchPostEditToNodes(*this, PropertyChangedEvent, *this);
diff --git a/Engine/Plugins/Runtime/StateTree/Source/StateTreeEditorModule/Private/StateTreeEditorStyle.cpp b/Engine/Plugins/Runtime/StateTree/Source/StateTreeEditorModule/Private/StateTreeEditorStyle.cpp
--- a/Engine/Plugins/Runtime/StateTree/Source/StateTreeEditorModule/Private/StateTreeEditorStyle.cpp
+++ b/Engine/Plugins/Runtime/StateTree/Source/StateTreeEditorModule/Private/StateTreeEditorStyle.cpp
@@ -78,6 +78,8 @@ FStateTreeEditorStyle::FStateTreeEditorStyle()
 			.SetEditableTextBoxStyle(StateTitleEditableText));
 
 		Set("StateTree.State.Border", new FSlateBorderBrush(NAME_None, FMargin(2.0f)));
+
+		Set("StateTree.State", new FSlateRoundedBoxBrush(FLinearColor::White, 2.0f));
 	}
 
 	// Details
@@ -101,17 +103,19 @@ FStateTreeEditorStyle::FStateTreeEditorStyle()
 
 	// Task
 	{
+		const FLinearColor ForegroundCol =  FStyleColors::Foreground.GetSpecifiedColor();
+
 		Set("StateTree.Task.Title", FTextBlockStyle(NormalText)
 			.SetFont(DEFAULT_FONT("Regular", 10))
-			.SetColorAndOpacity(FLinearColor(FColor(230, 230, 230, 192))));
+			.SetColorAndOpacity(ForegroundCol.CopyWithNewOpacity(0.8f)));
 
 		Set("StateTree.Task.Title.Bold", FTextBlockStyle(NormalText)
 			.SetFont(DEFAULT_FONT("Bold", 10))
-			.SetColorAndOpacity(FLinearColor(FColor(230, 230, 230, 192))));
+			.SetColorAndOpacity(ForegroundCol.CopyWithNewOpacity(0.8f)));
 
 		Set("StateTree.Task.Title.Subdued", FTextBlockStyle(NormalText)
 			.SetFont(DEFAULT_FONT("Regular", 10))
-			.SetColorAndOpacity(FLinearColor(FColor(230, 230, 230, 96))));
+			.SetColorAndOpacity(ForegroundCol.CopyWithNewOpacity(0.4f)));
 
 		// Tasks to be show up a bit darker than the state
 		Set("StateTree.Task.Rect", new FSlateColorBrush(FLinearColor(FVector3f(0.67f))));
@@ -152,6 +156,26 @@ FStateTreeEditorStyle::FStateTreeEditorStyle()
 			.SetFont(DEFAULT_FONT("Regular", 10)));
 	}
 
+	// Transition rich text
+	{
+		const FLinearColor ForegroundCol =  FStyleColors::White.GetSpecifiedColor();
+		Set("Transition.Normal", FTextBlockStyle(NormalText)
+			.SetColorAndOpacity(ForegroundCol.CopyWithNewOpacity(0.9f))
+			.SetFont(DEFAULT_FONT("Regular", 11)));
+
+		Set("Transition.Bold", FTextBlockStyle(NormalText)
+			.SetColorAndOpacity(ForegroundCol.CopyWithNewOpacity(0.9f))
+			.SetFont(DEFAULT_FONT("Bold", 11)));
+
+		Set("Transition.Italic", FTextBlockStyle(NormalText)
+			.SetColorAndOpacity(ForegroundCol.CopyWithNewOpacity(0.9f))
+			.SetFont(DEFAULT_FONT("Italic", 11)));
+
+		Set("Transition.Subdued", FTextBlockStyle(NormalText)
+			.SetColorAndOpacity(ForegroundCol.CopyWithNewOpacity(0.5f))
+			.SetFont(DEFAULT_FONT("Regular", 11)));
+	}
+
 	// Debugger
 	{
 		Set("StateTreeDebugger.Element.Normal",
@@ -344,6 +368,7 @@ FStateTreeEditorStyle::FStateTreeEditorStyle()
 		Set("StateTreeEditor.StateSubtree", new IMAGE_BRUSH_SVG("Icons/State_Subtree", CoreStyleConstants::Icon16x16));
 		Set("StateTreeEditor.StateLinked", new IMAGE_BRUSH_SVG("Icons/State_Linked", CoreStyleConstants::Icon16x16));
 
+		Set("StateTreeEditor.Transition.Dash", new IMAGE_BRUSH_SVG("Icons/Transition_Dash", CoreStyleConstants::Icon16x16, FStyleColors::Foreground));
 		Set("StateTreeEditor.Transition.Goto", new IMAGE_BRUSH_SVG("Icons/Transition_Goto", CoreStyleConstants::Icon16x16, FStyleColors::Foreground));
 		Set("StateTreeEditor.Transition.Next", new IMAGE_BRUSH_SVG("Icons/Transition_Next", CoreStyleConstants::Icon16x16, FStyleColors::Foreground));
 		Set("StateTreeEditor.Transition.Parent", new IMAGE_BRUSH_SVG("Icons/Transition_Parent", CoreStyleConstants::Icon16x16, FStyleColors::Foreground));
diff --git a/Engine/Plugins/Runtime/StateTree/Source/StateTreeEditorModule/Private/StateTreePropertyHelpers.h b/Engine/Plugins/Runtime/StateTree/Source/StateTreeEditorModule/Private/StateTreePropertyHelpers.h
--- a/Engine/Plugins/Runtime/StateTree/Source/StateTreeEditorModule/Private/StateTreePropertyHelpers.h
+++ b/Engine/Plugins/Runtime/StateTree/Source/StateTreeEditorModule/Private/StateTreePropertyHelpers.h
@@ -85,6 +85,32 @@ FPropertyAccess::Result GetStructValue(const TSharedPtr<const IPropertyHandle>&
 	return FPropertyAccess::Success;
 }
 
+/**
+ * Returns const pointer to struct container in the property.
+ * @param ValueProperty Handle to property where value is got from.
+ * @return Pointer to the struct, or nullptr if type does not match or multiple values.
+ */
+template<typename T>
+const T* GetStructPtr(const TSharedPtr<const IPropertyHandle>& ValueProperty)
+{
+	if (!ValueProperty)
+	{
+		return nullptr;
+	}
+
+	FStructProperty* StructProperty = CastFieldChecked<FStructProperty>(ValueProperty->GetProperty());
+	check(StructProperty);
+	check(StructProperty->Struct == TBaseStructure<T>::Get());
+
+	TArray<const void*> RawData;
+	ValueProperty->AccessRawData(RawData);
+	if (RawData.Num() == 1)
+	{
+		return static_cast<const T*>(RawData[0]); 
+	}
+
+	return nullptr;
+}
 /**
  * Sets a struct property to specific value, checks type before access. Expects T is struct.
  * @param ValueProperty Handle to property where value is got from.
diff --git a/Engine/Plugins/Runtime/StateTree/Source/StateTreeEditorModule/Private/StateTreeViewModel.cpp b/Engine/Plugins/Runtime/StateTree/Source/StateTreeEditorModule/Private/StateTreeViewModel.cpp
--- a/Engine/Plugins/Runtime/StateTree/Source/StateTreeEditorModule/Private/StateTreeViewModel.cpp
+++ b/Engine/Plugins/Runtime/StateTree/Source/StateTreeEditorModule/Private/StateTreeViewModel.cpp
@@ -275,6 +275,29 @@ const UStateTree* FStateTreeViewModel::GetStateTree() const
 	return nullptr;
 }
 
+const UStateTreeEditorData* FStateTreeViewModel::GetStateTreeEditorData() const
+{
+	return TreeDataWeak.Get();
+}
+
+const UStateTreeState* FStateTreeViewModel::GetStateByID(const FGuid StateID) const
+{
+	if (const UStateTreeEditorData* TreeData = TreeDataWeak.Get())
+	{
+		return const_cast<UStateTreeState*>(TreeData->GetStateByID(StateID));
+	}
+	return nullptr;
+}
+
+UStateTreeState* FStateTreeViewModel::GetMutableStateByID(const FGuid StateID) const
+{
+	if (UStateTreeEditorData* TreeData = TreeDataWeak.Get())
+	{
+		return TreeData->GetMutableStateByID(StateID);
+	}
+	return nullptr;
+}
+
 void FStateTreeViewModel::HandleIdentifierChanged(const UStateTree& StateTree) const
 {
 	if (GetStateTree() == &StateTree)
@@ -654,6 +677,12 @@ bool FStateTreeViewModel::HasSelection() const
 	return SelectedStates.Num() > 0;
 }
 
+void FStateTreeViewModel::BringNodeToFocus(UStateTreeState* State, const FGuid NodeID)
+{
+	SetSelection(State);
+	OnBringNodeToFocus.Broadcast(State, NodeID);
+}
+
 void FStateTreeViewModel::GetPersistentExpandedStates(TSet<TWeakObjectPtr<UStateTreeState>>& OutExpandedStates)
 {
 	OutExpandedStates.Reset();
@@ -1138,10 +1167,7 @@ void FStateTreeViewModel::MoveSelectedStates(UStateTreeState* TargetState, const
 			if (UStateTreeState* State = States[i])
 			{
 				State->Modify();
-				if (State->Parent)
-				{
-					AffectedParents.Add(State->Parent);
-				}
+				AffectedParents.Add(State->Parent);
 			}
 		}
 
diff --git a/Engine/Plugins/Runtime/StateTree/Source/StateTreeEditorModule/Private/StateTreeViewModel.h b/Engine/Plugins/Runtime/StateTree/Source/StateTreeEditorModule/Private/StateTreeViewModel.h
--- a/Engine/Plugins/Runtime/StateTree/Source/StateTreeEditorModule/Private/StateTreeViewModel.h
+++ b/Engine/Plugins/Runtime/StateTree/Source/StateTreeEditorModule/Private/StateTreeViewModel.h
@@ -40,6 +40,7 @@ public:
 	DECLARE_MULTICAST_DELEGATE_OneParam(FOnStatesRemoved, const TSet<UStateTreeState*>& /*AffectedParents*/);
 	DECLARE_MULTICAST_DELEGATE_TwoParams(FOnStatesMoved, const TSet<UStateTreeState*>& /*AffectedParents*/, const TSet<UStateTreeState*>& /*MovedStates*/);
 	DECLARE_MULTICAST_DELEGATE_OneParam(FOnSelectionChanged, const TArray<TWeakObjectPtr<UStateTreeState>>& /*SelectedStates*/);
+	DECLARE_MULTICAST_DELEGATE_TwoParams(FOnBringNodeToFocus, const UStateTreeState* /*State*/, const FGuid /*TaskID*/);
 
 	FStateTreeViewModel();
 	virtual ~FStateTreeViewModel() override;
@@ -60,8 +61,15 @@ public:
 	void GetSelectedStates(TArray<TWeakObjectPtr<UStateTreeState>>& OutSelectedStates) const;
 	bool HasSelection() const;
 
+	void BringNodeToFocus(UStateTreeState* State, const FGuid NodeID);
+	
 	// Returns associated state tree asset.
 	const UStateTree* GetStateTree() const;
+
+	const UStateTreeEditorData* GetStateTreeEditorData() const;
+
+	const UStateTreeState* GetStateByID(const FGuid StateID) const;
+	UStateTreeState* GetMutableStateByID(const FGuid StateID) const;
 	
 	// Returns array of subtrees to edit.
 	TArray<TObjectPtr<UStateTreeState>>* GetSubTrees() const;
@@ -126,6 +134,8 @@ public:
 	// Called each time the selection changes.
 	FOnSelectionChanged& GetOnSelectionChanged() { return OnSelectionChanged; }
 
+	FOnBringNodeToFocus& GetOnBringNodeToFocus() { return OnBringNodeToFocus; }
+	
 protected:
 	void GetExpandedStatesRecursive(UStateTreeState* State, TSet<TWeakObjectPtr<UStateTreeState>>& ExpandedStates);
 	void MoveSelectedStates(UStateTreeState* TargetState, const FStateTreeViewModelInsert RelativeLocation);
@@ -153,6 +163,7 @@ protected:
 	FOnStatesRemoved OnStatesRemoved;
 	FOnStatesMoved OnStatesMoved;
 	FOnSelectionChanged OnSelectionChanged;
+	FOnBringNodeToFocus OnBringNodeToFocus;
 };
 
 /** Helper class to allow to copy bindings into clipboard. */
diff --git a/Engine/Plugins/Runtime/StateTree/Source/StateTreeModule/Private/StateTreeDelegates.cpp b/Engine/Plugins/Runtime/StateTree/Source/StateTreeModule/Private/StateTreeDelegates.cpp
--- a/Engine/Plugins/Runtime/StateTree/Source/StateTreeModule/Private/StateTreeDelegates.cpp
+++ b/Engine/Plugins/Runtime/StateTree/Source/StateTreeModule/Private/StateTreeDelegates.cpp
@@ -10,6 +10,7 @@ FOnIdentifierChanged OnIdentifierChanged;
 FOnSchemaChanged OnSchemaChanged;
 FOnParametersChanged OnParametersChanged;
 FOnGlobalDataChanged OnGlobalDataChanged;
+FOnVisualThemeChanged OnVisualThemeChanged;
 FOnStateParametersChanged OnStateParametersChanged;
 FOnBreakpointsChanged OnBreakpointsChanged;
 FOnPostCompile OnPostCompile;
diff --git a/Engine/Plugins/Runtime/StateTree/Source/StateTreeModule/Public/StateTreeDelegates.h b/Engine/Plugins/Runtime/StateTree/Source/StateTreeModule/Public/StateTreeDelegates.h
--- a/Engine/Plugins/Runtime/StateTree/Source/StateTreeModule/Public/StateTreeDelegates.h
+++ b/Engine/Plugins/Runtime/StateTree/Source/StateTreeModule/Public/StateTreeDelegates.h
@@ -46,6 +46,13 @@ extern STATETREEMODULE_API FOnStateParametersChanged OnStateParametersChanged;
 DECLARE_MULTICAST_DELEGATE_OneParam(FOnGlobalDataChanged, const UStateTree& /*StateTree*/);
 extern STATETREEMODULE_API FOnGlobalDataChanged OnGlobalDataChanged;
 
+/**
+ * Called when the theme colors change.
+ * This should mainly used by the asset editor to maintain consistency in the UI for manipulations on the EditorData.
+ */
+DECLARE_MULTICAST_DELEGATE_OneParam(FOnVisualThemeChanged, const UStateTree& /*StateTree*/);
+extern STATETREEMODULE_API FOnVisualThemeChanged OnVisualThemeChanged;
+
 /**
  * Called when breakpoints of the StateTree EditorData changed.
  * This should mainly used by the asset editor to update the debugger.
