StateTree: Fix compiler warnings.

diff --git a/Engine/Plugins/Runtime/StateTree/Source/StateTreeEditorModule/Private/SStateTreeExpanderArrow.cpp b/Engine/Plugins/Runtime/StateTree/Source/StateTreeEditorModule/Private/SStateTreeExpanderArrow.cpp
--- a/Engine/Plugins/Runtime/StateTree/Source/StateTreeEditorModule/Private/SStateTreeExpanderArrow.cpp
+++ b/Engine/Plugins/Runtime/StateTree/Source/StateTreeEditorModule/Private/SStateTreeExpanderArrow.cpp
@@ -76,7 +76,7 @@ int32 SStateTreeExpanderArrow::OnPaint(const FPaintArgs& Args, const FGeometry&
 		const int32 NumLevels = NeedsWireByLevel.Num();
 		for (int32 Level = 1; Level < NumLevels; Level++)
 		{
-			const float CurrentIndent = Indent * (Level - 1);
+			const float CurrentIndent = Indent * static_cast<float>(Level - 1);
 
 			if (NeedsWireByLevel[Level])
 			{
@@ -96,7 +96,7 @@ int32 SStateTreeExpanderArrow::OnPaint(const FPaintArgs& Args, const FGeometry&
 		// For items that are the last expanded child in a list, we need to draw a special angle connector wire.
 		if (OwnerRow->IsLastChild())
 		{
-			const float CurrentIndent = Indent * (NumLevels - 2);
+			const float CurrentIndent = Indent * static_cast<float>(NumLevels - 2);
 			const FVector2f Offset(CurrentIndent + OffsetX, 0);
 			const FVector2f Size(WireThickness, VerticalWireLoc + HalfWireThickness);
 			FSlateDrawElement::MakeBox(
@@ -112,7 +112,7 @@ int32 SStateTreeExpanderArrow::OnPaint(const FPaintArgs& Args, const FGeometry&
 		// If this item is expanded, we need to draw a 1/2-height the line down to its first child cell.
 		if (OwnerRow->IsItemExpanded() && OwnerRow->DoesItemHaveChildren())
 		{
-			const float CurrentIndent = Indent * (NumLevels - 1);
+			const float CurrentIndent = Indent * static_cast<float>(NumLevels - 1);
 			const FVector2f Offset(CurrentIndent + OffsetX, ImageSize.Y + ImagePadding.Top);
 			const FVector2f Size(WireThickness,  (AllottedGeometry.Size.Y - (ImageSize.Y + ImagePadding.Top)));
 			FSlateDrawElement::MakeBox(
@@ -128,8 +128,8 @@ int32 SStateTreeExpanderArrow::OnPaint(const FPaintArgs& Args, const FGeometry&
 		// Draw horizontal connector from parent wire to child.
 		if (NumLevels > 1)
 		{
-			const float HorizontalWireStart = (NumLevels - 2) * Indent + OffsetX;
-			const float HorizontalWireEnd = (NumLevels - 1) * Indent + ImagePadding.Left - WireThickness + (OwnerRow->DoesItemHaveChildren() ? 0 : ImageSize.X);
+			const float HorizontalWireStart = static_cast<float>(NumLevels - 2) * Indent + OffsetX;
+			const float HorizontalWireEnd = static_cast<float>(NumLevels - 1) * Indent + ImagePadding.Left - WireThickness + (OwnerRow->DoesItemHaveChildren() ? 0.0f : ImageSize.X);
 			const FVector2f Offset(HorizontalWireStart + WireThickness, VerticalWireLoc - WireThickness * 0.5f);
 			const FVector2f Size(HorizontalWireEnd - HorizontalWireStart, WireThickness);
 			FSlateDrawElement::MakeBox(
@@ -209,8 +209,7 @@ const FSlateBrush* SStateTreeExpanderArrow::GetExpanderImage() const
 FMargin SStateTreeExpanderArrow::GetExpanderPadding() const
 {
 	const int32 NestingDepth = FMath::Max(0, OwnerRowPtr.Pin()->GetIndentLevel() - BaseIndentLevel);
-	const float Indent = IndentAmount;
 	FMargin Padding = ImagePadding;
-	Padding.Left += NestingDepth * Indent;
+	Padding.Left += static_cast<float>(NestingDepth) * IndentAmount;
 	return Padding;
 }
