StateTree: Allow conditions to receive state completion callbacks.

diff --git a/Engine/Plugins/Runtime/StateTree/Source/StateTreeEditorModule/Private/StateTreeCompiler.cpp b/Engine/Plugins/Runtime/StateTree/Source/StateTreeEditorModule/Private/StateTreeCompiler.cpp
--- a/Engine/Plugins/Runtime/StateTree/Source/StateTreeEditorModule/Private/StateTreeCompiler.cpp
+++ b/Engine/Plugins/Runtime/StateTree/Source/StateTreeEditorModule/Private/StateTreeCompiler.cpp
@@ -820,6 +820,19 @@ bool FStateTreeCompiler::CreateStateTransitions()
 			return false;
 		}
 
+		// Check if any of the enter conditions require state completion events, and cache that.
+		for (int32 ConditionIndex = (int32)CompactState.EnterConditionsBegin; ConditionIndex < Nodes.Num(); ConditionIndex++)
+		{
+			if (const FStateTreeConditionBase* Cond = Nodes[ConditionIndex].GetPtr<const FStateTreeConditionBase>())
+			{
+				if (Cond->bHasShouldCallStateChangeEvents)
+				{
+					CompactState.bHasStateChangeConditions = true;
+					break;
+				}
+			}
+		}
+		
 		// Linked state
 		if (SourceState->Type == EStateTreeStateType::Linked)
 		{
diff --git a/Engine/Plugins/Runtime/StateTree/Source/StateTreeModule/Private/StateTreeExecutionContext.cpp b/Engine/Plugins/Runtime/StateTree/Source/StateTreeModule/Private/StateTreeExecutionContext.cpp
--- a/Engine/Plugins/Runtime/StateTree/Source/StateTreeModule/Private/StateTreeExecutionContext.cpp
+++ b/Engine/Plugins/Runtime/StateTree/Source/StateTreeModule/Private/StateTreeExecutionContext.cpp
@@ -1641,6 +1641,40 @@ EStateTreeRunStatus FStateTreeExecutionContext::EnterState(FStateTreeTransitionR
 					*UEnum::GetDisplayValueAsText(CurrentTransition.ChangeType).ToString());
 			}
 
+			// Call state change events on conditions if needed.
+			if (bIsEnteringState && State.bHasStateChangeConditions)
+			{
+				for (int32 ConditionIndex = State.EnterConditionsBegin; ConditionIndex < (State.EnterConditionsBegin + State.EnterConditionsNum); ConditionIndex++)
+				{
+					const FStateTreeConditionBase& Cond = CurrentFrame.StateTree->Nodes[ConditionIndex].Get<const FStateTreeConditionBase>();
+					if (Cond.bHasShouldCallStateChangeEvents)
+					{
+						const bool bShouldCallStateChange = CurrentTransition.ChangeType == EStateTreeStateChangeType::Changed
+															|| (CurrentTransition.ChangeType == EStateTreeStateChangeType::Sustained && Cond.bShouldStateChangeOnReselect);
+
+						if (bShouldCallStateChange)
+						{
+							const FStateTreeDataView ConditionInstanceView = GetDataView(CurrentParentFrame, CurrentFrame, Cond.InstanceDataHandle);
+							FNodeInstanceDataScope DataScope(*this, Cond.InstanceDataHandle, ConditionInstanceView);
+
+							if (Cond.BindingsBatch.IsValid())
+							{
+								// Use validated copy, since we test in situations where the sources are not always valid (e.g. enter conditions may try to access inactive parent state).
+								CopyBatchOnActiveInstances(CurrentParentFrame, CurrentFrame, ConditionInstanceView, Cond.BindingsBatch);
+							}
+								
+							Cond.EnterState(*this, Transition);
+
+							// Reset copied properties that might contain object references.
+							if (Cond.BindingsBatch.IsValid())
+							{
+								CurrentFrame.StateTree->PropertyBindings.ResetObjects(Cond.BindingsBatch, ConditionInstanceView);
+							}
+						}
+					}
+				}
+			}
+
 			// Activate tasks on current state.
 			for (int32 TaskIndex = State.TasksBegin; TaskIndex < (State.TasksBegin + State.TasksNum); TaskIndex++)
 			{
@@ -1909,6 +1943,40 @@ void FStateTreeExecutionContext::ExitState(const FStateTreeTransitionResult& Tra
 						}
 					}
 				}
+
+				// Call state change events on conditions if needed.
+				if (State.bHasStateChangeConditions)
+				{
+					for (int32 ConditionIndex = (State.EnterConditionsBegin + State.EnterConditionsNum) - 1; ConditionIndex >= State.EnterConditionsBegin; ConditionIndex--)
+					{
+						const FStateTreeConditionBase& Cond = CurrentFrame.StateTree->Nodes[ConditionIndex].Get<const FStateTreeConditionBase>();
+						if (Cond.bHasShouldCallStateChangeEvents)
+						{
+							const bool bShouldCallStateChange = CurrentTransition.ChangeType == EStateTreeStateChangeType::Changed
+										|| (CurrentTransition.ChangeType == EStateTreeStateChangeType::Sustained && Cond.bShouldStateChangeOnReselect);
+
+							if (bShouldCallStateChange)
+							{
+								const FStateTreeDataView ConditionInstanceView = GetDataView(CurrentParentFrame, CurrentFrame, Cond.InstanceDataHandle);
+								FNodeInstanceDataScope DataScope(*this, Cond.InstanceDataHandle, ConditionInstanceView);
+
+								if (Cond.BindingsBatch.IsValid())
+								{
+									// Use validated copy, since we test in situations where the sources are not always valid (e.g. enter conditions may try to access inactive parent state).
+									CopyBatchOnActiveInstances(CurrentParentFrame, CurrentFrame, ConditionInstanceView, Cond.BindingsBatch);
+								}
+								
+								Cond.ExitState(*this, Transition);
+
+								// Reset copied properties that might contain object references.
+								if (Cond.BindingsBatch.IsValid())
+								{
+									CurrentFrame.StateTree->PropertyBindings.ResetObjects(Cond.BindingsBatch, ConditionInstanceView);
+								}
+							}
+						}
+					}
+				}
 			}
 
 			STATETREE_TRACE_STATE_EVENT(CurrentHandle, EStateTreeTraceEventType::OnExited);
@@ -1987,6 +2055,34 @@ void FStateTreeExecutionContext::StateCompleted()
 						Task.StateCompleted(*this, Exec.LastTickStatus, CurrentFrame.ActiveStates);
 					}
 				}
+				
+				// Call state change events on conditions if needed.
+				if (State.bHasStateChangeConditions)
+				{
+					for (int32 ConditionIndex = (State.EnterConditionsBegin + State.EnterConditionsNum) - 1; ConditionIndex >= State.EnterConditionsBegin; ConditionIndex--)
+					{
+						const FStateTreeConditionBase& Cond = CurrentFrame.StateTree->Nodes[ConditionIndex].Get<const FStateTreeConditionBase>();
+						if (Cond.bHasShouldCallStateChangeEvents)
+						{
+							const FStateTreeDataView ConditionInstanceView = GetDataView(CurrentParentFrame, CurrentFrame, Cond.InstanceDataHandle);
+							FNodeInstanceDataScope DataScope(*this, Cond.InstanceDataHandle, ConditionInstanceView);
+
+							if (Cond.BindingsBatch.IsValid())
+							{
+								// Use validated copy, since we test in situations where the sources are not always valid (e.g. enter conditions may try to access inactive parent state).
+								CopyBatchOnActiveInstances(CurrentParentFrame, CurrentFrame, ConditionInstanceView, Cond.BindingsBatch);
+							}
+							
+							Cond.StateCompleted(*this, Exec.LastTickStatus, CurrentFrame.ActiveStates);
+
+							// Reset copied properties that might contain object references.
+							if (Cond.BindingsBatch.IsValid())
+							{
+								CurrentFrame.StateTree->PropertyBindings.ResetObjects(Cond.BindingsBatch, ConditionInstanceView);
+							}
+						}
+					}
+				}
 			}
 		}
 	}
diff --git a/Engine/Plugins/Runtime/StateTree/Source/StateTreeModule/Public/StateTreeConditionBase.h b/Engine/Plugins/Runtime/StateTree/Source/StateTreeModule/Public/StateTreeConditionBase.h
--- a/Engine/Plugins/Runtime/StateTree/Source/StateTreeModule/Public/StateTreeConditionBase.h
+++ b/Engine/Plugins/Runtime/StateTree/Source/StateTreeModule/Public/StateTreeConditionBase.h
@@ -25,6 +25,36 @@ struct STATETREEMODULE_API FStateTreeConditionBase : public FStateTreeNodeBase
 	/** @return True if the condition passes. */
 	virtual bool TestCondition(FStateTreeExecutionContext& Context) const { return false; }
 
+	/**
+	 * Called when a new state is entered and task is part of active states.
+	 * Note: The condition instance data is shared between all the uses a State Tree asset.
+	 *       You should not modify the instance data in this callback.    
+	 * @param Context Reference to current execution context.
+	 * @param Transition Describes the states involved in the transition
+	 * @return Succeed/Failed will end the state immediately and trigger to select new state, Running will carry on to tick the state.
+	 */
+	virtual void EnterState(FStateTreeExecutionContext& Context, const FStateTreeTransitionResult& Transition) const {}
+
+	/**
+	 * Called when a current state is exited and task is part of active states.
+	 * Note: The condition instance data is shared between all the uses a State Tree asset.
+	 *       You should not modify the instance data in this callback.    
+	 * @param Context Reference to current execution context.
+	 * @param Transition Describes the states involved in the transition
+	 */
+	virtual void ExitState(FStateTreeExecutionContext& Context, const FStateTreeTransitionResult& Transition) const {}
+
+	/**
+	 * Called right after a state has been completed, but before new state has been selected. StateCompleted is called in reverse order to allow to propagate state to other Tasks that
+	 * are executed earlier in the tree. Note that StateCompleted is not called if conditional transition changes the state.
+	 * Note: The condition instance data is shared between all the uses a State Tree asset.
+	 *       You should not modify the instance data in this callback.    
+	 * @param Context Reference to current execution context.
+	 * @param CompletionStatus Describes the running status of the completed state (Succeeded/Failed).
+	 * @param CompletedActiveStates Active states at the time of completion.
+	 */
+	virtual void StateCompleted(FStateTreeExecutionContext& Context, const EStateTreeRunStatus CompletionStatus, const FStateTreeActiveStates& CompletedActiveStates) const {}
+
 	UPROPERTY()
 	EStateTreeExpressionOperand Operand = EStateTreeExpressionOperand::And;
 
@@ -33,6 +63,15 @@ struct STATETREEMODULE_API FStateTreeConditionBase : public FStateTreeNodeBase
 
 	UPROPERTY()
 	EStateTreeConditionEvaluationMode EvaluationMode = EStateTreeConditionEvaluationMode::Evaluated;
+
+	/** If set to true, EnterState, ExitState, and StateCompleted are called on the condition. */
+	uint8 bHasShouldCallStateChangeEvents : 1 = false;
+	
+	/**
+	 * If set to true, the condition will receive EnterState/ExitState even if the state was previously active.
+	 * Default value is true.
+	 */
+	uint8 bShouldStateChangeOnReselect : 1 = true;
 };
 
 /**
diff --git a/Engine/Plugins/Runtime/StateTree/Source/StateTreeModule/Public/StateTreeTypes.h b/Engine/Plugins/Runtime/StateTree/Source/StateTreeModule/Public/StateTreeTypes.h
--- a/Engine/Plugins/Runtime/StateTree/Source/StateTreeModule/Public/StateTreeTypes.h
+++ b/Engine/Plugins/Runtime/StateTree/Source/StateTreeModule/Public/StateTreeTypes.h
@@ -749,6 +749,10 @@ struct STATETREEMODULE_API FCompactStateTreeState
 	UPROPERTY()
 	uint8 bHasTransitionTasks : 1 = false;
 
+	/** True if the state contains conditions which require call to enter/completed/exit state. */
+	UPROPERTY()
+	uint8 bHasStateChangeConditions : 1 = false;
+
 	/** Should state's required event and enter conditions be evaluated when transition leads directly to it's child. */
 	UPROPERTY()
 	uint8 bCheckPrerequisitesWhenActivatingChildDirectly : 1 = false;
