StateTree: Fix deprecated NotSet link type for old assets. The original fix should have change the value of saved data but now we have 3 cases: correctly fixed data, fixed data with wrong type, and unfixed data.

diff --git a/Engine/Plugins/Runtime/StateTree/Source/StateTreeEditorModule/Private/StateTreeCompiler.cpp b/Engine/Plugins/Runtime/StateTree/Source/StateTreeEditorModule/Private/StateTreeCompiler.cpp
--- a/Engine/Plugins/Runtime/StateTree/Source/StateTreeEditorModule/Private/StateTreeCompiler.cpp
+++ b/Engine/Plugins/Runtime/StateTree/Source/StateTreeEditorModule/Private/StateTreeCompiler.cpp
@@ -1432,6 +1432,13 @@ bool FStateTreeCompiler::ResolveTransitionState(const UStateTreeState* SourceSta
 		OutTransitionHandle = FStateTreeStateHandle::Invalid;
 		return true;
 	}
+PRAGMA_DISABLE_DEPRECATION_WARNINGS
+	else if (Link.LinkType == EStateTreeTransitionType::NotSet)
+	{
+		OutTransitionHandle = FStateTreeStateHandle::Invalid;
+		return false;
+	}
+PRAGMA_ENABLE_DEPRECATION_WARNINGS
 	
 	return true;
 }
diff --git a/Engine/Plugins/Runtime/StateTree/Source/StateTreeModule/Private/StateTreeTypes.cpp b/Engine/Plugins/Runtime/StateTree/Source/StateTreeModule/Private/StateTreeTypes.cpp
--- a/Engine/Plugins/Runtime/StateTree/Source/StateTreeModule/Private/StateTreeTypes.cpp
+++ b/Engine/Plugins/Runtime/StateTree/Source/StateTreeModule/Private/StateTreeTypes.cpp
@@ -109,6 +109,10 @@ void FStateTreeStateLink::PostSerialize(const FArchive& Ar)
 	{
 PRAGMA_DISABLE_DEPRECATION_WARNINGS
 		LinkType = Type_DEPRECATED;
+		if (LinkType == EStateTreeTransitionType::NotSet)
+		{
+			LinkType = EStateTreeTransitionType::None;
+		}
 PRAGMA_ENABLE_DEPRECATION_WARNINGS
 	}
 #endif // WITH_EDITORONLY_DAT
