StateTree: Fix random determinism on transition random delay.

diff --git a/Engine/Plugins/Runtime/StateTree/Source/StateTreeModule/Private/StateTreeExecutionContext.cpp b/Engine/Plugins/Runtime/StateTree/Source/StateTreeModule/Private/StateTreeExecutionContext.cpp
--- a/Engine/Plugins/Runtime/StateTree/Source/StateTreeModule/Private/StateTreeExecutionContext.cpp
+++ b/Engine/Plugins/Runtime/StateTree/Source/StateTreeModule/Private/StateTreeExecutionContext.cpp
@@ -3312,7 +3312,7 @@ bool FStateTreeExecutionContext::TriggerTransitions()
 							if (!bIsDelayedTransitionExisting)
 							{
 								// Initialize new delayed transition.
-								const float DelayDuration = Transition.Delay.GetRandomDuration();
+								const float DelayDuration = Transition.Delay.GetRandomDuration(Exec.RandomStream);
 								if (DelayDuration > 0.0f)
 								{
 									FStateTreeTransitionDelayedState& DelayedState = Exec.DelayedTransitions.AddDefaulted_GetRef();
diff --git a/Engine/Plugins/Runtime/StateTree/Source/StateTreeModule/Public/StateTreeTypes.h b/Engine/Plugins/Runtime/StateTree/Source/StateTreeModule/Public/StateTreeTypes.h
--- a/Engine/Plugins/Runtime/StateTree/Source/StateTreeModule/Public/StateTreeTypes.h
+++ b/Engine/Plugins/Runtime/StateTree/Source/StateTreeModule/Public/StateTreeTypes.h
@@ -501,13 +501,20 @@ struct STATETREEMODULE_API FStateTreeRandomTimeDuration
 	bool IsEmpty() const { return Duration == 0 && RandomVariance == 0; }
 	
 	/** @return Returns random duration around Duration, varied by +-RandomVariation. */
-	float GetRandomDuration() const
+	float GetRandomDuration(const FRandomStream& RandomStream) const
 	{
 		const int32 MinVal = FMath::Max(0, static_cast<int32>(Duration) - static_cast<int32>(RandomVariance));
 		const int32 MaxVal = static_cast<int32>(Duration) + static_cast<int32>(RandomVariance);
-		return static_cast<decltype(Scale)>(FMath::RandRange(MinVal, MaxVal)) / Scale;
+		return static_cast<decltype(Scale)>(RandomStream.RandRange(MinVal, MaxVal)) / Scale;
 	}
-	
+
+	UE_DEPRECATED(5.5, "Use the version with random stream instead.")
+	float GetRandomDuration() const
+	{
+		const FRandomStream RandomStream((int32)FPlatformTime::Cycles());
+		return GetRandomDuration(RandomStream);
+	}
+
 protected:
 
 	static constexpr float Scale = 100.0f;
