StateTree: Fixing crash in StateTree utility selector when score is manually set to -1, resulting in no selected state.

diff --git a/Engine/Plugins/Runtime/StateTree/Source/StateTreeModule/Private/StateTreeExecutionContext.cpp b/Engine/Plugins/Runtime/StateTree/Source/StateTreeModule/Private/StateTreeExecutionContext.cpp
--- a/Engine/Plugins/Runtime/StateTree/Source/StateTreeModule/Private/StateTreeExecutionContext.cpp
+++ b/Engine/Plugins/Runtime/StateTree/Source/StateTreeModule/Private/StateTreeExecutionContext.cpp
@@ -4229,9 +4229,9 @@ bool FStateTreeExecutionContext::SelectStateInternal(
 				while (!NextLevelChildStates.IsEmpty())
 				{
 					//Find one with highest score in the remaining candidates
-					float HighestScore = .0f;
+					float HighestScore = -std::numeric_limits<float>::infinity();;
 					uint16 StateIndexWithHighestScore = FStateTreeStateHandle::InvalidIndex;
-					int32 StateArrayIndex = -1;
+					int32 ArrayIndexWithHighestScore = INDEX_NONE;
 					for (int32 Index = 0; Index < NextLevelChildStates.Num(); ++Index)
 					{
 						const uint16 CurrentStateIndex = NextLevelChildStates[Index];
@@ -4241,21 +4241,27 @@ bool FStateTreeExecutionContext::SelectStateInternal(
 						{
 							HighestScore = Score;
 							StateIndexWithHighestScore = CurrentStateIndex;
-							StateArrayIndex = Index;
+							ArrayIndexWithHighestScore = Index;
 						}
 					}
 
-					check(FStateTreeStateHandle::IsValidIndex(StateIndexWithHighestScore));
-					if (SelectStateInternal(CurrentParentFrame, CurrentFrame, CurrentFrameInActiveFrames, { FStateTreeStateHandle(StateIndexWithHighestScore) }, OutSelectionResult))
+					if (FStateTreeStateHandle::IsValidIndex(StateIndexWithHighestScore))
 					{
-						// Selection succeeded
-						bSucceededToSelectState = true;
+						if (SelectStateInternal(CurrentParentFrame, CurrentFrame, CurrentFrameInActiveFrames, { FStateTreeStateHandle(StateIndexWithHighestScore) }, OutSelectionResult))
+						{
+							// Selection succeeded
+							bSucceededToSelectState = true;
+							break;
+						}
+						
+						// Disqualify the state we failed to enter
+						NextLevelChildStates.RemoveAtSwap(ArrayIndexWithHighestScore, EAllowShrinking::No);
+					}
+					else
+					{
+						// No states in array were valid
 						break;
 					}
-					
-					//Disqualify the state we failed to enter
-					constexpr EAllowShrinking AllowShrinking = EAllowShrinking::No;
-					NextLevelChildStates.RemoveAtSwap(StateArrayIndex);
 				}
 
 				if (bSucceededToSelectState)
