Create basic state tree property function for int, bool, float, vector, actors
Create helper function to help create state tree property function
#rnx

diff --git a/Engine/Plugins/Runtime/GameplayStateTree/Source/GameplayStateTreeModule/Private/PropertyFunctions/StateTreeActorPropertyFunctions.cpp b/Engine/Plugins/Runtime/GameplayStateTree/Source/GameplayStateTreeModule/Private/PropertyFunctions/StateTreeActorPropertyFunctions.cpp
new file mode 100644
--- /dev/null
+++ b/Engine/Plugins/Runtime/GameplayStateTree/Source/GameplayStateTreeModule/Private/PropertyFunctions/StateTreeActorPropertyFunctions.cpp
@@ -0,0 +1,24 @@
+// Copyright Epic Games, Inc. All Rights Reserved.
+
+#include "PropertyFunctions/StateTreeActorPropertyFunctions.h"
+
+#include "GameFramework/Actor.h"
+#include "StateTreeExecutionContext.h"
+#include "StateTreeNodeDescriptionHelpers.h"
+
+#define LOCTEXT_NAMESPACE "StateTree"
+
+void FStateTreeGetActorLocationPropertyFunction::Execute(FStateTreeExecutionContext& Context) const
+{
+	FInstanceDataType& InstanceData = Context.GetInstanceData(*this);
+	InstanceData.Output = InstanceData.Input ? InstanceData.Input->GetActorLocation() : FVector::ZeroVector;
+}
+
+#if WITH_EDITOR
+FText FStateTreeGetActorLocationPropertyFunction::GetDescription(const FGuid& ID, FStateTreeDataView InstanceDataView, const IStateTreeBindingLookup& BindingLookup, EStateTreeNodeFormatting Formatting) const
+{
+	return UE::StateTree::DescHelpers::GetDescriptionForSingleParameterFunc<FInstanceDataType>(LOCTEXT("StateTreeActorLocation", "GetActorLocation"), ID, InstanceDataView, BindingLookup, Formatting);
+}
+#endif // WITH_EDITOR
+
+#undef LOCTEXT_NAMESPACE
diff --git a/Engine/Plugins/Runtime/GameplayStateTree/Source/GameplayStateTreeModule/Private/PropertyFunctions/StateTreeActorPropertyFunctions.h b/Engine/Plugins/Runtime/GameplayStateTree/Source/GameplayStateTreeModule/Private/PropertyFunctions/StateTreeActorPropertyFunctions.h
new file mode 100644
--- /dev/null
+++ b/Engine/Plugins/Runtime/GameplayStateTree/Source/GameplayStateTreeModule/Private/PropertyFunctions/StateTreeActorPropertyFunctions.h
@@ -0,0 +1,35 @@
+// Copyright Epic Games, Inc. All Rights Reserved.
+
+#pragma once
+
+#include "StateTreePropertyFunctionBase.h"
+
+#include "StateTreeActorPropertyFunctions.generated.h"
+
+USTRUCT()
+struct FStateTreeGetActorLocationPropertyFunctionInstanceData
+{
+	GENERATED_BODY()
+
+	UPROPERTY(EditAnywhere, Category = Parameter)
+	TObjectPtr<AActor> Input = nullptr;
+
+	UPROPERTY(EditAnywhere, Category = Output)
+	FVector Output = FVector::ZeroVector;
+};
+
+USTRUCT(DisplayName = "Get Actor Location")
+struct FStateTreeGetActorLocationPropertyFunction : public FStateTreePropertyFunctionCommonBase
+{
+	GENERATED_BODY()
+
+	using FInstanceDataType = FStateTreeGetActorLocationPropertyFunctionInstanceData;
+
+	virtual const UStruct* GetInstanceDataType() const override { return FInstanceDataType::StaticStruct(); }
+
+	virtual void Execute(FStateTreeExecutionContext& Context) const override;
+
+#if WITH_EDITOR
+	virtual FText GetDescription(const FGuid& ID, FStateTreeDataView InstanceDataView, const IStateTreeBindingLookup& BindingLookup, EStateTreeNodeFormatting Formatting) const;
+#endif
+};
\ No newline at end of file
diff --git a/Engine/Plugins/Runtime/StateTree/Source/StateTreeEditorModule/Private/Customizations/StateTreeBindingExtension.cpp b/Engine/Plugins/Runtime/StateTree/Source/StateTreeEditorModule/Private/Customizations/StateTreeBindingExtension.cpp
--- a/Engine/Plugins/Runtime/StateTree/Source/StateTreeEditorModule/Private/Customizations/StateTreeBindingExtension.cpp
+++ b/Engine/Plugins/Runtime/StateTree/Source/StateTreeEditorModule/Private/Customizations/StateTreeBindingExtension.cpp
@@ -1518,6 +1518,7 @@ void FStateTreeBindingExtension::ExtendWidgetRow(FDetailWidgetRow& InWidgetRow,
 				FBindingContextStruct& ContextStruct = BindingContextStructs.AddDefaulted_GetRef();
 				ContextStruct.DisplayText = FText::FromString(StructDesc.Name.ToString());
 				ContextStruct.Struct = const_cast<UStruct*>(Struct);
+				ContextStruct.Category = StructDesc.Category;
 
 				// Mare sure same section names get exact same FText representation (binding widget uses IsIdentical() to compare the section names).
 				if (const FText* SectionText = SectionNames.Find(StructDesc.StatePath))
diff --git a/Engine/Plugins/Runtime/StateTree/Source/StateTreeEditorModule/Private/StateTreeEditorData.cpp b/Engine/Plugins/Runtime/StateTree/Source/StateTreeEditorModule/Private/StateTreeEditorData.cpp
--- a/Engine/Plugins/Runtime/StateTree/Source/StateTreeEditorModule/Private/StateTreeEditorData.cpp
+++ b/Engine/Plugins/Runtime/StateTree/Source/StateTreeEditorModule/Private/StateTreeEditorData.cpp
@@ -519,6 +519,7 @@ EStateTreeVisitor UStateTreeEditorData::EnumerateBindablePropertyFunctionNodes(T
 					Desc.DataSource = EStateTreeBindableStructSource::PropertyFunction;
 					Desc.Name = FName(NodeStruct->GetDisplayNameText().ToString());
 					Desc.StatePath = UE::StateTree::Editor::PropertyFunctionStateName;
+					Desc.Category = NodeStruct->GetMetaData(TEXT("Category"));
 
 					if (InFunc(NodeStruct, Desc, FStateTreeDataView(InstanceDataStruct, nullptr)) == EStateTreeVisitor::Break)
 					{
diff --git a/Engine/Plugins/Runtime/StateTree/Source/StateTreeModule/Private/PropertyFunctions/StateTreeBooleanAlgebraPropertyFunctions.cpp b/Engine/Plugins/Runtime/StateTree/Source/StateTreeModule/Private/PropertyFunctions/StateTreeBooleanAlgebraPropertyFunctions.cpp
--- a/Engine/Plugins/Runtime/StateTree/Source/StateTreeModule/Private/PropertyFunctions/StateTreeBooleanAlgebraPropertyFunctions.cpp
+++ b/Engine/Plugins/Runtime/StateTree/Source/StateTreeModule/Private/PropertyFunctions/StateTreeBooleanAlgebraPropertyFunctions.cpp
@@ -25,17 +25,10 @@ namespace UE::StateTree::BooleanPropertyFunctions::Internal
 			RightValue = UE::StateTree::DescHelpers::GetBoolText(InstanceData.bRight, Formatting);
 		}
 
-		const FText Format = (Formatting == EStateTreeNodeFormatting::RichText)
-			? LOCTEXT("BoolFuncRich", "({Left} <s>{Operation}</> {Right})")
-			: LOCTEXT("BoolFunc", "({Left} {Operation} {Right})");
-
-		return FText::FormatNamed(Format,
-			TEXT("Left"), LeftValue,
-			TEXT("Operation"), OperationText,
-			TEXT("Right"), RightValue);
+		return UE::StateTree::DescHelpers::GetMathOperationText(OperationText, LeftValue, RightValue, Formatting);
 	}
 #endif // WITH_EDITOR
-}
+} // namespace UE::StateTree::BooleanPropertyFunctions::Internal
 
 void FStateTreeBooleanAndPropertyFunction::Execute(FStateTreeExecutionContext& Context) const
 {
@@ -63,4 +56,42 @@ FText FStateTreeBooleanOrPropertyFunction::GetDescription(const FGuid& ID, FStat
 }
 #endif // WITH_EDITOR
 
-#undef LOCTEXT_NAMESPACE
+void FStateTreeBooleanXOrPropertyFunction::Execute(FStateTreeExecutionContext& Context) const
+{
+	FStateTreeBooleanOperationPropertyFunctionInstanceData& InstanceData = Context.GetInstanceData(*this);
+	InstanceData.bResult = InstanceData.bLeft ^ InstanceData.bRight;
+}
+
+#if WITH_EDITOR
+FText FStateTreeBooleanXOrPropertyFunction::GetDescription(const FGuid& ID, FStateTreeDataView InstanceDataView, const IStateTreeBindingLookup& BindingLookup, EStateTreeNodeFormatting Formatting) const
+{
+	return UE::StateTree::BooleanPropertyFunctions::Internal::GetDescriptionForOperation(LOCTEXT("BoolXOr", "xor"), ID, InstanceDataView, BindingLookup, Formatting);
+}
+#endif // WITH_EDITOR
+
+void FStateTreeBooleanNotPropertyFunction::Execute(FStateTreeExecutionContext& Context) const
+{
+	FStateTreeBooleanNotOperationPropertyFunctionInstanceData& InstanceData = Context.GetInstanceData(*this);
+	InstanceData.bResult = !InstanceData.bInput;
+}
+
+#if WITH_EDITOR
+FText FStateTreeBooleanNotPropertyFunction::GetDescription(const FGuid& ID, FStateTreeDataView InstanceDataView, const IStateTreeBindingLookup& BindingLookup, EStateTreeNodeFormatting Formatting) const
+{
+	const FStateTreeBooleanNotOperationPropertyFunctionInstanceData& InstanceData = InstanceDataView.Get<FStateTreeBooleanNotOperationPropertyFunctionInstanceData>();
+
+	FText InputValue = BindingLookup.GetBindingSourceDisplayName(FStateTreePropertyPath(ID, GET_MEMBER_NAME_CHECKED(FStateTreeBooleanNotOperationPropertyFunctionInstanceData, bInput)), Formatting);
+	if (InputValue.IsEmpty())
+	{
+		InputValue = UE::StateTree::DescHelpers::GetBoolText(InstanceData.bInput, Formatting);
+	}
+
+	const FText Format = (Formatting == EStateTreeNodeFormatting::RichText)
+		? LOCTEXT("BoolNotFuncRich", "(<s>Not</> {Input})")
+		: LOCTEXT("BoolNotFunc", "(Not {Input})");
+
+	return FText::FormatNamed(Format, TEXT("Input"), InputValue);
+}
+#endif // WITH_EDITOR
+
+#undef LOCTEXT_NAMESPACE
\ No newline at end of file
diff --git a/Engine/Plugins/Runtime/StateTree/Source/StateTreeModule/Private/PropertyFunctions/StateTreeBooleanAlgebraPropertyFunctions.h b/Engine/Plugins/Runtime/StateTree/Source/StateTreeModule/Private/PropertyFunctions/StateTreeBooleanAlgebraPropertyFunctions.h
--- a/Engine/Plugins/Runtime/StateTree/Source/StateTreeModule/Private/PropertyFunctions/StateTreeBooleanAlgebraPropertyFunctions.h
+++ b/Engine/Plugins/Runtime/StateTree/Source/StateTreeModule/Private/PropertyFunctions/StateTreeBooleanAlgebraPropertyFunctions.h
@@ -25,7 +25,7 @@ struct STATETREEMODULE_API FStateTreeBooleanOperationPropertyFunctionInstanceDat
 /**
  * Performs 'And' operation on two booleans.
  */
-USTRUCT(DisplayName = "And")
+USTRUCT(meta=(DisplayName = "And", Category="Logic"))
 struct STATETREEMODULE_API FStateTreeBooleanAndPropertyFunction : public FStateTreePropertyFunctionCommonBase
 {
 	GENERATED_BODY()
@@ -44,7 +44,7 @@ struct STATETREEMODULE_API FStateTreeBooleanAndPropertyFunction : public FStateT
 /**
  * Performs 'Or' operation on two booleans.
  */
-USTRUCT(DisplayName = "Or")
+USTRUCT(meta=(DisplayName = "Or", Category="Logic"))
 struct STATETREEMODULE_API FStateTreeBooleanOrPropertyFunction : public FStateTreePropertyFunctionCommonBase
 {
 	GENERATED_BODY()
@@ -55,6 +55,56 @@ struct STATETREEMODULE_API FStateTreeBooleanOrPropertyFunction : public FStateTr
 
 	virtual void Execute(FStateTreeExecutionContext& Context) const override;
 
+#if WITH_EDITOR
+	virtual FText GetDescription(const FGuid& ID, FStateTreeDataView InstanceDataView, const IStateTreeBindingLookup& BindingLookup, EStateTreeNodeFormatting Formatting) const;
+#endif
+};
+
+/**
+ * Performs 'Exclusive Or' operation on two booleans.
+ */
+USTRUCT(meta=(DisplayName = "XOr", Category="Logic"))
+struct STATETREEMODULE_API FStateTreeBooleanXOrPropertyFunction : public FStateTreePropertyFunctionCommonBase
+{
+	GENERATED_BODY()
+
+	using FInstanceDataType = FStateTreeBooleanOperationPropertyFunctionInstanceData;
+
+	virtual const UStruct* GetInstanceDataType() const override { return FInstanceDataType::StaticStruct(); }
+
+	virtual void Execute(FStateTreeExecutionContext& Context) const override;
+
+#if WITH_EDITOR
+	virtual FText GetDescription(const FGuid& ID, FStateTreeDataView InstanceDataView, const IStateTreeBindingLookup& BindingLookup, EStateTreeNodeFormatting Formatting) const;
+#endif
+};
+
+USTRUCT()
+struct STATETREEMODULE_API FStateTreeBooleanNotOperationPropertyFunctionInstanceData
+{
+	GENERATED_BODY()
+
+	UPROPERTY(EditAnywhere, Category = Param)
+	bool bInput = false;
+
+	UPROPERTY(EditAnywhere, Category = Output)
+	bool bResult = false;
+};
+
+/**
+ * Performs 'Not' operation on a boolean.
+ */
+USTRUCT(meta=(DisplayName = "Not", Category="Logic"))
+struct STATETREEMODULE_API FStateTreeBooleanNotPropertyFunction : public FStateTreePropertyFunctionCommonBase
+{
+	GENERATED_BODY()
+
+	using FInstanceDataType = FStateTreeBooleanNotOperationPropertyFunctionInstanceData;
+
+	virtual const UStruct* GetInstanceDataType() const override { return FInstanceDataType::StaticStruct(); }
+
+	virtual void Execute(FStateTreeExecutionContext& Context) const override;
+
 #if WITH_EDITOR
 	virtual FText GetDescription(const FGuid& ID, FStateTreeDataView InstanceDataView, const IStateTreeBindingLookup& BindingLookup, EStateTreeNodeFormatting Formatting) const;
 #endif
diff --git a/Engine/Plugins/Runtime/StateTree/Source/StateTreeModule/Private/PropertyFunctions/StateTreeFloatPropertyFunctions.cpp b/Engine/Plugins/Runtime/StateTree/Source/StateTreeModule/Private/PropertyFunctions/StateTreeFloatPropertyFunctions.cpp
new file mode 100644
--- /dev/null
+++ b/Engine/Plugins/Runtime/StateTree/Source/StateTreeModule/Private/PropertyFunctions/StateTreeFloatPropertyFunctions.cpp
@@ -0,0 +1,85 @@
+// Copyright Epic Games, Inc. All Rights Reserved.
+
+#include "PropertyFunctions/StateTreeFloatPropertyFunctions.h"
+
+#include "StateTreeExecutionContext.h"
+#include "StateTreeNodeDescriptionHelpers.h"
+
+#define LOCTEXT_NAMESPACE "StateTree"
+
+void FStateTreeAddFloatPropertyFunction::Execute(FStateTreeExecutionContext& Context) const
+{
+	FInstanceDataType& InstanceData = Context.GetInstanceData(*this);
+	InstanceData.Result = InstanceData.Left + InstanceData.Right;
+}
+
+void FStateTreeSubtractFloatPropertyFunction::Execute(FStateTreeExecutionContext& Context) const
+{
+	FInstanceDataType& InstanceData = Context.GetInstanceData(*this);
+	InstanceData.Result = InstanceData.Left - InstanceData.Right;
+}
+
+void FStateTreeMultiplyFloatPropertyFunction::Execute(FStateTreeExecutionContext& Context) const
+{
+	FInstanceDataType& InstanceData = Context.GetInstanceData(*this);
+	InstanceData.Result = InstanceData.Left * InstanceData.Right;
+}
+
+void FStateTreeDivideFloatPropertyFunction::Execute(FStateTreeExecutionContext& Context) const
+{
+	FInstanceDataType& InstanceData = Context.GetInstanceData(*this);
+	if (InstanceData.Right != 0)
+	{
+		InstanceData.Result = InstanceData.Left / InstanceData.Right;
+	}
+	else
+	{
+		InstanceData.Result = 0;
+	}
+}
+
+void FStateTreeInvertFloatPropertyFunction::Execute(FStateTreeExecutionContext& Context) const
+{
+	FInstanceDataType& InstanceData = Context.GetInstanceData(*this);
+	InstanceData.Result = -InstanceData.Input;
+}
+
+void FStateTreeAbsoluteFloatPropertyFunction::Execute(FStateTreeExecutionContext& Context) const
+{
+	FInstanceDataType& InstanceData = Context.GetInstanceData(*this);
+	InstanceData.Result = FMath::Abs(InstanceData.Input);
+}
+
+#if WITH_EDITOR
+FText FStateTreeAddFloatPropertyFunction::GetDescription(const FGuid& ID, FStateTreeDataView InstanceDataView, const IStateTreeBindingLookup& BindingLookup, EStateTreeNodeFormatting Formatting) const
+{
+	return UE::StateTree::DescHelpers::GetDescriptionForMathOperation<FInstanceDataType>(LOCTEXT("FloatAdd", "+"), ID, InstanceDataView, BindingLookup, Formatting);
+}
+
+FText FStateTreeSubtractFloatPropertyFunction::GetDescription(const FGuid& ID, FStateTreeDataView InstanceDataView, const IStateTreeBindingLookup& BindingLookup, EStateTreeNodeFormatting Formatting) const
+{
+	return UE::StateTree::DescHelpers::GetDescriptionForMathOperation<FInstanceDataType>(LOCTEXT("FloatSubtract", "-"), ID, InstanceDataView, BindingLookup, Formatting);
+}
+
+FText FStateTreeMultiplyFloatPropertyFunction::GetDescription(const FGuid& ID, FStateTreeDataView InstanceDataView, const IStateTreeBindingLookup& BindingLookup, EStateTreeNodeFormatting Formatting) const
+{
+	return UE::StateTree::DescHelpers::GetDescriptionForMathOperation<FInstanceDataType>(LOCTEXT("FloatMultiply", "*"), ID, InstanceDataView, BindingLookup, Formatting);
+}
+
+FText FStateTreeDivideFloatPropertyFunction::GetDescription(const FGuid& ID, FStateTreeDataView InstanceDataView, const IStateTreeBindingLookup& BindingLookup, EStateTreeNodeFormatting Formatting) const
+{
+	return UE::StateTree::DescHelpers::GetDescriptionForMathOperation<FInstanceDataType>(LOCTEXT("TreeFloatDivide", "/"), ID, InstanceDataView, BindingLookup, Formatting);
+}
+
+FText FStateTreeInvertFloatPropertyFunction::GetDescription(const FGuid& ID, FStateTreeDataView InstanceDataView, const IStateTreeBindingLookup& BindingLookup, EStateTreeNodeFormatting Formatting) const
+{
+	return UE::StateTree::DescHelpers::GetDescriptionForSingleParameterFunc<FInstanceDataType>(LOCTEXT("FloatInvert", "-"), ID, InstanceDataView, BindingLookup, Formatting);
+}
+
+FText FStateTreeAbsoluteFloatPropertyFunction::GetDescription(const FGuid& ID, FStateTreeDataView InstanceDataView, const IStateTreeBindingLookup& BindingLookup, EStateTreeNodeFormatting Formatting) const
+{
+	return UE::StateTree::DescHelpers::GetDescriptionForSingleParameterFunc<FInstanceDataType>(LOCTEXT("FloatAbsolute", "Abs"), ID, InstanceDataView, BindingLookup, Formatting);
+}
+#endif // WITH_EDITOR
+
+#undef LOCTEXT_NAMESPACE
diff --git a/Engine/Plugins/Runtime/StateTree/Source/StateTreeModule/Private/PropertyFunctions/StateTreeFloatPropertyFunctions.h b/Engine/Plugins/Runtime/StateTree/Source/StateTreeModule/Private/PropertyFunctions/StateTreeFloatPropertyFunctions.h
new file mode 100644
--- /dev/null
+++ b/Engine/Plugins/Runtime/StateTree/Source/StateTreeModule/Private/PropertyFunctions/StateTreeFloatPropertyFunctions.h
@@ -0,0 +1,149 @@
+// Copyright Epic Games, Inc. All Rights Reserved.
+
+#pragma once
+
+#include "StateTreePropertyFunctionBase.h"
+#include "StateTreeFloatPropertyFunctions.generated.h"
+
+struct FStateTreeExecutionContext;
+
+USTRUCT()
+struct STATETREEMODULE_API FStateTreeFloatCombinaisonPropertyFunctionInstanceData
+{
+	GENERATED_BODY()
+
+	UPROPERTY(EditAnywhere, Category = Parameter)
+	float Left = 0.f;
+
+	UPROPERTY(EditAnywhere, Category = Parameter)
+	float Right = 0.f;
+
+	UPROPERTY(EditAnywhere, Category = Output)
+	float Result = 0.f;
+};
+
+/**
+ * Add two floats.
+ */
+USTRUCT(meta=(DisplayName = "Add", Category = "Math|Float"))
+struct STATETREEMODULE_API FStateTreeAddFloatPropertyFunction : public FStateTreePropertyFunctionCommonBase
+{
+	GENERATED_BODY()
+
+	using FInstanceDataType = FStateTreeFloatCombinaisonPropertyFunctionInstanceData;
+
+	virtual const UStruct* GetInstanceDataType() const override { return FInstanceDataType::StaticStruct(); }
+
+	virtual void Execute(FStateTreeExecutionContext& Context) const override;
+
+#if WITH_EDITOR
+	virtual FText GetDescription(const FGuid& ID, FStateTreeDataView InstanceDataView, const IStateTreeBindingLookup& BindingLookup, EStateTreeNodeFormatting Formatting) const;
+#endif
+};
+
+/**
+ * Subtract right float from left float.
+ */
+USTRUCT(meta=(DisplayName = "Subtract", Category = "Math|Float"))
+struct STATETREEMODULE_API FStateTreeSubtractFloatPropertyFunction : public FStateTreePropertyFunctionCommonBase
+{
+	GENERATED_BODY()
+
+	using FInstanceDataType = FStateTreeFloatCombinaisonPropertyFunctionInstanceData;
+
+	virtual const UStruct* GetInstanceDataType() const override { return FInstanceDataType::StaticStruct(); }
+
+	virtual void Execute(FStateTreeExecutionContext& Context) const override;
+
+#if WITH_EDITOR
+	virtual FText GetDescription(const FGuid& ID, FStateTreeDataView InstanceDataView, const IStateTreeBindingLookup& BindingLookup, EStateTreeNodeFormatting Formatting) const;
+#endif
+};
+
+/**
+ * Multiply the two given float.
+ */
+USTRUCT(meta=(DisplayName = "Multiply", Category = "Math|Float"))
+struct STATETREEMODULE_API FStateTreeMultiplyFloatPropertyFunction : public FStateTreePropertyFunctionCommonBase
+{
+	GENERATED_BODY()
+
+	using FInstanceDataType = FStateTreeFloatCombinaisonPropertyFunctionInstanceData;
+
+	virtual const UStruct* GetInstanceDataType() const override { return FInstanceDataType::StaticStruct(); }
+
+	virtual void Execute(FStateTreeExecutionContext& Context) const override;
+
+#if WITH_EDITOR
+	virtual FText GetDescription(const FGuid& ID, FStateTreeDataView InstanceDataView, const IStateTreeBindingLookup& BindingLookup, EStateTreeNodeFormatting Formatting) const;
+#endif
+};
+
+/**
+ * Divide left float by right float.
+ */
+USTRUCT(meta=(DisplayName = "Divide", Category = "Math|Float"))
+struct STATETREEMODULE_API FStateTreeDivideFloatPropertyFunction : public FStateTreePropertyFunctionCommonBase
+{
+	GENERATED_BODY()
+
+	using FInstanceDataType = FStateTreeFloatCombinaisonPropertyFunctionInstanceData;
+
+	virtual const UStruct* GetInstanceDataType() const override { return FInstanceDataType::StaticStruct(); }
+
+	virtual void Execute(FStateTreeExecutionContext& Context) const override;
+
+#if WITH_EDITOR
+	virtual FText GetDescription(const FGuid& ID, FStateTreeDataView InstanceDataView, const IStateTreeBindingLookup& BindingLookup, EStateTreeNodeFormatting Formatting) const;
+#endif
+};
+
+USTRUCT()
+struct STATETREEMODULE_API FStateTreeSingleFloatPropertyFunctionInstanceData
+{
+	GENERATED_BODY()
+
+	UPROPERTY(EditAnywhere, Category = Parameter)
+	float Input = 0.f;
+
+	UPROPERTY(EditAnywhere, Category = Output)
+	float Result = 0.f;
+};
+
+/**
+ * Invert the given float.
+ */
+USTRUCT(meta=(DisplayName = "Invert", Category = "Math|Float"))
+struct STATETREEMODULE_API FStateTreeInvertFloatPropertyFunction : public FStateTreePropertyFunctionCommonBase
+{
+	GENERATED_BODY()
+
+	using FInstanceDataType = FStateTreeSingleFloatPropertyFunctionInstanceData;
+
+	virtual const UStruct* GetInstanceDataType() const override { return FInstanceDataType::StaticStruct(); }
+
+	virtual void Execute(FStateTreeExecutionContext& Context) const override;
+
+#if WITH_EDITOR
+	virtual FText GetDescription(const FGuid& ID, FStateTreeDataView InstanceDataView, const IStateTreeBindingLookup& BindingLookup, EStateTreeNodeFormatting Formatting) const;
+#endif
+};
+
+/**
+ * Gives the absolute value of the given float.
+ */
+USTRUCT(meta=(DisplayName = "Absolute", Category = "Math|Float"))
+struct STATETREEMODULE_API FStateTreeAbsoluteFloatPropertyFunction : public FStateTreePropertyFunctionCommonBase
+{
+	GENERATED_BODY()
+
+	using FInstanceDataType = FStateTreeSingleFloatPropertyFunctionInstanceData;
+
+	virtual const UStruct* GetInstanceDataType() const override { return FInstanceDataType::StaticStruct(); }
+
+	virtual void Execute(FStateTreeExecutionContext& Context) const override;
+
+#if WITH_EDITOR
+	virtual FText GetDescription(const FGuid& ID, FStateTreeDataView InstanceDataView, const IStateTreeBindingLookup& BindingLookup, EStateTreeNodeFormatting Formatting) const;
+#endif
+};
diff --git a/Engine/Plugins/Runtime/StateTree/Source/StateTreeModule/Private/PropertyFunctions/StateTreeIntPropertyFunctions.cpp b/Engine/Plugins/Runtime/StateTree/Source/StateTreeModule/Private/PropertyFunctions/StateTreeIntPropertyFunctions.cpp
new file mode 100644
--- /dev/null
+++ b/Engine/Plugins/Runtime/StateTree/Source/StateTreeModule/Private/PropertyFunctions/StateTreeIntPropertyFunctions.cpp
@@ -0,0 +1,85 @@
+// Copyright Epic Games, Inc. All Rights Reserved.
+
+#include "PropertyFunctions/StateTreeIntPropertyFunctions.h"
+
+#include "StateTreeExecutionContext.h"
+#include "StateTreeNodeDescriptionHelpers.h"
+
+#define LOCTEXT_NAMESPACE "StateTree"
+
+void FStateTreeAddIntPropertyFunction::Execute(FStateTreeExecutionContext& Context) const
+{
+	FInstanceDataType& InstanceData = Context.GetInstanceData(*this);
+	InstanceData.Result = InstanceData.Left + InstanceData.Right;
+}
+
+void FStateTreeSubtractIntPropertyFunction::Execute(FStateTreeExecutionContext& Context) const
+{
+	FInstanceDataType& InstanceData = Context.GetInstanceData(*this);
+	InstanceData.Result = InstanceData.Left - InstanceData.Right;
+}
+
+void FStateTreeMultiplyIntPropertyFunction::Execute(FStateTreeExecutionContext& Context) const
+{
+	FInstanceDataType& InstanceData = Context.GetInstanceData(*this);
+	InstanceData.Result = InstanceData.Left * InstanceData.Right;
+}
+
+void FStateTreeDivideIntPropertyFunction::Execute(FStateTreeExecutionContext& Context) const
+{
+	FInstanceDataType& InstanceData = Context.GetInstanceData(*this);
+	if (InstanceData.Right != 0)
+	{
+		InstanceData.Result = InstanceData.Left / InstanceData.Right;
+	}
+	else
+	{
+		InstanceData.Result = 0;
+	}
+}
+
+void FStateTreeInvertIntPropertyFunction::Execute(FStateTreeExecutionContext& Context) const
+{
+	FInstanceDataType& InstanceData = Context.GetInstanceData(*this);
+	InstanceData.Result = -InstanceData.Input;
+}
+
+void FStateTreeAbsoluteIntPropertyFunction::Execute(FStateTreeExecutionContext& Context) const
+{
+	FInstanceDataType& InstanceData = Context.GetInstanceData(*this);
+	InstanceData.Result = FMath::Abs(InstanceData.Input);
+}
+
+#if WITH_EDITOR
+FText FStateTreeAddIntPropertyFunction::GetDescription(const FGuid& ID, FStateTreeDataView InstanceDataView, const IStateTreeBindingLookup& BindingLookup, EStateTreeNodeFormatting Formatting) const
+{
+	return UE::StateTree::DescHelpers::GetDescriptionForMathOperation<FInstanceDataType>(LOCTEXT("IntAdd", "+"), ID, InstanceDataView, BindingLookup, Formatting);
+}
+
+FText FStateTreeSubtractIntPropertyFunction::GetDescription(const FGuid& ID, FStateTreeDataView InstanceDataView, const IStateTreeBindingLookup& BindingLookup, EStateTreeNodeFormatting Formatting) const
+{
+	return UE::StateTree::DescHelpers::GetDescriptionForMathOperation<FInstanceDataType>(LOCTEXT("IntSubtract", "-"), ID, InstanceDataView, BindingLookup, Formatting);
+}
+
+FText FStateTreeMultiplyIntPropertyFunction::GetDescription(const FGuid& ID, FStateTreeDataView InstanceDataView, const IStateTreeBindingLookup& BindingLookup, EStateTreeNodeFormatting Formatting) const
+{
+	return UE::StateTree::DescHelpers::GetDescriptionForMathOperation<FInstanceDataType>(LOCTEXT("IntMultiply", "*"), ID, InstanceDataView, BindingLookup, Formatting);
+}
+
+FText FStateTreeDivideIntPropertyFunction::GetDescription(const FGuid& ID, FStateTreeDataView InstanceDataView, const IStateTreeBindingLookup& BindingLookup, EStateTreeNodeFormatting Formatting) const
+{
+	return UE::StateTree::DescHelpers::GetDescriptionForMathOperation<FInstanceDataType>(LOCTEXT("TreeIntDivide", "/"), ID, InstanceDataView, BindingLookup, Formatting);
+}
+
+FText FStateTreeInvertIntPropertyFunction::GetDescription(const FGuid& ID, FStateTreeDataView InstanceDataView, const IStateTreeBindingLookup& BindingLookup, EStateTreeNodeFormatting Formatting) const
+{
+	return UE::StateTree::DescHelpers::GetDescriptionForSingleParameterFunc<FInstanceDataType>(LOCTEXT("IntInvert", "-"), ID, InstanceDataView, BindingLookup, Formatting);
+}
+
+FText FStateTreeAbsoluteIntPropertyFunction::GetDescription(const FGuid& ID, FStateTreeDataView InstanceDataView, const IStateTreeBindingLookup& BindingLookup, EStateTreeNodeFormatting Formatting) const
+{
+	return UE::StateTree::DescHelpers::GetDescriptionForSingleParameterFunc<FInstanceDataType>(LOCTEXT("IntAbsolute", "Abs"), ID, InstanceDataView, BindingLookup, Formatting);
+}
+#endif // WITH_EDITOR
+
+#undef LOCTEXT_NAMESPACE
diff --git a/Engine/Plugins/Runtime/StateTree/Source/StateTreeModule/Private/PropertyFunctions/StateTreeIntPropertyFunctions.h b/Engine/Plugins/Runtime/StateTree/Source/StateTreeModule/Private/PropertyFunctions/StateTreeIntPropertyFunctions.h
new file mode 100644
--- /dev/null
+++ b/Engine/Plugins/Runtime/StateTree/Source/StateTreeModule/Private/PropertyFunctions/StateTreeIntPropertyFunctions.h
@@ -0,0 +1,149 @@
+// Copyright Epic Games, Inc. All Rights Reserved.
+
+#pragma once
+
+#include "StateTreePropertyFunctionBase.h"
+#include "StateTreeIntPropertyFunctions.generated.h"
+
+struct FStateTreeExecutionContext;
+
+USTRUCT()
+struct STATETREEMODULE_API FStateTreeIntCombinaisonPropertyFunctionInstanceData
+{
+	GENERATED_BODY()
+
+	UPROPERTY(EditAnywhere, Category = Parameter)
+	int32 Left = 0;
+
+	UPROPERTY(EditAnywhere, Category = Parameter)
+	int32 Right = 0;
+
+	UPROPERTY(EditAnywhere, Category = Output)
+	int32 Result = 0;
+};
+
+/**
+ * Add two ints.
+ */
+USTRUCT(meta=(DisplayName = "Add", Category = "Math|Integer"))
+struct STATETREEMODULE_API FStateTreeAddIntPropertyFunction : public FStateTreePropertyFunctionCommonBase
+{
+	GENERATED_BODY()
+
+	using FInstanceDataType = FStateTreeIntCombinaisonPropertyFunctionInstanceData;
+
+	virtual const UStruct* GetInstanceDataType() const override { return FInstanceDataType::StaticStruct(); }
+
+	virtual void Execute(FStateTreeExecutionContext& Context) const override;
+
+#if WITH_EDITOR
+	virtual FText GetDescription(const FGuid& ID, FStateTreeDataView InstanceDataView, const IStateTreeBindingLookup& BindingLookup, EStateTreeNodeFormatting Formatting) const;
+#endif
+};
+
+/**
+ * Subtract right int from left int.
+ */
+USTRUCT(meta=(DisplayName = "Subtract", Category = "Math|Integer"))
+struct STATETREEMODULE_API FStateTreeSubtractIntPropertyFunction : public FStateTreePropertyFunctionCommonBase
+{
+	GENERATED_BODY()
+
+	using FInstanceDataType = FStateTreeIntCombinaisonPropertyFunctionInstanceData;
+
+	virtual const UStruct* GetInstanceDataType() const override { return FInstanceDataType::StaticStruct(); }
+
+	virtual void Execute(FStateTreeExecutionContext& Context) const override;
+
+#if WITH_EDITOR
+	virtual FText GetDescription(const FGuid& ID, FStateTreeDataView InstanceDataView, const IStateTreeBindingLookup& BindingLookup, EStateTreeNodeFormatting Formatting) const;
+#endif
+};
+
+/**
+ * Multiply the two given ints.
+ */
+USTRUCT(meta=(DisplayName = "Multiply", Category = "Math|Integer"))
+struct STATETREEMODULE_API FStateTreeMultiplyIntPropertyFunction : public FStateTreePropertyFunctionCommonBase
+{
+	GENERATED_BODY()
+
+	using FInstanceDataType = FStateTreeIntCombinaisonPropertyFunctionInstanceData;
+
+	virtual const UStruct* GetInstanceDataType() const override { return FInstanceDataType::StaticStruct(); }
+
+	virtual void Execute(FStateTreeExecutionContext& Context) const override;
+
+#if WITH_EDITOR
+	virtual FText GetDescription(const FGuid& ID, FStateTreeDataView InstanceDataView, const IStateTreeBindingLookup& BindingLookup, EStateTreeNodeFormatting Formatting) const;
+#endif
+};
+
+/**
+ * Divide left int by right int.
+ */
+USTRUCT(meta=(DisplayName = "Divide", Category = "Math|Integer"))
+struct STATETREEMODULE_API FStateTreeDivideIntPropertyFunction : public FStateTreePropertyFunctionCommonBase
+{
+	GENERATED_BODY()
+
+	using FInstanceDataType = FStateTreeIntCombinaisonPropertyFunctionInstanceData;
+
+	virtual const UStruct* GetInstanceDataType() const override { return FInstanceDataType::StaticStruct(); }
+
+	virtual void Execute(FStateTreeExecutionContext& Context) const override;
+
+#if WITH_EDITOR
+	virtual FText GetDescription(const FGuid& ID, FStateTreeDataView InstanceDataView, const IStateTreeBindingLookup& BindingLookup, EStateTreeNodeFormatting Formatting) const;
+#endif
+};
+
+USTRUCT()
+struct STATETREEMODULE_API FStateTreeSingleIntPropertyFunctionInstanceData
+{
+	GENERATED_BODY()
+
+	UPROPERTY(EditAnywhere, Category = Parameter)
+	int32 Input = 0;
+
+	UPROPERTY(EditAnywhere, Category = Output)
+	int32 Result = 0;
+};
+
+/**
+ * Invert the given int.
+ */
+USTRUCT(meta=(DisplayName = "Invert", Category = "Math|Integer"))
+struct STATETREEMODULE_API FStateTreeInvertIntPropertyFunction : public FStateTreePropertyFunctionCommonBase
+{
+	GENERATED_BODY()
+
+	using FInstanceDataType = FStateTreeSingleIntPropertyFunctionInstanceData;
+
+	virtual const UStruct* GetInstanceDataType() const override { return FInstanceDataType::StaticStruct(); }
+
+	virtual void Execute(FStateTreeExecutionContext& Context) const override;
+
+#if WITH_EDITOR
+	virtual FText GetDescription(const FGuid& ID, FStateTreeDataView InstanceDataView, const IStateTreeBindingLookup& BindingLookup, EStateTreeNodeFormatting Formatting) const;
+#endif
+};
+
+/**
+ * Gives the absolute value of the given int.
+ */
+USTRUCT(meta=(DisplayName = "Absolute", Category = "Math|Integer"))
+struct STATETREEMODULE_API FStateTreeAbsoluteIntPropertyFunction : public FStateTreePropertyFunctionCommonBase
+{
+	GENERATED_BODY()
+
+	using FInstanceDataType = FStateTreeSingleIntPropertyFunctionInstanceData;
+
+	virtual const UStruct* GetInstanceDataType() const override { return FInstanceDataType::StaticStruct(); }
+
+	virtual void Execute(FStateTreeExecutionContext& Context) const override;
+
+#if WITH_EDITOR
+	virtual FText GetDescription(const FGuid& ID, FStateTreeDataView InstanceDataView, const IStateTreeBindingLookup& BindingLookup, EStateTreeNodeFormatting Formatting) const;
+#endif
+};
diff --git a/Engine/Plugins/Runtime/StateTree/Source/StateTreeModule/Private/StateTreeNodeDescriptionHelpers.cpp b/Engine/Plugins/Runtime/StateTree/Source/StateTreeModule/Private/StateTreeNodeDescriptionHelpers.cpp
--- a/Engine/Plugins/Runtime/StateTree/Source/StateTreeModule/Private/StateTreeNodeDescriptionHelpers.cpp
+++ b/Engine/Plugins/Runtime/StateTree/Source/StateTreeModule/Private/StateTreeNodeDescriptionHelpers.cpp
@@ -164,6 +164,50 @@ FText GetExactMatchText(bool bExactMatch, EStateTreeNodeFormatting Formatting)
 	}
 	return bExactMatch ? LOCTEXT("Exactly", "exactly ") : FText::GetEmpty();
 }
+
+FText GetText(const FVector& Value, EStateTreeNodeFormatting Formatting)
+{
+	return Value.ToCompactText();
+}
+
+FText GetText(float Value, EStateTreeNodeFormatting Formatting)
+{
+	return FText::AsNumber(Value);
+}
+
+FText GetText(int32 Value, EStateTreeNodeFormatting Formatting)
+{
+	return FText::AsNumber(Value);
+}
+
+FText GetText(const UObject* Value, EStateTreeNodeFormatting Formatting)
+{
+	return FText::FromName(GetFNameSafe(Value));
+}
+
+FText GetMathOperationText(const FText& OperationText, const FText& LeftValue, const FText& RightValue, EStateTreeNodeFormatting Formatting)
+{
+	const FText Format = (Formatting == EStateTreeNodeFormatting::RichText)
+		? LOCTEXT("MathFuncRich", "({Left} <s>{Operation}</> {Right})")
+		: LOCTEXT("MathFunc", "({Left} {Operation} {Right})");
+
+	return FText::FormatNamed(Format,
+		TEXT("Left"), LeftValue,
+		TEXT("Operation"), OperationText,
+		TEXT("Right"), RightValue);
+}
+
+FText GetSingleParamFunctionText(const FText& FunctionText, const FText& ParamText, EStateTreeNodeFormatting Formatting)
+{
+	const FText Format = (Formatting == EStateTreeNodeFormatting::RichText)
+		? LOCTEXT("SingleParamFuncRich", "<s>{Function}</>({Input})")
+		: LOCTEXT("SingleParamFunc", "{Function}({Input})");
+
+	return FText::FormatNamed(Format,
+		TEXT("Function"), FunctionText,
+		TEXT("Input"), ParamText);
+}
+
 } // UE::StateTree::Helpers
 
 #undef LOCTEXT_NAMESPACE
diff --git a/Engine/Plugins/Runtime/StateTree/Source/StateTreeModule/Public/StateTreeNodeDescriptionHelpers.h b/Engine/Plugins/Runtime/StateTree/Source/StateTreeModule/Public/StateTreeNodeDescriptionHelpers.h
--- a/Engine/Plugins/Runtime/StateTree/Source/StateTreeModule/Public/StateTreeNodeDescriptionHelpers.h
+++ b/Engine/Plugins/Runtime/StateTree/Source/StateTreeModule/Public/StateTreeNodeDescriptionHelpers.h
@@ -4,6 +4,8 @@
 
 #include "AITypes.h"
 #include "Internationalization/Text.h"
+#include "StateTreeTypes.h"
+#include "StateTreePropertyBindings.h"
 
 enum EStateTreeNodeFormatting : uint8;
 struct FGameplayTagContainer;
@@ -34,5 +36,60 @@ extern STATETREEMODULE_API FText GetGameplayTagQueryAsText(const FGameplayTagQue
 /** @return description for exact match, used for Gameplay Tag matching functions (returns "Exactly" plus space). */
 extern STATETREEMODULE_API FText GetExactMatchText(bool bExactMatch, EStateTreeNodeFormatting Formatting);
 
-#endif // WITH_EDITOR 
-} // UE::StateTree::Helpers
\ No newline at end of file
+/** @return description of a vector value. */
+extern STATETREEMODULE_API FText GetText(const FVector& Value, EStateTreeNodeFormatting Formatting);
+
+/** @return description of a float value. */
+extern STATETREEMODULE_API FText GetText(float Value, EStateTreeNodeFormatting Formatting);
+
+/** @return description of an int value. */
+extern STATETREEMODULE_API FText GetText(int32 Value, EStateTreeNodeFormatting Formatting);
+
+/** @return description of a UObject value. */
+extern STATETREEMODULE_API FText GetText(const UObject* Value, EStateTreeNodeFormatting Formatting);
+
+extern STATETREEMODULE_API FText GetMathOperationText(const FText& OperationText, const FText& LeftText, const FText& RightText, EStateTreeNodeFormatting Formatting);
+
+/** @return description in the form of (Left OperationText Right).
+*	Expect TInstanceDataType to have a member Left and Right whose types have an overloaded UE::StateTree::DescHelpers::GetText function.
+*/
+template <typename TInstanceDataType>
+FText GetDescriptionForMathOperation(FText OperationText, const FGuid& ID, FStateTreeDataView InstanceDataView, const IStateTreeBindingLookup& BindingLookup, EStateTreeNodeFormatting Formatting)
+{
+	const TInstanceDataType& InstanceData = InstanceDataView.Get<TInstanceDataType>();
+
+	FText LeftValue = BindingLookup.GetBindingSourceDisplayName(FStateTreePropertyPath(ID, GET_MEMBER_NAME_CHECKED(TInstanceDataType, Left)), Formatting);
+	if (LeftValue.IsEmpty())
+	{
+		LeftValue = UE::StateTree::DescHelpers::GetText(InstanceData.Left, Formatting);
+	}
+
+	FText RightValue = BindingLookup.GetBindingSourceDisplayName(FStateTreePropertyPath(ID, GET_MEMBER_NAME_CHECKED(TInstanceDataType, Right)), Formatting);
+	if (RightValue.IsEmpty())
+	{
+		RightValue = UE::StateTree::DescHelpers::GetText(InstanceData.Right, Formatting);
+	}
+
+	return GetMathOperationText(OperationText, LeftValue, RightValue, Formatting);
+}
+
+extern STATETREEMODULE_API FText GetSingleParamFunctionText(const FText& FunctionText, const FText& ParamText, EStateTreeNodeFormatting Formatting);
+
+/** @return description in the form of OperationText(Input).
+ *	Expect TInstanceDataType to have a member input whose type has an overloaded UE::StateTree::DescHelpers::GetText function.
+ */
+template <typename TInstanceDataType>
+FText GetDescriptionForSingleParameterFunc(FText OperationText, const FGuid& ID, FStateTreeDataView InstanceDataView, const IStateTreeBindingLookup& BindingLookup, EStateTreeNodeFormatting Formatting)
+{
+	const TInstanceDataType& InstanceData = InstanceDataView.Get<TInstanceDataType>();
+
+	FText InputValue = BindingLookup.GetBindingSourceDisplayName(FStateTreePropertyPath(ID, GET_MEMBER_NAME_CHECKED(TInstanceDataType, Input)), Formatting);
+	if (InputValue.IsEmpty())
+	{
+		InputValue = UE::StateTree::DescHelpers::GetText(InstanceData.Input, Formatting);
+	}
+
+	return GetSingleParamFunctionText(OperationText, InputValue, Formatting);
+}
+#endif // WITH_EDITOR
+} // namespace UE::StateTree::DescHelpers
\ No newline at end of file
diff --git a/Engine/Plugins/Runtime/StateTree/Source/StateTreeModule/Public/StateTreePropertyBindings.h b/Engine/Plugins/Runtime/StateTree/Source/StateTreeModule/Public/StateTreePropertyBindings.h
--- a/Engine/Plugins/Runtime/StateTree/Source/StateTreeModule/Public/StateTreePropertyBindings.h
+++ b/Engine/Plugins/Runtime/StateTree/Source/StateTreeModule/Public/StateTreePropertyBindings.h
@@ -201,6 +201,10 @@ struct STATETREEMODULE_API FStateTreeBindableStructDesc
 	/** In Editor path to State containting the data. */
 	UPROPERTY(Transient)
 	FString StatePath;
+
+	/** Category of the bindable struct. Can be used to display the category in a menu. */
+	UPROPERTY()
+	FString Category;
 #endif
 };
 
