StateTree: Use a thread safe delegate for internal bind to delegates that are out of StateTree controls. This is a quick fix to unblock internal projects from building. A better solution is a link/compile manager that knows the current instances or use the TObjectIterator to find the instances.

diff --git a/Engine/Plugins/Runtime/StateTree/Source/StateTreeModule/Private/StateTree.cpp b/Engine/Plugins/Runtime/StateTree/Source/StateTreeModule/Private/StateTree.cpp
--- a/Engine/Plugins/Runtime/StateTree/Source/StateTreeModule/Private/StateTree.cpp
+++ b/Engine/Plugins/Runtime/StateTree/Source/StateTreeModule/Private/StateTree.cpp
@@ -18,6 +18,7 @@
 
 #if WITH_EDITOR
 #include "Editor.h"
+#include "StateTreeModuleImpl.h"
 #include "StructUtils/UserDefinedStruct.h"
 #include "StructUtilsDelegates.h"
 #endif
@@ -250,27 +251,23 @@ void UStateTree::PostInitProperties()
 {
 	Super::PostInitProperties();
 	
-	OnObjectsReinstancedHandle = FCoreUObjectDelegates::OnObjectsReinstanced.AddUObject(this, &UStateTree::OnObjectsReinstanced);
-	OnUserDefinedStructReinstancedHandle = UE::StructUtils::Delegates::OnUserDefinedStructReinstanced.AddUObject(this, &UStateTree::OnUserDefinedStructReinstanced);
-	FEditorDelegates::PreBeginPIE.AddUObject(this, &UStateTree::OnPreBeginPIE);
+	if (!HasAnyFlags(RF_ClassDefaultObject))
+	{
+		FStateTreeModule& StateTreeModule = FModuleManager::GetModuleChecked<FStateTreeModule>("StateTreeModule");
+		OnObjectsReinstancedHandle = StateTreeModule.OnObjectsReinstanced.AddUObject(this, &UStateTree::OnObjectsReinstanced);
+		OnUserDefinedStructReinstancedHandle = StateTreeModule.OnUserDefinedStructReinstanced.AddUObject(this, &UStateTree::OnUserDefinedStructReinstanced);
+		OnPreBeginPIEHandle = StateTreeModule.OnPreBeginPIE.AddUObject(this, &UStateTree::OnPreBeginPIE);
+	}
 }
 
 void UStateTree::BeginDestroy()
 {
-	if (OnObjectsReinstancedHandle.IsValid())
-	{
-		FCoreUObjectDelegates::OnObjectsReinstanced.Remove(OnObjectsReinstancedHandle);
-		OnObjectsReinstancedHandle.Reset();
-	}
-	if (OnUserDefinedStructReinstancedHandle.IsValid())
+	if (FStateTreeModule* StateTreeModule = FModuleManager::GetModulePtr<FStateTreeModule>("StateTreeModule"))
 	{
-		UE::StructUtils::Delegates::OnUserDefinedStructReinstanced.Remove(OnUserDefinedStructReinstancedHandle);
-		OnUserDefinedStructReinstancedHandle.Reset();
+		StateTreeModule->OnObjectsReinstanced.Remove(OnObjectsReinstancedHandle);
+		StateTreeModule->OnUserDefinedStructReinstanced.Remove(OnUserDefinedStructReinstancedHandle);
+		StateTreeModule->OnPreBeginPIE.Remove(OnPreBeginPIEHandle);
 	}
-	
-#if WITH_EDITOR
-	FEditorDelegates::PreBeginPIE.RemoveAll(this);
-#endif // WITH_EDITOR
 
 	Super::BeginDestroy();
 }
diff --git a/Engine/Plugins/Runtime/StateTree/Source/StateTreeModule/Private/StateTreeModule.cpp b/Engine/Plugins/Runtime/StateTree/Source/StateTreeModule/Private/StateTreeModule.cpp
--- a/Engine/Plugins/Runtime/StateTree/Source/StateTreeModule/Private/StateTreeModule.cpp
+++ b/Engine/Plugins/Runtime/StateTree/Source/StateTreeModule/Private/StateTreeModule.cpp
@@ -1,6 +1,7 @@
 // Copyright Epic Games, Inc. All Rights Reserved.
 
 #include "StateTreeModule.h"
+#include "StateTreeModuleImpl.h"
 
 #include "StateTreeTypes.h"
 
@@ -27,68 +28,43 @@
 #include "StateTreeInstanceData.h"
 #endif // WITH_EDITORONLY_DATA
 
-#define LOCTEXT_NAMESPACE "StateTree"
-
-class FStateTreeModule : public IStateTreeModule
-{
-	// Begin IModuleInterface
-	virtual void StartupModule() override;
-	virtual void ShutdownModule() override;
+#if WITH_EDITOR
+#include "StructUtilsDelegates.h"
+#endif
 
-	virtual bool StartTraces(int32& OutTraceId) override;
-	virtual bool IsTracing() const override;
-	virtual void StopTraces() override;
+#define LOCTEXT_NAMESPACE "StateTree"
 
 #if WITH_STATETREE_TRACE_DEBUGGER
-	/**
-	 * Gets the store client.
-	 */
-	virtual UE::Trace::FStoreClient* GetStoreClient() override
+UE::Trace::FStoreClient* FStateTreeModule::GetStoreClient()
+{
+	if (!StoreClient.IsValid())
 	{
-		if (!StoreClient.IsValid())
-		{
-			StoreClient = TUniquePtr<UE::Trace::FStoreClient>(UE::Trace::FStoreClient::Connect(TEXT("localhost")));
-		}
-		return StoreClient.Get();
+		StoreClient = TUniquePtr<UE::Trace::FStoreClient>(UE::Trace::FStoreClient::Connect(TEXT("localhost")));
 	}
-
-	TSharedPtr<TraceServices::IAnalysisService> TraceAnalysisService;
-	TSharedPtr<TraceServices::IModuleService> TraceModuleService;
-
-	/** The client used to connect to the trace store. */
-	TUniquePtr<UE::Trace::FStoreClient> StoreClient;
-
-	FStateTreeTraceModule StateTreeTraceModule;
+	return StoreClient.Get();
+}
 #endif // WITH_STATETREE_TRACE_DEBUGGER
 
+FStateTreeModule::FStateTreeModule()
 #if WITH_STATETREE_TRACE
-	TArray<const FString> ChannelsToRestore;
-
-	/** Keep track if StartTraces was explicitly called. */
-	bool bIsTracing = false;
-
-	FAutoConsoleCommand StartDebuggerTracesCommand = FAutoConsoleCommand(
+	: StartDebuggerTracesCommand(FAutoConsoleCommand(
 		TEXT("statetree.startdebuggertraces"),
 		TEXT("Turns on StateTree debugger traces if not already active."),
 		FConsoleCommandDelegate::CreateLambda([]
 			{
-				IStateTreeModule& StateTreeModule = FModuleManager::GetModuleChecked<IStateTreeModule>("StateTreeModule");
 				int32 TraceId = 0;
-				StateTreeModule.StartTraces(TraceId);
-			}));
-
-	FAutoConsoleCommand StopDebuggerTracesCommand = FAutoConsoleCommand(
+				IStateTreeModule::Get().StartTraces(TraceId);
+			})))
+	, StopDebuggerTracesCommand(FAutoConsoleCommand(
 		TEXT("statetree.stopdebuggertraces"),
 		TEXT("Turns off StateTree debugger traces if active."),
 		FConsoleCommandDelegate::CreateLambda([]
 			{
-				IStateTreeModule& StateTreeModule = FModuleManager::GetModuleChecked<IStateTreeModule>("StateTreeModule");
-				StateTreeModule.StopTraces();
-			}));
+				IStateTreeModule::Get().StopTraces();
+			})))
 #endif // WITH_STATETREE_TRACE
-};
-
-IMPLEMENT_MODULE(FStateTreeModule, StateTreeModule)
+{
+}
 
 void FStateTreeModule::StartupModule()
 {
@@ -118,10 +94,23 @@ void FStateTreeModule::StartupModule()
 #if WITH_EDITORONLY_DATA
 	UE::StateTree::RegisterInstanceDataForLocalization();
 #endif // WITH_EDITORONLY_DATA
+
+#if WITH_EDITOR
+	// Register thread safe delegates, which allows the StateTree objects to safely register to these delegates e.g. in PostInitProperties() which may be called from another thread.
+	OnObjectsReinstancedHandle = FCoreUObjectDelegates::OnObjectsReinstanced.AddRaw(this, &FStateTreeModule::HandleObjectsReinstanced);
+	OnUserDefinedStructReinstancedHandle = UE::StructUtils::Delegates::OnUserDefinedStructReinstanced.AddRaw(this, &FStateTreeModule::HandleUserDefinedStructReinstanced);
+	OnPreBeginPIEHandle = FEditorDelegates::PreBeginPIE.AddRaw(this, &FStateTreeModule::HandlePreBeginPIE);
+#endif //WITH_EDITOR
 }
 
 void FStateTreeModule::ShutdownModule()
 {
+#if WITH_EDITOR
+	FCoreUObjectDelegates::OnObjectsReinstanced.Remove(OnObjectsReinstancedHandle);
+	UE::StructUtils::Delegates::OnUserDefinedStructReinstanced.Remove(OnUserDefinedStructReinstancedHandle);
+	FEditorDelegates::PreBeginPIE.Remove(OnPreBeginPIEHandle);
+#endif //WITH_EDITOR
+
 #if WITH_STATETREE_TRACE
 	StopTraces();
 
@@ -261,4 +250,23 @@ void FStateTreeModule::StopTraces()
 #endif // WITH_STATETREE_TRACE
 }
 
+#if WITH_EDITOR
+void FStateTreeModule::HandleObjectsReinstanced(const FReplacementObjectMap& ObjectMap)
+{
+	OnObjectsReinstanced.Broadcast(ObjectMap);
+}
+
+void FStateTreeModule::HandlePreBeginPIE(const bool bIsSimulating)
+{
+	OnPreBeginPIE.Broadcast(bIsSimulating);
+}
+
+void FStateTreeModule::HandleUserDefinedStructReinstanced(const UUserDefinedStruct& UserDefinedStruct)
+{
+	OnUserDefinedStructReinstanced.Broadcast(UserDefinedStruct);
+}
+#endif
+
+IMPLEMENT_MODULE(FStateTreeModule, StateTreeModule)
+
 #undef LOCTEXT_NAMESPACE
diff --git a/Engine/Plugins/Runtime/StateTree/Source/StateTreeModule/Private/StateTreeModuleImpl.h b/Engine/Plugins/Runtime/StateTree/Source/StateTreeModule/Private/StateTreeModuleImpl.h
new file mode 100644
--- /dev/null
+++ b/Engine/Plugins/Runtime/StateTree/Source/StateTreeModule/Private/StateTreeModuleImpl.h
@@ -0,0 +1,83 @@
+// Copyright Epic Games, Inc. All Rights Reserved.
+
+#pragma once
+
+#include "StateTreeModule.h"
+#include "Debugger/StateTreeTraceModule.h"
+#include "HAL/IConsoleManager.h"
+#include "UObject/UObjectGlobals.h"
+
+#if WITH_EDITOR
+#include "Editor.h"
+#endif
+
+namespace TraceServices { class IAnalysisService; }
+namespace TraceServices { class IModuleService; }
+namespace UE::Trace { class FStoreClient; }
+class UUserDefinedStruct;
+
+class FStateTreeModule : public IStateTreeModule
+{
+public:
+	FStateTreeModule();
+
+	//~ Begin IModuleInterface
+	virtual void StartupModule() override;
+	virtual void ShutdownModule() override;
+	//~ End IModuleInterface
+
+	virtual bool StartTraces(int32& OutTraceId) override;
+	virtual bool IsTracing() const override;
+	virtual void StopTraces() override;
+
+#if WITH_EDITOR
+	using FReplacementObjectMap = FCoreUObjectDelegates::FReplacementObjectMap;
+	DECLARE_TS_MULTICAST_DELEGATE_OneParam(FOnObjectsReinstanced, const FReplacementObjectMap&);
+	FOnObjectsReinstanced OnObjectsReinstanced;
+
+	DECLARE_TS_MULTICAST_DELEGATE_OneParam(FOnPIEEvent, bool);
+	FOnPIEEvent OnPreBeginPIE;
+
+	DECLARE_TS_MULTICAST_DELEGATE_OneParam(FOnUserDefinedStructReinstanced, const UUserDefinedStruct& /*UserDefinedStruct*/);
+	FOnUserDefinedStructReinstanced OnUserDefinedStructReinstanced;
+#endif
+
+private:
+#if WITH_EDITOR
+	void HandleObjectsReinstanced(const FReplacementObjectMap& ObjectMap);
+	void HandlePreBeginPIE(const bool bIsSimulating);
+	void HandleUserDefinedStructReinstanced(const UUserDefinedStruct& UserDefinedStruct);
+#endif
+
+#if WITH_STATETREE_TRACE_DEBUGGER
+	/**
+	 * Gets the store client.
+	 */
+	virtual UE::Trace::FStoreClient* GetStoreClient() override;
+
+	TSharedPtr<TraceServices::IAnalysisService> TraceAnalysisService;
+	TSharedPtr<TraceServices::IModuleService> TraceModuleService;
+
+	/** The client used to connect to the trace store. */
+	TUniquePtr<UE::Trace::FStoreClient> StoreClient;
+
+	FStateTreeTraceModule StateTreeTraceModule;
+#endif // WITH_STATETREE_TRACE_DEBUGGER
+
+#if WITH_STATETREE_TRACE
+	TArray<const FString> ChannelsToRestore;
+
+	/** Keep track if StartTraces was explicitly called. */
+	bool bIsTracing = false;
+
+	FAutoConsoleCommand StartDebuggerTracesCommand;
+
+	FAutoConsoleCommand StopDebuggerTracesCommand;
+#endif // WITH_STATETREE_TRACE
+
+#if WITH_EDITOR
+	FDelegateHandle OnObjectsReinstancedHandle;
+	FDelegateHandle OnUserDefinedStructReinstancedHandle;
+	FDelegateHandle OnPreBeginPIEHandle;
+#endif 
+};
diff --git a/Engine/Plugins/Runtime/StateTree/Source/StateTreeModule/Public/StateTree.h b/Engine/Plugins/Runtime/StateTree/Source/StateTreeModule/Public/StateTree.h
--- a/Engine/Plugins/Runtime/StateTree/Source/StateTreeModule/Public/StateTree.h
+++ b/Engine/Plugins/Runtime/StateTree/Source/StateTreeModule/Public/StateTree.h
@@ -187,6 +187,7 @@ public:
 
 	FDelegateHandle OnObjectsReinstancedHandle;
 	FDelegateHandle OnUserDefinedStructReinstancedHandle;
+	FDelegateHandle OnPreBeginPIEHandle;
 #endif
 
 	/** Hash of the editor data from last compile. Also used to detect mismatching events from recorded traces. */
