StateTree: Move the module callback to static variables to prevent issues with the initialization. For commandlet, they can load objects before the module has a chance to load. This is temporary, we should create a subsystem that handles link and compilation.

diff --git a/Engine/Plugins/Runtime/StateTree/Source/StateTreeModule/Private/StateTree.cpp b/Engine/Plugins/Runtime/StateTree/Source/StateTreeModule/Private/StateTree.cpp
--- a/Engine/Plugins/Runtime/StateTree/Source/StateTreeModule/Private/StateTree.cpp
+++ b/Engine/Plugins/Runtime/StateTree/Source/StateTreeModule/Private/StateTree.cpp
@@ -253,21 +253,17 @@ void UStateTree::PostInitProperties()
 	
 	if (!HasAnyFlags(RF_ClassDefaultObject))
 	{
-		FStateTreeModule& StateTreeModule = FModuleManager::GetModuleChecked<FStateTreeModule>("StateTreeModule");
-		OnObjectsReinstancedHandle = StateTreeModule.OnObjectsReinstanced.AddUObject(this, &UStateTree::OnObjectsReinstanced);
-		OnUserDefinedStructReinstancedHandle = StateTreeModule.OnUserDefinedStructReinstanced.AddUObject(this, &UStateTree::OnUserDefinedStructReinstanced);
-		OnPreBeginPIEHandle = StateTreeModule.OnPreBeginPIE.AddUObject(this, &UStateTree::OnPreBeginPIE);
+		OnObjectsReinstancedHandle = FStateTreeModule::OnObjectsReinstanced.AddUObject(this, &UStateTree::OnObjectsReinstanced);
+		OnUserDefinedStructReinstancedHandle = FStateTreeModule::OnUserDefinedStructReinstanced.AddUObject(this, &UStateTree::OnUserDefinedStructReinstanced);
+		OnPreBeginPIEHandle = FStateTreeModule::OnPreBeginPIE.AddUObject(this, &UStateTree::OnPreBeginPIE);
 	}
 }
 
 void UStateTree::BeginDestroy()
 {
-	if (FStateTreeModule* StateTreeModule = FModuleManager::GetModulePtr<FStateTreeModule>("StateTreeModule"))
-	{
-		StateTreeModule->OnObjectsReinstanced.Remove(OnObjectsReinstancedHandle);
-		StateTreeModule->OnUserDefinedStructReinstanced.Remove(OnUserDefinedStructReinstancedHandle);
-		StateTreeModule->OnPreBeginPIE.Remove(OnPreBeginPIEHandle);
-	}
+	FStateTreeModule::OnObjectsReinstanced.Remove(OnObjectsReinstancedHandle);
+	FStateTreeModule::OnUserDefinedStructReinstanced.Remove(OnUserDefinedStructReinstancedHandle);
+	FStateTreeModule::OnPreBeginPIE.Remove(OnPreBeginPIEHandle);
 
 	Super::BeginDestroy();
 }
diff --git a/Engine/Plugins/Runtime/StateTree/Source/StateTreeModule/Private/StateTreeModule.cpp b/Engine/Plugins/Runtime/StateTree/Source/StateTreeModule/Private/StateTreeModule.cpp
--- a/Engine/Plugins/Runtime/StateTree/Source/StateTreeModule/Private/StateTreeModule.cpp
+++ b/Engine/Plugins/Runtime/StateTree/Source/StateTreeModule/Private/StateTreeModule.cpp
@@ -34,6 +34,12 @@
 
 #define LOCTEXT_NAMESPACE "StateTree"
 
+#if WITH_EDITOR
+FStateTreeModule::FOnObjectsReinstanced FStateTreeModule::OnObjectsReinstanced;
+FStateTreeModule::FOnPIEEvent FStateTreeModule::OnPreBeginPIE;
+FStateTreeModule::FOnUserDefinedStructReinstanced FStateTreeModule::OnUserDefinedStructReinstanced;
+#endif
+
 #if WITH_STATETREE_TRACE_DEBUGGER
 UE::Trace::FStoreClient* FStateTreeModule::GetStoreClient()
 {
diff --git a/Engine/Plugins/Runtime/StateTree/Source/StateTreeModule/Private/StateTreeModuleImpl.h b/Engine/Plugins/Runtime/StateTree/Source/StateTreeModule/Private/StateTreeModuleImpl.h
--- a/Engine/Plugins/Runtime/StateTree/Source/StateTreeModule/Private/StateTreeModuleImpl.h
+++ b/Engine/Plugins/Runtime/StateTree/Source/StateTreeModule/Private/StateTreeModuleImpl.h
@@ -33,13 +33,13 @@ public:
 #if WITH_EDITOR
 	using FReplacementObjectMap = FCoreUObjectDelegates::FReplacementObjectMap;
 	DECLARE_TS_MULTICAST_DELEGATE_OneParam(FOnObjectsReinstanced, const FReplacementObjectMap&);
-	FOnObjectsReinstanced OnObjectsReinstanced;
+	static FOnObjectsReinstanced OnObjectsReinstanced;
 
 	DECLARE_TS_MULTICAST_DELEGATE_OneParam(FOnPIEEvent, bool);
-	FOnPIEEvent OnPreBeginPIE;
+	static FOnPIEEvent OnPreBeginPIE;
 
 	DECLARE_TS_MULTICAST_DELEGATE_OneParam(FOnUserDefinedStructReinstanced, const UUserDefinedStruct& /*UserDefinedStruct*/);
-	FOnUserDefinedStructReinstanced OnUserDefinedStructReinstanced;
+	static FOnUserDefinedStructReinstanced OnUserDefinedStructReinstanced;
 #endif
 
 private:
