StateTree: Fix LOCTEXT conflict.
#rnx

diff --git a/Engine/Plugins/Runtime/StateTree/Source/StateTreeEditorModule/Private/SStateTreeViewRow.cpp b/Engine/Plugins/Runtime/StateTree/Source/StateTreeEditorModule/Private/SStateTreeViewRow.cpp
--- a/Engine/Plugins/Runtime/StateTree/Source/StateTreeEditorModule/Private/SStateTreeViewRow.cpp
+++ b/Engine/Plugins/Runtime/StateTree/Source/StateTreeEditorModule/Private/SStateTreeViewRow.cpp
@@ -1605,7 +1605,7 @@ TSharedRef<SWidget> SStateTreeViewRow::MakeTransitionWidgets(const EStateTreeTra
 		else
 		{
 			DescItems.Emplace(
-				LOCTEXT("TransitionActionRoot", "<i>Root</>"),
+				LOCTEXT("TransitionActionRootRich", "<i>Root</>"),
 				LOCTEXT("TransitionActionRoot", "Transition to Root State.")
 				);
 		}
diff --git a/Engine/Plugins/Runtime/StateTree/Source/StateTreeEditorModule/Private/StateTreeDescriptionHelpers.cpp b/Engine/Plugins/Runtime/StateTree/Source/StateTreeEditorModule/Private/StateTreeDescriptionHelpers.cpp
--- a/Engine/Plugins/Runtime/StateTree/Source/StateTreeEditorModule/Private/StateTreeDescriptionHelpers.cpp
+++ b/Engine/Plugins/Runtime/StateTree/Source/StateTreeEditorModule/Private/StateTreeDescriptionHelpers.cpp
@@ -1,4 +1,4 @@
-﻿// Copyright Epic Games, Inc. All Rights Reserved.
+// Copyright Epic Games, Inc. All Rights Reserved.
 
 #include "StateTreeDescriptionHelpers.h"
 
@@ -168,25 +168,25 @@ FText GetTransitionDesc(const UStateTreeEditorData* EditorData, const FStateTree
 	if (Transition.Trigger == EStateTreeTransitionTrigger::OnStateCompleted)
 	{
 		TriggerText = Formatting == EStateTreeNodeFormatting::RichText
-			? LOCTEXT("TransitionOnStateCompleted", "<b>On State Completed</>")
+			? LOCTEXT("TransitionOnStateCompletedRich", "<b>On State Completed</>")
 			: LOCTEXT("TransitionOnStateCompleted", "On State Completed");
 	}
 	else if (Transition.Trigger == EStateTreeTransitionTrigger::OnStateSucceeded)
 	{
 		TriggerText = Formatting == EStateTreeNodeFormatting::RichText
-			? LOCTEXT("TransitionOnStateSucceeded", "<b>On State Succeeded</b>")
+			? LOCTEXT("TransitionOnStateSucceededRich", "<b>On State Succeeded</b>")
 			: LOCTEXT("TransitionOnStateSucceeded", "On State Succeeded");
 	}
 	else if (Transition.Trigger == EStateTreeTransitionTrigger::OnStateFailed)
 	{
 		TriggerText = Formatting == EStateTreeNodeFormatting::RichText
-			? LOCTEXT("TransitionOnStateFailed", "<b>On State Failed</>")
+			? LOCTEXT("TransitionOnStateFailedRich", "<b>On State Failed</>")
 			: LOCTEXT("TransitionOnStateFailed", "On State Failed");
 	}
 	else if (Transition.Trigger == EStateTreeTransitionTrigger::OnTick)
 	{
 		TriggerText = Formatting == EStateTreeNodeFormatting::RichText
-			? LOCTEXT("TransitionOnTick", "<b>On Tick</b>")
+			? LOCTEXT("TransitionOnTickRich", "<b>On Tick</b>")
 			: LOCTEXT("TransitionOnTick", "On Tick");
 	}
 	else if (Transition.Trigger == EStateTreeTransitionTrigger::OnEvent)
@@ -198,7 +198,7 @@ FText GetTransitionDesc(const UStateTreeEditorData* EditorData, const FStateTree
 			if (Transition.RequiredEvent.Tag.IsValid())
 			{
 				const FText TagFormat = Formatting == EStateTreeNodeFormatting::RichText
-					? LOCTEXT("TransitionEventTag", "<s>Tag:</> '{0}'")
+					? LOCTEXT("TransitionEventTagRich", "<s>Tag:</> '{0}'")
 					: LOCTEXT("TransitionEventTag", "Tag: '{0}'");
 				PayloadItems.Add(FText::Format(TagFormat, FText::FromName(Transition.RequiredEvent.Tag.GetTagName())));
 			}
@@ -206,7 +206,7 @@ FText GetTransitionDesc(const UStateTreeEditorData* EditorData, const FStateTree
 			if (Transition.RequiredEvent.PayloadStruct)
 			{
 				const FText PayloadFormat = Formatting == EStateTreeNodeFormatting::RichText
-					? LOCTEXT("TransitionEventPayload", "<s>Payload:</> '{0}'")
+					? LOCTEXT("TransitionEventPayloadRich", "<s>Payload:</> '{0}'")
 					: LOCTEXT("TransitionEventPayload", "Payload: '{0}'");
 				PayloadItems.Add(FText::Format(PayloadFormat, Transition.RequiredEvent.PayloadStruct->GetDisplayNameText()));
 			}
@@ -217,14 +217,14 @@ FText GetTransitionDesc(const UStateTreeEditorData* EditorData, const FStateTree
 		}
 
 		const FText TransitionFormat = Formatting == EStateTreeNodeFormatting::RichText
-			? LOCTEXT("TransitionOnEvent", "<b>On Event</> ({0})")
+			? LOCTEXT("TransitionOnEventRich", "<b>On Event</> ({0})")
 			: LOCTEXT("TransitionOnEvent", "On Event ({0})");
 		
 		TriggerText = FText::Format(TransitionFormat, FText::Join(INVTEXT(", "), PayloadItems));
 	}
 
 	FText ActionText = Formatting == EStateTreeNodeFormatting::RichText
-		? LOCTEXT("ActionGotoRich", "<s>go to</>")
+		? LOCTEXT("ActionGotoRichRich", "<s>go to</>")
 		: LOCTEXT("ActionGoto", "go to");
 	
 	if (Transition.State.LinkType == EStateTreeTransitionType::Succeeded
