StateTree: Prevent access to invalid StateTree in FStateTreeExecutionContext.

diff --git a/Engine/Plugins/Runtime/StateTree/Source/StateTreeModule/Private/StateTreeExecutionContext.cpp b/Engine/Plugins/Runtime/StateTree/Source/StateTreeModule/Private/StateTreeExecutionContext.cpp
--- a/Engine/Plugins/Runtime/StateTree/Source/StateTreeModule/Private/StateTreeExecutionContext.cpp
+++ b/Engine/Plugins/Runtime/StateTree/Source/StateTreeModule/Private/StateTreeExecutionContext.cpp
@@ -154,7 +154,7 @@ FStateTreeExecutionContext::FStateTreeExecutionContext(const FStateTreeExecution
 FStateTreeExecutionContext::~FStateTreeExecutionContext()
 {
 	// Mark external data indices as invalid
-	FStateTreeExecutionState& Exec = GetExecState();
+	FStateTreeExecutionState& Exec = InstanceData.GetMutableStorage().GetMutableExecutionState();
 	for (FStateTreeExecutionFrame& Frame : Exec.ActiveFrames)
 	{
 		Frame.ExternalDataBaseIndex = {};
@@ -1046,6 +1046,13 @@ FStateTreeDataView FStateTreeExecutionContext::GetDataView(const FStateTreeExecu
 
 EStateTreeRunStatus FStateTreeExecutionContext::ForceTransition(const FRecordedStateTreeTransitionResult& Transition)
 {
+	if (!IsValid())
+	{
+		STATETREE_LOG(Warning, TEXT("%hs: StateTree context is not initialized properly ('%s' using StateTree '%s')"),
+			__FUNCTION__, *GetNameSafe(&Owner), *GetFullNameSafe(&RootStateTree));
+		return EStateTreeRunStatus::Failed;
+	}
+
 	FStateTreeTransitionResult TransitionResult = FStateTreeTransitionResult(Transition);
 
 	ExitState(TransitionResult);
@@ -4418,6 +4425,13 @@ FString FStateTreeExecutionContext::GetStateStatusString(const FStateTreeExecuti
 
 EStateTreeRunStatus FStateTreeExecutionContext::GetLastTickStatus() const
 {
+	if (!IsValid())
+	{
+		STATETREE_LOG(Warning, TEXT("%hs: StateTree context is not initialized properly ('%s' using StateTree '%s')"),
+			__FUNCTION__, *GetNameSafe(&Owner), *GetFullNameSafe(&RootStateTree));
+		return EStateTreeRunStatus::Failed;
+	}
+
 	const FStateTreeExecutionState& Exec = GetExecState();
 	return Exec.LastTickStatus;
 }
@@ -4439,6 +4453,13 @@ FString FStateTreeExecutionContext::GetInstanceDescription() const
 
 TConstArrayView<FStateTreeExecutionFrame> FStateTreeExecutionContext::GetActiveFrames() const
 {
+	if (!IsValid())
+	{
+		STATETREE_LOG(Warning, TEXT("%hs: StateTree context is not initialized properly ('%s' using StateTree '%s')"),
+			__FUNCTION__, *GetNameSafe(&Owner), *GetFullNameSafe(&RootStateTree));
+		return TConstArrayView<FStateTreeExecutionFrame>();
+	}
+
 	const FStateTreeExecutionState& Exec = GetExecState();
 	return Exec.ActiveFrames;
 }
@@ -4448,77 +4469,85 @@ TConstArrayView<FStateTreeExecutionFrame> FStateTreeExecutionContext::GetActiveF
 
 FString FStateTreeExecutionContext::GetDebugInfoString() const
 {
-	const FStateTreeExecutionState& Exec = GetExecState();
 
 	FString DebugString = FString::Printf(TEXT("StateTree (asset: '%s')\n"), *GetFullNameSafe(&RootStateTree));
 
-	DebugString += TEXT("Status: ");
-	switch (Exec.TreeRunStatus)
+	if (IsValid())
 	{
-	case EStateTreeRunStatus::Failed:
-		DebugString += TEXT("Failed\n");
-		break;
-	case EStateTreeRunStatus::Succeeded:
-		DebugString += TEXT("Succeeded\n");
-		break;
-	case EStateTreeRunStatus::Running:
-		DebugString += TEXT("Running\n");
-		break;
-	default:
-		DebugString += TEXT("--\n");
-	}
+		const FStateTreeExecutionState& Exec = GetExecState();
 
-	// Active States
-	DebugString += TEXT("Current State:\n");
-	for (const FStateTreeExecutionFrame& CurrentFrame : Exec.ActiveFrames)
-	{
-		const UStateTree* CurrentStateTree = CurrentFrame.StateTree;
+		DebugString += TEXT("Status: ");
+		switch (Exec.TreeRunStatus)
+		{
+		case EStateTreeRunStatus::Failed:
+			DebugString += TEXT("Failed\n");
+			break;
+		case EStateTreeRunStatus::Succeeded:
+			DebugString += TEXT("Succeeded\n");
+			break;
+		case EStateTreeRunStatus::Running:
+			DebugString += TEXT("Running\n");
+			break;
+		default:
+			DebugString += TEXT("--\n");
+		}
 
-		if (CurrentFrame.bIsGlobalFrame)
+		// Active States
+		DebugString += TEXT("Current State:\n");
+		for (const FStateTreeExecutionFrame& CurrentFrame : Exec.ActiveFrames)
 		{
-			DebugString += FString::Printf(TEXT("\nEvaluators\n  [ %-30s | %8s | %15s ]\n"),
-				TEXT("Name"), TEXT("Bindings"), TEXT("Data Handle"));
-			for (int32 EvalIndex = CurrentStateTree->EvaluatorsBegin; EvalIndex < (CurrentStateTree->EvaluatorsBegin + CurrentStateTree->EvaluatorsNum); EvalIndex++)
-			{
-				const FStateTreeEvaluatorBase& Eval = CurrentStateTree->Nodes[EvalIndex].Get<const FStateTreeEvaluatorBase>();
-				DebugString += FString::Printf(TEXT("| %-30s | %8d | %15s |\n"),
-					*Eval.Name.ToString(), Eval.BindingsBatch.Get(), *Eval.InstanceDataHandle.Describe());
-			}
-			
-			DebugString += FString::Printf(TEXT("\nGlobal Tasks\n"));
-			for (int32 TaskIndex = CurrentStateTree->GlobalTasksBegin; TaskIndex < (CurrentStateTree->GlobalTasksBegin + CurrentStateTree->GlobalTasksNum); TaskIndex++)
+			const UStateTree* CurrentStateTree = CurrentFrame.StateTree;
+
+			if (CurrentFrame.bIsGlobalFrame)
 			{
-				const FStateTreeTaskBase& Task = CurrentStateTree->Nodes[TaskIndex].Get<const FStateTreeTaskBase>();
-				if (Task.bTaskEnabled)
+				DebugString += FString::Printf(TEXT("\nEvaluators\n  [ %-30s | %8s | %15s ]\n"),
+					TEXT("Name"), TEXT("Bindings"), TEXT("Data Handle"));
+				for (int32 EvalIndex = CurrentStateTree->EvaluatorsBegin; EvalIndex < (CurrentStateTree->EvaluatorsBegin + CurrentStateTree->EvaluatorsNum); EvalIndex++)
 				{
-					Task.AppendDebugInfoString(DebugString, *this);
+					const FStateTreeEvaluatorBase& Eval = CurrentStateTree->Nodes[EvalIndex].Get<const FStateTreeEvaluatorBase>();
+					DebugString += FString::Printf(TEXT("| %-30s | %8d | %15s |\n"),
+						*Eval.Name.ToString(), Eval.BindingsBatch.Get(), *Eval.InstanceDataHandle.Describe());
+				}
+
+				DebugString += FString::Printf(TEXT("\nGlobal Tasks\n"));
+				for (int32 TaskIndex = CurrentStateTree->GlobalTasksBegin; TaskIndex < (CurrentStateTree->GlobalTasksBegin + CurrentStateTree->GlobalTasksNum); TaskIndex++)
+				{
+					const FStateTreeTaskBase& Task = CurrentStateTree->Nodes[TaskIndex].Get<const FStateTreeTaskBase>();
+					if (Task.bTaskEnabled)
+					{
+						Task.AppendDebugInfoString(DebugString, *this);
+					}
 				}
 			}
-		}
-		
-		for (int32 Index = 0; Index < CurrentFrame.ActiveStates.Num(); Index++)
-		{
-			FStateTreeStateHandle Handle = CurrentFrame.ActiveStates[Index];
-			if (Handle.IsValid())
-			{
-				const FCompactStateTreeState& State = RootStateTree.States[Handle.Index];
-				DebugString += FString::Printf(TEXT("[%s]\n"), *State.Name.ToString());
 
-				if (State.TasksNum > 0)
+			for (int32 Index = 0; Index < CurrentFrame.ActiveStates.Num(); Index++)
+			{
+				FStateTreeStateHandle Handle = CurrentFrame.ActiveStates[Index];
+				if (Handle.IsValid())
 				{
-					DebugString += TEXT("\nTasks:\n");
-					for (int32 TaskIndex = State.TasksBegin; TaskIndex < (State.TasksBegin + State.TasksNum); TaskIndex++)
+					const FCompactStateTreeState& State = RootStateTree.States[Handle.Index];
+					DebugString += FString::Printf(TEXT("[%s]\n"), *State.Name.ToString());
+
+					if (State.TasksNum > 0)
 					{
-						const FStateTreeTaskBase& Task = RootStateTree.Nodes[TaskIndex].Get<const FStateTreeTaskBase>();
-						if (Task.bTaskEnabled)
+						DebugString += TEXT("\nTasks:\n");
+						for (int32 TaskIndex = State.TasksBegin; TaskIndex < (State.TasksBegin + State.TasksNum); TaskIndex++)
 						{
-							Task.AppendDebugInfoString(DebugString, *this);
+							const FStateTreeTaskBase& Task = RootStateTree.Nodes[TaskIndex].Get<const FStateTreeTaskBase>();
+							if (Task.bTaskEnabled)
+							{
+								Task.AppendDebugInfoString(DebugString, *this);
+							}
 						}
 					}
 				}
 			}
 		}
 	}
+	else
+	{
+		DebugString += TEXT("StateTree context is not initialized properly.");
+	}
 
 	return DebugString;
 }
@@ -4632,6 +4661,13 @@ void FStateTreeExecutionContext::DebugPrintInternalLayout()
 
 int32 FStateTreeExecutionContext::GetStateChangeCount() const
 {
+	if (!IsValid())
+	{
+		STATETREE_LOG(Warning, TEXT("%hs: StateTree context is not initialized properly ('%s' using StateTree '%s')"),
+			__FUNCTION__, *GetNameSafe(&Owner), *GetFullNameSafe(&RootStateTree));
+		return 0;
+	}
+
 	const FStateTreeExecutionState& Exec = GetExecState();
 	return Exec.StateChangeCount;
 }
@@ -4639,7 +4675,14 @@ int32 FStateTreeExecutionContext::GetStateChangeCount() const
 #endif // WITH_STATETREE_DEBUG
 
 FString FStateTreeExecutionContext::GetActiveStateName() const
-{	
+{
+	if (!IsValid())
+	{
+		STATETREE_LOG(Warning, TEXT("%hs: StateTree context is not initialized properly ('%s' using StateTree '%s')"),
+			__FUNCTION__, *GetNameSafe(&Owner), *GetFullNameSafe(&RootStateTree));
+		return FString();
+	}
+
 	const FStateTreeExecutionState& Exec = GetExecState();
 
 	FString FullStateName;
@@ -4701,7 +4744,14 @@ FString FStateTreeExecutionContext::GetActiveStateName() const
 
 TArray<FName> FStateTreeExecutionContext::GetActiveStateNames() const
 {
-	TArray<FName> Result;	
+	if (!IsValid())
+	{
+		STATETREE_LOG(Warning, TEXT("%hs: StateTree context is not initialized properly ('%s' using StateTree '%s')"),
+			__FUNCTION__, *GetNameSafe(&Owner), *GetFullNameSafe(&RootStateTree));
+		return TArray<FName>();
+	}
+
+	TArray<FName> Result;
 	const FStateTreeExecutionState& Exec = GetExecState();
 
 	// Active States
