[State Tree] Changed weighted random utility selection behavior to not select states with zero score

diff --git a/Engine/Plugins/Runtime/StateTree/Source/StateTreeModule/Private/StateTreeExecutionContext.cpp b/Engine/Plugins/Runtime/StateTree/Source/StateTreeModule/Private/StateTreeExecutionContext.cpp
--- a/Engine/Plugins/Runtime/StateTree/Source/StateTreeModule/Private/StateTreeExecutionContext.cpp
+++ b/Engine/Plugins/Runtime/StateTree/Source/StateTreeModule/Private/StateTreeExecutionContext.cpp
@@ -4311,7 +4311,8 @@ bool FStateTreeExecutionContext::SelectStateInternal(
 
 						if (RandomScore < AccumulatedScore || (Index == (NextLevelChildStates.Num() - 1)))
 						{
-							if (SelectStateInternal(CurrentParentFrame, CurrentFrame, CurrentFrameInActiveFrames, { FStateTreeStateHandle(StateIndex) }, OutSelectionResult))
+							// States with zero possibility won't be selected
+							if (StateScore != 0.f && SelectStateInternal(CurrentParentFrame, CurrentFrame, CurrentFrameInActiveFrames, { FStateTreeStateHandle(StateIndex) }, OutSelectionResult))
 							{
 								// Selection succeeded
 								bSucceededToSelectState = true;
