StateTree editor
- StateTreeEditorMode, which now holds all of the state-tree editing behavior previously part of StateTreeEditor
- StateTreeEditorModeToolkit, responsible for UI/widgets used by StateTreeEditorMode
- StateTreeEditorModeUILayer used to host StateTreeEditorMode inside of the StateTreeEditor
- Exposed StateTreeViewModel
- Introduced IStateTreeEditorHost, interface to-be implemented for supporting StateTreeEditorMode within an Asset editor
- Moved compilation/validation to UStateTreeEditingSubsystem (editor sub system)
- Added layout extensions for StateTreeEditor

diff --git a/Engine/Plugins/Runtime/StateTree/Source/StateTreeEditorModule/Private/Commandlets/StateTreeCompileAllCommandlet.cpp b/Engine/Plugins/Runtime/StateTree/Source/StateTreeEditorModule/Private/Commandlets/StateTreeCompileAllCommandlet.cpp
--- a/Engine/Plugins/Runtime/StateTree/Source/StateTreeEditorModule/Private/Commandlets/StateTreeCompileAllCommandlet.cpp
+++ b/Engine/Plugins/Runtime/StateTree/Source/StateTreeEditorModule/Private/Commandlets/StateTreeCompileAllCommandlet.cpp
@@ -11,6 +11,7 @@
 #include "StateTreeEditor.h"
 #include "StateTreeCompiler.h"
 #include "StateTreeDelegates.h"
+#include "StateTreeEditingSubsystem.h"
 
 #include UE_INLINE_GENERATED_CPP_BY_NAME(StateTreeCompileAllCommandlet)
 
@@ -60,7 +61,7 @@ int32 UStateTreeCompileAllCommandlet::Main(const FString& Params)
 		}
 		else
 		{
-			CompileAndSaveStateTree(*StateTree);
+			CompileAndSaveStateTree(StateTree);
 		}
 		Counter++;
 	}
@@ -68,36 +69,20 @@ int32 UStateTreeCompileAllCommandlet::Main(const FString& Params)
 	return 0;
 }
 
-bool UStateTreeCompileAllCommandlet::CompileAndSaveStateTree(UStateTree& StateTree) const
+bool UStateTreeCompileAllCommandlet::CompileAndSaveStateTree(TNonNullPtr<UStateTree> StateTree) const
 {
-	UPackage* Package = StateTree.GetPackage();
+	UPackage* Package = StateTree->GetPackage();
 	const FString PackageFileName = SourceControlHelpers::PackageFilename(Package);
 
 	// Compile the StateTree asset.
-	UE::StateTree::Editor::ValidateAsset(StateTree);
-	const uint32 EditorDataHash = UE::StateTree::Editor::CalcAssetHash(StateTree);
+	UStateTreeEditingSubsystem::ValidateStateTree(StateTree);
+	const uint32 EditorDataHash = UStateTreeEditingSubsystem::CalculateStateTreeHash(StateTree);
 
 	FStateTreeCompilerLog Log;
-	FStateTreeCompiler Compiler(Log);
+	const bool bSuccess = UStateTreeEditingSubsystem::CompileStateTree(StateTree, Log);
 
-	const bool bSuccess = Compiler.Compile(StateTree);
-
-	if (bSuccess)
-	{
-		// Success
-		StateTree.LastCompiledEditorDataHash = EditorDataHash;
-		UE::StateTree::Delegates::OnPostCompile.Broadcast(StateTree);
-		UE_LOG(LogStateTreeCompile, Log, TEXT("Compile StateTree %s succeeded."), *PackageFileName, EditorDataHash);
-	}
-	else
+	if (!bSuccess)
 	{
-		// Make sure not to leave stale data on failed compile.
-		StateTree.ResetCompiled();
-		StateTree.LastCompiledEditorDataHash = 0;
-
-		UE_LOG(LogStateTreeCompile, Error, TEXT("Failed to compile StateTree %s, errors follow."), *PackageFileName);
-		Log.DumpToLog(LogStateTreeCompile);
-
 		return false;
 	}
 
diff --git a/Engine/Plugins/Runtime/StateTree/Source/StateTreeEditorModule/Private/Commandlets/StateTreeCompileAllCommandlet.h b/Engine/Plugins/Runtime/StateTree/Source/StateTreeEditorModule/Private/Commandlets/StateTreeCompileAllCommandlet.h
--- a/Engine/Plugins/Runtime/StateTree/Source/StateTreeEditorModule/Private/Commandlets/StateTreeCompileAllCommandlet.h
+++ b/Engine/Plugins/Runtime/StateTree/Source/StateTreeEditorModule/Private/Commandlets/StateTreeCompileAllCommandlet.h
@@ -24,7 +24,7 @@ public:
 	//~ End UCommandlet Interface
 
 protected:
-	bool CompileAndSaveStateTree(UStateTree& StateTree) const;
+	bool CompileAndSaveStateTree(TNonNullPtr<UStateTree> StateTree) const;
 
 	ISourceControlProvider* SourceControlProvider = nullptr;
 };
diff --git a/Engine/Plugins/Runtime/StateTree/Source/StateTreeEditorModule/Private/Customizations/StateTreeBindingExtension.cpp b/Engine/Plugins/Runtime/StateTree/Source/StateTreeEditorModule/Private/Customizations/StateTreeBindingExtension.cpp
--- a/Engine/Plugins/Runtime/StateTree/Source/StateTreeEditorModule/Private/Customizations/StateTreeBindingExtension.cpp
+++ b/Engine/Plugins/Runtime/StateTree/Source/StateTreeEditorModule/Private/Customizations/StateTreeBindingExtension.cpp
@@ -1790,8 +1790,11 @@ bool FStateTreeBindingsChildrenCustomization::ShouldCustomizeChildren(TSharedRef
 		FStateTreePropertyPath TargetPath;
 		UE::StateTree::PropertyBinding::MakeStructPropertyPathFromPropertyHandle(InPropertyHandle, TargetPath);
 		IStateTreeEditorPropertyBindingsOwner* BindingOwner = Cast<IStateTreeEditorPropertyBindingsOwner>(UE::StateTree::PropertyBinding::FindEditorBindingsOwner(OuterObjects[0]));
-		FStateTreeEditorPropertyBindings* EditorBindings = BindingOwner->GetPropertyEditorBindings();
-		return EditorBindings->HasPropertyBinding(TargetPath);
+		if (!TargetPath.IsPathEmpty() && BindingOwner)
+		{
+			FStateTreeEditorPropertyBindings* EditorBindings = BindingOwner->GetPropertyEditorBindings();
+			return EditorBindings->HasPropertyBinding(TargetPath);
+		}
 	}
 
 	return false;
diff --git a/Engine/Plugins/Runtime/StateTree/Source/StateTreeEditorModule/Private/Customizations/Widgets/SCompactStateTreeView.cpp b/Engine/Plugins/Runtime/StateTree/Source/StateTreeEditorModule/Private/Customizations/Widgets/SCompactStateTreeView.cpp
--- a/Engine/Plugins/Runtime/StateTree/Source/StateTreeEditorModule/Private/Customizations/Widgets/SCompactStateTreeView.cpp
+++ b/Engine/Plugins/Runtime/StateTree/Source/StateTreeEditorModule/Private/Customizations/Widgets/SCompactStateTreeView.cpp
@@ -5,6 +5,7 @@
 #include "StateTreeDelegates.h"
 #include "StateTreeDragDrop.h"
 #include "StateTreeEditorData.h"
+#include "StateTreeEditorModule.h"
 #include "StateTreeEditorStyle.h"
 #include "StateTreeViewModel.h"
 #include "Widgets/Input/SSearchBox.h"
@@ -193,36 +194,7 @@ void SCompactStateTreeView::CacheState(TSharedPtr<FStateTreeStateItem> ParentNod
 			StateItem->Color = FoundColor->Color;
 		}
 		
-		// Figure out icon.
-		if (State->SelectionBehavior == EStateTreeStateSelectionBehavior::None)
-		{
-			StateItem->Icon = FStateTreeEditorStyle::Get().GetBrush("StateTreeEditor.SelectNone");
-		}
-		else if (State->SelectionBehavior == EStateTreeStateSelectionBehavior::TryEnterState)
-		{
-			StateItem->Icon = FStateTreeEditorStyle::Get().GetBrush("StateTreeEditor.TryEnterState");			
-		}
-		else if (State->SelectionBehavior == EStateTreeStateSelectionBehavior::TrySelectChildrenInOrder)
-		{
-			if (State->Children.IsEmpty()
-				|| State->Type == EStateTreeStateType::Linked
-				|| State->Type == EStateTreeStateType::LinkedAsset)
-			{
-				StateItem->Icon = FStateTreeEditorStyle::Get().GetBrush("StateTreeEditor.TryEnterState");			
-			}
-			else
-			{
-				StateItem->Icon = FStateTreeEditorStyle::Get().GetBrush("StateTreeEditor.TrySelectChildrenInOrder");
-			}
-		}
-		else if (State->SelectionBehavior == EStateTreeStateSelectionBehavior::TrySelectChildrenAtUniformRandom)
-		{
-			StateItem->Icon = FStateTreeEditorStyle::Get().GetBrush("StateTreeEditor.TrySelectChildrenAtRandom");
-		}
-		else if (State->SelectionBehavior == EStateTreeStateSelectionBehavior::TryFollowTransitions)
-		{
-			StateItem->Icon = FStateTreeEditorStyle::Get().GetBrush("StateTreeEditor.TryFollowTransitions");
-		}
+		StateItem->Icon = FStateTreeEditorStyle::GetBrushForSelectionBehaviorType(State->SelectionBehavior, !State->Children.IsEmpty(), State->Type);
 
 		// Linked states
 		if (State->Type == EStateTreeStateType::Linked)
diff --git a/Engine/Plugins/Runtime/StateTree/Source/StateTreeEditorModule/Private/SStateTreeViewRow.cpp b/Engine/Plugins/Runtime/StateTree/Source/StateTreeEditorModule/Private/SStateTreeViewRow.cpp
--- a/Engine/Plugins/Runtime/StateTree/Source/StateTreeEditorModule/Private/SStateTreeViewRow.cpp
+++ b/Engine/Plugins/Runtime/StateTree/Source/StateTreeEditorModule/Private/SStateTreeViewRow.cpp
@@ -13,6 +13,7 @@
 #include "StateTreeConditionBase.h"
 #include "StateTreeDescriptionHelpers.h"
 #include "StateTreeDragDrop.h"
+#include "StateTreeEditorModule.h"
 #include "StateTreeState.h"
 #include "StateTreeTaskBase.h"
 #include "StateTreeViewModel.h"
@@ -1004,70 +1005,7 @@ const FSlateBrush* SStateTreeViewRow::GetSelectorIcon() const
 {
 	if (const UStateTreeState* State = WeakState.Get())
 	{
-		if (State->SelectionBehavior == EStateTreeStateSelectionBehavior::None)
-		{
-			return FStateTreeEditorStyle::Get().GetBrush("StateTreeEditor.SelectNone");
-		}
-		else if (State->SelectionBehavior == EStateTreeStateSelectionBehavior::TryEnterState)
-		{
-			return FStateTreeEditorStyle::Get().GetBrush("StateTreeEditor.TryEnterState");
-		}
-		else if (State->SelectionBehavior == EStateTreeStateSelectionBehavior::TrySelectChildrenInOrder)
-		{
-			if (State->Children.IsEmpty()
-				|| State->Type == EStateTreeStateType::Linked
-				|| State->Type == EStateTreeStateType::LinkedAsset)
-			{
-				// Backwards compatible behavior
-				return FStateTreeEditorStyle::Get().GetBrush("StateTreeEditor.TryEnterState");
-			}
-			else
-			{
-				return FStateTreeEditorStyle::Get().GetBrush("StateTreeEditor.TrySelectChildrenInOrder");
-			}
-		}
-		else if (State->SelectionBehavior == EStateTreeStateSelectionBehavior::TrySelectChildrenAtUniformRandom)
-		{
-			if (State->Children.IsEmpty()
-				|| State->Type == EStateTreeStateType::Linked
-				|| State->Type == EStateTreeStateType::LinkedAsset)
-			{
-				// Backwards compatible behavior
-				return FStateTreeEditorStyle::Get().GetBrush("StateTreeEditor.TryEnterState");
-			}
-
-			return FStateTreeEditorStyle::Get().GetBrush("StateTreeEditor.TrySelectCHildrenAtRandom");
-		}
-		else if (State->SelectionBehavior == EStateTreeStateSelectionBehavior::TrySelectChildrenWithHighestUtility)
-		{
-			if (State->Children.IsEmpty()
-				|| State->Type == EStateTreeStateType::Linked
-				|| State->Type == EStateTreeStateType::LinkedAsset)
-			{
-				// Backwards compatible behavior
-				return FStateTreeEditorStyle::Get().GetBrush("StateTreeEditor.TryEnterState");
-			}
-
-			//place holder
-			return FStateTreeEditorStyle::Get().GetBrush("StateTreeEditor.TrySelectChildrenInOrder");
-		}
-		else if (State->SelectionBehavior == EStateTreeStateSelectionBehavior::TrySelectChildrenBasedOnRelativeUtility)
-		{
-			if (State->Children.IsEmpty()
-				|| State->Type == EStateTreeStateType::Linked
-				|| State->Type == EStateTreeStateType::LinkedAsset)
-			{
-				// Backwards compatible behavior
-				return FStateTreeEditorStyle::Get().GetBrush("StateTreeEditor.TryEnterState");
-			}
-
-			//place holder
-			return FStateTreeEditorStyle::Get().GetBrush("StateTreeEditor.TrySelectChildrenInOrder");
-		}
-		else if (State->SelectionBehavior == EStateTreeStateSelectionBehavior::TryFollowTransitions)
-		{
-			return FStateTreeEditorStyle::Get().GetBrush("StateTreeEditor.TryFollowTransitions");
-		}
+		return FStateTreeEditorStyle::GetBrushForSelectionBehaviorType(State->SelectionBehavior, !State->Children.IsEmpty(), State->Type);		
 	}
 
 	return nullptr;
diff --git a/Engine/Plugins/Runtime/StateTree/Source/StateTreeEditorModule/Private/StandaloneStateTreeEditorHost.cpp b/Engine/Plugins/Runtime/StateTree/Source/StateTreeEditorModule/Private/StandaloneStateTreeEditorHost.cpp
new file mode 100644
--- /dev/null
+++ b/Engine/Plugins/Runtime/StateTree/Source/StateTreeEditorModule/Private/StandaloneStateTreeEditorHost.cpp
@@ -0,0 +1,53 @@
+// Copyright Epic Games, Inc. All Rights Reserved.
+
+#include "StandaloneStateTreeEditorHost.h"
+#include "StateTreeEditor.h"
+
+void FStandaloneStateTreeEditorHost::Init(const TWeakPtr<FStateTreeEditor>& InWeakStateTreeEditor)
+{
+	WeakStateTreeEditor = InWeakStateTreeEditor;
+}
+
+UStateTree* FStandaloneStateTreeEditorHost::GetStateTree() const
+{
+	if (TSharedPtr<FStateTreeEditor> SharedEditor = WeakStateTreeEditor.Pin())
+	{
+		return SharedEditor->StateTree;
+	}
+	return nullptr;
+}
+
+FName FStandaloneStateTreeEditorHost::GetCompilerLogName() const
+{
+	return FStateTreeEditor::CompilerLogListingName;
+}
+
+FName FStandaloneStateTreeEditorHost::GetCompilerTabName() const
+{
+	return FStateTreeEditor::CompilerResultsTabId;
+}
+
+FSimpleMulticastDelegate& FStandaloneStateTreeEditorHost::OnStateTreeChanged()
+{
+	return OnStateTreeChangedDelegate;
+}
+
+TSharedPtr<IDetailsView> FStandaloneStateTreeEditorHost::GetAssetDetailsView()
+{
+	if (TSharedPtr<FStateTreeEditor> SharedEditor = WeakStateTreeEditor.Pin())
+	{
+		return SharedEditor->AssetDetailsView;
+	}
+	
+	return nullptr;	
+}
+
+TSharedPtr<IDetailsView> FStandaloneStateTreeEditorHost::GetDetailsView()
+{
+	if (TSharedPtr<FStateTreeEditor> SharedEditor = WeakStateTreeEditor.Pin())
+	{
+		return SharedEditor->SelectionDetailsView;
+	}
+	
+	return nullptr;	
+}
\ No newline at end of file
diff --git a/Engine/Plugins/Runtime/StateTree/Source/StateTreeEditorModule/Private/StandaloneStateTreeEditorHost.h b/Engine/Plugins/Runtime/StateTree/Source/StateTreeEditorModule/Private/StandaloneStateTreeEditorHost.h
new file mode 100644
--- /dev/null
+++ b/Engine/Plugins/Runtime/StateTree/Source/StateTreeEditorModule/Private/StandaloneStateTreeEditorHost.h
@@ -0,0 +1,26 @@
+// Copyright Epic Games, Inc. All Rights Reserved.
+
+#pragma once
+
+#include "IStateTreeEditorHost.h"
+
+class FStateTreeEditor;
+
+class FStandaloneStateTreeEditorHost : public IStateTreeEditorHost
+{	
+public:
+	void Init(const TWeakPtr<FStateTreeEditor>& InWeakStateTreeEditor);
+	
+	// IStateTreeEditorHost overrides
+	virtual UStateTree* GetStateTree() const override;
+	virtual FSimpleMulticastDelegate& OnStateTreeChanged() override;
+	virtual TSharedPtr<IDetailsView> GetAssetDetailsView() override;
+	virtual TSharedPtr<IDetailsView> GetDetailsView() override;
+
+	virtual FName GetCompilerLogName() const override;
+	virtual FName GetCompilerTabName() const override;
+
+protected:
+	TWeakPtr<FStateTreeEditor> WeakStateTreeEditor;
+	FSimpleMulticastDelegate OnStateTreeChangedDelegate;
+};
\ No newline at end of file
diff --git a/Engine/Plugins/Runtime/StateTree/Source/StateTreeEditorModule/Private/StateTreeEditingSubsystem.cpp b/Engine/Plugins/Runtime/StateTree/Source/StateTreeEditorModule/Private/StateTreeEditingSubsystem.cpp
new file mode 100644
--- /dev/null
+++ b/Engine/Plugins/Runtime/StateTree/Source/StateTreeEditorModule/Private/StateTreeEditingSubsystem.cpp
@@ -0,0 +1,288 @@
+// Copyright Epic Games, Inc. All Rights Reserved.
+
+#include "StateTreeEditingSubsystem.h"
+
+#include "SStateTreeView.h"
+#include "StateTreeCompiler.h"
+#include "StateTreeCompilerLog.h"
+#include "StateTreeDelegates.h"
+#include "StateTreeEditorData.h"
+#include "StateTreeEditorModule.h"
+#include "StateTreeObjectHash.h"
+#include "StateTreeTaskBase.h"
+
+bool UStateTreeEditingSubsystem::CompileStateTree(const TNonNullPtr<UStateTree> InStateTree, FStateTreeCompilerLog& InOutLog)
+{
+	ValidateStateTree(InStateTree);
+	const uint32 EditorDataHash = CalculateStateTreeHash(InStateTree);
+	
+	FStateTreeCompiler Compiler(InOutLog);
+
+	const bool bCompilationResult = Compiler.Compile(*InStateTree);
+	if(bCompilationResult)
+	{
+		// Success
+		InStateTree->LastCompiledEditorDataHash = EditorDataHash;
+		UE::StateTree::Delegates::OnPostCompile.Broadcast(*InStateTree);
+		UE_LOG(LogStateTreeEditor, Log, TEXT("Compile StateTree '%s' succeeded."), *InStateTree->GetFullName());
+	}
+	else
+	{
+		// Make sure not to leave stale data on failed compile.
+		InStateTree->ResetCompiled();
+		InStateTree->LastCompiledEditorDataHash = 0;
+		
+		UE_LOG(LogStateTreeEditor, Error, TEXT("Failed to compile '%s', errors follow."), *InStateTree->GetFullName());
+		InOutLog.DumpToLog(LogStateTreeEditor);
+	}	
+
+	return bCompilationResult;
+}
+
+TSharedRef<FStateTreeViewModel> UStateTreeEditingSubsystem::FindOrAddViewModel(const TNonNullPtr<UStateTree> InStateTree)
+{
+	if(TSharedPtr<FStateTreeViewModel>* ModelPtr = StateTreeViewModels.Find(InStateTree.Get()))
+	{
+		return ModelPtr->ToSharedRef();
+	}
+
+	TSharedRef<FStateTreeViewModel> SharedModel = StateTreeViewModels.Add(InStateTree.Get(), MakeShared<FStateTreeViewModel>()).ToSharedRef();
+	UStateTreeEditorData* EditorData = Cast<UStateTreeEditorData>(InStateTree->EditorData);
+	if (EditorData == nullptr)
+	{
+		EditorData = NewObject<UStateTreeEditorData>(InStateTree, FName(), RF_Transactional);
+		EditorData->AddRootState();
+		InStateTree->EditorData = EditorData;
+
+		FStateTreeCompilerLog Log;
+		CompileStateTree(InStateTree, Log);
+	}
+	
+	for (UStateTreeState* SubTree : EditorData->SubTrees)
+	{
+		TArray<UStateTreeState*> Stack;
+
+		Stack.Add(SubTree);
+		while (!Stack.IsEmpty())
+		{
+			if (UStateTreeState* State = Stack.Pop())
+			{
+				State->SetFlags(RF_Transactional);
+				
+				for (UStateTreeState* ChildState : State->Children)
+				{
+					Stack.Add(ChildState);
+				}
+			}
+		}
+	}
+		
+	SharedModel->Init(EditorData);
+
+	return SharedModel;
+}
+
+TSharedRef<SWidget> UStateTreeEditingSubsystem::GetStateTreeView(TSharedRef<FStateTreeViewModel> InViewModel, const TSharedRef<FUICommandList>& TreeViewCommandList)
+{
+	return SNew(SStateTreeView, InViewModel, TreeViewCommandList);
+}
+
+void UStateTreeEditingSubsystem::ValidateStateTree(const TNonNullPtr<UStateTree> InStateTree)
+{
+	auto FixChangedStateLinkName = [](FStateTreeStateLink& StateLink, const TMap<FGuid, FName>& IDToName) -> bool
+	{
+		if (StateLink.ID.IsValid())
+		{
+			const FName* Name = IDToName.Find(StateLink.ID);
+			if (Name == nullptr)
+			{
+				// Missing link, we'll show these in the UI
+				return false;
+			}
+			if (StateLink.Name != *Name)
+			{
+				// Name changed, fix!
+				StateLink.Name = *Name;
+				return true;
+			}
+		}
+		return false;
+	};
+
+	auto ValidateLinkedStates = [FixChangedStateLinkName](const UStateTree& StateTree)
+	{
+		UStateTreeEditorData* TreeData = Cast<UStateTreeEditorData>(StateTree.EditorData);
+		if (!TreeData)
+		{
+			return;
+		}
+
+		TreeData->Modify();
+
+		// Make sure all state links are valid and update the names if needed.
+
+		// Create ID to state name map.
+		TMap<FGuid, FName> IDToName;
+
+		TreeData->VisitHierarchy([&IDToName](const UStateTreeState& State, UStateTreeState* /*ParentState*/)
+		{
+			IDToName.Add(State.ID, State.Name);
+			return EStateTreeVisitor::Continue;
+		});
+		
+		// Fix changed names.
+		TreeData->VisitHierarchy([&IDToName, FixChangedStateLinkName](UStateTreeState& State, UStateTreeState* /*ParentState*/)
+		{
+			State.Modify();
+			if (State.Type == EStateTreeStateType::Linked)
+			{
+				FixChangedStateLinkName(State.LinkedSubtree, IDToName);
+			}
+					
+			for (FStateTreeTransition& Transition : State.Transitions)
+			{
+				FixChangedStateLinkName(Transition.State, IDToName);
+			}
+
+			return EStateTreeVisitor::Continue;
+		});
+	};
+
+	auto UpdateParents = [](const UStateTree& StateTree)
+	{
+		UStateTreeEditorData* TreeData = Cast<UStateTreeEditorData>(StateTree.EditorData);
+		if (!TreeData)
+		{
+			return;
+		}
+
+		TreeData->Modify();
+		TreeData->ReparentStates();
+	};
+
+	auto ApplySchema =[](const UStateTree& StateTree)
+	{
+		UStateTreeEditorData* TreeData = Cast<UStateTreeEditorData>(StateTree.EditorData);
+		if (!TreeData)
+		{
+			return;
+		
+		}
+		const UStateTreeSchema* Schema = TreeData->Schema;
+		if (!Schema)
+		{
+			return;
+		}
+
+		TreeData->Modify();
+		
+		// Clear evaluators if not allowed.
+		if (Schema->AllowEvaluators() == false && TreeData->Evaluators.Num() > 0)
+		{
+			UE_LOG(LogStateTreeEditor, Warning, TEXT("%s: Resetting Evaluators due to current schema restrictions."), *GetNameSafe(&StateTree));
+			TreeData->Evaluators.Reset();
+		}
+
+
+		TreeData->VisitHierarchy([&StateTree, Schema](UStateTreeState& State, UStateTreeState* /*ParentState*/)
+		{
+			State.Modify();
+
+			// Clear enter conditions if not allowed.
+			if (Schema->AllowEnterConditions() == false && State.EnterConditions.Num() > 0)
+			{
+				UE_LOG(LogStateTreeEditor, Warning, TEXT("%s: Resetting Enter Conditions in state %s due to current schema restrictions."), *GetNameSafe(&StateTree), *GetNameSafe(&State));
+				State.EnterConditions.Reset();
+			}
+
+			// Clear Utility if not allowed
+			if (Schema->AllowUtilityConsiderations() == false && State.Considerations.Num() > 0)
+			{
+				UE_LOG(LogStateTreeEditor, Warning, TEXT("%s: Resetting Utility Considerations in state %s due to current schema restrictions."), *GetNameSafe(&StateTree), *GetNameSafe(&State));
+				State.Considerations.Reset();
+			}
+
+			// Keep single and many tasks based on what is allowed.
+			if (Schema->AllowMultipleTasks() == false)
+			{
+				if (State.Tasks.Num() > 0)
+				{
+					State.Tasks.Reset();
+					UE_LOG(LogStateTreeEditor, Warning, TEXT("%s: Resetting Tasks in state %s due to current schema restrictions."), *GetNameSafe(&StateTree), *GetNameSafe(&State));
+				}
+				
+				// Task name is the same as state name.
+				if (FStateTreeTaskBase* Task = State.SingleTask.Node.GetMutablePtr<FStateTreeTaskBase>())
+				{
+					Task->Name = State.Name;
+				}
+			}
+			else
+			{
+				if (State.SingleTask.Node.IsValid())
+				{
+					State.SingleTask.Reset();
+					UE_LOG(LogStateTreeEditor, Warning, TEXT("%s: Resetting Single Task in state %s due to current schema restrictions."), *GetNameSafe(&StateTree), *GetNameSafe(&State));
+				}
+			}
+			
+			return EStateTreeVisitor::Continue;
+		});
+	};
+	
+
+	auto RemoveUnusedBindings = [](const UStateTree& StateTree)
+	{
+		UStateTreeEditorData* TreeData = Cast<UStateTreeEditorData>(StateTree.EditorData);
+		if (!TreeData)
+		{
+			return;
+		}
+
+		TMap<FGuid, const FStateTreeDataView> AllStructValues;
+		TreeData->GetAllStructValues(AllStructValues);
+		TreeData->Modify();
+		TreeData->GetPropertyEditorBindings()->RemoveUnusedBindings(AllStructValues);
+	};
+
+	auto UpdateLinkedStateParameters = [](const UStateTree& StateTree)
+	{
+		UStateTreeEditorData* TreeData = Cast<UStateTreeEditorData>(StateTree.EditorData);
+		if (!TreeData)
+		{
+			return;
+		}
+
+		TreeData->Modify();
+
+		const EStateTreeVisitor Result = TreeData->VisitHierarchy([](UStateTreeState& State, UStateTreeState* /*ParentState*/)
+		{
+			if (State.Type == EStateTreeStateType::Linked
+				|| State.Type == EStateTreeStateType::LinkedAsset)
+			{
+				State.Modify();
+				State.UpdateParametersFromLinkedSubtree();
+			}
+			return EStateTreeVisitor::Continue;
+		});
+	};
+
+
+	UpdateParents(*InStateTree);
+	ApplySchema(*InStateTree);
+	RemoveUnusedBindings(*InStateTree);
+	ValidateLinkedStates(*InStateTree);
+	UpdateLinkedStateParameters(*InStateTree);
+}
+
+uint32 UStateTreeEditingSubsystem::CalculateStateTreeHash(const TNonNullPtr<const UStateTree> InStateTree)
+{
+	uint32 EditorDataHash = 0;
+	if (InStateTree->EditorData != nullptr)
+	{
+		FStateTreeObjectCRC32 Archive;
+		EditorDataHash = Archive.Crc32(InStateTree->EditorData, 0);
+	}
+
+	return EditorDataHash;
+}
diff --git a/Engine/Plugins/Runtime/StateTree/Source/StateTreeEditorModule/Private/StateTreeEditor.cpp b/Engine/Plugins/Runtime/StateTree/Source/StateTreeEditorModule/Private/StateTreeEditor.cpp
--- a/Engine/Plugins/Runtime/StateTree/Source/StateTreeEditorModule/Private/StateTreeEditor.cpp
+++ b/Engine/Plugins/Runtime/StateTree/Source/StateTreeEditorModule/Private/StateTreeEditor.cpp
@@ -8,10 +8,13 @@
 #include "Blueprint/StateTreeConsiderationBlueprintBase.h"
 #include "ContentBrowserModule.h"
 #include "ClassViewerFilter.h"
+#include "ContextObjectStore.h"
 #include "Kismet2/KismetEditorUtilities.h"
 #include "Kismet2/BlueprintEditorUtils.h"
 #include "Customizations/StateTreeBindingExtension.h"
 #include "DetailsViewArgs.h"
+#include "EditorModeManager.h"
+#include "EdModeInteractiveToolsContext.h"
 #include "IDetailsView.h"
 #include "IContentBrowserSingleton.h"
 #include "Framework/Application/SlateApplication.h"
@@ -30,7 +33,6 @@
 #include "StateTreeEditorSettings.h"
 #include "StateTreeObjectHash.h"
 #include "StateTreeTaskBase.h"
-#include "StateTreeToolMenuContext.h"
 #include "StateTreeViewModel.h"
 #include "ToolMenus.h"
 #include "ToolMenu.h"
@@ -38,8 +40,12 @@
 #include "ToolMenuEntry.h"
 #include "Widgets/Input/SMultiLineEditableTextBox.h"
 #include "FileHelpers.h"
+#include "StateTreeEditorMode.h"
 #include "PropertyPath.h"
 #include "SStateTreeOutliner.h"
+#include "StandaloneStateTreeEditorHost.h"
+#include "StateTreeEditingSubsystem.h"
+#include "StateTreeEditorUILayer.h"
 #include "Debugger/SStateTreeDebuggerView.h"
 #include "StateTreeSettings.h"
 #include "ToolMenuSection.h"
@@ -51,13 +57,10 @@ const FName StateTreeEditorAppName(TEXT("StateTreeEditorApp"));
 const FName FStateTreeEditor::SelectionDetailsTabId(TEXT("StateTreeEditor_SelectionDetails"));
 const FName FStateTreeEditor::AssetDetailsTabId(TEXT("StateTreeEditor_AssetDetails"));
 const FName FStateTreeEditor::StateTreeViewTabId(TEXT("StateTreeEditor_StateTreeView"));
-const FName FStateTreeEditor::StateTreeOutlinerTabId(TEXT("StateTreeEditor_StateTreeOutliner"));
-const FName FStateTreeEditor::StateTreeStatisticsTabId(TEXT("StateTreeEditor_StateTreeStatistics"));
 const FName FStateTreeEditor::CompilerResultsTabId(TEXT("StateTreeEditor_CompilerResults"));
-#if WITH_STATETREE_TRACE_DEBUGGER
-const FName FStateTreeEditor::DebuggerTabId(TEXT("StateTreeEditor_Debugger"));
-#endif // WITH_STATETREE_TRACE_DEBUGGER
-
+const FName FStateTreeEditor::CompilerLogListingName(TEXT("StateTreeCompiler"));
+const FName FStateTreeEditor::LayoutLeftStackId("LeftStackId");
+const FName FStateTreeEditor::LayoutBottomMiddleStackId("BottomMiddleStackId");
 
 namespace UE::StateTree::Editor
 {
@@ -69,14 +72,6 @@ FAutoConsoleVariableRef CVarDisplayItemIds(
 	TEXT("Appends Id to task and state names in the treeview and expose Ids in the details view."));
 }
 
-void FStateTreeEditor::PostUndo(bool bSuccess)
-{
-}
-
-void FStateTreeEditor::PostRedo(bool bSuccess)
-{
-}
-
 void FStateTreeEditor::AddReferencedObjects(FReferenceCollector& Collector)
 {
 	if (StateTree != nullptr)
@@ -106,28 +101,11 @@ void FStateTreeEditor::RegisterTabSpawners(const TSharedRef<class FTabManager>&
 		.SetDisplayName(NSLOCTEXT("StateTreeEditor", "StateTreeViewTab", "States"))
 		.SetGroup(WorkspaceMenuCategoryRef)
 		.SetIcon(FSlateIcon(FAppStyle::GetAppStyleSetName(), "LevelEditor.Tabs.Outliner"));
-
-	InTabManager->RegisterTabSpawner(StateTreeOutlinerTabId, FOnSpawnTab::CreateSP(this, &FStateTreeEditor::SpawnTab_StateTreeOutliner))
-		.SetDisplayName(NSLOCTEXT("StateTreeEditor", "StateTreeOutlinerTab", "Outliner"))
-		.SetGroup(WorkspaceMenuCategoryRef)
-		.SetIcon(FSlateIcon(FAppStyle::GetAppStyleSetName(), "LevelEditor.Tabs.Outliner"));
-
-	InTabManager->RegisterTabSpawner(StateTreeStatisticsTabId, FOnSpawnTab::CreateSP(this, &FStateTreeEditor::SpawnTab_StateTreeStatistics))
-		.SetDisplayName(NSLOCTEXT("StateTreeEditor", "StatisticsTab", "StateTree Statistics"))
-		.SetGroup(WorkspaceMenuCategoryRef)
-		.SetIcon(FSlateIcon(FAppStyle::GetAppStyleSetName(), "AssetEditor.ToggleStats"));
 	
 	InTabManager->RegisterTabSpawner(CompilerResultsTabId, FOnSpawnTab::CreateSP(this, &FStateTreeEditor::SpawnTab_CompilerResults))
 		.SetDisplayName(NSLOCTEXT("StateTreeEditor", "CompilerResultsTab", "Compiler Results"))
 		.SetGroup(WorkspaceMenuCategoryRef)
 		.SetIcon(FSlateIcon(FAppStyle::GetAppStyleSetName(), "Log.TabIcon"));
-	
-#if WITH_STATETREE_TRACE_DEBUGGER
-	InTabManager->RegisterTabSpawner(DebuggerTabId, FOnSpawnTab::CreateSP(this, &FStateTreeEditor::SpawnTab_Debugger))
-	   .SetDisplayName(NSLOCTEXT("StateTreeEditor", "DebuggerTab", "Debugger"))
-	   .SetGroup(WorkspaceMenuCategoryRef)
-	   .SetIcon(FSlateIcon(FAppStyle::GetAppStyleSetName(), "Debug"));
-#endif // WITH_STATETREE_TRACE_DEBUGGER
 }
 
 
@@ -138,12 +116,7 @@ void FStateTreeEditor::UnregisterTabSpawners(const TSharedRef<class FTabManager>
 	InTabManager->UnregisterTabSpawner(SelectionDetailsTabId);
 	InTabManager->UnregisterTabSpawner(AssetDetailsTabId);
 	InTabManager->UnregisterTabSpawner(StateTreeViewTabId);
-	InTabManager->UnregisterTabSpawner(StateTreeOutlinerTabId);
-	InTabManager->UnregisterTabSpawner(StateTreeStatisticsTabId);
 	InTabManager->UnregisterTabSpawner(CompilerResultsTabId);
-#if WITH_STATETREE_TRACE_DEBUGGER
-	InTabManager->UnregisterTabSpawner(DebuggerTabId);
-#endif // WITH_STATETREE_TRACE_DEBUGGER
 }
 
 void FStateTreeEditor::InitEditor( const EToolkitMode::Type Mode, const TSharedPtr< class IToolkitHost >& InitToolkitHost, UStateTree* InStateTree)
@@ -151,48 +124,10 @@ void FStateTreeEditor::InitEditor( const EToolkitMode::Type Mode, const TSharedP
 	StateTree = InStateTree;
 	check(StateTree != NULL);
 
-	UStateTreeEditorData* EditorData = Cast<UStateTreeEditorData>(StateTree->EditorData);
-	if (EditorData == NULL)
-	{
-		EditorData = NewObject<UStateTreeEditorData>(StateTree, FName(), RF_Transactional);
-		EditorData->AddRootState();
-		StateTree->EditorData = EditorData;
-		Compile();
-	}
-
-	EditorDataHash = UE::StateTree::Editor::CalcAssetHash(*StateTree);
-	
-	// @todo: Temporary fix
-	// Make sure all states are transactional
-	for (UStateTreeState* SubTree : EditorData->SubTrees)
-	{
-		TArray<UStateTreeState*> Stack;
-
-		Stack.Add(SubTree);
-		while (!Stack.IsEmpty())
-		{
-			if (UStateTreeState* State = Stack.Pop())
-			{
-				State->SetFlags(RF_Transactional);
-				
-				for (UStateTreeState* ChildState : State->Children)
-				{
-					Stack.Add(ChildState);
-				}
-			}
-		}
-	}
+	UStateTreeEditingSubsystem* StateTreeEditingSubsystem = GEditor->GetEditorSubsystem<UStateTreeEditingSubsystem>();
+	check(StateTreeEditingSubsystem);
 
-
-	StateTreeViewModel = MakeShareable(new FStateTreeViewModel());
-	StateTreeViewModel->Init(EditorData);
-
-	StateTreeViewModel->GetOnAssetChanged().AddSP(this, &FStateTreeEditor::HandleModelAssetChanged);
-	StateTreeViewModel->GetOnStateAdded().AddSPLambda(this, [this](UStateTreeState* , UStateTreeState*){ UpdateAsset(); });
-	StateTreeViewModel->GetOnStatesRemoved().AddSPLambda(this, [this](const TSet<UStateTreeState*>&){ UpdateAsset(); });
-	StateTreeViewModel->GetOnStatesMoved().AddSPLambda(this, [this](const TSet<UStateTreeState*>&, const TSet<UStateTreeState*>&){ UpdateAsset(); });
-	StateTreeViewModel->GetOnSelectionChanged().AddSP(this, &FStateTreeEditor::HandleModelSelectionChanged);
-	StateTreeViewModel->GetOnBringNodeToFocus().AddSP(this, &FStateTreeEditor::HandleModelBringNodeToFocus);
+	StateTreeViewModel = StateTreeEditingSubsystem->FindOrAddViewModel(InStateTree);
 
 	FMessageLogModule& MessageLogModule = FModuleManager::LoadModuleChecked<FMessageLogModule>("MessageLog");
 	FMessageLogInitializationOptions LogOptions;
@@ -201,12 +136,11 @@ void FStateTreeEditor::InitEditor( const EToolkitMode::Type Mode, const TSharedP
 	LogOptions.bShowFilters = false;
 	LogOptions.bAllowClear = false;
 	LogOptions.MaxPageCount = 1;
-	CompilerResultsListing = MessageLogModule.CreateLogListing("StateTreeCompiler", LogOptions);
-	CompilerResults = MessageLogModule.CreateLogListingWidget(CompilerResultsListing.ToSharedRef());
 
-	CompilerResultsListing->OnMessageTokenClicked().AddSP(this, &FStateTreeEditor::HandleMessageTokenClicked);
+	MessageLogModule.RegisterLogListing(CompilerLogListingName, FText::FromName(CompilerLogListingName), LogOptions);
+	CompilerResultsListing = MessageLogModule.GetLogListing(CompilerLogListingName);
+	CompilerResults = MessageLogModule.CreateLogListingWidget(CompilerResultsListing.ToSharedRef());
 
-	
 	TSharedRef<FTabManager::FLayout> StandaloneDefaultLayout = FTabManager::NewLayout("Standalone_StateTree_Layout_v4")
 	->AddArea
 	(
@@ -218,9 +152,8 @@ void FStateTreeEditor::InitEditor( const EToolkitMode::Type Mode, const TSharedP
 			(
 				FTabManager::NewStack()
 				->SetSizeCoefficient(0.2f)
+				->SetExtensionId(LayoutLeftStackId)
 				->AddTab(AssetDetailsTabId, ETabState::OpenedTab)
-				->AddTab(StateTreeStatisticsTabId, ETabState::OpenedTab)
-				->AddTab(StateTreeOutlinerTabId, ETabState::OpenedTab)
 				->SetForegroundTab(AssetDetailsTabId)
 			)
 			->Split
@@ -237,10 +170,8 @@ void FStateTreeEditor::InitEditor( const EToolkitMode::Type Mode, const TSharedP
 				(
 					FTabManager::NewStack()
 					->SetSizeCoefficient(0.25f)
+					->SetExtensionId(LayoutBottomMiddleStackId)
 					->AddTab(CompilerResultsTabId, ETabState::ClosedTab)
-#if WITH_STATETREE_TRACE_DEBUGGER
-					->AddTab(DebuggerTabId, ETabState::ClosedTab)
-#endif // WITH_STATETREE_TRACE_DEBUGGER
 				)
 			)
 			->Split
@@ -253,24 +184,59 @@ void FStateTreeEditor::InitEditor( const EToolkitMode::Type Mode, const TSharedP
 		)
 	);
 
+	FLayoutExtender LayoutExtender;
+	FStateTreeEditorModule& StateTreeEditorModule = FModuleManager::LoadModuleChecked<FStateTreeEditorModule>("StateTreeEditorModule");
+	StateTreeEditorModule.OnRegisterLayoutExtensions().Broadcast(LayoutExtender);
+	StandaloneDefaultLayout->ProcessExtensions(LayoutExtender);
 
-	const bool bCreateDefaultStandaloneMenu = true;
-	const bool bCreateDefaultToolbar = true;
+	constexpr bool bCreateDefaultStandaloneMenu = true;
+	constexpr bool bCreateDefaultToolbar = true;
 	FAssetEditorToolkit::InitAssetEditor(Mode, InitToolkitHost, StateTreeEditorAppName, StandaloneDefaultLayout, bCreateDefaultStandaloneMenu, bCreateDefaultToolbar, StateTree);
 
-	BindCommands();
 	RegisterToolbar();
 	
-	FStateTreeEditorModule& StateTreeEditorModule = FModuleManager::LoadModuleChecked<FStateTreeEditorModule>("StateTreeEditorModule");
 	AddMenuExtender(StateTreeEditorModule.GetMenuExtensibilityManager()->GetAllExtenders(GetToolkitCommands(), GetEditingObjects()));
 
 	RegenerateMenusAndToolbars();
+}
+
+void FStateTreeEditor::PostInitAssetEditor()
+{
+	IStateTreeEditor::PostInitAssetEditor();
+	
+	ModeUILayer = MakeShared<FStateTreeEditorModeUILayer>(ToolkitHost.Pin().Get());
+	ModeUILayer->SetModeMenuCategory(WorkspaceMenuCategory);
+	ModeUILayer->SetSecondaryModeToolbarName(GetToolMenuToolbarName());
+	ToolkitCommands->Append(ModeUILayer->GetModeCommands());
+	
+	if (UContextObjectStore* ContextStore = EditorModeManager->GetInteractiveToolsContext()->ContextObjectStore)
+	{
+		UStateTreeEditorContext* StateTreeEditorContext = ContextStore->FindContext<UStateTreeEditorContext>();
+		if (!StateTreeEditorContext)
+		{
+			TSharedPtr<FStandaloneStateTreeEditorHost> Host = MakeShared<FStandaloneStateTreeEditorHost>();
+			Host->Init(StaticCastSharedRef<FStateTreeEditor>(AsShared()));
+
+			StateTreeEditorContext = NewObject<UStateTreeEditorContext>();
+			StateTreeEditorContext->EditorHostInterface = Host;
+			ContextStore->AddContextObject(StateTreeEditorContext);
+		}
+	}
+
+	EditorModeManager->SetDefaultMode(UStateTreeEditorMode::EM_StateTree);
+	EditorModeManager->ActivateDefaultMode();
+}
+
+void FStateTreeEditor::OnToolkitHostingStarted(const TSharedRef<IToolkit>& Toolkit)
+{
+	ModeUILayer->OnToolkitHostingStarted(Toolkit);
+	HostedToolkit = Toolkit;
+}
 
-	UE::StateTree::Delegates::OnIdentifierChanged.AddSP(this, &FStateTreeEditor::OnIdentifierChanged);
-	UE::StateTree::Delegates::OnSchemaChanged.AddSP(this, &FStateTreeEditor::OnSchemaChanged);
-	UE::StateTree::Delegates::OnParametersChanged.AddSP(this, &FStateTreeEditor::OnRefreshDetailsView);
-	UE::StateTree::Delegates::OnGlobalDataChanged.AddSP(this, &FStateTreeEditor::OnRefreshDetailsView);
-	UE::StateTree::Delegates::OnStateParametersChanged.AddSP(this, &FStateTreeEditor::OnStateParametersChanged);
+void FStateTreeEditor::OnToolkitHostingFinished(const TSharedRef<IToolkit>& Toolkit)
+{
+	ModeUILayer->OnToolkitHostingFinished(Toolkit);
+	HostedToolkit = nullptr;
 }
 
 FName FStateTreeEditor::GetToolkitFName() const
@@ -293,23 +259,12 @@ FLinearColor FStateTreeEditor::GetWorldCentricTabColorScale() const
 	return FLinearColor( 0.0f, 0.0f, 0.2f, 0.5f );
 }
 
-void FStateTreeEditor::InitToolMenuContext(FToolMenuContext& MenuContext)
-{
-	UStateTreeToolMenuContext* Context = NewObject<UStateTreeToolMenuContext>();
-	Context->StateTreeEditor = SharedThis(this);
-	MenuContext.AddObject(Context);
-}
-
-void FStateTreeEditor::HandleMessageTokenClicked(const TSharedRef<IMessageToken>& InMessageToken)
+void FStateTreeEditor::OnClose()
 {
-	if (InMessageToken->GetType() == EMessageToken::Object)
+	if (HostedToolkit.IsValid())
 	{
-		const TSharedRef<FUObjectToken> ObjectToken = StaticCastSharedRef<FUObjectToken>(InMessageToken);
-		UStateTreeState* State = Cast<UStateTreeState>(ObjectToken->GetObject().Get());
-		if (State)
-		{
-			StateTreeViewModel->SetSelection(State);
-		}
+		UToolMenus::UnregisterOwner(&(*HostedToolkit));	
+		HostedToolkit = nullptr;
 	}
 }
 
@@ -325,18 +280,6 @@ TSharedRef<SDockTab> FStateTreeEditor::SpawnTab_StateTreeView(const FSpawnTabArg
 		];
 }
 
-TSharedRef<SDockTab> FStateTreeEditor::SpawnTab_StateTreeOutliner(const FSpawnTabArgs& Args)
-{
-	check(Args.GetTabId() == StateTreeOutlinerTabId);
-
-	return SNew(SDockTab)
-		.Label(NSLOCTEXT("StateTreeEditor", "StateTreeOutlinerTab", "Outliner"))
-		.TabColorScale(GetTabColorScale())
-		[
-			SAssignNew(StateTreeOutliner, SStateTreeOutliner, StateTreeViewModel.ToSharedRef(), TreeViewCommandList)
-		];
-}
-
 TSharedRef<SDockTab> FStateTreeEditor::SpawnTab_SelectionDetails(const FSpawnTabArgs& Args)
 {
 	check(Args.GetTabId() == SelectionDetailsTabId);
@@ -347,9 +290,6 @@ TSharedRef<SDockTab> FStateTreeEditor::SpawnTab_SelectionDetails(const FSpawnTab
 
 	SelectionDetailsView = PropertyEditorModule.CreateDetailView(DetailsViewArgs);
 	SelectionDetailsView->SetObject(nullptr);
-	SelectionDetailsView->OnFinishedChangingProperties().AddSP(this, &FStateTreeEditor::OnSelectionFinishedChangingProperties);
-
-	FStateTreeEditorModule::SetDetailPropertyHandlers(*SelectionDetailsView);
 
 	TSharedRef<SDockTab> SpawnedTab = SNew(SDockTab)
 		.Label(NSLOCTEXT("StateTreeEditor", "SelectionDetailsTab", "Details"))
@@ -370,9 +310,6 @@ TSharedRef<SDockTab> FStateTreeEditor::SpawnTab_AssetDetails(const FSpawnTabArgs
 
 	AssetDetailsView = PropertyEditorModule.CreateDetailView(DetailsViewArgs);
 	AssetDetailsView->SetObject(StateTree ? StateTree->EditorData : nullptr);
-	AssetDetailsView->OnFinishedChangingProperties().AddSP(this, &FStateTreeEditor::OnAssetFinishedChangingProperties);
-
-	FStateTreeEditorModule::SetDetailPropertyHandlers(*AssetDetailsView);
 
 	TSharedRef<SDockTab> SpawnedTab = SNew(SDockTab)
 		.Label(NSLOCTEXT("StateTreeEditor", "AssetDetailsTabLabel", "Asset Details"))
@@ -383,24 +320,7 @@ TSharedRef<SDockTab> FStateTreeEditor::SpawnTab_AssetDetails(const FSpawnTabArgs
 	return SpawnedTab;
 }
 
-TSharedRef<SDockTab> FStateTreeEditor::SpawnTab_StateTreeStatistics(const FSpawnTabArgs& Args)
-{
-	check(Args.GetTabId() == StateTreeStatisticsTabId);
-	TSharedRef<SDockTab> SpawnedTab = SNew(SDockTab)
-		.Label(LOCTEXT("StatisticsTitle", "StateTree Statistics"))
-		[
-			SNew(SMultiLineEditableTextBox)
-			.Padding(10.0f)
-			.Style(FAppStyle::Get(), "Log.TextBox")
-			.Font(FCoreStyle::GetDefaultFontStyle("Mono", 9))
-			.ForegroundColor(FLinearColor::Gray)
-			.IsReadOnly(true)
-			.Text(this, &FStateTreeEditor::GetStatisticsText)
-		];
-	return SpawnedTab;
-}
-
-TSharedRef<SDockTab> FStateTreeEditor::SpawnTab_CompilerResults(const FSpawnTabArgs& Args)
+TSharedRef<SDockTab> FStateTreeEditor::SpawnTab_CompilerResults(const FSpawnTabArgs& Args) const
 {
 	check(Args.GetTabId() == CompilerResultsTabId);
 	TSharedRef<SDockTab> SpawnedTab = SNew(SDockTab)
@@ -414,152 +334,6 @@ TSharedRef<SDockTab> FStateTreeEditor::SpawnTab_CompilerResults(const FSpawnTabA
 	return SpawnedTab;
 }
 
-#if WITH_STATETREE_TRACE_DEBUGGER
-TSharedRef<SDockTab> FStateTreeEditor::SpawnTab_Debugger(const FSpawnTabArgs& Args)
-{
-	TSharedPtr<SWidget> Widget = SNullWidget::NullWidget;
-	if (StateTree != nullptr)
-	{
-		// Reuse existing view if Tab is reopened
-		if (!DebuggerView.IsValid())
-		{
-			DebuggerView = SNew(SStateTreeDebuggerView, *StateTree, StateTreeViewModel.ToSharedRef(), TreeViewCommandList);
-		}
-		Widget = DebuggerView;
-	}
-	
-	check(Args.GetTabId() == DebuggerTabId);
-
-	return SNew(SDockTab)
-		.Label(LOCTEXT("DebuggerTitle", "Debugger"))
-		.TabColorScale(GetTabColorScale())
-		[
-			Widget.ToSharedRef()
-		];
-}
-#endif // WITH_STATETREE_TRACE_DEBUGGER
-
-FText FStateTreeEditor::GetStatisticsText() const
-{
-	if (!StateTree)
-	{
-		return FText::GetEmpty();
-	}
-
-
-	TArray<FStateTreeMemoryUsage> MemoryUsages = StateTree->CalculateEstimatedMemoryUsage();
-	if (MemoryUsages.IsEmpty())
-	{
-		return FText::GetEmpty();
-	}
-
-	TArray<FText> Rows;
-
-	for (const FStateTreeMemoryUsage& Usage : MemoryUsages)
-	{
-		const FText SizeText = FText::AsMemory(Usage.EstimatedMemoryUsage);
-		const FText NumNodesText = FText::AsNumber(Usage.NodeCount);
-		Rows.Add(FText::Format(LOCTEXT("UsageRow", "{0}: {1}, {2} nodes"), FText::FromString(Usage.Name), SizeText, NumNodesText));
-	}
-
-	return FText::Join(FText::FromString(TEXT("\n")), Rows);
-}
-
-void FStateTreeEditor::HandleModelAssetChanged()
-{
-	UpdateAsset();
-}
-
-void FStateTreeEditor::HandleModelSelectionChanged(const TArray<TWeakObjectPtr<UStateTreeState>>& SelectedStates)
-{
-	if (SelectionDetailsView)
-	{
-		TArray<UObject*> Selected;
-		for (const TWeakObjectPtr<UStateTreeState>& WeakState : SelectedStates)
-		{
-			if (UStateTreeState* State = WeakState.Get())
-			{
-				Selected.Add(State);
-			}
-		}
-		SelectionDetailsView->SetObjects(Selected);
-	}
-}
-
-void FStateTreeEditor::HandleModelBringNodeToFocus(const UStateTreeState* State, const FGuid NodeID)
-{
-	if (SelectionDetailsView && State)
-	{
-		FPropertyPath HighlightPath;
-
-		if (!HighlightPath.IsValid())
-		{
-			FArrayProperty* TasksProperty = CastFieldChecked<FArrayProperty>(UStateTreeState::StaticClass()->FindPropertyByName(GET_MEMBER_NAME_CHECKED(UStateTreeState, Tasks)));
-			const int32 TaskIndex = State->Tasks.IndexOfByPredicate([&NodeID](const FStateTreeEditorNode& Node)
-			{
-				return Node.ID == NodeID;
-			});
-			if (TaskIndex != INDEX_NONE)
-			{
-				HighlightPath.AddProperty(FPropertyInfo(TasksProperty));
-				HighlightPath.AddProperty(FPropertyInfo(TasksProperty->Inner, TaskIndex));
-			}
-		}
-
-		if (!HighlightPath.IsValid())
-		{
-			FProperty* SingleTaskProperty = CastFieldChecked<FProperty>(UStateTreeState::StaticClass()->FindPropertyByName(GET_MEMBER_NAME_CHECKED(UStateTreeState, SingleTask)));
-			if (State->SingleTask.ID == NodeID)
-			{
-				HighlightPath.AddProperty(FPropertyInfo(SingleTaskProperty));
-			}
-		}
-
-		if (!HighlightPath.IsValid())
-		{
-			FArrayProperty* TransitionsProperty = CastFieldChecked<FArrayProperty>(UStateTreeState::StaticClass()->FindPropertyByName(GET_MEMBER_NAME_CHECKED(UStateTreeState, Transitions)));
-			const int32 TransitionIndex = State->Transitions.IndexOfByPredicate([&NodeID](const FStateTreeTransition& Transition)
-			{
-				return Transition.ID == NodeID;
-			});
-			if (TransitionIndex != INDEX_NONE)
-			{
-				HighlightPath.AddProperty(FPropertyInfo(TransitionsProperty));
-				HighlightPath.AddProperty(FPropertyInfo(TransitionsProperty->Inner, TransitionIndex));
-			}
-		}
-
-		if (!HighlightPath.IsValid())
-		{
-			FArrayProperty* EnterConditionsProperty = CastFieldChecked<FArrayProperty>(UStateTreeState::StaticClass()->FindPropertyByName(GET_MEMBER_NAME_CHECKED(UStateTreeState, EnterConditions)));
-			const int32 EnterConditionIndex = State->EnterConditions.IndexOfByPredicate([&NodeID](const FStateTreeEditorNode& Node)
-			{
-				return Node.ID == NodeID;
-			});
-			if (EnterConditionIndex != INDEX_NONE)
-			{
-				HighlightPath.AddProperty(FPropertyInfo(EnterConditionsProperty));
-				HighlightPath.AddProperty(FPropertyInfo(EnterConditionsProperty->Inner, EnterConditionIndex));
-			}
-		}
-
-		if (HighlightPath.IsValid())
-		{
-			SelectionDetailsView->ScrollPropertyIntoView(HighlightPath, /*bExpandProperty*/true);
-			SelectionDetailsView->HighlightProperty(HighlightPath);
-			
-			GEditor->GetTimerManager()->SetTimer(
-				HighlighTimerHandle,
-				FTimerDelegate::CreateLambda([SelectionDetailsView = SelectionDetailsView]()
-				{
-					SelectionDetailsView->HighlightProperty({});
-				}),
-				1.0f,
-				/*Loop*/false);
-		}
-	}	
-}
-
 void FStateTreeEditor::SaveAsset_Execute()
 {
 	// Remember the treeview expansion state
@@ -568,711 +342,19 @@ void FStateTreeEditor::SaveAsset_Execute()
 		StateTreeView->SavePersistentExpandedStates();
 	}
 
-	UpdateAsset();
-
 	// save it
 	FAssetEditorToolkit::SaveAsset_Execute();
 }
 
-void FStateTreeEditor::OnIdentifierChanged(const UStateTree& InStateTree)
-{
-	if (StateTree == &InStateTree)
-	{
-		UpdateAsset();
-	}
-}
-
-void FStateTreeEditor::OnSchemaChanged(const UStateTree& InStateTree)
-{
-	if (StateTree == &InStateTree)
-	{
-		UpdateAsset();
-		
-		if (StateTreeViewModel)
-		{
-			StateTreeViewModel->NotifyAssetChangedExternally();
-		}
-
-		if (SelectionDetailsView.IsValid())
-		{
-			SelectionDetailsView->ForceRefresh();
-		}
-	}
-}
-
-void FStateTreeEditor::OnRefreshDetailsView(const UStateTree& InStateTree)
-{
-	if (StateTree == &InStateTree)
-	{
-		// Accessible structs might be different after modifying parameters so forcing refresh
-		// so the FStateTreeBindingExtension can rebuild the list of bindable structs
-		if (SelectionDetailsView.IsValid())
-		{
-			SelectionDetailsView->ForceRefresh();
-		}
-	}
-}
-
-void FStateTreeEditor::OnStateParametersChanged(const UStateTree& InStateTree, const FGuid ChangedStateID)
-{
-	if (StateTree == &InStateTree)
-	{
-		if (const UStateTreeEditorData* TreeData = Cast<UStateTreeEditorData>(StateTree->EditorData))
-		{
-			TreeData->VisitHierarchy([&ChangedStateID](UStateTreeState& State, UStateTreeState* /*ParentState*/)
-			{
-				if (State.Type == EStateTreeStateType::Linked && State.LinkedSubtree.ID == ChangedStateID)
-				{
-					State.UpdateParametersFromLinkedSubtree();
-				}
-				return EStateTreeVisitor::Continue;
-			});
-		}
-
-		// Accessible structs might be different after modifying parameters so forcing refresh
-		// so the FStateTreeBindingExtension can rebuild the list of bindable structs
-		if (SelectionDetailsView.IsValid())
-		{
-			SelectionDetailsView->ForceRefresh();
-		}
-	}
-}
-
-void FStateTreeEditor::OnAssetFinishedChangingProperties(const FPropertyChangedEvent& PropertyChangedEvent)
-{
-	// Make sure nodes get updates when properties are changed.
-	if (StateTreeViewModel)
-	{
-		StateTreeViewModel->NotifyAssetChangedExternally();
-	}
-}
-
-void FStateTreeEditor::OnSelectionFinishedChangingProperties(const FPropertyChangedEvent& PropertyChangedEvent)
-{
-	// Make sure nodes get updates when properties are changed.
-	if (StateTreeViewModel)
-	{
-		const TArray<TWeakObjectPtr<UObject>>& SelectedObjects = SelectionDetailsView->GetSelectedObjects();
-		TSet<UStateTreeState*> ChangedStates;
-		for (const TWeakObjectPtr<UObject>& WeakObject : SelectedObjects)
-		{
-			if (UObject* Object = WeakObject.Get())
-			{
-				if (UStateTreeState* State = Cast<UStateTreeState>(Object))
-				{
-					ChangedStates.Add(State);
-				}
-			}
-		}
-		if (ChangedStates.Num() > 0)
-		{
-			StateTreeViewModel->NotifyStatesChangedExternally(ChangedStates, PropertyChangedEvent);
-			UpdateAsset();
-		}
-	}
-}
-
-namespace UE::StateTree::Editor::Internal
-{
-	bool FixChangedStateLinkName(FStateTreeStateLink& StateLink, const TMap<FGuid, FName>& IDToName)
-	{
-		if (StateLink.ID.IsValid())
-		{
-			const FName* Name = IDToName.Find(StateLink.ID);
-			if (Name == nullptr)
-			{
-				// Missing link, we'll show these in the UI
-				return false;
-			}
-			if (StateLink.Name != *Name)
-			{
-				// Name changed, fix!
-				StateLink.Name = *Name;
-				return true;
-			}
-		}
-		return false;
-	}
-
-	void ValidateLinkedStates(UStateTree& StateTree)
-	{
-		UStateTreeEditorData* TreeData = Cast<UStateTreeEditorData>(StateTree.EditorData);
-		if (!TreeData)
-		{
-			return;
-		}
-
-		TreeData->Modify();
-
-		// Make sure all state links are valid and update the names if needed.
-
-		// Create ID to state name map.
-		TMap<FGuid, FName> IDToName;
-
-		TreeData->VisitHierarchy([&IDToName](const UStateTreeState& State, UStateTreeState* /*ParentState*/)
-		{
-			IDToName.Add(State.ID, State.Name);
-			return EStateTreeVisitor::Continue;
-		});
-		
-		// Fix changed names.
-		TreeData->VisitHierarchy([&IDToName](UStateTreeState& State, UStateTreeState* /*ParentState*/)
-		{
-			State.Modify();
-			if (State.Type == EStateTreeStateType::Linked)
-			{
-				FixChangedStateLinkName(State.LinkedSubtree, IDToName);
-			}
-					
-			for (FStateTreeTransition& Transition : State.Transitions)
-			{
-				FixChangedStateLinkName(Transition.State, IDToName);
-			}
-
-			return EStateTreeVisitor::Continue;
-		});
-	}
-
-	void UpdateParents(UStateTree& StateTree)
-	{
-		UStateTreeEditorData* TreeData = Cast<UStateTreeEditorData>(StateTree.EditorData);
-		if (!TreeData)
-		{
-			return;
-		}
-
-		TreeData->Modify();
-		TreeData->ReparentStates();
-	}
-
-	void ApplySchema(UStateTree& StateTree)
-	{
-		UStateTreeEditorData* TreeData = Cast<UStateTreeEditorData>(StateTree.EditorData);
-		if (!TreeData)
-		{
-			return;
-		
-		}
-		const UStateTreeSchema* Schema = TreeData->Schema;
-		if (!Schema)
-		{
-			return;
-		}
-
-		TreeData->Modify();
-		
-		// Clear evaluators if not allowed.
-		if (Schema->AllowEvaluators() == false && TreeData->Evaluators.Num() > 0)
-		{
-			UE_LOG(LogStateTreeEditor, Warning, TEXT("%s: Resetting Evaluators due to current schema restrictions."), *GetNameSafe(&StateTree));
-			TreeData->Evaluators.Reset();
-		}
-
-
-		TreeData->VisitHierarchy([&StateTree, Schema](UStateTreeState& State, UStateTreeState* /*ParentState*/)
-		{
-			State.Modify();
-
-			// Clear enter conditions if not allowed.
-			if (Schema->AllowEnterConditions() == false && State.EnterConditions.Num() > 0)
-			{
-				UE_LOG(LogStateTreeEditor, Warning, TEXT("%s: Resetting Enter Conditions in state %s due to current schema restrictions."), *GetNameSafe(&StateTree), *GetNameSafe(&State));
-				State.EnterConditions.Reset();
-			}
-
-			// Clear Utility if not allowed
-			if (Schema->AllowUtilityConsiderations() == false && State.Considerations.Num() > 0)
-			{
-				UE_LOG(LogStateTreeEditor, Warning, TEXT("%s: Resetting Utility Considerations in state %s due to current schema restrictions."), *GetNameSafe(&StateTree), *GetNameSafe(&State));
-				State.Considerations.Reset();
-			}
-
-			// Keep single and many tasks based on what is allowed.
-			if (Schema->AllowMultipleTasks() == false)
-			{
-				if (State.Tasks.Num() > 0)
-				{
-					State.Tasks.Reset();
-					UE_LOG(LogStateTreeEditor, Warning, TEXT("%s: Resetting Tasks in state %s due to current schema restrictions."), *GetNameSafe(&StateTree), *GetNameSafe(&State));
-				}
-				
-				// Task name is the same as state name.
-				if (FStateTreeTaskBase* Task = State.SingleTask.Node.GetMutablePtr<FStateTreeTaskBase>())
-				{
-					Task->Name = State.Name;
-				}
-			}
-			else
-			{
-				if (State.SingleTask.Node.IsValid())
-				{
-					State.SingleTask.Reset();
-					UE_LOG(LogStateTreeEditor, Warning, TEXT("%s: Resetting Single Task in state %s due to current schema restrictions."), *GetNameSafe(&StateTree), *GetNameSafe(&State));
-				}
-			}
-			
-			return EStateTreeVisitor::Continue;
-		});
-	}
-
-	void RemoveUnusedBindings(UStateTree& StateTree)
-	{
-		UStateTreeEditorData* TreeData = Cast<UStateTreeEditorData>(StateTree.EditorData);
-		if (!TreeData)
-		{
-			return;
-		}
-
-		TMap<FGuid, const FStateTreeDataView> AllStructValues;
-		TreeData->GetAllStructValues(AllStructValues);
-		TreeData->Modify();
-		TreeData->GetPropertyEditorBindings()->RemoveUnusedBindings(AllStructValues);
-	}
-
-	void UpdateLinkedStateParameters(UStateTree& StateTree)
-	{
-		UStateTreeEditorData* TreeData = Cast<UStateTreeEditorData>(StateTree.EditorData);
-		if (!TreeData)
-		{
-			return;
-		}
-
-		TreeData->Modify();
-
-		TreeData->VisitHierarchy([](UStateTreeState& State, UStateTreeState* /*ParentState*/)
-		{
-			if (State.Type == EStateTreeStateType::Linked
-				|| State.Type == EStateTreeStateType::LinkedAsset)
-			{
-				State.Modify();
-				State.UpdateParametersFromLinkedSubtree();
-			}
-			return EStateTreeVisitor::Continue;
-		});
-	}
-
-	static void MakeSaveOnCompileSubMenu(UToolMenu* InMenu)
-	{
-		FToolMenuSection& Section = InMenu->AddSection("Section");
-		const FStateTreeEditorCommands& Commands = FStateTreeEditorCommands::Get();
-		Section.AddMenuEntry(Commands.SaveOnCompile_Never);
-		Section.AddMenuEntry(Commands.SaveOnCompile_SuccessOnly);
-		Section.AddMenuEntry(Commands.SaveOnCompile_Always);
-	}
-
-	static void GenerateCompileOptionsMenu(UToolMenu* InMenu)
-	{
-		FToolMenuSection& Section = InMenu->AddSection("Section");
-		const FStateTreeEditorCommands& Commands = FStateTreeEditorCommands::Get();
-
-		// @TODO: disable the menu and change up the tooltip when all sub items are disabled
-		Section.AddSubMenu(
-			"SaveOnCompile",
-			LOCTEXT("SaveOnCompileSubMenu", "Save on Compile"),
-			LOCTEXT("SaveOnCompileSubMenu_ToolTip", "Determines how the StateTree is saved whenever you compile it."),
-			FNewToolMenuDelegate::CreateStatic(&MakeSaveOnCompileSubMenu));
-	}
-
-	static void SetSaveOnCompileSetting(const EStateTreeSaveOnCompile NewSetting)
-	{
-		UStateTreeEditorSettings* Settings = GetMutableDefault<UStateTreeEditorSettings>();
-		Settings->SaveOnCompile = NewSetting;
-		Settings->SaveConfig();
-	}
-
-	static bool IsSaveOnCompileOptionSet(TWeakPtr<FStateTreeEditor> Editor, const EStateTreeSaveOnCompile Option)
-	{
-		const UStateTreeEditorSettings* Settings = GetDefault<UStateTreeEditorSettings>();
-
-		EStateTreeSaveOnCompile CurrentSetting = Settings->SaveOnCompile;
-		if (!Editor.IsValid() || !Editor.Pin()->IsSaveOnCompileEnabled())
-		{
-			// If save-on-compile is disabled for the StateTree, then we want to 
-			// show "Never" as being selected
-			// 
-			// @TODO: a tooltip explaining why would be nice too
-			CurrentSetting = EStateTreeSaveOnCompile::Never;
-		}
-
-		return (CurrentSetting == Option);
-	}
-} // UE::StateTree::Editor::Internal
-
-namespace UE::StateTree::Editor
-{
-	void ValidateAsset(UStateTree& StateTree)
-	{
-		UE::StateTree::Editor::Internal::UpdateParents(StateTree);
-		UE::StateTree::Editor::Internal::ApplySchema(StateTree);
-		UE::StateTree::Editor::Internal::RemoveUnusedBindings(StateTree);
-		UE::StateTree::Editor::Internal::ValidateLinkedStates(StateTree);
-		UE::StateTree::Editor::Internal::UpdateLinkedStateParameters(StateTree);
-	}
-
-	uint32 CalcAssetHash(const UStateTree& StateTree)
-	{
-		uint32 EditorDataHash = 0;
-		if (StateTree.EditorData != nullptr)
-		{
-			FStateTreeObjectCRC32 Archive;
-			EditorDataHash = Archive.Crc32(StateTree.EditorData, 0);
-		}
-
-		return EditorDataHash;
-	}
-
-	template <typename ClassType, typename = typename TEnableIf<TIsDerivedFrom<ClassType, UStateTreeNodeBlueprintBase>::Value>::Type>
-	class FEditorNodeClassFilter : public IClassViewerFilter
-	{
-	public:
-		virtual bool IsClassAllowed(const FClassViewerInitializationOptions& InInitOptions, const UClass* InClass, TSharedRef< FClassViewerFilterFuncs > InFilterFuncs) override
-		{
-			check(InClass);
-			return InClass->IsChildOf(ClassType::StaticClass());
-		}
-
-		virtual bool IsUnloadedClassAllowed(const FClassViewerInitializationOptions& InInitOptions, const TSharedRef< const IUnloadedBlueprintData > InUnloadedClassData, TSharedRef< FClassViewerFilterFuncs > InFilterFuncs) override
-		{
-			return InUnloadedClassData->IsChildOf(ClassType::StaticClass());
-		}
-	};
-
-	using FStateTreeTaskBPClassFilter = FEditorNodeClassFilter<UStateTreeTaskBlueprintBase>;
-	using FStateTreeConditionBPClassFilter = FEditorNodeClassFilter<UStateTreeConditionBlueprintBase>;
-	using FStateTreeConsiderationBPClassFilter = FEditorNodeClassFilter<UStateTreeConsiderationBlueprintBase>;
-}; // UE::StateTree::Editor
-
-void FStateTreeEditor::BindCommands()
-{
-	const FStateTreeEditorCommands& Commands = FStateTreeEditorCommands::Get();
-
-	ToolkitCommands->MapAction(
-		Commands.Compile,
-		FExecuteAction::CreateSP(this, &FStateTreeEditor::Compile),
-		FCanExecuteAction::CreateSP(this, &FStateTreeEditor::CanCompile));
-
-	TWeakPtr<FStateTreeEditor> WeakThisPtr = SharedThis(this);
-	ToolkitCommands->MapAction(
-		FStateTreeEditorCommands::Get().SaveOnCompile_Never,
-		FExecuteAction::CreateStatic(&UE::StateTree::Editor::Internal::SetSaveOnCompileSetting, EStateTreeSaveOnCompile::Never),
-		FCanExecuteAction::CreateSP(this, &FStateTreeEditor::IsSaveOnCompileEnabled),
-		FIsActionChecked::CreateStatic(&UE::StateTree::Editor::Internal::IsSaveOnCompileOptionSet, WeakThisPtr, EStateTreeSaveOnCompile::Never)
-	);
-	ToolkitCommands->MapAction(
-		FStateTreeEditorCommands::Get().SaveOnCompile_SuccessOnly,
-		FExecuteAction::CreateStatic(&UE::StateTree::Editor::Internal::SetSaveOnCompileSetting, EStateTreeSaveOnCompile::SuccessOnly),
-		FCanExecuteAction::CreateSP(this, &FStateTreeEditor::IsSaveOnCompileEnabled),
-		FIsActionChecked::CreateStatic(&UE::StateTree::Editor::Internal::IsSaveOnCompileOptionSet, WeakThisPtr, EStateTreeSaveOnCompile::SuccessOnly)
-	);
-	ToolkitCommands->MapAction(
-		FStateTreeEditorCommands::Get().SaveOnCompile_Always,
-		FExecuteAction::CreateStatic(&UE::StateTree::Editor::Internal::SetSaveOnCompileSetting, EStateTreeSaveOnCompile::Always),
-		FCanExecuteAction::CreateSP(this, &FStateTreeEditor::IsSaveOnCompileEnabled),
-		FIsActionChecked::CreateStatic(&UE::StateTree::Editor::Internal::IsSaveOnCompileOptionSet, WeakThisPtr, EStateTreeSaveOnCompile::Always)
-	);
-}
-
-bool FStateTreeEditor::IsSaveOnCompileEnabled() const
-{
-	return true;
-}
-
 void FStateTreeEditor::RegisterToolbar()
 {
 	UToolMenus* ToolMenus = UToolMenus::Get();
-	UToolMenu* ToolBar;
 	FName ParentName;
 	const FName MenuName = GetToolMenuToolbarName(ParentName);
-	if (ToolMenus->IsMenuRegistered(MenuName))
+	if (!ToolMenus->IsMenuRegistered(MenuName))
 	{
-		ToolBar = ToolMenus->ExtendMenu(MenuName);
+		UToolMenus::Get()->RegisterMenu(MenuName, ParentName, EMultiBoxType::ToolBar);
 	}
-	else
-	{
-		ToolBar = UToolMenus::Get()->RegisterMenu(MenuName, ParentName, EMultiBoxType::ToolBar);
-	}
-
-	static const FToolMenuInsert InsertAfterAssetSection("Asset", EToolMenuInsertType::After);
-
-	FToolMenuSection& CompileSection = ToolBar->AddSection("Compile", TAttribute<FText>(), InsertAfterAssetSection);
-	
-	CompileSection.AddDynamicEntry("CompileCommands", FNewToolMenuSectionDelegate::CreateLambda([](FToolMenuSection& InSection)
-	{
-		const UStateTreeToolMenuContext* Context = InSection.FindContext<UStateTreeToolMenuContext>();
-		if (Context && Context->StateTreeEditor.IsValid())
-		{
-			const TSharedPtr<FStateTreeEditor> StateTreeEditor = Context->StateTreeEditor.Pin();
-			if (StateTreeEditor.IsValid())
-			{
-				const FStateTreeEditorCommands& Commands = FStateTreeEditorCommands::Get();
-
-				FToolMenuEntry& CompileButton = InSection.AddEntry(FToolMenuEntry::InitToolBarButton(
-					Commands.Compile,
-					TAttribute<FText>(),
-					TAttribute<FText>(),
-					TAttribute<FSlateIcon>(StateTreeEditor.ToSharedRef(), &FStateTreeEditor::GetCompileStatusImage)));
-				CompileButton.StyleNameOverride = "CalloutToolbar";
-
-				FToolMenuEntry& CompileOptions = InSection.AddEntry(FToolMenuEntry::InitComboButton(
-					"CompileComboButton",
-					FUIAction(),
-					FNewToolMenuDelegate::CreateStatic(&UE::StateTree::Editor::Internal::GenerateCompileOptionsMenu),
-					LOCTEXT("CompileOptions_ToolbarTooltip", "Options to customize how State Trees compile")
-				));
-				CompileOptions.StyleNameOverride = "CalloutToolbar";
-				CompileOptions.ToolBarData.bSimpleComboBox = true;
-			}
-		}
-	}));
-
-	static const FToolMenuInsert InsertAfterCompileSection("Compile", EToolMenuInsertType::After);
-
-	FToolMenuSection& CreateNewNodeSection = ToolBar->AddSection("CreateNewNodes", TAttribute<FText>(), InsertAfterCompileSection);
-
-	CreateNewNodeSection.AddDynamicEntry("CreateNewNodes", FNewToolMenuSectionDelegate::CreateLambda([](FToolMenuSection& InSection)
-	{
-		const UStateTreeToolMenuContext* Context = InSection.FindContext<UStateTreeToolMenuContext>();
-		if (Context && Context->StateTreeEditor.IsValid())
-		{
-			const TSharedPtr<FStateTreeEditor> StateTreeEditor = Context->StateTreeEditor.Pin();
-			if (StateTreeEditor.IsValid())
-			{
-				const TSharedRef<FStateTreeEditor> StateTreeEditorRef = StateTreeEditor.ToSharedRef();
-				FToolMenuEntry& CreateNewTaskDropdown = InSection.AddEntry(FToolMenuEntry::InitComboButton(
-					"CreateNewTaskComboButton",
-					FUIAction(),
-					FOnGetContent::CreateSP(StateTreeEditorRef, &FStateTreeEditor::GenerateTaskBPBaseClassesMenu),
-					LOCTEXT("CreateNewTask_Title", "New Task"),
-					LOCTEXT("CreateNewTask_ToolbarTooltip", "Create a new Blueprint State Tree Task"),
-					TAttribute<FSlateIcon>(StateTreeEditorRef, &FStateTreeEditor::GetNewTaskButtonImage)
-				));
-
-				FToolMenuEntry& CreateNewConditionDropdown = InSection.AddEntry(FToolMenuEntry::InitComboButton(
-					"CreateNewConditionComboButton",
-					FUIAction(),
-					FOnGetContent::CreateSP(StateTreeEditorRef, &FStateTreeEditor::GenerateConditionBPBaseClassesMenu),
-					LOCTEXT("CreateNewCondition_Title", "New Condition"),
-					LOCTEXT("CreateNewCondition_ToolbarTooltip", "Create a new Blueprint State Tree Condition"),
-					TAttribute<FSlateIcon>(StateTreeEditorRef, &FStateTreeEditor::GetNewConditionButtonImage)
-				));
-
-				FToolMenuEntry& CreateNewConsiderationDropdown = InSection.AddEntry(FToolMenuEntry::InitComboButton(
-					"CreateNewConsiderationComboButton",
-					FUIAction(),
-					FOnGetContent::CreateSP(StateTreeEditorRef, &FStateTreeEditor::GenerateConsiderationBPBaseClassesMenu),
-					LOCTEXT("CreateNewConsideration_Title", "New Consideration"),
-					LOCTEXT("CreateNewConsideration_ToolbarTooltip", "Create a new Blueprint State Tree Utility Consideration"),
-					TAttribute<FSlateIcon>(StateTreeEditorRef, &FStateTreeEditor::GetNewConsiderationButtonImage)
-				));
-			}
-		}
-	}));
-}
-
-void FStateTreeEditor::Compile()
-{
-	if (!StateTree)
-	{
-		return;
-	}
-
-	// Note: If the compilation process changes, also update UStateTreeCompileAllCommandlet and UStateTreeFactory::FactoryCreateNew.
-	
-	UpdateAsset();
-	
-	if (CompilerResultsListing.IsValid())
-	{
-		CompilerResultsListing->ClearMessages();
-	}
-
-	FStateTreeCompilerLog Log;
-	FStateTreeCompiler Compiler(Log);
-
-	bLastCompileSucceeded = Compiler.Compile(*StateTree);
-
-	if (CompilerResultsListing.IsValid())
-	{
-		Log.AppendToLog(CompilerResultsListing.Get());
-	}
-
-	if (bLastCompileSucceeded)
-	{
-		// Success
-		StateTree->LastCompiledEditorDataHash = EditorDataHash;
-
-		UE::StateTree::Delegates::OnPostCompile.Broadcast(*StateTree);
-	}
-	else
-	{
-		// Make sure not to leave stale data on failed compile.
-		StateTree->ResetCompiled();
-		StateTree->LastCompiledEditorDataHash = 0;
-
-		// Show log
-		TabManager->TryInvokeTab(CompilerResultsTabId);
-	}
-
-	const UStateTreeEditorSettings* Settings = GetMutableDefault<UStateTreeEditorSettings>();
-	const bool bShouldSaveOnCompile = ((Settings->SaveOnCompile == EStateTreeSaveOnCompile::Always)
-									|| ((Settings->SaveOnCompile == EStateTreeSaveOnCompile::SuccessOnly) && bLastCompileSucceeded));
-
-	if (bShouldSaveOnCompile)
-	{
-		const TArray<UPackage*> PackagesToSave { StateTree->GetOutermost() };
-		FEditorFileUtils::PromptForCheckoutAndSave(PackagesToSave, /*bCheckDirty =*/true, /*bPromptToSave =*/false);
-	}
-}
-
-bool FStateTreeEditor::CanCompile() const
-{
-	if (StateTree == nullptr)
-	{
-		return false;
-	}
-	
-	// We can't recompile while in PIE
-	if (GEditor->IsPlaySessionInProgress())
-	{
-		return false;
-	}
-
-	return true;
-}
-
-FSlateIcon FStateTreeEditor::GetCompileStatusImage() const
-{
-	static const FName CompileStatusBackground("Blueprint.CompileStatus.Background");
-	static const FName CompileStatusUnknown("Blueprint.CompileStatus.Overlay.Unknown");
-	static const FName CompileStatusError("Blueprint.CompileStatus.Overlay.Error");
-	static const FName CompileStatusGood("Blueprint.CompileStatus.Overlay.Good");
-	static const FName CompileStatusWarning("Blueprint.CompileStatus.Overlay.Warning");
-
-	if (StateTree == nullptr)
-	{
-		return FSlateIcon(FAppStyle::GetAppStyleSetName(), CompileStatusBackground, NAME_None, CompileStatusUnknown);
-	}
-	
-	const bool bCompiledDataResetDuringLoad = StateTree->LastCompiledEditorDataHash == EditorDataHash && !StateTree->IsReadyToRun();
-
-	if (!bLastCompileSucceeded || bCompiledDataResetDuringLoad)
-	{
-		return FSlateIcon(FAppStyle::GetAppStyleSetName(), CompileStatusBackground, NAME_None, CompileStatusError);
-	}
-	
-	if (StateTree->LastCompiledEditorDataHash != EditorDataHash)
-	{
-		return FSlateIcon(FAppStyle::GetAppStyleSetName(), CompileStatusBackground, NAME_None, CompileStatusUnknown);
-	}
-	
-	return FSlateIcon(FAppStyle::GetAppStyleSetName(), CompileStatusBackground, NAME_None, CompileStatusGood);
-}
-
-FSlateIcon FStateTreeEditor::GetNewConditionButtonImage() const
-{
-	//placeholder
-	return FSlateIcon(FAppStyle::GetAppStyleSetName(), NAME_None, NAME_None, NAME_None);
-}
-
-TSharedRef<SWidget> FStateTreeEditor::GenerateConditionBPBaseClassesMenu() const
-{
-	FClassViewerInitializationOptions Options;
-	Options.NameTypeToDisplay = EClassViewerNameTypeToDisplay::DisplayName;
-	Options.ClassFilters.Add(MakeShareable(new UE::StateTree::Editor::FStateTreeConditionBPClassFilter));
-
-	FOnClassPicked OnPicked(FOnClassPicked::CreateSP(this, &FStateTreeEditor::OnNodeBPBaseClassPicked));
-
-	return FModuleManager::LoadModuleChecked<FClassViewerModule>("ClassViewer").CreateClassViewer(Options, OnPicked);
-}
-
-FSlateIcon FStateTreeEditor::GetNewConsiderationButtonImage() const
-{
-	//placeholder
-	return FSlateIcon(FAppStyle::GetAppStyleSetName(), NAME_None, NAME_None, NAME_None);
-}
-
-TSharedRef<SWidget> FStateTreeEditor::GenerateConsiderationBPBaseClassesMenu() const
-{
-	FClassViewerInitializationOptions Options;
-	Options.NameTypeToDisplay = EClassViewerNameTypeToDisplay::DisplayName;
-	Options.ClassFilters.Add(MakeShareable(new UE::StateTree::Editor::FStateTreeConsiderationBPClassFilter));
-
-	FOnClassPicked OnPicked(FOnClassPicked::CreateSP(this, &FStateTreeEditor::OnNodeBPBaseClassPicked));
-
-	return FModuleManager::LoadModuleChecked<FClassViewerModule>("ClassViewer").CreateClassViewer(Options, OnPicked);
-}
-
-FSlateIcon FStateTreeEditor::GetNewTaskButtonImage() const
-{
-	//placeholder
-	return FSlateIcon(FAppStyle::GetAppStyleSetName(), NAME_None, NAME_None, NAME_None);
-}
-
-TSharedRef<SWidget> FStateTreeEditor::GenerateTaskBPBaseClassesMenu() const
-{
-	FClassViewerInitializationOptions Options;
-	Options.NameTypeToDisplay = EClassViewerNameTypeToDisplay::DisplayName;
-	Options.ClassFilters.Add(MakeShareable(new UE::StateTree::Editor::FStateTreeTaskBPClassFilter));
-
-	FOnClassPicked OnPicked(FOnClassPicked::CreateSP(this, &FStateTreeEditor::OnNodeBPBaseClassPicked));
-
-	return FModuleManager::LoadModuleChecked<FClassViewerModule>("ClassViewer").CreateClassViewer(Options, OnPicked);
-}
-
-void FStateTreeEditor::OnNodeBPBaseClassPicked(UClass* NodeClass) const
-{
-	check(NodeClass);
-	
-	if (!StateTree)
-	{
-		return;
-	}
-
-	const FString ClassName = FBlueprintEditorUtils::GetClassNameWithoutSuffix(NodeClass);
-	const FString PathName = FPaths::GetPath(StateTree->GetOutermost()->GetPathName());
-
-	// Now that we've generated some reasonable default locations/names for the package, allow the user to have the final say
-	// before we create the package and initialize the blueprint inside of it.
-	FSaveAssetDialogConfig SaveAssetDialogConfig;
-	SaveAssetDialogConfig.DialogTitleOverride = LOCTEXT("SaveAssetDialogTitle", "Save Asset As");
-	SaveAssetDialogConfig.DefaultPath = PathName;
-	SaveAssetDialogConfig.DefaultAssetName = ClassName + TEXT("_New");
-	SaveAssetDialogConfig.ExistingAssetPolicy = ESaveAssetDialogExistingAssetPolicy::Disallow;
-
-	const FContentBrowserModule& ContentBrowserModule = FModuleManager::LoadModuleChecked<FContentBrowserModule>("ContentBrowser");
-	const FString SaveObjectPath = ContentBrowserModule.Get().CreateModalSaveAssetDialog(SaveAssetDialogConfig);
-	if (!SaveObjectPath.IsEmpty())
-	{
-		const FString SavePackageName = FPackageName::ObjectPathToPackageName(SaveObjectPath);
-		const FString SavePackagePath = FPaths::GetPath(SavePackageName);
-		const FString SaveAssetName = FPaths::GetBaseFilename(SavePackageName);
-
-		if (UPackage* Package = CreatePackage(*SavePackageName))
-		{
-			// Create and init a new Blueprint
-			if (UBlueprint* NewBP = FKismetEditorUtilities::CreateBlueprint(NodeClass, Package, FName(*SaveAssetName), BPTYPE_Normal, UBlueprint::StaticClass(), UBlueprintGeneratedClass::StaticClass()))
-			{
-				GEditor->GetEditorSubsystem<UAssetEditorSubsystem>()->OpenEditorForAsset(NewBP);
-
-				// Notify the asset registry
-				FAssetRegistryModule::AssetCreated(NewBP);
-
-				Package->MarkPackageDirty();
-			}
-		}
-	}
-
-	FSlateApplication::Get().DismissAllMenus();
-}
-
-void FStateTreeEditor::UpdateAsset()
-{
-	if (!StateTree)
-	{
-		return;
-	}
-
-	UE::StateTree::Editor::ValidateAsset(*StateTree);
-	EditorDataHash = UE::StateTree::Editor::CalcAssetHash(*StateTree);
 }
 
 #undef LOCTEXT_NAMESPACE
diff --git a/Engine/Plugins/Runtime/StateTree/Source/StateTreeEditorModule/Private/StateTreeEditor.h b/Engine/Plugins/Runtime/StateTree/Source/StateTreeEditorModule/Private/StateTreeEditor.h
--- a/Engine/Plugins/Runtime/StateTree/Source/StateTreeEditorModule/Private/StateTreeEditor.h
+++ b/Engine/Plugins/Runtime/StateTree/Source/StateTreeEditorModule/Private/StateTreeEditor.h
@@ -5,30 +5,32 @@
 #include "EditorUndoClient.h"
 #include "IStateTreeEditor.h"
 
+class FStateTreeEditorModeUILayer;
 class FSpawnTabArgs;
 class IMessageToken;
 class UStateTreeState;
 class UStateTreeNodeBlueprintBase;
-
+class IMessageLogListing;
 class IDetailsView;
 class UStateTree;
 class FStateTreeViewModel;
 
 namespace UE::StateTree::Editor
 {
-	// Validates asset state
-	void ValidateAsset(UStateTree& StateTree);
-
-	// Calculates editor data hash of the asset. 
-	uint32 CalcAssetHash(const UStateTree& StateTree);
-
 	extern bool GbDisplayItemIds;
 } // UE::StateTree::Editor
 
-
 class FStateTreeEditor : public IStateTreeEditor, public FSelfRegisteringEditorUndoClient, public FGCObject
 {
+	
+private:
+	friend class FStandaloneStateTreeEditorHost;	
 public:
+	static const FName LayoutLeftStackId;
+	static const FName LayoutBottomMiddleStackId;
+	static const FName CompilerLogListingName;
+	static const FName CompilerResultsTabId;
+	
 	FStateTreeEditor() : TreeViewCommandList(new FUICommandList()) {}
 
 	virtual void RegisterTabSpawners(const TSharedRef<class FTabManager>& TabManager) override;
@@ -36,23 +38,19 @@ public:
 
 	void InitEditor(const EToolkitMode::Type Mode, const TSharedPtr< class IToolkitHost >& InitToolkitHost, UStateTree* StateTree);
 
-	bool IsSaveOnCompileEnabled() const;
+	//~ Begin FAssetEditorToolkit Interface
+	virtual void PostInitAssetEditor() override;
+	virtual void OnToolkitHostingFinished(const TSharedRef<IToolkit>& Toolkit) override;
+	virtual void OnToolkitHostingStarted(const TSharedRef<IToolkit>& Toolkit) override;	
+	//~ End FAssetEditorToolkit Interface
 
 	//~ Begin IToolkit Interface
 	virtual FName GetToolkitFName() const override;
 	virtual FText GetBaseToolkitName() const override;
 	virtual FString GetWorldCentricTabPrefix() const override;
 	virtual FLinearColor GetWorldCentricTabColorScale() const override;
-	virtual void InitToolMenuContext(FToolMenuContext& MenuContext) override;
+	virtual void OnClose() override;
 	//~ End IToolkit Interface
-
-	void OnAssetFinishedChangingProperties(const FPropertyChangedEvent& PropertyChangedEvent);
-	void OnSelectionFinishedChangingProperties(const FPropertyChangedEvent& PropertyChangedEvent);
-
-	/** FEditorUndoClient interface */
-	virtual void PostUndo(bool bSuccess) override;
-	virtual void PostRedo(bool bSuccess) override;
-
 	/** FGCObject interface */
 	virtual void AddReferencedObjects(FReferenceCollector& Collector) override;
 	virtual FString GetReferencerName() const override
@@ -65,61 +63,25 @@ protected:
 	virtual void SaveAsset_Execute() override;
 
 private:
-
-	void HandleMessageTokenClicked(const TSharedRef<IMessageToken>& InMessageToken);
-	void HandleModelSelectionChanged(const TArray<TWeakObjectPtr<UStateTreeState>>& SelectedStates);
-	void HandleModelBringNodeToFocus(const UStateTreeState* State, const FGuid NodeID);
-
 	/** Spawns the tab with the update graph inside */
 	TSharedRef<SDockTab> SpawnTab_StateTreeView(const FSpawnTabArgs& Args);
-	TSharedRef<SDockTab> SpawnTab_StateTreeOutliner(const FSpawnTabArgs& Args);
 	TSharedRef<SDockTab> SpawnTab_SelectionDetails(const FSpawnTabArgs& Args);
 	TSharedRef<SDockTab> SpawnTab_AssetDetails(const FSpawnTabArgs& Args);
-	TSharedRef<SDockTab> SpawnTab_StateTreeStatistics(const FSpawnTabArgs& Args);
-	TSharedRef<SDockTab> SpawnTab_CompilerResults(const FSpawnTabArgs& Args);
+	TSharedRef<SDockTab> SpawnTab_CompilerResults(const FSpawnTabArgs& Args) const;
 
-	void BindCommands();
 	void RegisterToolbar();
-
-	void Compile();
-	bool CanCompile() const;
-	FSlateIcon GetCompileStatusImage() const;
-
-	FSlateIcon GetNewTaskButtonImage() const;
-	TSharedRef<SWidget> GenerateTaskBPBaseClassesMenu() const;
-
-	FSlateIcon GetNewConditionButtonImage() const;
-	TSharedRef<SWidget> GenerateConditionBPBaseClassesMenu() const;
 	
-	FSlateIcon GetNewConsiderationButtonImage() const;
-	TSharedRef<SWidget> GenerateConsiderationBPBaseClassesMenu() const;
-
-	void OnNodeBPBaseClassPicked(UClass* NodeClass) const;
-
-	void UpdateAsset();
-
-	void HandleModelAssetChanged();
-	void OnIdentifierChanged(const UStateTree& StateTree);
-	void OnSchemaChanged(const UStateTree& StateTree);
-	void OnRefreshDetailsView(const UStateTree& StateTree);
-	void OnStateParametersChanged(const UStateTree& StateTree, const FGuid StateID);
-
-	FText GetStatisticsText() const;
-
 	/** State Tree being edited */
 	TObjectPtr<UStateTree> StateTree = nullptr;
 
-	uint32 EditorDataHash = 0;
-	bool bLastCompileSucceeded = true;
-
 	/** The command list used by the tree view. Stored here, so that other windows (e.g. debugger) can add commands to it, even if the tree view is not spawned yet. */
 	TSharedRef<FUICommandList> TreeViewCommandList;
 	
 	/** Selection Property View */
-	TSharedPtr<class IDetailsView> SelectionDetailsView;
+	TSharedPtr<IDetailsView> SelectionDetailsView;
 
 	/** Asset Property View */
-	TSharedPtr<class IDetailsView> AssetDetailsView;
+	TSharedPtr<IDetailsView> AssetDetailsView;
 
 	/** Tree View */
 	TSharedPtr<class SStateTreeView> StateTreeView;
@@ -128,24 +90,16 @@ private:
 	TSharedPtr<class SStateTreeOutliner> StateTreeOutliner;
 
 	/** Compiler Results log */
-	TSharedPtr<class SWidget> CompilerResults;
-	TSharedPtr<class IMessageLogListing> CompilerResultsListing;
+	TSharedPtr<SWidget> CompilerResults;
+	TSharedPtr<IMessageLogListing> CompilerResultsListing;
 	
 	TSharedPtr<FStateTreeViewModel> StateTreeViewModel;
 
-	FTimerHandle HighlighTimerHandle;
-	
-#if WITH_STATETREE_TRACE_DEBUGGER
-	TSharedRef<SDockTab> SpawnTab_Debugger(const FSpawnTabArgs& Args);
-
-	TSharedPtr<class SStateTreeDebuggerView> DebuggerView;
-	static const FName DebuggerTabId;
-#endif // WITH_STATETREE_TRACE_DEBUGGER
+	TSharedPtr<FStateTreeEditorModeUILayer> ModeUILayer;
+	TSharedPtr<IToolkit> HostedToolkit;
+	TSharedPtr<FWorkspaceItem> WorkspaceMenuCategory;
 
 	static const FName StateTreeViewTabId;
-	static const FName StateTreeOutlinerTabId;
 	static const FName SelectionDetailsTabId;
 	static const FName AssetDetailsTabId;
-	static const FName StateTreeStatisticsTabId;
-	static const FName CompilerResultsTabId;
 };
diff --git a/Engine/Plugins/Runtime/StateTree/Source/StateTreeEditorModule/Private/StateTreeEditorMode.cpp b/Engine/Plugins/Runtime/StateTree/Source/StateTreeEditorModule/Private/StateTreeEditorMode.cpp
new file mode 100644
--- /dev/null
+++ b/Engine/Plugins/Runtime/StateTree/Source/StateTreeEditorModule/Private/StateTreeEditorMode.cpp
@@ -0,0 +1,621 @@
+// Copyright Epic Games, Inc. All Rights Reserved.
+
+#include "StateTreeEditorMode.h"
+
+#include "ContextObjectStore.h"
+#include "EditorModeManager.h"
+#include "FileHelpers.h"
+
+#include "StateTreeEditorModeToolkit.h"
+#include "StateTreeEditorSettings.h"
+#include "StateTreeCompilerLog.h"
+#include "StateTreeDelegates.h"
+#include "EdModeInteractiveToolsContext.h"
+#include "IMessageLogListing.h"
+#include "InteractiveToolManager.h"
+#include "PropertyPath.h"
+#include "StateTreeEditorCommands.h"
+#include "Misc/UObjectToken.h"
+#include "Toolkits/ToolkitManager.h"
+#include "Modules/ModuleManager.h"
+
+#include "IStateTreeEditorHost.h"
+#include "MessageLogModule.h"
+#include "StateTreeEditingSubsystem.h"
+#include "Customizations/StateTreeBindingExtension.h"
+
+#define LOCTEXT_NAMESPACE "UStateTreeEditorMode"
+
+class FUObjectToken;
+const FEditorModeID UStateTreeEditorMode::EM_StateTree("StateTreeEditorMode");
+
+UStateTreeEditorMode::UStateTreeEditorMode()
+{
+	Info = FEditorModeInfo(UStateTreeEditorMode::EM_StateTree,
+		LOCTEXT("StateTreeEditorModeName", "StateTreeEditorMode"),
+		FSlateIcon(),
+		false);
+}
+
+void UStateTreeEditorMode::Enter()
+{
+	Super::Enter();
+
+	DetailsViewExtensionHandler = MakeShared<FStateTreeBindingExtension>();
+	DetailsViewChildrenCustomizationHandler = MakeShared<FStateTreeBindingsChildrenCustomization>();
+	
+	if (const UContextObjectStore* ContextObjectStore = GetToolManager()->GetContextObjectStore())
+	{
+		if (const UStateTreeEditorContext* Context = ContextObjectStore->FindContext<UStateTreeEditorContext>())
+		{
+			TSharedRef<IStateTreeEditorHost> Host = Context->EditorHostInterface.ToSharedRef();
+			Host->OnStateTreeChanged().AddUObject(this, &UStateTreeEditorMode::OnStateTreeChanged);
+
+			if (TSharedPtr<IMessageLogListing> MessageLogListing = GetMessageLogListing())
+			{
+				MessageLogListing->OnMessageTokenClicked().AddUObject(this, &UStateTreeEditorMode::HandleMessageTokenClicked);
+			}
+
+			if (TSharedPtr<IDetailsView> DetailsView = GetDetailsView())
+			{
+				DetailsView->OnFinishedChangingProperties().AddUObject(this, &UStateTreeEditorMode::OnSelectionFinishedChangingProperties);
+
+				DetailsView->SetExtensionHandler(DetailsViewExtensionHandler);
+				DetailsView->SetChildrenCustomizationHandler(DetailsViewChildrenCustomizationHandler);
+			}
+			
+			if (TSharedPtr<IDetailsView> AssetDetailsView = GetAssetDetailsView())
+			{
+				AssetDetailsView->OnFinishedChangingProperties().AddUObject(this, &UStateTreeEditorMode::OnAssetFinishedChangingProperties);
+				
+				AssetDetailsView->SetExtensionHandler(DetailsViewExtensionHandler);
+				AssetDetailsView->SetChildrenCustomizationHandler(DetailsViewChildrenCustomizationHandler);
+			}
+		}
+	}
+
+	UE::StateTree::Delegates::OnIdentifierChanged.AddUObject(this, &UStateTreeEditorMode::OnIdentifierChanged);
+	UE::StateTree::Delegates::OnSchemaChanged.AddUObject(this, &UStateTreeEditorMode::OnSchemaChanged);
+	UE::StateTree::Delegates::OnParametersChanged.AddUObject(this, &UStateTreeEditorMode::OnRefreshDetailsView);
+	UE::StateTree::Delegates::OnGlobalDataChanged.AddUObject(this, &UStateTreeEditorMode::OnRefreshDetailsView);
+	UE::StateTree::Delegates::OnStateParametersChanged.AddUObject(this, &UStateTreeEditorMode::OnStateParametersChanged);
+
+	OnStateTreeChanged();
+}
+
+
+void UStateTreeEditorMode::OnIdentifierChanged(const UStateTree& InStateTree)
+{
+	if (GetStateTree() == &InStateTree)
+	{
+		UpdateAsset();
+	}
+}
+
+void UStateTreeEditorMode::OnSchemaChanged(const UStateTree& InStateTree)
+{
+	if (GetStateTree() == &InStateTree)
+	{
+		UpdateAsset();
+
+		if(UStateTreeEditingSubsystem* StateTreeEditingSubsystem = GEditor->GetEditorSubsystem<UStateTreeEditingSubsystem>())
+		{
+			TSharedRef<FStateTreeViewModel> ViewModel = StateTreeEditingSubsystem->FindOrAddViewModel(GetStateTree());		
+			ViewModel->NotifyAssetChangedExternally();
+		}
+		
+		ForceRefreshDetailsView();
+	}
+}
+
+void UStateTreeEditorMode::ForceRefreshDetailsView() const
+{ 					
+	if (TSharedPtr<IDetailsView> DetailsView = GetDetailsView())
+	{
+		DetailsView->ForceRefresh();
+	}
+}
+
+void UStateTreeEditorMode::OnRefreshDetailsView(const UStateTree& InStateTree) const
+{
+	if (GetStateTree() == &InStateTree)
+	{
+		// Accessible structs might be different after modifying parameters so forcing refresh
+		// so the FStateTreeBindingExtension can rebuild the list of bindable structs
+		ForceRefreshDetailsView();
+	}
+}
+
+void UStateTreeEditorMode::OnStateParametersChanged(const UStateTree& InStateTree, const FGuid ChangedStateID) const
+{
+	UStateTree* StateTree = GetStateTree(); 
+	if (StateTree == &InStateTree)
+	{
+		if (const UStateTreeEditorData* TreeData = Cast<UStateTreeEditorData>(StateTree->EditorData))
+		{
+			TreeData->VisitHierarchy([&ChangedStateID](UStateTreeState& State, UStateTreeState* /*ParentState*/)
+			{
+				if (State.Type == EStateTreeStateType::Linked && State.LinkedSubtree.ID == ChangedStateID)
+				{
+					State.UpdateParametersFromLinkedSubtree();
+				}
+				return EStateTreeVisitor::Continue;
+			});
+		}
+
+		// Accessible structs might be different after modifying parameters so forcing refresh
+		// so the FStateTreeBindingExtension can rebuild the list of bindable structs
+		ForceRefreshDetailsView();
+	}
+}
+
+void UStateTreeEditorMode::HandleMessageTokenClicked(const TSharedRef<IMessageToken>& InMessageToken) const
+{
+	if (InMessageToken->GetType() == EMessageToken::Object)
+	{
+		const TSharedRef<FUObjectToken> ObjectToken = StaticCastSharedRef<FUObjectToken>(InMessageToken);
+
+		if (UStateTreeState* State = Cast<UStateTreeState>(ObjectToken->GetObject().Get()))
+		{
+			if(UStateTreeEditingSubsystem* StateTreeEditingSubsystem = GEditor->GetEditorSubsystem<UStateTreeEditingSubsystem>())
+			{
+				StateTreeEditingSubsystem->FindOrAddViewModel(GetStateTree())->SetSelection(State);
+			}
+			
+		}
+	}
+}
+
+void UStateTreeEditorMode::Exit()
+{
+	if (Toolkit.IsValid())
+	{
+		FToolkitManager::Get().CloseToolkit(Toolkit.ToSharedRef());
+		Toolkit.Reset();
+	}
+	
+	if (UContextObjectStore* ContextObjectStore = GetToolManager()->GetContextObjectStore())
+	{
+		if (const UStateTreeEditorContext* Context = ContextObjectStore->FindContext<UStateTreeEditorContext>())
+		{
+			Context->EditorHostInterface->OnStateTreeChanged().RemoveAll(this);
+			
+			if (TSharedPtr<IMessageLogListing> MessageLogListing = GetMessageLogListing())
+			{
+				MessageLogListing->OnMessageTokenClicked().RemoveAll(this);
+			}
+			
+			if (TSharedPtr<IDetailsView> DetailsView = GetDetailsView())
+			{
+				DetailsView->OnFinishedChangingProperties().RemoveAll(this);				
+				DetailsView->SetExtensionHandler(nullptr);
+				DetailsView->SetChildrenCustomizationHandler(nullptr);
+			}
+			
+			if (TSharedPtr<IDetailsView> AssetDetailsView = GetAssetDetailsView())
+			{
+				AssetDetailsView->OnFinishedChangingProperties().RemoveAll(this);				
+				AssetDetailsView->SetExtensionHandler(nullptr);
+				AssetDetailsView->SetChildrenCustomizationHandler(nullptr);
+			}
+		}
+	}
+
+	if (CachedStateTree.IsValid())
+	{
+		if (UStateTreeEditingSubsystem* StateTreeEditingSubsystem = GEditor->GetEditorSubsystem<UStateTreeEditingSubsystem>())
+		{		
+			TSharedRef<FStateTreeViewModel> ViewModel = StateTreeEditingSubsystem->FindOrAddViewModel(CachedStateTree.Get());				
+			{
+				ViewModel->GetOnAssetChanged().RemoveAll(this);
+				ViewModel->GetOnStateAdded().RemoveAll(this);
+				ViewModel->GetOnStatesRemoved().RemoveAll(this);
+				ViewModel->GetOnStatesMoved().RemoveAll(this);
+				ViewModel->GetOnSelectionChanged().RemoveAll(this);
+				ViewModel->GetOnBringNodeToFocus().RemoveAll(this);
+			}
+		}			
+	}
+
+	UE::StateTree::Delegates::OnIdentifierChanged.RemoveAll(this);
+	UE::StateTree::Delegates::OnSchemaChanged.RemoveAll(this);
+	UE::StateTree::Delegates::OnParametersChanged.RemoveAll(this);
+	UE::StateTree::Delegates::OnGlobalDataChanged.RemoveAll(this);
+	UE::StateTree::Delegates::OnStateParametersChanged.RemoveAll(this);
+	
+	Super::Exit();
+}
+
+void UStateTreeEditorMode::CreateToolkit()
+{
+	Toolkit = MakeShareable(new FStateTreeEditorModeToolkit(this));
+}
+
+void UStateTreeEditorMode::OnStateTreeChanged()
+{
+	UContextObjectStore* ContextStore = GetInteractiveToolsContext()->ToolManager->GetContextObjectStore();
+	if (const UStateTreeEditorContext* Context = ContextStore->FindContext<UStateTreeEditorContext>())
+	{
+		if (UStateTreeEditingSubsystem* StateTreeEditingSubsystem = GEditor->GetEditorSubsystem<UStateTreeEditingSubsystem>())
+		{
+			if (CachedStateTree.IsValid())
+			{
+				TSharedRef<FStateTreeViewModel> OldViewModel = StateTreeEditingSubsystem->FindOrAddViewModel(CachedStateTree.Get());				
+				{
+					OldViewModel->GetOnAssetChanged().RemoveAll(this);
+					OldViewModel->GetOnStateAdded().RemoveAll(this);
+					OldViewModel->GetOnStatesRemoved().RemoveAll(this);
+					OldViewModel->GetOnStatesMoved().RemoveAll(this);
+					OldViewModel->GetOnSelectionChanged().RemoveAll(this);
+					OldViewModel->GetOnBringNodeToFocus().RemoveAll(this);
+				}
+			}			
+		}
+		
+		CachedStateTree = Context->EditorHostInterface->GetStateTree();
+
+		if (CachedStateTree.IsValid())
+		{
+			if (UStateTreeEditingSubsystem* StateTreeEditingSubsystem = GEditor->GetEditorSubsystem<UStateTreeEditingSubsystem>())
+			{				
+				TSharedRef<FStateTreeViewModel> NewViewModel = StateTreeEditingSubsystem->FindOrAddViewModel(CachedStateTree.Get());
+				{
+					NewViewModel->GetOnAssetChanged().AddUObject(this, &UStateTreeEditorMode::HandleModelAssetChanged);
+					NewViewModel->GetOnStateAdded().AddUObject(this, &UStateTreeEditorMode::HandleStateAdded);
+					NewViewModel->GetOnStatesRemoved().AddUObject(this, &UStateTreeEditorMode::HandleStatesRemoved);
+					NewViewModel->GetOnStatesMoved().AddUObject(this, &UStateTreeEditorMode::HandleOnStatesMoved);
+					NewViewModel->GetOnSelectionChanged().AddUObject(this, &UStateTreeEditorMode::HandleModelSelectionChanged);
+					NewViewModel->GetOnBringNodeToFocus().AddUObject(this, &UStateTreeEditorMode::HandleModelBringNodeToFocus);
+				}
+			}
+		}
+	}
+
+	if (Toolkit)
+	{
+		StaticCastSharedPtr<FStateTreeEditorModeToolkit>(Toolkit)->OnStateTreeChanged();
+	}
+}
+
+
+namespace UE::StateTree::Editor::Internal
+{
+static void SetSaveOnCompileSetting(const EStateTreeSaveOnCompile NewSetting)
+{
+	UStateTreeEditorSettings* Settings = GetMutableDefault<UStateTreeEditorSettings>();
+	Settings->SaveOnCompile = NewSetting;
+	Settings->SaveConfig();
+}
+
+static bool IsSaveOnCompileOptionSet(const EStateTreeSaveOnCompile Option)
+{
+	const UStateTreeEditorSettings* Settings = GetDefault<UStateTreeEditorSettings>();
+	return (Settings->SaveOnCompile == Option);
+}
+}
+
+void UStateTreeEditorMode::BindToolkitCommands(const TSharedRef<FUICommandList>& ToolkitCommands)
+{
+	FStateTreeEditorCommands::Register();
+	const FStateTreeEditorCommands& Commands = FStateTreeEditorCommands::Get();
+
+	ToolkitCommands->MapAction(
+		Commands.Compile,
+		FExecuteAction::CreateUObject(this, &UStateTreeEditorMode::Compile),
+		FCanExecuteAction::CreateUObject(this, &UStateTreeEditorMode::CanCompile),
+		FIsActionChecked(),
+		FIsActionButtonVisible::CreateUObject(this, &UStateTreeEditorMode::HasValidStateTree));
+
+	ToolkitCommands->MapAction(
+		FStateTreeEditorCommands::Get().SaveOnCompile_Never,
+		FExecuteAction::CreateStatic(&UE::StateTree::Editor::Internal::SetSaveOnCompileSetting, EStateTreeSaveOnCompile::Never),
+		FCanExecuteAction(),
+		FIsActionChecked::CreateStatic(&UE::StateTree::Editor::Internal::IsSaveOnCompileOptionSet, EStateTreeSaveOnCompile::Never),
+		FIsActionButtonVisible::CreateUObject(this, &UStateTreeEditorMode::HasValidStateTree)
+	);
+	ToolkitCommands->MapAction(
+		FStateTreeEditorCommands::Get().SaveOnCompile_SuccessOnly,
+		FExecuteAction::CreateStatic(&UE::StateTree::Editor::Internal::SetSaveOnCompileSetting, EStateTreeSaveOnCompile::SuccessOnly),
+		FCanExecuteAction(),
+		FIsActionChecked::CreateStatic(&UE::StateTree::Editor::Internal::IsSaveOnCompileOptionSet,  EStateTreeSaveOnCompile::SuccessOnly),
+		FIsActionButtonVisible::CreateUObject(this, &UStateTreeEditorMode::HasValidStateTree)
+	);
+	ToolkitCommands->MapAction(
+		FStateTreeEditorCommands::Get().SaveOnCompile_Always,
+		FExecuteAction::CreateStatic(&UE::StateTree::Editor::Internal::SetSaveOnCompileSetting, EStateTreeSaveOnCompile::Always),
+		FCanExecuteAction(),
+		FIsActionChecked::CreateStatic(&UE::StateTree::Editor::Internal::IsSaveOnCompileOptionSet,  EStateTreeSaveOnCompile::Always),
+		FIsActionButtonVisible::CreateUObject(this, &UStateTreeEditorMode::HasValidStateTree)
+	);
+}
+
+void UStateTreeEditorMode::BindCommands()
+{
+	UEdMode::BindCommands();
+	const TSharedRef<FUICommandList>& CommandList = Toolkit->GetToolkitCommands();
+	BindToolkitCommands(CommandList);
+}
+
+void UStateTreeEditorMode::Compile()
+{
+	UStateTree* StateTree = GetStateTree();
+
+	if (!StateTree)
+	{
+		return;
+	}
+
+	UpdateAsset();
+
+	if (TSharedPtr<IMessageLogListing> Listing = GetMessageLogListing())
+	{
+		Listing->ClearMessages();
+	}
+	
+	FStateTreeCompilerLog Log;
+	bLastCompileSucceeded = UStateTreeEditingSubsystem::CompileStateTree(StateTree, Log);
+
+	if (TSharedPtr<IMessageLogListing> Listing = GetMessageLogListing())
+	{					
+		Log.AppendToLog(Listing.Get());
+
+		if (!bLastCompileSucceeded)
+		{
+			// Show log
+			ShowCompilerTab();
+		}
+	}
+	
+
+	const UStateTreeEditorSettings* Settings = GetMutableDefault<UStateTreeEditorSettings>();
+	const bool bShouldSaveOnCompile = ((Settings->SaveOnCompile == EStateTreeSaveOnCompile::Always)
+									|| ((Settings->SaveOnCompile == EStateTreeSaveOnCompile::SuccessOnly) && bLastCompileSucceeded));
+
+	if (bShouldSaveOnCompile)
+	{
+		const TArray<UPackage*> PackagesToSave { StateTree->GetOutermost() };
+		FEditorFileUtils::PromptForCheckoutAndSave(PackagesToSave, /*bCheckDirty =*/true, /*bPromptToSave =*/false);
+	}
+}
+
+bool UStateTreeEditorMode::CanCompile() const
+{
+	if (GetStateTree() == nullptr)
+	{
+		return false;
+	}
+    
+	// We can't recompile while in PIE
+	if (GEditor->IsPlaySessionInProgress())
+	{
+		return false;
+	}
+
+	return true;
+}
+
+bool UStateTreeEditorMode::HasValidStateTree() const
+{
+	return GetStateTree() != nullptr;
+}
+
+void UStateTreeEditorMode::HandleModelAssetChanged()
+{
+	UpdateAsset();
+}
+
+void UStateTreeEditorMode::HandleModelSelectionChanged(const TArray<TWeakObjectPtr<UStateTreeState>>& SelectedStates) const
+{
+	if (TSharedPtr<IDetailsView> DetailsView = GetDetailsView())
+	{
+		TArray<UObject*> Selected;
+		for (const TWeakObjectPtr<UStateTreeState>& WeakState : SelectedStates)
+		{
+			if (UStateTreeState* State = WeakState.Get())
+			{
+				Selected.Add(State);
+			}
+		}
+		DetailsView->SetObjects(Selected);
+	}
+}
+
+void UStateTreeEditorMode::HandleModelBringNodeToFocus(const UStateTreeState* State, const FGuid NodeID) const
+{
+	TSharedPtr<IDetailsView> DetailsView = GetDetailsView();
+	if (DetailsView && State)
+	{
+		FPropertyPath HighlightPath;
+	
+		if (!HighlightPath.IsValid())
+		{
+			FArrayProperty* TasksProperty = CastFieldChecked<FArrayProperty>(UStateTreeState::StaticClass()->FindPropertyByName(GET_MEMBER_NAME_CHECKED(UStateTreeState, Tasks)));
+			const int32 TaskIndex = State->Tasks.IndexOfByPredicate([&NodeID](const FStateTreeEditorNode& Node)
+			{
+				return Node.ID == NodeID;
+			});
+			if (TaskIndex != INDEX_NONE)
+			{
+				HighlightPath.AddProperty(FPropertyInfo(TasksProperty));
+				HighlightPath.AddProperty(FPropertyInfo(TasksProperty->Inner, TaskIndex));
+			}
+		}
+	
+		if (!HighlightPath.IsValid())
+		{
+			FProperty* SingleTaskProperty = CastFieldChecked<FProperty>(UStateTreeState::StaticClass()->FindPropertyByName(GET_MEMBER_NAME_CHECKED(UStateTreeState, SingleTask)));
+			if (State->SingleTask.ID == NodeID)
+			{
+				HighlightPath.AddProperty(FPropertyInfo(SingleTaskProperty));
+			}
+		}
+	
+		if (!HighlightPath.IsValid())
+		{
+			FArrayProperty* TransitionsProperty = CastFieldChecked<FArrayProperty>(UStateTreeState::StaticClass()->FindPropertyByName(GET_MEMBER_NAME_CHECKED(UStateTreeState, Transitions)));
+			const int32 TransitionIndex = State->Transitions.IndexOfByPredicate([&NodeID](const FStateTreeTransition& Transition)
+			{
+				return Transition.ID == NodeID;
+			});
+			if (TransitionIndex != INDEX_NONE)
+			{
+				HighlightPath.AddProperty(FPropertyInfo(TransitionsProperty));
+				HighlightPath.AddProperty(FPropertyInfo(TransitionsProperty->Inner, TransitionIndex));
+			}
+		}
+	
+		if (!HighlightPath.IsValid())
+		{
+			FArrayProperty* EnterConditionsProperty = CastFieldChecked<FArrayProperty>(UStateTreeState::StaticClass()->FindPropertyByName(GET_MEMBER_NAME_CHECKED(UStateTreeState, EnterConditions)));
+			const int32 EnterConditionIndex = State->EnterConditions.IndexOfByPredicate([&NodeID](const FStateTreeEditorNode& Node)
+			{
+				return Node.ID == NodeID;
+			});
+			if (EnterConditionIndex != INDEX_NONE)
+			{
+				HighlightPath.AddProperty(FPropertyInfo(EnterConditionsProperty));
+				HighlightPath.AddProperty(FPropertyInfo(EnterConditionsProperty->Inner, EnterConditionIndex));
+			}
+		}
+	
+		if (HighlightPath.IsValid())
+		{
+			DetailsView->ScrollPropertyIntoView(HighlightPath, /*bExpandProperty*/true);
+			DetailsView->HighlightProperty(HighlightPath);
+
+			FTimerHandle HighlightTimerHandle;
+			GEditor->GetTimerManager()->SetTimer(
+				HighlightTimerHandle,
+				FTimerDelegate::CreateLambda([SelectionDetailsView = DetailsView]()
+				{
+					SelectionDetailsView->HighlightProperty({});
+				}),
+				1.0f,
+				/*Loop*/false);
+		}
+	}
+}
+
+void UStateTreeEditorMode::UpdateAsset()
+{
+	UStateTree* StateTree = GetStateTree();
+	if (!StateTree)
+	{
+		return;
+	}
+
+	UStateTreeEditingSubsystem::ValidateStateTree(StateTree);
+	EditorDataHash = UStateTreeEditingSubsystem::CalculateStateTreeHash(StateTree);
+}
+
+TSharedPtr<IDetailsView> UStateTreeEditorMode::GetDetailsView() const
+{
+	if (const UContextObjectStore* ContextObjectStore = GetToolManager()->GetContextObjectStore())
+	{
+		if (const UStateTreeEditorContext* Context = ContextObjectStore->FindContext<UStateTreeEditorContext>())
+		{
+			return Context->EditorHostInterface->GetDetailsView();
+		}
+	}
+
+	return nullptr;
+}
+
+TSharedPtr<IDetailsView> UStateTreeEditorMode::GetAssetDetailsView() const
+{
+	if (UContextObjectStore* ContextObjectStore = GetToolManager()->GetContextObjectStore())
+	{
+		if (const UStateTreeEditorContext* Context = ContextObjectStore->FindContext<UStateTreeEditorContext>())
+		{
+			return Context->EditorHostInterface->GetAssetDetailsView();
+		}
+	}
+
+	return nullptr;
+}
+
+TSharedPtr<IMessageLogListing> UStateTreeEditorMode::GetMessageLogListing() const
+{
+	if (UContextObjectStore* ContextObjectStore = GetToolManager()->GetContextObjectStore())
+	{
+		if (const UStateTreeEditorContext* Context = ContextObjectStore->FindContext<UStateTreeEditorContext>())
+		{
+			FMessageLogModule& MessageLogModule = FModuleManager::LoadModuleChecked<FMessageLogModule>("MessageLog");
+			return MessageLogModule.GetLogListing(Context->EditorHostInterface->GetCompilerLogName());
+		}
+	}
+
+	return nullptr;
+}
+
+void UStateTreeEditorMode::ShowCompilerTab() const
+{
+	if (UContextObjectStore* ContextObjectStore = GetToolManager()->GetContextObjectStore())
+	{
+		if (const UStateTreeEditorContext* Context = ContextObjectStore->FindContext<UStateTreeEditorContext>())
+		{
+			if(TSharedPtr<FTabManager> TabManager = GetModeManager()->GetToolkitHost()->GetTabManager())
+			{
+				TabManager->TryInvokeTab(Context->EditorHostInterface->GetCompilerTabName());
+			}
+		}
+	}
+}
+
+UStateTree* UStateTreeEditorMode::GetStateTree() const
+{
+	return CachedStateTree.Get();
+}
+
+void UStateTreeEditorMode::OnAssetFinishedChangingProperties(const FPropertyChangedEvent& PropertyChangedEvent) const
+{
+	// Make sure nodes get updates when properties are changed.
+	if(UStateTreeEditingSubsystem* StateTreeEditingSubsystem = GEditor->GetEditorSubsystem<UStateTreeEditingSubsystem>())
+	{
+		const int32 NumEditedObjects = PropertyChangedEvent.GetNumObjectsBeingEdited();
+		if (NumEditedObjects > 0)
+		{			
+			for (int32 Index = 0; Index < NumEditedObjects; ++Index)
+			{
+				if (const UStateTree* EditedStateTree = Cast<UStateTree>(PropertyChangedEvent.GetObjectBeingEdited(Index)))
+				{
+					if (EditedStateTree == GetStateTree())
+					{
+						StateTreeEditingSubsystem->FindOrAddViewModel(GetStateTree())->NotifyAssetChangedExternally();
+						break;
+					}
+				}
+			}
+		}
+	}
+}
+
+void UStateTreeEditorMode::OnSelectionFinishedChangingProperties(const FPropertyChangedEvent& PropertyChangedEvent)
+{
+	// Make sure nodes get updates when properties are changed.
+	if(UStateTreeEditingSubsystem* StateTreeEditingSubsystem = GEditor->GetEditorSubsystem<UStateTreeEditingSubsystem>())
+	{
+		if (TSharedPtr<IDetailsView> DetailsView = GetDetailsView())
+		{
+			const TArray<TWeakObjectPtr<UObject>> SelectedObjects = GetDetailsView()->GetSelectedObjects();
+			TSet<UStateTreeState*> ChangedStates;
+			for (const TWeakObjectPtr<UObject>& WeakObject : SelectedObjects)
+			{
+				if (UObject* Object = WeakObject.Get())
+				{
+					if (UStateTreeState* State = Cast<UStateTreeState>(Object))
+					{
+						ChangedStates.Add(State);
+					}
+				}
+			}
+			if (ChangedStates.Num() > 0)
+			{
+				StateTreeEditingSubsystem->FindOrAddViewModel(GetStateTree())->NotifyStatesChangedExternally(ChangedStates, PropertyChangedEvent);
+				UpdateAsset();
+			}
+		}
+	}
+}
+
+#undef LOCTEXT_NAMESPACE // "UStateTreeEditorMode"
diff --git a/Engine/Plugins/Runtime/StateTree/Source/StateTreeEditorModule/Private/StateTreeEditorModeToolkit.cpp b/Engine/Plugins/Runtime/StateTree/Source/StateTreeEditorModule/Private/StateTreeEditorModeToolkit.cpp
new file mode 100644
--- /dev/null
+++ b/Engine/Plugins/Runtime/StateTree/Source/StateTreeEditorModule/Private/StateTreeEditorModeToolkit.cpp
@@ -0,0 +1,506 @@
+// Copyright Epic Games, Inc. All Rights Reserved.
+
+#include "StateTreeEditorModeToolkit.h"
+
+#include "ClassViewerModule.h"
+#include "ContextObjectStore.h"
+#include "InteractiveToolManager.h"
+#include "StateTreeEditorCommands.h"
+#include "Blueprint/StateTreeConditionBlueprintBase.h"
+#include "Blueprint/StateTreeConsiderationBlueprintBase.h"
+#include "Blueprint/StateTreeNodeBlueprintBase.h"
+#include "Blueprint/StateTreeTaskBlueprintBase.h"
+
+#include "ClassViewerFilter.h"
+#include "ContentBrowserModule.h"
+#include "IContentBrowserSingleton.h"
+#include "SStateTreeOutliner.h"
+#include "StateTreeEditingSubsystem.h"
+#include "AssetRegistry/AssetRegistryModule.h"
+#include "Debugger/SStateTreeDebuggerView.h"
+#include "Kismet2/BlueprintEditorUtils.h"
+#include "Kismet2/KismetEditorUtilities.h"
+#include "Toolkits/AssetEditorModeUILayer.h"
+#include "Widgets/Input/SMultiLineEditableTextBox.h"
+#include "Widgets/Docking/SDockTab.h"
+#include "ToolMenus.h"
+#include "Subsystems/AssetEditorSubsystem.h"
+#include "Framework/Application/SlateApplication.h"
+
+#define LOCTEXT_NAMESPACE "StateTreeModeToolkit"
+
+const FName FStateTreeEditorModeToolkit::StateTreeOutlinerTabId(TEXT("StateTreeEditor_StateTreeOutliner"));		
+const FName FStateTreeEditorModeToolkit::StateTreeStatisticsTabId(TEXT("StateTreeEditor_StateTreeStatistics"));
+#if WITH_STATETREE_TRACE_DEBUGGER
+const FName FStateTreeEditorModeToolkit::DebuggerTabId(TEXT("StateTreeEditor_Debugger"));
+#endif // WITH_STATETREE_TRACE_DEBUGGER
+
+FStateTreeEditorModeToolkit::FStateTreeEditorModeToolkit(UStateTreeEditorMode* InEditorMode) : WeakEditorMode(InEditorMode)
+{
+}
+
+void FStateTreeEditorModeToolkit::RequestModeUITabs()
+{
+	FModeToolkit::RequestModeUITabs();
+	
+	if (ModeUILayer.IsValid())
+	{
+		TSharedPtr<FAssetEditorModeUILayer> ModeUILayerPtr = ModeUILayer.Pin();
+		TSharedPtr<FWorkspaceItem> MenuGroup = ModeUILayerPtr->GetModeMenuCategory();		
+		if(!MenuGroup)
+		{ 
+			return;
+		}
+
+		FMinorTabConfig OutlinerTabInfo;
+		OutlinerTabInfo.OnSpawnTab = FOnSpawnTab::CreateLambda([this](const FSpawnTabArgs& Args)-> TSharedRef<SDockTab>
+		{
+			ensure(Args.GetTabId() == UAssetEditorUISubsystem::TopLeftTabID);
+			TSharedRef<SDockTab> Tab = SNew(SDockTab)
+			.Label(LOCTEXT("StateTreeOutlinerTab", "Outliner"))
+			.TabColorScale(GetTabColorScale());
+
+			WeakOutlinerTab = Tab;
+			UpdateStateTreeOutliner();
+
+			return Tab;
+		});
+		OutlinerTabInfo.TabLabel = LOCTEXT("StateTreeOutlinerTab", "Outliner");
+		
+		OutlinerTabInfo.TabId = StateTreeOutlinerTabId;
+		OutlinerTabInfo.TabIcon = FSlateIcon(FAppStyle::GetAppStyleSetName(), "LevelEditor.Tabs.Outliner");
+		OutlinerTabInfo.WorkspaceGroup = MenuGroup;
+		ModeUILayerPtr->SetModePanelInfo(UAssetEditorUISubsystem::TopLeftTabID, OutlinerTabInfo);
+		
+		FMinorTabConfig StatisticsTabInfo;
+		StatisticsTabInfo.OnSpawnTab = FOnSpawnTab::CreateLambda([this](const FSpawnTabArgs& Args)-> TSharedRef<SDockTab>
+		{
+			ensure(Args.GetTabId() == UAssetEditorUISubsystem::BottomRightTabID);
+			return SNew(SDockTab)
+			.Label(LOCTEXT("StatisticsTitle", "StateTree Statistics"))
+			[
+				SNew(SMultiLineEditableTextBox)
+				.Padding(10.0f)
+				.Style(FAppStyle::Get(), "Log.TextBox")
+				.Font(FCoreStyle::GetDefaultFontStyle("Mono", 9))
+				.ForegroundColor(FLinearColor::Gray)
+				.IsReadOnly(true)
+				.Text(this, &FStateTreeEditorModeToolkit::GetStatisticsText)
+			];
+		});
+		StatisticsTabInfo.TabLabel = LOCTEXT("StatisticsTitle", "StateTree Statistics");
+		StatisticsTabInfo.TabIcon = FSlateIcon(FAppStyle::GetAppStyleSetName(), "AssetEditor.ToggleStats");	
+		StatisticsTabInfo.TabId = StateTreeStatisticsTabId;
+		StatisticsTabInfo.WorkspaceGroup = MenuGroup;	
+		
+		ModeUILayerPtr->SetModePanelInfo(UAssetEditorUISubsystem::BottomRightTabID, StatisticsTabInfo);
+
+
+#if WITH_STATETREE_TRACE_DEBUGGER
+		FMinorTabConfig DebuggerTabInfo;
+		DebuggerTabInfo.OnSpawnTab = FOnSpawnTab::CreateLambda([this](const FSpawnTabArgs& Args)-> TSharedRef<SDockTab>
+		{
+			ensure(Args.GetTabId() == UAssetEditorUISubsystem::TopRightTabID);
+			TSharedRef<SDockTab> Tab = SNew(SDockTab)
+			.Label(LOCTEXT("DebuggerTitle", "Debugger"))
+			.TabColorScale(GetTabColorScale());
+
+			WeakDebuggerTab = Tab;
+			UpdateDebuggerView();
+
+			return Tab;
+		});
+		DebuggerTabInfo.TabLabel = LOCTEXT("DebuggerTab", "Debugger");		
+		DebuggerTabInfo.TabId = DebuggerTabId;
+		DebuggerTabInfo.TabIcon = FSlateIcon(FAppStyle::GetAppStyleSetName(), "Debug");
+		DebuggerTabInfo.WorkspaceGroup = MenuGroup;
+		ModeUILayerPtr->SetModePanelInfo(UAssetEditorUISubsystem::TopRightTabID, DebuggerTabInfo);
+#endif // WITH_STATETREE_TRACE_DEBUGGER
+	}
+}
+
+void FStateTreeEditorModeToolkit::Init(const TSharedPtr<IToolkitHost>& InitToolkitHost, TWeakObjectPtr<UEdMode> InOwningMode)
+{
+	FModeToolkit::Init(InitToolkitHost, InOwningMode);
+	
+	if (UContextObjectStore* ContextObjectStore = InOwningMode->GetToolManager()->GetContextObjectStore())
+	{
+		if (UStateTreeEditorContext* Context = ContextObjectStore->FindContext<UStateTreeEditorContext>())
+		{
+			EditorHost = Context->EditorHostInterface;
+		}
+	}
+	
+}
+
+void FStateTreeEditorModeToolkit::InvokeUI()
+{
+	if (ModeUILayer.IsValid())
+	{
+		TSharedPtr<FAssetEditorModeUILayer> ModeUILayerPtr = ModeUILayer.Pin();
+		ModeUILayerPtr->GetTabManager()->TryInvokeTab(UAssetEditorUISubsystem::TopLeftTabID);
+		ModeUILayerPtr->GetTabManager()->TryInvokeTab(UAssetEditorUISubsystem::BottomRightTabID);
+	}
+}
+
+FName FStateTreeEditorModeToolkit::GetToolkitFName() const
+{
+	return FName("StateTreeMode");
+}
+
+FText FStateTreeEditorModeToolkit::GetBaseToolkitName() const
+{
+	return LOCTEXT("ToolkitName", "State Tree Mode");
+}
+
+
+FSlateIcon FStateTreeEditorModeToolkit::GetCompileStatusImage() const
+{
+	static const FName CompileStatusBackground("Blueprint.CompileStatus.Background");
+	static const FName CompileStatusUnknown("Blueprint.CompileStatus.Overlay.Unknown");
+	static const FName CompileStatusError("Blueprint.CompileStatus.Overlay.Error");
+	static const FName CompileStatusGood("Blueprint.CompileStatus.Overlay.Good");
+
+	UStateTree* StateTree = WeakEditorMode.Pin()->GetStateTree();
+	UStateTreeEditorMode* SharedEditorMode = WeakEditorMode.Get();
+	
+	if (StateTree && SharedEditorMode)
+	{
+		const bool bCompiledDataResetDuringLoad = StateTree->LastCompiledEditorDataHash == SharedEditorMode->EditorDataHash && !StateTree->IsReadyToRun();
+
+		if (!SharedEditorMode->bLastCompileSucceeded || bCompiledDataResetDuringLoad)
+		{
+			return FSlateIcon(FAppStyle::GetAppStyleSetName(), CompileStatusBackground, NAME_None, CompileStatusError);
+		}
+	
+		if (StateTree->LastCompiledEditorDataHash != SharedEditorMode->EditorDataHash)
+		{
+			return FSlateIcon(FAppStyle::GetAppStyleSetName(), CompileStatusBackground, NAME_None, CompileStatusUnknown);
+		}
+	
+		return FSlateIcon(FAppStyle::GetAppStyleSetName(), CompileStatusBackground, NAME_None, CompileStatusGood);
+	}
+	
+	return FSlateIcon(FAppStyle::GetAppStyleSetName(), CompileStatusBackground, NAME_None, CompileStatusUnknown);
+}
+
+
+namespace UE::StateTree::Editor::Internal
+{
+static void MakeSaveOnCompileSubMenu(UToolMenu* InMenu)
+{
+	FToolMenuSection& Section = InMenu->AddSection("Section");
+	const FStateTreeEditorCommands& Commands = FStateTreeEditorCommands::Get();
+	Section.AddMenuEntry(Commands.SaveOnCompile_Never);
+	Section.AddMenuEntry(Commands.SaveOnCompile_SuccessOnly);
+	Section.AddMenuEntry(Commands.SaveOnCompile_Always);
+}
+
+static void GenerateCompileOptionsMenu(UToolMenu* InMenu)
+{
+	FToolMenuSection& Section = InMenu->AddSection("Section");
+
+	// @TODO: disable the menu and change up the tooltip when all sub items are disabled
+	Section.AddSubMenu(
+		"SaveOnCompile",
+		LOCTEXT("SaveOnCompileSubMenu", "Save on Compile"),
+		LOCTEXT("SaveOnCompileSubMenu_ToolTip", "Determines how the StateTree is saved whenever you compile it."),
+		FNewToolMenuDelegate::CreateStatic(&MakeSaveOnCompileSubMenu));
+}
+}
+
+//
+void FStateTreeEditorModeToolkit::ExtendSecondaryModeToolbar(UToolMenu* ToolBar)
+{
+	ToolBar->Context.AppendCommandList(ToolkitCommands);
+
+	const FStateTreeEditorCommands& Commands = FStateTreeEditorCommands::Get();
+	ensure(ToolBar->Context.GetActionForCommand(Commands.Compile));
+
+	static const FToolMenuInsert InsertLast(NAME_None, EToolMenuInsertType::Last);
+
+	FToolMenuSection& CompileSection = ToolBar->AddSection("Compile", FText(), InsertLast);
+
+
+	TWeakPtr<FStateTreeEditorModeToolkit> WeakToolkit = StaticCastSharedRef<FStateTreeEditorModeToolkit>(AsShared());
+	CompileSection.AddDynamicEntry("CompileCommands", FNewToolMenuSectionDelegate::CreateLambda([WeakToolkit, WeakMode = WeakEditorMode](FToolMenuSection& InSection)
+	{
+		if (UStateTreeEditorMode* Mode = WeakMode.Get())
+		{
+			UInteractiveToolManager* ToolManager = Mode->GetToolManager();
+			if (const UContextObjectStore* ContextObjectStore = ToolManager ? ToolManager->GetContextObjectStore() : nullptr)
+			{
+				if (ContextObjectStore->FindContext<UStateTreeEditorContext>())
+				{
+					const FStateTreeEditorCommands& Commands = FStateTreeEditorCommands::Get();
+
+					FToolMenuEntry& CompileButton = InSection.AddEntry(FToolMenuEntry::InitToolBarButton(
+						Commands.Compile,
+						TAttribute<FText>(),
+						TAttribute<FText>(),
+						TAttribute<FSlateIcon>::CreateLambda([WeakToolkit]() -> FSlateIcon
+						{
+							if (TSharedPtr<FStateTreeEditorModeToolkit> SharedThis = WeakToolkit.Pin())
+							{
+								return SharedThis->GetCompileStatusImage();
+							}
+
+							return FSlateIcon();
+						}))
+					);
+					CompileButton.StyleNameOverride = "CalloutToolbar";
+
+					FToolMenuEntry& CompileOptions = InSection.AddEntry(FToolMenuEntry::InitComboButton(
+						"CompileComboButton",
+						FUIAction(),
+						FNewToolMenuDelegate::CreateStatic(&UE::StateTree::Editor::Internal::GenerateCompileOptionsMenu),
+						LOCTEXT("CompileOptions_ToolbarTooltip", "Options to customize how State Trees compile")
+					));
+					CompileOptions.StyleNameOverride = "CalloutToolbar";
+					CompileOptions.ToolBarData.bSimpleComboBox = true;
+				}
+			}
+		}
+	}));
+
+	static const FToolMenuInsert InsertAfterCompileSection("Compile", EToolMenuInsertType::After);
+
+	FToolMenuSection& CreateNewNodeSection = ToolBar->AddSection("CreateNewNodes", TAttribute<FText>(), InsertAfterCompileSection);
+
+	CreateNewNodeSection.AddDynamicEntry("CreateNewNodes", FNewToolMenuSectionDelegate::CreateLambda([this, WeakMode = WeakEditorMode](FToolMenuSection& InSection)
+	{
+		if (UStateTreeEditorMode* Mode = WeakMode.Get())
+		{
+			if (UContextObjectStore* ContextObjectStore = Mode->GetToolManager()->GetContextObjectStore())
+			{
+				if (ContextObjectStore->FindContext<UStateTreeEditorContext>())
+				{
+					InSection.AddEntry(FToolMenuEntry::InitComboButton(
+						 "CreateNewTaskComboButton",
+						 FUIAction(),
+						 FOnGetContent::CreateSP(this, &FStateTreeEditorModeToolkit::GenerateTaskBPBaseClassesMenu),
+						 LOCTEXT("CreateNewTask_Title", "New Task"),
+						 LOCTEXT("CreateNewTask_ToolbarTooltip", "Create a new Blueprint State Tree Task"),
+						 TAttribute<FSlateIcon>::CreateSP(this, &FStateTreeEditorModeToolkit::GetNewTaskButtonImage)
+					 ));
+						
+					InSection.AddEntry(FToolMenuEntry::InitComboButton(
+						 "CreateNewConditionComboButton",
+						 FUIAction(),
+						 FOnGetContent::CreateSP(this, &FStateTreeEditorModeToolkit::GenerateConditionBPBaseClassesMenu),
+						 LOCTEXT("CreateNewCondition_Title", "New Condition"),
+						 LOCTEXT("CreateNewCondition_ToolbarTooltip", "Create a new Blueprint State Tree Condition"),
+						 TAttribute<FSlateIcon>::CreateSP(this, &FStateTreeEditorModeToolkit::GetNewConditionButtonImage)
+					 ));
+						
+					 InSection.AddEntry(FToolMenuEntry::InitComboButton(
+						 "CreateNewConsiderationComboButton",
+						 FUIAction(),
+						 FOnGetContent::CreateSP(this, &FStateTreeEditorModeToolkit::GenerateConsiderationBPBaseClassesMenu),
+						 LOCTEXT("CreateNewConsideration_Title", "New Consideration"),
+						 LOCTEXT("CreateNewConsideration_ToolbarTooltip", "Create a new Blueprint State Tree Utility Consideration"),
+						 TAttribute<FSlateIcon>::CreateSP(this, &FStateTreeEditorModeToolkit::GetNewConsiderationButtonImage)
+					 ));
+				}
+			}
+		}
+	}));
+}
+
+void FStateTreeEditorModeToolkit::OnStateTreeChanged()
+{
+	// Update underlying state tree state
+	UpdateStateTreeOutliner();
+	
+#if WITH_STATETREE_TRACE_DEBUGGER
+	UpdateDebuggerView();	
+#endif // WITH_STATETREE_TRACE_DEBUGGER
+}
+
+
+namespace UE::StateTree::Editor
+{
+template <typename ClassType, typename = typename TEnableIf<TIsDerivedFrom<ClassType, UStateTreeNodeBlueprintBase>::Value>::Type>
+class FEditorNodeClassFilter : public IClassViewerFilter
+{
+public:
+	virtual bool IsClassAllowed(const FClassViewerInitializationOptions& InInitOptions, const UClass* InClass, TSharedRef< FClassViewerFilterFuncs > InFilterFuncs) override
+	{
+		check(InClass);
+		return InClass->IsChildOf(ClassType::StaticClass());
+	}
+
+	virtual bool IsUnloadedClassAllowed(const FClassViewerInitializationOptions& InInitOptions, const TSharedRef< const IUnloadedBlueprintData > InUnloadedClassData, TSharedRef< FClassViewerFilterFuncs > InFilterFuncs) override
+	{
+		return InUnloadedClassData->IsChildOf(ClassType::StaticClass());
+	}
+};
+
+using FStateTreeTaskBPClassFilter = FEditorNodeClassFilter<UStateTreeTaskBlueprintBase>;
+using FStateTreeConditionBPClassFilter = FEditorNodeClassFilter<UStateTreeConditionBlueprintBase>;
+using FStateTreeConsiderationBPClassFilter = FEditorNodeClassFilter<UStateTreeConsiderationBlueprintBase>;
+}; // UE::StateTree::Editor
+
+FSlateIcon FStateTreeEditorModeToolkit::GetNewTaskButtonImage() const
+{
+	//placeholder
+	return FSlateIcon(FAppStyle::GetAppStyleSetName(), NAME_None, NAME_None, NAME_None);
+}
+
+TSharedRef<SWidget> FStateTreeEditorModeToolkit::GenerateTaskBPBaseClassesMenu() const
+{
+	FClassViewerInitializationOptions Options;
+	Options.NameTypeToDisplay = EClassViewerNameTypeToDisplay::DisplayName;
+	Options.ClassFilters.Add(MakeShareable(new UE::StateTree::Editor::FStateTreeTaskBPClassFilter));
+
+	FOnClassPicked OnPicked(FOnClassPicked::CreateSP(this, &FStateTreeEditorModeToolkit::OnNodeBPBaseClassPicked));
+
+	return FModuleManager::LoadModuleChecked<FClassViewerModule>("ClassViewer").CreateClassViewer(Options, OnPicked);
+}
+
+FSlateIcon FStateTreeEditorModeToolkit::GetNewConditionButtonImage() const
+{
+	//placeholder
+    return FSlateIcon(FAppStyle::GetAppStyleSetName(), NAME_None, NAME_None, NAME_None);
+}
+
+TSharedRef<SWidget> FStateTreeEditorModeToolkit::GenerateConditionBPBaseClassesMenu() const
+{
+	FClassViewerInitializationOptions Options;
+	Options.NameTypeToDisplay = EClassViewerNameTypeToDisplay::DisplayName;
+	Options.ClassFilters.Add(MakeShareable(new UE::StateTree::Editor::FStateTreeConditionBPClassFilter));
+
+	FOnClassPicked OnPicked(FOnClassPicked::CreateSP(this, &FStateTreeEditorModeToolkit::OnNodeBPBaseClassPicked));
+
+	return FModuleManager::LoadModuleChecked<FClassViewerModule>("ClassViewer").CreateClassViewer(Options, OnPicked);
+}
+
+FSlateIcon FStateTreeEditorModeToolkit::GetNewConsiderationButtonImage() const
+{
+	//placeholder
+    return FSlateIcon(FAppStyle::GetAppStyleSetName(), NAME_None, NAME_None, NAME_None);
+}
+
+TSharedRef<SWidget> FStateTreeEditorModeToolkit::GenerateConsiderationBPBaseClassesMenu() const
+{
+	FClassViewerInitializationOptions Options;
+    Options.NameTypeToDisplay = EClassViewerNameTypeToDisplay::DisplayName;
+    Options.ClassFilters.Add(MakeShareable(new UE::StateTree::Editor::FStateTreeConsiderationBPClassFilter));
+
+    FOnClassPicked OnPicked(FOnClassPicked::CreateSP(this, &FStateTreeEditorModeToolkit::OnNodeBPBaseClassPicked));
+
+    return FModuleManager::LoadModuleChecked<FClassViewerModule>("ClassViewer").CreateClassViewer(Options, OnPicked);
+}
+
+void FStateTreeEditorModeToolkit::OnNodeBPBaseClassPicked(UClass* NodeClass) const
+{
+	check(NodeClass);
+
+	UStateTree* StateTree = WeakEditorMode->GetStateTree();
+    if (!StateTree)
+    {
+    	return;
+    }
+
+    const FString ClassName = FBlueprintEditorUtils::GetClassNameWithoutSuffix(NodeClass);
+    const FString PathName = FPaths::GetPath(StateTree->GetOutermost()->GetPathName());
+
+    // Now that we've generated some reasonable default locations/names for the package, allow the user to have the final say
+    // before we create the package and initialize the blueprint inside of it.
+    FSaveAssetDialogConfig SaveAssetDialogConfig;
+    SaveAssetDialogConfig.DialogTitleOverride = LOCTEXT("SaveAssetDialogTitle", "Save Asset As");
+    SaveAssetDialogConfig.DefaultPath = PathName;
+    SaveAssetDialogConfig.DefaultAssetName = ClassName + TEXT("_New");
+    SaveAssetDialogConfig.ExistingAssetPolicy = ESaveAssetDialogExistingAssetPolicy::Disallow;
+
+    const FContentBrowserModule& ContentBrowserModule = FModuleManager::LoadModuleChecked<FContentBrowserModule>("ContentBrowser");
+    const FString SaveObjectPath = ContentBrowserModule.Get().CreateModalSaveAssetDialog(SaveAssetDialogConfig);
+    if (!SaveObjectPath.IsEmpty())
+    {
+    	const FString SavePackageName = FPackageName::ObjectPathToPackageName(SaveObjectPath);
+    	const FString SavePackagePath = FPaths::GetPath(SavePackageName);
+    	const FString SaveAssetName = FPaths::GetBaseFilename(SavePackageName);
+
+    	if (UPackage* Package = CreatePackage(*SavePackageName))
+    	{
+    		// Create and init a new Blueprint
+    		if (UBlueprint* NewBP = FKismetEditorUtilities::CreateBlueprint(NodeClass, Package, FName(*SaveAssetName), BPTYPE_Normal, UBlueprint::StaticClass(), UBlueprintGeneratedClass::StaticClass()))
+    		{
+    			GEditor->GetEditorSubsystem<UAssetEditorSubsystem>()->OpenEditorForAsset(NewBP);
+
+    			// Notify the asset registry
+    			FAssetRegistryModule::AssetCreated(NewBP);
+
+    			Package->MarkPackageDirty();
+    		}
+    	}
+    }
+
+    FSlateApplication::Get().DismissAllMenus();
+}
+
+FText FStateTreeEditorModeToolkit::GetStatisticsText() const
+{
+	UStateTree* StateTree = WeakEditorMode->GetStateTree(); 
+	if (!StateTree)
+   {
+	return FText::GetEmpty();
+   }
+
+
+   TArray<FStateTreeMemoryUsage> MemoryUsages = StateTree->CalculateEstimatedMemoryUsage();
+   if (MemoryUsages.IsEmpty())
+   {
+	return FText::GetEmpty();
+   }
+
+	TArray<FText> Rows;
+
+	for (const FStateTreeMemoryUsage& Usage : MemoryUsages)
+	{
+		const FText SizeText = FText::AsMemory(Usage.EstimatedMemoryUsage);
+		const FText NumNodesText = FText::AsNumber(Usage.NodeCount);
+		Rows.Add(FText::Format(LOCTEXT("UsageRow", "{0}: {1}, {2} nodes"), FText::FromString(Usage.Name), SizeText, NumNodesText));
+	}
+
+	return FText::Join(FText::FromString(TEXT("\n")), Rows);
+}
+
+void FStateTreeEditorModeToolkit::UpdateStateTreeOutliner()
+{
+	StateTreeOutliner = SNullWidget::NullWidget;
+	if (WeakEditorMode->GetStateTree())
+	{
+		if (UStateTreeEditingSubsystem* Subsystem = GEditor->GetEditorSubsystem<UStateTreeEditingSubsystem>())
+		{
+			StateTreeOutliner = SNew(SStateTreeOutliner, Subsystem->FindOrAddViewModel(WeakEditorMode->GetStateTree()), GetToolkitCommands());
+		}		
+	}
+	
+	if(TSharedPtr<SDockTab> SharedOutlinerTab = WeakOutlinerTab.Pin())
+	{
+		SharedOutlinerTab->SetContent(StateTreeOutliner.ToSharedRef());
+	}
+}
+
+#if WITH_STATETREE_TRACE_DEBUGGER
+void FStateTreeEditorModeToolkit::UpdateDebuggerView()
+{
+	DebuggerView = SNullWidget::NullWidget;
+	
+	if (UStateTree* StateTree = WeakEditorMode->GetStateTree())
+    {
+    	if (UStateTreeEditingSubsystem* Subsystem = GEditor->GetEditorSubsystem<UStateTreeEditingSubsystem>())
+    	{
+    		DebuggerView = SNew(SStateTreeDebuggerView, *StateTree, Subsystem->FindOrAddViewModel(StateTree), GetToolkitCommands());
+    	}		
+    }
+
+	if(TSharedPtr<SDockTab> SharedDebuggerTab = WeakDebuggerTab.Pin())
+	{
+		SharedDebuggerTab->SetContent(DebuggerView.ToSharedRef());
+	}
+}
+#endif // WITH_STATETREE_TRACE_DEBUGGER
+
+#undef LOCTEXT_NAMESPACE // "StateTreeModeToolkit"
diff --git a/Engine/Plugins/Runtime/StateTree/Source/StateTreeEditorModule/Private/StateTreeEditorModule.cpp b/Engine/Plugins/Runtime/StateTree/Source/StateTreeEditorModule/Private/StateTreeEditorModule.cpp
--- a/Engine/Plugins/Runtime/StateTree/Source/StateTreeEditorModule/Private/StateTreeEditorModule.cpp
+++ b/Engine/Plugins/Runtime/StateTree/Source/StateTreeEditorModule/Private/StateTreeEditorModule.cpp
@@ -25,6 +25,7 @@
 #include "StateTreeCompilerLog.h"
 #include "Debugger/StateTreeDebuggerCommands.h"
 #include "StateTreeDelegates.h"
+#include "StateTreeEditingSubsystem.h"
 #include "StateTreeEditor.h"
 #include "StateTreeEditorCommands.h"
 #include "StateTreeEditorStyle.h"
@@ -42,33 +43,15 @@ namespace UE::StateTree::Editor
 	// @todo Could we make this a IModularFeature?
 	static bool CompileStateTree(UStateTree& StateTree)
 	{
-		// Compile the StateTree asset.
-		UE::StateTree::Editor::ValidateAsset(StateTree);
-		const uint32 EditorDataHash = UE::StateTree::Editor::CalcAssetHash(StateTree);
-
-		FStateTreeCompilerLog Log;
-		FStateTreeCompiler Compiler(Log);
-
-		const bool bSuccess = Compiler.Compile(StateTree);
-
-		if (bSuccess)
-		{
-			// Success
-			StateTree.LastCompiledEditorDataHash = EditorDataHash;
-			UE::StateTree::Delegates::OnPostCompile.Broadcast(StateTree);
-			UE_LOG(LogStateTreeEditor, Log, TEXT("Compile StateTree '%s' succeeded."), *StateTree.GetFullName());
-		}
-		else
+		if (UStateTreeEditingSubsystem* StateTreeEditingSubsystem = GEditor->GetEditorSubsystem<UStateTreeEditingSubsystem>())
 		{
-			// Make sure not to leave stale data on failed compile.
-			StateTree.ResetCompiled();
-			StateTree.LastCompiledEditorDataHash = 0;
+			FStateTreeCompilerLog Log;
+			const bool bSuccess = UStateTreeEditingSubsystem::CompileStateTree(&StateTree, Log);
 
-			UE_LOG(LogStateTreeEditor, Error, TEXT("Failed to compile '%s', errors follow."), *StateTree.GetFullName());
-			Log.DumpToLog(LogStateTreeEditor);
+			return bSuccess;
 		}
 
-		return bSuccess;
+		return false;
 	}
 
 }; // UE::StateTree::Editor
@@ -76,7 +59,7 @@ namespace UE::StateTree::Editor
 void FStateTreeEditorModule::StartupModule()
 {
 	UE::StateTree::Delegates::OnRequestCompile.BindStatic(&UE::StateTree::Editor::CompileStateTree);
-	UE::StateTree::Delegates::OnRequestEditorHash.BindStatic(&UE::StateTree::Editor::CalcAssetHash);
+	UE::StateTree::Delegates::OnRequestEditorHash.BindLambda([](const UStateTree& InStateTree) -> uint32 { return UStateTreeEditingSubsystem::CalculateStateTreeHash(&InStateTree); });
 
 #if WITH_STATETREE_TRACE_DEBUGGER
 	FStateTreeDebuggerCommands::Register();
diff --git a/Engine/Plugins/Runtime/StateTree/Source/StateTreeEditorModule/Private/StateTreeEditorStyle.cpp b/Engine/Plugins/Runtime/StateTree/Source/StateTreeEditorModule/Private/StateTreeEditorStyle.cpp
--- a/Engine/Plugins/Runtime/StateTree/Source/StateTreeEditorModule/Private/StateTreeEditorStyle.cpp
+++ b/Engine/Plugins/Runtime/StateTree/Source/StateTreeEditorModule/Private/StateTreeEditorStyle.cpp
@@ -9,6 +9,7 @@
 #include "Styling/SlateTypes.h"
 #include "Misc/Paths.h"
 #include "Styling/StyleColors.h"
+#include "StateTreeTypes.h"
 #include "Styling/SlateStyleMacros.h"
 
 
@@ -404,3 +405,40 @@ FStateTreeEditorStyle& FStateTreeEditorStyle::Get()
 	static FStateTreeEditorStyle Instance;
 	return Instance;
 }
+
+const FSlateBrush* FStateTreeEditorStyle::GetBrushForSelectionBehaviorType(EStateTreeStateSelectionBehavior InBehaviour, bool bHasChildren, EStateTreeStateType StateType)
+{	
+	if (InBehaviour == EStateTreeStateSelectionBehavior::None)
+	{
+		return Get().GetBrush("StateTreeEditor.SelectNone");
+	}
+	else if (InBehaviour == EStateTreeStateSelectionBehavior::TryEnterState)
+	{
+		return Get().GetBrush("StateTreeEditor.TryEnterState");			
+	}
+	else if (InBehaviour == EStateTreeStateSelectionBehavior::TrySelectChildrenInOrder
+		|| InBehaviour == EStateTreeStateSelectionBehavior::TrySelectChildrenWithHighestUtility
+		|| InBehaviour == EStateTreeStateSelectionBehavior::TrySelectChildrenBasedOnRelativeUtility)
+	{
+		if (!bHasChildren
+			|| StateType == EStateTreeStateType::Linked
+			|| StateType == EStateTreeStateType::LinkedAsset)
+		{
+			return Get().GetBrush("StateTreeEditor.TryEnterState");			
+		}
+		else
+		{
+			return Get().GetBrush("StateTreeEditor.TrySelectChildrenInOrder");
+		}
+	}
+	else if (InBehaviour == EStateTreeStateSelectionBehavior::TrySelectChildrenAtUniformRandom)
+	{
+		return Get().GetBrush("StateTreeEditor.TrySelectChildrenAtRandom");
+	}
+	else if (InBehaviour == EStateTreeStateSelectionBehavior::TryFollowTransitions)
+	{
+		return Get().GetBrush("StateTreeEditor.TryFollowTransitions");
+	}
+
+	return nullptr;
+}
diff --git a/Engine/Plugins/Runtime/StateTree/Source/StateTreeEditorModule/Private/StateTreeEditorUILayer.cpp b/Engine/Plugins/Runtime/StateTree/Source/StateTreeEditorModule/Private/StateTreeEditorUILayer.cpp
new file mode 100644
--- /dev/null
+++ b/Engine/Plugins/Runtime/StateTree/Source/StateTreeEditorModule/Private/StateTreeEditorUILayer.cpp
@@ -0,0 +1,72 @@
+// Copyright Epic Games, Inc. All Rights Reserved.
+
+#include "StateTreeEditorUILayer.h"
+#include "StateTreeEditor.h"
+#include "StateTreeEditorModule.h"
+#include "Toolkits/IToolkit.h"
+#include "ToolMenus.h"
+#include "Modules/ModuleManager.h"
+
+FStateTreeEditorModeUILayer::FStateTreeEditorModeUILayer(const IToolkitHost* InToolkitHost) : FAssetEditorModeUILayer(InToolkitHost)
+{
+}
+
+void FStateTreeEditorModeUILayer::OnToolkitHostingStarted(const TSharedRef<IToolkit>& Toolkit)
+{
+	if (!Toolkit->IsAssetEditor())
+	{
+		FAssetEditorModeUILayer::OnToolkitHostingStarted(Toolkit);
+		HostedToolkit = Toolkit;
+		Toolkit->SetModeUILayer(SharedThis(this));
+		Toolkit->RegisterTabSpawners(ToolkitHost->GetTabManager().ToSharedRef());
+		RegisterModeTabSpawners();	
+
+		OnToolkitHostReadyForUI.Execute();
+
+		// Set up an owner for the current scope so that we can cleanly clean up the toolbar extension on hosting finish
+		FToolMenuOwnerScoped Owner(&(*Toolkit));
+		UToolMenu* SecondaryModeToolbar = UToolMenus::Get()->ExtendMenu(GetSecondaryModeToolbarName());
+		OnRegisterSecondaryModeToolbarExtension.ExecuteIfBound(SecondaryModeToolbar);		
+	}
+}
+
+void FStateTreeEditorModeUILayer::OnToolkitHostingFinished(const TSharedRef<IToolkit>& Toolkit)
+{	
+	if (HostedToolkit.IsValid() && HostedToolkit.Pin() == Toolkit)
+	{
+		FAssetEditorModeUILayer::OnToolkitHostingFinished(Toolkit);
+	}
+}
+
+void FStateTreeEditorModeUILayer::SetModeMenuCategory(const TSharedPtr<FWorkspaceItem>& MenuCategoryIn)
+{
+	MenuCategory = MenuCategoryIn;
+}
+
+TSharedPtr<FWorkspaceItem> FStateTreeEditorModeUILayer::GetModeMenuCategory() const
+{
+	check(MenuCategory);
+	return MenuCategory;
+}
+
+void UStateTreeEditorUISubsystem::Initialize(FSubsystemCollectionBase& Collection)
+{
+	FStateTreeEditorModule& LevelEditorModule = FModuleManager::LoadModuleChecked<FStateTreeEditorModule>("StateTreeEditorModule");
+	LevelEditorModule.OnRegisterLayoutExtensions().AddUObject(this, &UStateTreeEditorUISubsystem::RegisterLayoutExtensions);
+}
+
+void UStateTreeEditorUISubsystem::Deinitialize()
+{
+	FStateTreeEditorModule& LevelEditorModule = FModuleManager::LoadModuleChecked<FStateTreeEditorModule>("StateTreeEditorModule");
+	LevelEditorModule.OnRegisterLayoutExtensions().RemoveAll(this);
+}
+
+void UStateTreeEditorUISubsystem::RegisterLayoutExtensions(FLayoutExtender& Extender)
+{	
+	Extender.ExtendStack(FStateTreeEditor::LayoutLeftStackId, ELayoutExtensionPosition::After, FTabManager::FTab(UAssetEditorUISubsystem::TopLeftTabID, ETabState::ClosedTab));
+	Extender.ExtendStack(FStateTreeEditor::LayoutLeftStackId, ELayoutExtensionPosition::After, FTabManager::FTab(UAssetEditorUISubsystem::BottomRightTabID, ETabState::ClosedTab));
+
+#if WITH_STATETREE_TRACE_DEBUGGER
+	Extender.ExtendStack(FStateTreeEditor::LayoutBottomMiddleStackId, ELayoutExtensionPosition::After, FTabManager::FTab(UAssetEditorUISubsystem::TopRightTabID, ETabState::ClosedTab));
+#endif // WITH_STATETREE_TRACE_DEBUGGER
+}
diff --git a/Engine/Plugins/Runtime/StateTree/Source/StateTreeEditorModule/Private/StateTreeEditorUILayer.h b/Engine/Plugins/Runtime/StateTree/Source/StateTreeEditorModule/Private/StateTreeEditorUILayer.h
new file mode 100644
--- /dev/null
+++ b/Engine/Plugins/Runtime/StateTree/Source/StateTreeEditorModule/Private/StateTreeEditorUILayer.h
@@ -0,0 +1,31 @@
+// Copyright Epic Games, Inc. All Rights Reserved.
+
+#pragma once
+
+#include "Toolkits/AssetEditorModeUILayer.h"
+
+#include "StateTreeEditorUILayer.generated.h"
+
+class FStateTreeEditorModeUILayer : public FAssetEditorModeUILayer
+{
+public:
+	FStateTreeEditorModeUILayer(const IToolkitHost* InToolkitHost);
+	virtual void OnToolkitHostingStarted(const TSharedRef<IToolkit>& Toolkit) override;
+	virtual void OnToolkitHostingFinished(const TSharedRef<IToolkit>& Toolkit) override;
+	virtual TSharedPtr<FWorkspaceItem> GetModeMenuCategory() const override;
+	
+	void SetModeMenuCategory(const TSharedPtr<FWorkspaceItem>& MenuCategoryIn);
+protected:
+	TSharedPtr<FWorkspaceItem> MenuCategory;
+};
+
+UCLASS()
+class UStateTreeEditorUISubsystem : public UAssetEditorUISubsystem
+{
+	GENERATED_BODY()
+
+public:
+	virtual void Initialize(FSubsystemCollectionBase& Collection) override;
+	virtual void Deinitialize() override;
+	virtual void RegisterLayoutExtensions(FLayoutExtender& Extender) override;
+};
\ No newline at end of file
diff --git a/Engine/Plugins/Runtime/StateTree/Source/StateTreeEditorModule/Private/StateTreeFactory.cpp b/Engine/Plugins/Runtime/StateTree/Source/StateTreeEditorModule/Private/StateTreeFactory.cpp
--- a/Engine/Plugins/Runtime/StateTree/Source/StateTreeEditorModule/Private/StateTreeFactory.cpp
+++ b/Engine/Plugins/Runtime/StateTree/Source/StateTreeEditorModule/Private/StateTreeFactory.cpp
@@ -11,6 +11,7 @@
 #include "StateTreeEditor.h"
 #include "StateTreeCompiler.h"
 #include "StateTreeDelegates.h"
+#include "StateTreeEditingSubsystem.h"
 #include "Modules/ModuleManager.h"
 
 #include UE_INLINE_GENERATED_CPP_BY_NAME(StateTreeFactory)
@@ -121,24 +122,15 @@ UObject* UStateTreeFactory::FactoryCreateNew(UClass* Class, UObject* InParent, F
 	NewStateTree->EditorData = EditorData;
 
 	// Compile the StateTree asset, so that it is valid initially.
-	UE::StateTree::Editor::ValidateAsset(*NewStateTree);
+	UStateTreeEditingSubsystem::ValidateStateTree(NewStateTree);
 
 	FStateTreeCompilerLog Log;
 	FStateTreeCompiler Compiler(Log);
 
-	const bool bSuccess = Compiler.Compile(*NewStateTree);
+	const bool bSuccess = UStateTreeEditingSubsystem::CompileStateTree(NewStateTree, Log);
 
-	if (bSuccess)
+	if (!bSuccess)
 	{
-		// Success
-		const uint32 EditorDataHash = UE::StateTree::Editor::CalcAssetHash(*NewStateTree);
-		NewStateTree->LastCompiledEditorDataHash = EditorDataHash;
-		UE::StateTree::Delegates::OnPostCompile.Broadcast(*NewStateTree);
-	}
-	else
-	{
-		// Should not happen.
-		Log.DumpToLog(LogStateTreeEditor);
 		NewStateTree = nullptr;
 	}
 	
diff --git a/Engine/Plugins/Runtime/StateTree/Source/StateTreeEditorModule/Private/StateTreeViewModel.h b/Engine/Plugins/Runtime/StateTree/Source/StateTreeEditorModule/Private/StateTreeViewModel.h
deleted file mode 100644
diff --git a/Engine/Plugins/Runtime/StateTree/Source/StateTreeEditorModule/Public/IStateTreeEditorHost.h b/Engine/Plugins/Runtime/StateTree/Source/StateTreeEditorModule/Public/IStateTreeEditorHost.h
new file mode 100644
--- /dev/null
+++ b/Engine/Plugins/Runtime/StateTree/Source/StateTreeEditorModule/Public/IStateTreeEditorHost.h
@@ -0,0 +1,35 @@
+// Copyright Epic Games, Inc. All Rights Reserved.
+
+#pragma once
+
+#include "CoreMinimal.h"
+#include "UObject/Interface.h"
+
+#include "IStateTreeEditorHost.generated.h"
+
+class UStateTree;
+class IMessageLogListing;
+class IDetailsView;
+
+// Interface required for re-using StateTreeEditor mode across different AssetEditors
+class IStateTreeEditorHost : public TSharedFromThis<IStateTreeEditorHost>
+{
+public:
+	virtual ~IStateTreeEditorHost() = default;
+	virtual FName GetCompilerLogName() const = 0;
+	virtual FName GetCompilerTabName() const = 0;
+	
+	virtual UStateTree* GetStateTree() const = 0;
+	virtual FSimpleMulticastDelegate& OnStateTreeChanged() = 0;
+	virtual TSharedPtr<IDetailsView> GetAssetDetailsView() = 0;
+	virtual TSharedPtr<IDetailsView> GetDetailsView() = 0;
+};
+
+UCLASS()
+class STATETREEEDITORMODULE_API UStateTreeEditorContext : public UObject
+{
+	GENERATED_BODY()
+public:
+	TSharedPtr<IStateTreeEditorHost> EditorHostInterface;
+};
+
diff --git a/Engine/Plugins/Runtime/StateTree/Source/StateTreeEditorModule/Public/StateTreeEditingSubsystem.h b/Engine/Plugins/Runtime/StateTree/Source/StateTreeEditorModule/Public/StateTreeEditingSubsystem.h
new file mode 100644
--- /dev/null
+++ b/Engine/Plugins/Runtime/StateTree/Source/StateTreeEditorModule/Public/StateTreeEditingSubsystem.h
@@ -0,0 +1,37 @@
+// Copyright Epic Games, Inc. All Rights Reserved.
+
+#pragma once
+
+#include "StateTree.h"
+#include "StateTreeViewModel.h"
+#include "EditorSubsystem.h"
+
+#include "StateTreeEditingSubsystem.generated.h"
+
+class SWidget;
+class FStateTreeViewModel;
+class FUICommandList;
+struct FStateTreeCompilerLog;
+
+UCLASS()
+class STATETREEEDITORMODULE_API UStateTreeEditingSubsystem : public UEditorSubsystem
+{
+	GENERATED_BODY()
+public:
+	UStateTreeEditingSubsystem() {}
+	
+	TSharedRef<FStateTreeViewModel> FindOrAddViewModel(const TNonNullPtr<UStateTree> InStateTree);
+	
+	static bool CompileStateTree(const TNonNullPtr<UStateTree>InStateTree,  FStateTreeCompilerLog& InOutLog);
+	
+	static TSharedRef<SWidget> GetStateTreeView(TSharedRef<FStateTreeViewModel> InViewModel, const TSharedRef<FUICommandList>& TreeViewCommandList);
+	
+	// Validates asset state
+	static void ValidateStateTree(const TNonNullPtr<UStateTree> InStateTree);
+
+	// Calculates editor data hash of the asset.
+	static uint32 CalculateStateTreeHash(const TNonNullPtr<const UStateTree> InStateTree);
+	
+protected:
+	TMap<TSoftObjectPtr<UStateTree>, TSharedPtr<FStateTreeViewModel>> StateTreeViewModels;
+};
diff --git a/Engine/Plugins/Runtime/StateTree/Source/StateTreeEditorModule/Public/StateTreeEditorMode.h b/Engine/Plugins/Runtime/StateTree/Source/StateTreeEditorModule/Public/StateTreeEditorMode.h
new file mode 100644
--- /dev/null
+++ b/Engine/Plugins/Runtime/StateTree/Source/StateTreeEditorModule/Public/StateTreeEditorMode.h
@@ -0,0 +1,78 @@
+// Copyright Epic Games, Inc. All Rights Reserved.
+
+#pragma once
+
+#include "EdMode.h"
+#include "StateTree.h"
+#include "StateTreeEditorData.h"
+
+#include "StateTreeEditorMode.generated.h"
+
+class FStateTreeBindingExtension;
+class FStateTreeBindingsChildrenCustomization;
+class IDetailsView;
+class IMessageLogListing;
+
+UCLASS(Transient)
+class STATETREEEDITORMODULE_API UStateTreeEditorMode : public UEdMode
+{
+	GENERATED_BODY()
+public:
+	const static FEditorModeID EM_StateTree;
+	
+	UStateTreeEditorMode();
+
+	virtual void Enter() override;
+	virtual void Exit() override;
+	virtual void CreateToolkit() override;
+	virtual void BindCommands() override;
+
+protected:
+	void OnStateTreeChanged();
+	void BindToolkitCommands(const TSharedRef<FUICommandList>& ToolkitCommands);
+
+	void OnIdentifierChanged(const UStateTree& InStateTree);
+	void OnSchemaChanged(const UStateTree& InStateTree);
+	void ForceRefreshDetailsView() const;
+	void OnRefreshDetailsView(const UStateTree& InStateTree) const;
+	void OnStateParametersChanged(const UStateTree& InStateTree, const FGuid ChangedStateID) const;
+
+	void HandleMessageTokenClicked(const TSharedRef<IMessageToken>& InMessageToken) const;
+
+	void OnAssetFinishedChangingProperties(const FPropertyChangedEvent& PropertyChangedEvent) const;
+	void OnSelectionFinishedChangingProperties(const FPropertyChangedEvent& PropertyChangedEvent);
+
+	void Compile();
+	bool CanCompile() const;
+
+	bool HasValidStateTree() const;
+	
+	void HandleModelAssetChanged();
+	void HandleModelSelectionChanged(const TArray<TWeakObjectPtr<UStateTreeState>>& SelectedStates) const;
+	void HandleModelBringNodeToFocus(const UStateTreeState* State, const FGuid NodeID) const;
+
+	void HandleStateAdded(UStateTreeState* , UStateTreeState*) { UpdateAsset(); }
+	void HandleStatesRemoved(const TSet<UStateTreeState*>&) { UpdateAsset(); }
+	void HandleOnStatesMoved(const TSet<UStateTreeState*>&, const TSet<UStateTreeState*>&) { UpdateAsset(); }
+	void UpdateAsset();
+
+	TSharedPtr<IDetailsView> GetDetailsView() const;
+	TSharedPtr<IDetailsView> GetAssetDetailsView() const;
+	
+	TSharedPtr<IMessageLogListing> GetMessageLogListing() const;	
+	void ShowCompilerTab() const;
+
+	UStateTree* GetStateTree() const;
+	
+	friend class FStateTreeEditorModeToolkit;
+
+protected:
+	uint32 EditorDataHash = 0;
+	bool bLastCompileSucceeded = true;
+
+	TWeakObjectPtr<UStateTree> CachedStateTree = nullptr;
+
+	TSharedPtr<FStateTreeBindingExtension> DetailsViewExtensionHandler;
+	TSharedPtr<FStateTreeBindingsChildrenCustomization> DetailsViewChildrenCustomizationHandler;
+};
+
diff --git a/Engine/Plugins/Runtime/StateTree/Source/StateTreeEditorModule/Public/StateTreeEditorModeToolkit.h b/Engine/Plugins/Runtime/StateTree/Source/StateTreeEditorModule/Public/StateTreeEditorModeToolkit.h
new file mode 100644
--- /dev/null
+++ b/Engine/Plugins/Runtime/StateTree/Source/StateTreeEditorModule/Public/StateTreeEditorModeToolkit.h
@@ -0,0 +1,68 @@
+// Copyright Epic Games, Inc. All Rights Reserved.
+
+#pragma once
+
+#include "StateTreeEditorMode.h"
+#include "IStateTreeEditorHost.h"
+#include "Toolkits/BaseToolkit.h"
+
+class UStateTreeEditorMode;
+
+class STATETREEEDITORMODULE_API FStateTreeEditorModeToolkit : public FModeToolkit
+{
+public:
+
+	FStateTreeEditorModeToolkit(UStateTreeEditorMode* InEditorMode);
+	
+	/** IToolkit interface */
+	virtual FName GetToolkitFName() const override;
+	virtual FText GetBaseToolkitName() const override;
+
+	virtual void Init(const TSharedPtr<IToolkitHost>& InitToolkitHost, TWeakObjectPtr<UEdMode> InOwningMode) override;
+	virtual void InvokeUI() override;
+	virtual void RequestModeUITabs() override;
+	
+	virtual void ExtendSecondaryModeToolbar(UToolMenu* InModeToolbarMenu) override;
+		
+	void OnStateTreeChanged();
+protected:	
+	FSlateIcon GetCompileStatusImage() const;
+
+	FSlateIcon GetNewTaskButtonImage() const;
+	TSharedRef<SWidget> GenerateTaskBPBaseClassesMenu() const;
+
+	FSlateIcon GetNewConditionButtonImage() const;
+	TSharedRef<SWidget> GenerateConditionBPBaseClassesMenu() const;
+    
+	FSlateIcon GetNewConsiderationButtonImage() const;
+	TSharedRef<SWidget> GenerateConsiderationBPBaseClassesMenu() const;
+
+	void OnNodeBPBaseClassPicked(UClass* NodeClass) const;
+	
+	FText GetStatisticsText() const;
+
+	void UpdateStateTreeOutliner();
+protected:
+	TWeakObjectPtr<UStateTreeEditorMode> WeakEditorMode;
+	bool bLastCompileSucceeded = true;
+	uint32 EditorDataHash = 0;
+
+	TSharedPtr<IStateTreeEditorHost> EditorHost;
+
+	/** Tree Outliner */
+	TSharedPtr<SWidget> StateTreeOutliner = nullptr;
+	TWeakPtr<SDockTab> WeakOutlinerTab = nullptr;
+
+#if WITH_STATETREE_TRACE_DEBUGGER
+	void UpdateDebuggerView();
+	TSharedPtr<SWidget> DebuggerView = nullptr;	
+	TWeakPtr<SDockTab> WeakDebuggerTab = nullptr;
+public:
+	static const FName DebuggerTabId;
+#endif // WITH_STATETREE_TRACE_DEBUGGER
+	
+	static const FName StateTreeOutlinerTabId;
+	static const FName StateTreeStatisticsTabId;
+};
+
+
diff --git a/Engine/Plugins/Runtime/StateTree/Source/StateTreeEditorModule/Public/StateTreeEditorModule.h b/Engine/Plugins/Runtime/StateTree/Source/StateTreeEditorModule/Public/StateTreeEditorModule.h
--- a/Engine/Plugins/Runtime/StateTree/Source/StateTreeEditorModule/Public/StateTreeEditorModule.h
+++ b/Engine/Plugins/Runtime/StateTree/Source/StateTreeEditorModule/Public/StateTreeEditorModule.h
@@ -34,12 +34,16 @@ public:
 
 	TSharedPtr<FStateTreeNodeClassCache> GetNodeClassCache();
 	
+	DECLARE_EVENT_OneParam(FStateTreeEditorModule, FOnRegisterLayoutExtensions, FLayoutExtender&);
+	FOnRegisterLayoutExtensions& OnRegisterLayoutExtensions() { return RegisterLayoutExtensions; }
+
 protected:
 	TSharedPtr<FExtensibilityManager> MenuExtensibilityManager;
 	TSharedPtr<FExtensibilityManager> ToolBarExtensibilityManager;
 	TSharedPtr<FStateTreeNodeClassCache> NodeClassCache;
 
 	FDelegateHandle OnUserDefinedStructReinstancedHandle;
+	FOnRegisterLayoutExtensions RegisterLayoutExtensions;
 };
 
 #if UE_ENABLE_INCLUDE_ORDER_DEPRECATED_IN_5_2
diff --git a/Engine/Plugins/Runtime/StateTree/Source/StateTreeEditorModule/Public/StateTreeEditorStyle.h b/Engine/Plugins/Runtime/StateTree/Source/StateTreeEditorModule/Public/StateTreeEditorStyle.h
--- a/Engine/Plugins/Runtime/StateTree/Source/StateTreeEditorModule/Public/StateTreeEditorStyle.h
+++ b/Engine/Plugins/Runtime/StateTree/Source/StateTreeEditorModule/Public/StateTreeEditorStyle.h
@@ -4,6 +4,9 @@
 
 #include "Styling/SlateStyle.h"
 
+enum class EStateTreeStateSelectionBehavior : uint8;
+enum class EStateTreeStateType : uint8;
+
 class ISlateStyle;
 
 class STATETREEEDITORMODULE_API FStateTreeEditorStyle
@@ -11,6 +14,8 @@ class STATETREEEDITORMODULE_API FStateTreeEditorStyle
 {
 public:
 	static FStateTreeEditorStyle& Get();
+	
+	static const FSlateBrush* GetBrushForSelectionBehaviorType(EStateTreeStateSelectionBehavior InBehaviour, bool bHasChildren, EStateTreeStateType StateType);	
 
 protected:
 	friend class FStateTreeEditorModule;
diff --git a/Engine/Plugins/Runtime/StateTree/Source/StateTreeEditorModule/Public/StateTreeToolMenuContext.h b/Engine/Plugins/Runtime/StateTree/Source/StateTreeEditorModule/Public/StateTreeToolMenuContext.h
deleted file mode 100644
diff --git a/Engine/Plugins/Runtime/StateTree/Source/StateTreeEditorModule/Public/StateTreeViewModel.h b/Engine/Plugins/Runtime/StateTree/Source/StateTreeEditorModule/Public/StateTreeViewModel.h
new file mode 100644
--- /dev/null
+++ b/Engine/Plugins/Runtime/StateTree/Source/StateTreeEditorModule/Public/StateTreeViewModel.h
@@ -0,0 +1,177 @@
+// Copyright Epic Games, Inc. All Rights Reserved.
+
+#pragma once
+
+#include "UObject/WeakObjectPtr.h"
+#include "EditorUndoClient.h"
+#include "StateTreeViewModel.generated.h"
+
+class FMenuBuilder;
+class UStateTree;
+class UStateTreeEditorData;
+class UStateTreeState;
+
+enum class ECheckBoxState : uint8;
+enum class EStateTreeBreakpointType : uint8;
+
+struct FPropertyChangedEvent;
+struct FStateTreeDebugger;
+struct FStateTreeDebuggerBreakpoint;
+struct FStateTreeEditorBreakpoint;
+struct FStateTreePropertyPathBinding;
+
+enum class FStateTreeViewModelInsert : uint8
+{
+	Before,
+	After,
+	Into,
+};
+
+/**
+ * ModelView for editing StateTreeEditorData.
+ */
+class STATETREEEDITORMODULE_API FStateTreeViewModel : public FEditorUndoClient, public TSharedFromThis<FStateTreeViewModel>
+{
+public:
+
+	DECLARE_MULTICAST_DELEGATE(FOnAssetChanged);
+	DECLARE_MULTICAST_DELEGATE_TwoParams(FOnStatesChanged, const TSet<UStateTreeState*>& /*AffectedStates*/, const FPropertyChangedEvent& /*PropertyChangedEvent*/);
+	DECLARE_MULTICAST_DELEGATE_TwoParams(FOnStateAdded, UStateTreeState* /*ParentState*/, UStateTreeState* /*NewState*/);
+	DECLARE_MULTICAST_DELEGATE_OneParam(FOnStatesRemoved, const TSet<UStateTreeState*>& /*AffectedParents*/);
+	DECLARE_MULTICAST_DELEGATE_TwoParams(FOnStatesMoved, const TSet<UStateTreeState*>& /*AffectedParents*/, const TSet<UStateTreeState*>& /*MovedStates*/);
+	DECLARE_MULTICAST_DELEGATE_OneParam(FOnSelectionChanged, const TArray<TWeakObjectPtr<UStateTreeState>>& /*SelectedStates*/);
+	DECLARE_MULTICAST_DELEGATE_TwoParams(FOnBringNodeToFocus, const UStateTreeState* /*State*/, const FGuid /*TaskID*/);
+
+	FStateTreeViewModel();
+	virtual ~FStateTreeViewModel() override;
+
+	void Init(UStateTreeEditorData* InTreeData);
+
+	//~ FEditorUndoClient
+	virtual void PostUndo(bool bSuccess) override;
+	virtual void PostRedo(bool bSuccess) override;
+
+	// Selection handling.
+	void ClearSelection();
+	void SetSelection(UStateTreeState* Selected);
+	void SetSelection(const TArray<TWeakObjectPtr<UStateTreeState>>& InSelection);
+	bool IsSelected(const UStateTreeState* State) const;
+	bool IsChildOfSelection(const UStateTreeState* State) const;
+	void GetSelectedStates(TArray<UStateTreeState*>& OutSelectedStates) const;
+	void GetSelectedStates(TArray<TWeakObjectPtr<UStateTreeState>>& OutSelectedStates) const;
+	bool HasSelection() const;
+
+	void BringNodeToFocus(UStateTreeState* State, const FGuid NodeID);
+	
+	// Returns associated state tree asset.
+	const UStateTree* GetStateTree() const;
+
+	const UStateTreeEditorData* GetStateTreeEditorData() const;
+
+	const UStateTreeState* GetStateByID(const FGuid StateID) const;
+	UStateTreeState* GetMutableStateByID(const FGuid StateID) const;
+	
+	// Returns array of subtrees to edit.
+	TArray<TObjectPtr<UStateTreeState>>* GetSubTrees() const;
+	int32 GetSubTreeCount() const;
+	void GetSubTrees(TArray<TWeakObjectPtr<UStateTreeState>>& OutSubtrees) const;
+
+	// Gets and sets StateTree view expansion state store in the asset.
+	void SetPersistentExpandedStates(TSet<TWeakObjectPtr<UStateTreeState>>& InExpandedStates);
+	void GetPersistentExpandedStates(TSet<TWeakObjectPtr<UStateTreeState>>& OutExpandedStates);
+
+	// State manipulation commands
+	void AddState(UStateTreeState* AfterState);
+	void AddChildState(UStateTreeState* ParentState);
+	void RenameState(UStateTreeState* State, FName NewName);
+	void RemoveSelectedStates();
+	void CopySelectedStates();
+	bool CanPasteStatesFromClipboard() const;
+	void PasteStatesFromClipboard(UStateTreeState* AfterState);
+	void PasteStatesAsChildrenFromClipboard(UStateTreeState* ParentState);
+	void DuplicateSelectedStates();
+	void MoveSelectedStatesBefore(UStateTreeState* TargetState);
+	void MoveSelectedStatesAfter(UStateTreeState* TargetState);
+	void MoveSelectedStatesInto(UStateTreeState* TargetState);
+	bool CanEnableStates() const;
+	bool CanDisableStates() const;
+	void SetSelectedStatesEnabled(bool bEnable);
+	
+	// Force to update the view externally.
+	void NotifyAssetChangedExternally() const;
+	void NotifyStatesChangedExternally(const TSet<UStateTreeState*>& ChangedStates, const FPropertyChangedEvent& PropertyChangedEvent) const;
+
+	// Debugging
+#if WITH_STATETREE_TRACE_DEBUGGER
+	bool CanAddStateBreakpoint(EStateTreeBreakpointType Type) const;
+	bool CanRemoveStateBreakpoint(EStateTreeBreakpointType Type) const;
+	ECheckBoxState GetStateBreakpointCheckState(EStateTreeBreakpointType Type) const;
+	void HandleEnableStateBreakpoint(EStateTreeBreakpointType Type);
+
+	UStateTreeState* FindStateAssociatedToBreakpoint(FStateTreeDebuggerBreakpoint Breakpoint) const;
+
+	TSharedRef<FStateTreeDebugger> GetDebugger() const { return Debugger; }
+	void RefreshDebuggerBreakpoints();
+#endif // WITH_STATETREE_TRACE_DEBUGGER
+
+	bool IsStateActiveInDebugger(const UStateTreeState& State) const;
+
+	// Called when the whole asset is updated (i.e. undo/redo).
+	FOnAssetChanged& GetOnAssetChanged() { return OnAssetChanged; }
+	
+	// Called when States are changed (i.e. change name or properties).
+	FOnStatesChanged& GetOnStatesChanged() { return OnStatesChanged; }
+	
+	// Called each time a state is added.
+	FOnStateAdded& GetOnStateAdded() { return OnStateAdded; }
+
+	// Called each time a states are removed.
+	FOnStatesRemoved& GetOnStatesRemoved() { return OnStatesRemoved; }
+
+	// Called each time a state is removed.
+	FOnStatesMoved& GetOnStatesMoved() { return OnStatesMoved; }
+
+	// Called each time the selection changes.
+	FOnSelectionChanged& GetOnSelectionChanged() { return OnSelectionChanged; }
+
+	FOnBringNodeToFocus& GetOnBringNodeToFocus() { return OnBringNodeToFocus; }
+	
+protected:
+	void GetExpandedStatesRecursive(UStateTreeState* State, TSet<TWeakObjectPtr<UStateTreeState>>& ExpandedStates);
+	void MoveSelectedStates(UStateTreeState* TargetState, const FStateTreeViewModelInsert RelativeLocation);
+
+	void PasteStatesAsChildrenFromText(const FString& TextToImport, UStateTreeState* ParentState, const int32 IndexToInsertAt);
+
+	void HandleIdentifierChanged(const UStateTree& StateTree) const;
+	
+	void BindToDebuggerDelegates();
+	
+	TWeakObjectPtr<UStateTreeEditorData> TreeDataWeak;
+	TSet<TWeakObjectPtr<UStateTreeState>> SelectedStates;
+
+#if WITH_STATETREE_TRACE_DEBUGGER
+	void HandleBreakpointsChanged(const UStateTree& StateTree);
+	void HandlePostCompile(const UStateTree& StateTree);
+
+	TSharedRef<FStateTreeDebugger> Debugger;
+	TArray<FGuid> ActiveStates;
+#endif // WITH_STATETREE_TRACE_DEBUGGER
+	
+	FOnAssetChanged OnAssetChanged;
+	FOnStatesChanged OnStatesChanged;
+	FOnStateAdded OnStateAdded;
+	FOnStatesRemoved OnStatesRemoved;
+	FOnStatesMoved OnStatesMoved;
+	FOnSelectionChanged OnSelectionChanged;
+	FOnBringNodeToFocus OnBringNodeToFocus;
+};
+
+/** Helper class to allow to copy bindings into clipboard. */
+UCLASS(Hidden)
+class STATETREEEDITORMODULE_API UStateTreeClipboardBindings : public UObject
+{
+	GENERATED_BODY()
+public:
+	UPROPERTY()
+	TArray<FStateTreePropertyPathBinding> Bindings;
+};
\ No newline at end of file
diff --git a/Engine/Plugins/Runtime/StateTree/Source/StateTreeEditorModule/StateTreeEditorModule.Build.cs b/Engine/Plugins/Runtime/StateTree/Source/StateTreeEditorModule/StateTreeEditorModule.Build.cs
--- a/Engine/Plugins/Runtime/StateTree/Source/StateTreeEditorModule/StateTreeEditorModule.Build.cs
+++ b/Engine/Plugins/Runtime/StateTree/Source/StateTreeEditorModule/StateTreeEditorModule.Build.cs
@@ -33,6 +33,7 @@ namespace UnrealBuildTool.Rules
 				"PropertyAccessEditor",
 				"StructUtilsEditor",
 				"GameplayTags",
+				"EditorSubsystem"
 			}
 			);
 
@@ -50,7 +51,9 @@ namespace UnrealBuildTool.Rules
 				"RewindDebuggerInterface",
 				"DetailCustomizations",
 				"AppFramework",
-				"KismetCompiler"
+				"KismetCompiler",
+				"EditorInteractiveToolsFramework",
+				"InteractiveToolsFramework",
 			}
 			);
 
diff --git a/Engine/Plugins/Runtime/StateTree/Source/StateTreeModule/Private/StateTree.cpp b/Engine/Plugins/Runtime/StateTree/Source/StateTreeModule/Private/StateTree.cpp
--- a/Engine/Plugins/Runtime/StateTree/Source/StateTreeModule/Private/StateTree.cpp
+++ b/Engine/Plugins/Runtime/StateTree/Source/StateTreeModule/Private/StateTree.cpp
@@ -280,6 +280,11 @@ void UStateTree::GetAssetRegistryTags(FAssetRegistryTagsContext Context) const
 	const FString SchemaClassName = Schema ? Schema->GetClass()->GetPathName() : TEXT("");
 	Context.AddTag(FAssetRegistryTag(UE::StateTree::SchemaTag, SchemaClassName, FAssetRegistryTag::TT_Alphabetical));
 
+	if (Schema)
+	{
+		Schema->GetAssetRegistryTags(Context);
+	}
+
 	Super::GetAssetRegistryTags(Context);
 }
 
diff --git a/Engine/Plugins/Runtime/StateTree/Source/StateTreeModule/StateTreeModule.Build.cs b/Engine/Plugins/Runtime/StateTree/Source/StateTreeModule/StateTreeModule.Build.cs
--- a/Engine/Plugins/Runtime/StateTree/Source/StateTreeModule/StateTreeModule.Build.cs
+++ b/Engine/Plugins/Runtime/StateTree/Source/StateTreeModule/StateTreeModule.Build.cs
@@ -37,12 +37,14 @@ namespace UnrealBuildTool.Rules
 				PublicDependencyModuleNames.AddRange(
 					new [] {
 						"UnrealEd",
-						"BlueprintGraph"
+						"BlueprintGraph",
 					}
 				);
 				PrivateDependencyModuleNames.AddRange(
 					new [] {
-						"StructUtilsEditor"
+						"StructUtilsEditor",
+						"EditorSubsystem",
+						"EditorFramework"
 					}
 				);
 			}
