StateTree: Deprecate FStateTreeTaskBase::ExitState and EnterState with the ChangeType arguments. It was already deprecated by comments but didn't had the deprecate attribute.
#rnx

diff --git a/Engine/Plugins/Runtime/StateTree/Source/StateTreeModule/Public/StateTreeTaskBase.h b/Engine/Plugins/Runtime/StateTree/Source/StateTreeModule/Public/StateTreeTaskBase.h
--- a/Engine/Plugins/Runtime/StateTree/Source/StateTreeModule/Public/StateTreeTaskBase.h
+++ b/Engine/Plugins/Runtime/StateTree/Source/StateTreeModule/Public/StateTreeTaskBase.h
@@ -28,11 +28,18 @@ struct STATETREEMODULE_API FStateTreeTaskBase : public FStateTreeNodeBase
 	{
 	}
 	
+	UE_DEPRECATED(5.5, "Use EnterState without the EStateTreeStateChangeType.")
 	/**
 	 * Note: The API has been deprecated. ChangeType is moved into FStateTreeTransitionResult.
 	 * You can configure the task to be only called on state changes (that is, never call sustained changes) by setting bShouldStateChangeOnReselect to true.
 	 */
 	virtual EStateTreeRunStatus EnterState(FStateTreeExecutionContext& Context, const EStateTreeStateChangeType ChangeType, const FStateTreeTransitionResult& Transition) const final { return EStateTreeRunStatus::Running; }
+
+	UE_DEPRECATED(5.5, "Use ExitState without the EStateTreeStateChangeType.")
+	/**
+	 * Note: The API has been deprecated. ChangeType is moved into FStateTreeTransitionResult.
+	 * You can configure the task to be only called on state changes (that is, never call sustained changes) by setting bShouldStateChangeOnReselect to true.
+	 */
 	virtual void ExitState(FStateTreeExecutionContext& Context, const EStateTreeStateChangeType ChangeType, const FStateTreeTransitionResult& Transition) const final {}
 
 	/**
