StateTree: Add option to print the result of a StateTree compilation. Moved  FStateTreeExecutionContext::GetDebugPrintString to UStateTree. Added transitions, global tasks and conditions. Update other FStateTreeExecutionContext functions to use StringBuilder.

diff --git a/Engine/Plugins/Runtime/StateTree/Source/StateTreeEditorModule/Private/StateTreeCompiler.cpp b/Engine/Plugins/Runtime/StateTree/Source/StateTreeEditorModule/Private/StateTreeCompiler.cpp
--- a/Engine/Plugins/Runtime/StateTree/Source/StateTreeEditorModule/Private/StateTreeCompiler.cpp
+++ b/Engine/Plugins/Runtime/StateTree/Source/StateTreeEditorModule/Private/StateTreeCompiler.cpp
@@ -5,6 +5,7 @@
 #include "StateTreeAnyEnum.h"
 #include "StateTreeCompilerLog.h"
 #include "StateTreeEditorData.h"
+#include "StateTreeEditorModule.h"
 #include "StateTreeEvaluatorBase.h"
 #include "StateTreeTaskBase.h"
 #include "StateTreeConditionBase.h"
@@ -19,6 +20,12 @@
 
 namespace UE::StateTree::Compiler
 {
+	FAutoConsoleVariable CVarLogCompiledStateTree(
+		TEXT("StateTree.LogCompiledResult"),
+		false,
+		TEXT("After a StateTree compiles, log the internal content of the StateTree.")
+	);
+
 	// Helper archive that checks that the all instanced sub-objects have correct outer. 
 	class FCheckOutersArchive : public FArchiveUObject
 	{
@@ -347,7 +354,11 @@ bool FStateTreeCompiler::Compile(UStateTree& InStateTree)
 
 	UE::StateTree::Compiler::FCheckOutersArchive CheckOuters(*StateTree, *EditorData, Log);
 	StateTree->Serialize(CheckOuters);
-	
+
+	if (UE::StateTree::Compiler::CVarLogCompiledStateTree->GetBool())
+	{
+		UE_LOG(LogStateTreeEditor, Log, TEXT("%s"), *StateTree->DebugInternalLayoutAsString());
+	}
 
 	return true;
 }
diff --git a/Engine/Plugins/Runtime/StateTree/Source/StateTreeModule/Private/StateTree.cpp b/Engine/Plugins/Runtime/StateTree/Source/StateTreeModule/Private/StateTree.cpp
--- a/Engine/Plugins/Runtime/StateTree/Source/StateTreeModule/Private/StateTree.cpp
+++ b/Engine/Plugins/Runtime/StateTree/Source/StateTreeModule/Private/StateTree.cpp
@@ -1069,5 +1069,168 @@ void UStateTree::CompileIfChanged()
 		}
 	}
 }
-#endif // WITH_EDITOR
+#endif //WITH_EDITOR
+
+#if WITH_EDITOR || WITH_STATETREE_DEBUG
+FString UStateTree::DebugInternalLayoutAsString() const
+{
+	FStringBuilderBase DebugString;
+	DebugString << TEXT("StateTree (asset: '");
+	GetFullName(DebugString);
+	DebugString << TEXT("')\n");
+
+	auto PrintObjectNameSafe = [&DebugString](const UObject* Obj)
+		{
+			DebugString << TEXT("  ");
+			if (Obj)
+			{
+				DebugString << Obj->GetFName();
+			}
+			else
+			{
+				DebugString << TEXT("null");
+			}
+			DebugString << TEXT('\n');
+		};
+	auto PrintViewNameSafe = [&DebugString](const FConstStructView& View)
+		{
+			DebugString << TEXT("  ");
+			if (View.IsValid())
+			{
+				DebugString << View.GetScriptStruct()->GetFName();
+			}
+			else
+			{
+				DebugString << TEXT("null");
+			}
+			DebugString << TEXT('\n');
+		};
+
+	// Tree items (e.g. tasks, evaluators, conditions)
+	DebugString.Appendf(TEXT("\nNodes(%d)\n"), Nodes.Num());
+	for (int32 Index = 0; Index < Nodes.Num(); Index++)
+	{
+		const FConstStructView Node = Nodes[Index];
+		PrintViewNameSafe(Node);
+	}
+
+	// Instance InstanceData data (e.g. tasks)
+	DebugString.Appendf(TEXT("\nInstance Data(%d)\n"), DefaultInstanceData.Num());
+	for (int32 Index = 0; Index < DefaultInstanceData.Num(); Index++)
+	{
+		if (DefaultInstanceData.IsObject(Index))
+		{
+			const UObject* Data = DefaultInstanceData.GetObject(Index);
+			PrintObjectNameSafe(Data);
+		}
+		else
+		{
+			const FConstStructView Data = DefaultInstanceData.GetStruct(Index);
+			PrintViewNameSafe(Data);
+		}
+	}
+
+	// External data (e.g. fragments, subsystems)
+	DebugString.Appendf(TEXT("\nExternal Data(%d)\n"), ExternalDataDescs.Num());
+	if (ExternalDataDescs.Num())
+	{
+		DebugString.Appendf(TEXT("  [% -40s | % -8s | % 15s]\n"), TEXT("Name"), TEXT("Optional"), TEXT("Handle"));
+		for (const FStateTreeExternalDataDesc& Desc : ExternalDataDescs)
+		{
+			DebugString.Appendf(TEXT("  | %-40s | %8s | %15s |\n"), Desc.Struct ? *Desc.Struct->GetName() : TEXT("null"), *UEnum::GetDisplayValueAsText(Desc.Requirement).ToString(), *Desc.Handle.DataHandle.Describe());
+		}
+	}
+
+	// Bindings
+	DebugString << PropertyBindings.DebugInternalLayoutAsString();
+
+	// States
+	DebugString.Appendf(TEXT("\nStates(%d)\n"), States.Num());
+	if (States.Num())
+	{
+		DebugString.Appendf(TEXT("  [ %-30s | %15s | %5s [%3s:%-3s[ | Begin Idx : %4s %4s %4s %4s | Num : %4s %4s %4s %4s ]\n"),
+			TEXT("Name"), TEXT("Parent"), TEXT("Child"), TEXT("Beg"), TEXT("End"),
+			TEXT("Cond"), TEXT("Tr"), TEXT("Tsk"), TEXT("Uti"), TEXT("Cond"), TEXT("Tr"), TEXT("Tsk"), TEXT("Uti"));
+		for (const FCompactStateTreeState& State : States)
+		{
+			DebugString.Appendf(TEXT("  | %-30s | %15s | %5s [%3d:%-3d[ | %9s   %4d %4d %4d %4d | %3s   %4d %4d %4d %4d |\n"),
+				*State.Name.ToString(), *State.Parent.Describe(),
+				TEXT(""), State.ChildrenBegin, State.ChildrenEnd,
+				TEXT(""), State.EnterConditionsBegin, State.TransitionsBegin, State.TasksBegin, State.UtilityConsiderationsBegin,
+				TEXT(""), State.EnterConditionsNum, State.TransitionsNum, State.TasksNum, State.UtilityConsiderationsNum);
+		}
+	}
+
+	// Transitions
+	DebugString.Appendf(TEXT("\nTransitions(%d)\n"), Transitions.Num());
+	if (Transitions.Num())
+	{
+		DebugString.Appendf(TEXT("  [ %-3s | %15s | %-20s | %-40s | %-40s | %-8s ]\n")
+			, TEXT("Idx"), TEXT("State"), TEXT("Transition Trigger"), TEXT("Transition Event Tag"), TEXT("Transition Event Payload"), TEXT("Cond:Num"));
+		for (const FCompactStateTransition& Transition : Transitions)
+		{
+			DebugString.Appendf(TEXT("  | %3d | %15s | %-20s | %-40s | %-40s | %4d:%3d |\n"),
+				Transition.ConditionsBegin, *Transition.State.Describe(),
+				*UEnum::GetDisplayValueAsText(Transition.Trigger).ToString(),
+				*Transition.RequiredEvent.Tag.ToString(),
+				Transition.RequiredEvent.PayloadStruct ? *Transition.RequiredEvent.PayloadStruct->GetName() : TEXT("None"),
+				Transition.ConditionsBegin, Transition.ConditionsNum);
+		}
+	}
+
+	// Evaluators
+	DebugString.Appendf(TEXT("\nEvaluators(%d)\n"), EvaluatorsNum);
+	if (EvaluatorsNum)
+	{
+		DebugString.Appendf(TEXT("  [ %-30s | %8s | %10s ]\n"),
+			TEXT("Name"), TEXT("Bindings"), TEXT("Struct Idx"));
+		for (int32 EvalIndex = EvaluatorsBegin; EvalIndex < (EvaluatorsBegin + EvaluatorsNum); EvalIndex++)
+		{
+			const FStateTreeEvaluatorBase& Eval = Nodes[EvalIndex].Get<const FStateTreeEvaluatorBase>();
+			DebugString.Appendf(TEXT("  | %-30s | %8d | %10s |\n"),
+				*Eval.Name.ToString(), Eval.BindingsBatch.Get(), *Eval.InstanceDataHandle.Describe());
+		}
+	}
+
+	// Tasks
+	DebugString.Appendf(TEXT("\nTasks\n  [ %-30s | %-30s | %8s | %10s ]\n"),
+		TEXT("State"), TEXT("Name"), TEXT("Bindings"), TEXT("Struct Idx"));
+	for (const FCompactStateTreeState& State : States)
+	{
+		if (State.TasksNum)
+		{
+			for (int32 TaskIndex = State.TasksBegin; TaskIndex < (State.TasksBegin + State.TasksNum); TaskIndex++)
+			{
+				const FStateTreeTaskBase& Task = Nodes[TaskIndex].Get<const FStateTreeTaskBase>();
+				DebugString.Appendf(TEXT("  | %-30s | %-30s | %8d | %10s |\n"),
+					*State.Name.ToString(), *Task.Name.ToString(), Task.BindingsBatch.Get(), *Task.InstanceDataHandle.Describe());
+			}
+		}
+	}
+	for (int32 TaskIndex = GlobalTasksBegin; TaskIndex < (GlobalTasksBegin + GlobalTasksNum); TaskIndex++)
+	{
+		const FStateTreeTaskBase& Task = Nodes[TaskIndex].Get<const FStateTreeTaskBase>();
+		DebugString.Appendf(TEXT("  | %-30s | %-30s | %8d | %10s |\n"),
+			TEXT("Global"), *Task.Name.ToString(), Task.BindingsBatch.Get(), *Task.InstanceDataHandle.Describe());
+	}
+
+	// Conditions
+	DebugString.Appendf(TEXT("\nCondition\n  [ %-30s | %-30s | %8s | %12s | %10s ]\n"),
+		TEXT("State"), TEXT("Name"), TEXT("Operand"), TEXT("Evaluation"), TEXT("Struct Idx"));
+	for (const FCompactStateTreeState& State : States)
+	{
+		if (State.EnterConditionsNum)
+		{
+			for (int32 CondIndex = State.EnterConditionsBegin; CondIndex < (State.EnterConditionsBegin + State.EnterConditionsNum); CondIndex++)
+			{
+				const FStateTreeConditionBase& Cond = Nodes[CondIndex].Get<const FStateTreeConditionBase>();
+				DebugString.Appendf(TEXT("  | %-30s | %-30s | %8s | %12s | %10s |\n"),
+					*State.Name.ToString(), *Cond.Name.ToString(), *UEnum::GetDisplayValueAsText(Cond.Operand).ToString(), *UEnum::GetDisplayValueAsText(Cond.EvaluationMode).ToString(), *Cond.InstanceDataHandle.Describe());
+			}
+		}
+	}
+
+	return DebugString.ToString();
+}
+#endif // WITH_EDITOR || WITH_STATETREE_DEBUG
 
diff --git a/Engine/Plugins/Runtime/StateTree/Source/StateTreeModule/Private/StateTreeExecutionContext.cpp b/Engine/Plugins/Runtime/StateTree/Source/StateTreeModule/Private/StateTreeExecutionContext.cpp
--- a/Engine/Plugins/Runtime/StateTree/Source/StateTreeModule/Private/StateTreeExecutionContext.cpp
+++ b/Engine/Plugins/Runtime/StateTree/Source/StateTreeModule/Private/StateTreeExecutionContext.cpp
@@ -2057,7 +2057,7 @@ void FStateTreeExecutionContext::StateCompleted()
 							STATETREE_LOG(VeryVerbose, TEXT("%*sSkipped 'StateCompleted' for disabled Task: '%s'"), UE::StateTree::DebugIndentSize, TEXT(""), *Task.Name.ToString());
 							continue;
 						}
-						
+
 						STATETREE_LOG(Verbose, TEXT("%*s  Task '%s'"), Index*UE::StateTree::DebugIndentSize, TEXT(""), *Task.Name.ToString());
 						Task.StateCompleted(*this, Exec.LastTickStatus, CurrentFrame.ActiveStates);
 					}
@@ -2915,8 +2915,8 @@ void FStateTreeExecutionContext::EvaluatePropertyFunctionsWithValidation(const F
 
 FString FStateTreeExecutionContext::DebugGetEventsAsString() const
 {
-	FStringBuilderBase StrBuilder;
-	
+	TStringBuilder<512> StrBuilder;
+
 	if (EventQueue)
 	{
 		for (const FStateTreeSharedEvent& Event : EventQueue->GetEventsView())
@@ -2925,7 +2925,7 @@ FString FStateTreeExecutionContext::DebugGetEventsAsString() const
 			{
 				if (StrBuilder.Len() > 0)
 				{
-					StrBuilder += TEXT(", ");
+					StrBuilder << TEXT(", ");
 				}
 
 				const bool bHasTag = Event->Tag.IsValid();
@@ -2933,26 +2933,30 @@ FString FStateTreeExecutionContext::DebugGetEventsAsString() const
 			
 				if (bHasTag || bHasPayload)
 				{
-					StrBuilder.Appendf(TEXT("("));
+					StrBuilder << (TEXT('('));
 				
 					if (bHasTag)
 					{
-						StrBuilder.Appendf(TEXT("Tag: '%s'"), *Event->Tag.ToString()); 
+						StrBuilder << TEXT("Tag: '");
+						StrBuilder << Event->Tag.GetTagName();
+						StrBuilder << TEXT('\'');
 					}
 					if (bHasTag && bHasPayload)
 					{
-						StrBuilder.Appendf(TEXT(", "));
+						StrBuilder << TEXT(", ");
 					}
 					if (bHasPayload)
 					{
-						StrBuilder.Appendf(TEXT(" Payload: '%s'"), *Event->Payload.GetScriptStruct()->GetName()); 
+						StrBuilder << TEXT(" Payload: '");
+						StrBuilder << Event->Payload.GetScriptStruct()->GetFName();
+						StrBuilder << TEXT('\'');
 					}
-					StrBuilder.Appendf(TEXT(") "));
+					StrBuilder << TEXT(") ");
 				}
 			}
 		}
 	}
-	
+
 	return StrBuilder.ToString();
 }
 
@@ -4470,53 +4474,46 @@ TConstArrayView<FStateTreeExecutionFrame> FStateTreeExecutionContext::GetActiveF
 
 FString FStateTreeExecutionContext::GetDebugInfoString() const
 {
-
-	FString DebugString = FString::Printf(TEXT("StateTree (asset: '%s')\n"), *GetFullNameSafe(&RootStateTree));
+	TStringBuilder<2048> DebugString;
+	DebugString << TEXT("StateTree (asset: '");
+	RootStateTree.GetFullName(DebugString);
+	DebugString << TEXT("')");
 
 	if (IsValid())
 	{
 		const FStateTreeExecutionState& Exec = GetExecState();
 
-		DebugString += TEXT("Status: ");
-		switch (Exec.TreeRunStatus)
-		{
-		case EStateTreeRunStatus::Failed:
-			DebugString += TEXT("Failed\n");
-			break;
-		case EStateTreeRunStatus::Succeeded:
-			DebugString += TEXT("Succeeded\n");
-			break;
-		case EStateTreeRunStatus::Running:
-			DebugString += TEXT("Running\n");
-			break;
-		default:
-			DebugString += TEXT("--\n");
-		}
+		DebugString << TEXT("Status: ");
+		DebugString << UEnum::GetDisplayValueAsText(Exec.TreeRunStatus).ToString();
+		DebugString << TEXT("\n");
 
 		// Active States
-		DebugString += TEXT("Current State:\n");
+		DebugString << TEXT("Current State:\n");
 		for (const FStateTreeExecutionFrame& CurrentFrame : Exec.ActiveFrames)
 		{
 			const UStateTree* CurrentStateTree = CurrentFrame.StateTree;
 
 			if (CurrentFrame.bIsGlobalFrame)
 			{
-				DebugString += FString::Printf(TEXT("\nEvaluators\n  [ %-30s | %8s | %15s ]\n"),
+				DebugString.Appendf(TEXT("\nEvaluators\n  [ %-30s | %8s | %15s ]\n"),
 					TEXT("Name"), TEXT("Bindings"), TEXT("Data Handle"));
 				for (int32 EvalIndex = CurrentStateTree->EvaluatorsBegin; EvalIndex < (CurrentStateTree->EvaluatorsBegin + CurrentStateTree->EvaluatorsNum); EvalIndex++)
 				{
 					const FStateTreeEvaluatorBase& Eval = CurrentStateTree->Nodes[EvalIndex].Get<const FStateTreeEvaluatorBase>();
-					DebugString += FString::Printf(TEXT("| %-30s | %8d | %15s |\n"),
+					DebugString.Appendf(TEXT("| %-30s | %8d | %15s |\n"),
 						*Eval.Name.ToString(), Eval.BindingsBatch.Get(), *Eval.InstanceDataHandle.Describe());
 				}
 
-				DebugString += FString::Printf(TEXT("\nGlobal Tasks\n"));
+				DebugString << TEXT("\nGlobal Tasks\n");
 				for (int32 TaskIndex = CurrentStateTree->GlobalTasksBegin; TaskIndex < (CurrentStateTree->GlobalTasksBegin + CurrentStateTree->GlobalTasksNum); TaskIndex++)
 				{
 					const FStateTreeTaskBase& Task = CurrentStateTree->Nodes[TaskIndex].Get<const FStateTreeTaskBase>();
 					if (Task.bTaskEnabled)
 					{
-						Task.AppendDebugInfoString(DebugString, *this);
+						FString TempString;
+						Task.AppendDebugInfoString(TempString, *this);
+						DebugString << TempString;
+
 					}
 				}
 			}
@@ -4527,7 +4524,9 @@ FString FStateTreeExecutionContext::GetDebugInfoString() const
 				if (Handle.IsValid())
 				{
 					const FCompactStateTreeState& State = RootStateTree.States[Handle.Index];
-					DebugString += FString::Printf(TEXT("[%s]\n"), *State.Name.ToString());
+					DebugString << TEXT('[');
+					DebugString << State.Name;
+					DebugString << TEXT("]\n");
 
 					if (State.TasksNum > 0)
 					{
@@ -4537,7 +4536,9 @@ FString FStateTreeExecutionContext::GetDebugInfoString() const
 							const FStateTreeTaskBase& Task = RootStateTree.Nodes[TaskIndex].Get<const FStateTreeTaskBase>();
 							if (Task.bTaskEnabled)
 							{
-								Task.AppendDebugInfoString(DebugString, *this);
+								FString TempString;
+								Task.AppendDebugInfoString(TempString, *this);
+								DebugString << TempString;
 							}
 						}
 					}
@@ -4547,10 +4548,10 @@ FString FStateTreeExecutionContext::GetDebugInfoString() const
 	}
 	else
 	{
-		DebugString += TEXT("StateTree context is not initialized properly.");
+		DebugString << TEXT("StateTree context is not initialized properly.");
 	}
 
-	return DebugString;
+	return DebugString.ToString();
 }
 #endif // WITH_GAMEPLAY_DEBUGGER
 
@@ -4558,106 +4559,7 @@ FString FStateTreeExecutionContext::GetDebugInfoString() const
 void FStateTreeExecutionContext::DebugPrintInternalLayout()
 {
 	LOG_SCOPE_VERBOSITY_OVERRIDE(LogStateTree, ELogVerbosity::Log);
-
-	// @todo: this looks like more a UStateTree thing...
-	
-	FString DebugString = FString::Printf(TEXT("StateTree (asset: '%s')\n"), *GetFullNameSafe(&RootStateTree));
-
-	// Tree items (e.g. tasks, evaluators, conditions)
-	DebugString += FString::Printf(TEXT("\nItems(%d)\n"), RootStateTree.Nodes.Num());
-	for (int32 Index = 0; Index < RootStateTree.Nodes.Num(); Index++)
-	{
-		const FConstStructView Node = RootStateTree.Nodes[Index];
-		DebugString += FString::Printf(TEXT("  %s\n"), Node.IsValid() ? *Node.GetScriptStruct()->GetName() : TEXT("null"));
-	}
-
-	// Instance InstanceData data (e.g. tasks)
-	DebugString += FString::Printf(TEXT("\nInstance Data(%d)\n"), RootStateTree.DefaultInstanceData.Num());
-	for (int32 Index = 0; Index < RootStateTree.DefaultInstanceData.Num(); Index++)
-	{
-		if (RootStateTree.DefaultInstanceData.IsObject(Index))
-		{
-			const UObject* Data = RootStateTree.DefaultInstanceData.GetObject(Index);
-			DebugString += FString::Printf(TEXT("  %s\n"), *GetNameSafe(Data));
-		}
-		else
-		{
-			const FConstStructView Data = RootStateTree.DefaultInstanceData.GetStruct(Index);
-			DebugString += FString::Printf(TEXT("  %s\n"), Data.IsValid() ? *Data.GetScriptStruct()->GetName() : TEXT("null"));
-		}
-	}
-
-	// External data (e.g. fragments, subsystems)
-	DebugString += FString::Printf(TEXT("\nExternal Data(%d)\n  [ %-40s | %-8s | %15s ]\n"), RootStateTree.ExternalDataDescs.Num(), TEXT("Name"), TEXT("Optional"), TEXT("Handle"));
-	for (const FStateTreeExternalDataDesc& Desc : RootStateTree.ExternalDataDescs)
-	{
-		DebugString += FString::Printf(TEXT("  | %-40s | %8s | %15s |\n"), Desc.Struct ? *Desc.Struct->GetName() : TEXT("null"), *UEnum::GetDisplayValueAsText(Desc.Requirement).ToString(), *Desc.Handle.DataHandle.Describe());
-	}
-
-	// Bindings
-	RootStateTree.PropertyBindings.DebugPrintInternalLayout(DebugString);
-
-	// Transitions
-	DebugString += FString::Printf(TEXT("\nTransitions(%d)\n  [ %-3s | %15s | %-20s | %-40s | %-40s | %-8s ]\n"), RootStateTree.Transitions.Num()
-		, TEXT("Idx"), TEXT("State"), TEXT("Transition Trigger"), TEXT("Transition Event Tag"), TEXT("Transition Event Payload"), TEXT("Num Cond"));
-	for (const FCompactStateTransition& Transition : RootStateTree.Transitions)
-	{
-		DebugString += FString::Printf(TEXT("  | %3d | %15s | %-20s | %-40s | %-40s | %8d |\n"),
-									Transition.ConditionsBegin, *Transition.State.Describe(),
-									*UEnum::GetDisplayValueAsText(Transition.Trigger).ToString(),
-									*Transition.RequiredEvent.Tag.ToString(),
-									Transition.RequiredEvent.PayloadStruct ? *Transition.RequiredEvent.PayloadStruct->GetName() : TEXT("None"),
-									Transition.ConditionsNum);
-	}
-
-	// States
-	DebugString += FString::Printf(TEXT("\nStates(%d)\n"
-		"  [ %-30s | %15s | %5s [%3s:%-3s[ | Begin Idx : %4s %4s %4s %4s | Num : %4s %4s %4s %4s | Transitions : %-16s %-40s %-16s %-40s ]\n"),
-		RootStateTree.States.Num(),
-		TEXT("Name"), TEXT("Parent"), TEXT("Child"), TEXT("Beg"), TEXT("End"),
-		TEXT("Cond"), TEXT("Tr"), TEXT("Tsk"), TEXT("Evt"), TEXT("Cond"), TEXT("Tr"), TEXT("Tsk"), TEXT("Evt"),
-		TEXT("Done State"), TEXT("Done Type"), TEXT("Failed State"), TEXT("Failed Type")
-		);
-	for (const FCompactStateTreeState& State : RootStateTree.States)
-	{
-		DebugString += FString::Printf(TEXT("  | %-30s | %15s | %5s [%3d:%-3d[ | %9s   %4d %4d %4d | %3s   %4d %4d %4d\n"),
-									*State.Name.ToString(), *State.Parent.Describe(),
-									TEXT(""), State.ChildrenBegin, State.ChildrenEnd,
-									TEXT(""), State.EnterConditionsBegin, State.TransitionsBegin, State.TasksBegin,
-									TEXT(""), State.EnterConditionsNum, State.TransitionsNum, State.TasksNum);
-	}
-
-	// Evaluators
-	if (RootStateTree.EvaluatorsNum)
-	{
-		DebugString += FString::Printf(TEXT("\nEvaluators\n  [ %-30s | %8s | %10s ]\n"),
-			TEXT("Name"), TEXT("Bindings"), TEXT("Struct Idx"));
-		for (int32 EvalIndex = RootStateTree.EvaluatorsBegin; EvalIndex < (RootStateTree.EvaluatorsBegin + RootStateTree.EvaluatorsNum); EvalIndex++)
-		{
-			const FStateTreeEvaluatorBase& Eval = RootStateTree.Nodes[EvalIndex].Get<const FStateTreeEvaluatorBase>();
-			DebugString += FString::Printf(TEXT("| %-30s | %8d | %10s |\n"),
-				*Eval.Name.ToString(), Eval.BindingsBatch.Get(), *Eval.InstanceDataHandle.Describe());
-		}
-	}
-
-
-	DebugString += FString::Printf(TEXT("\nTasks\n  [ %-30s | %-30s | %8s | %10s ]\n"),
-		TEXT("State"), TEXT("Name"), TEXT("Bindings"), TEXT("Struct Idx"));
-	for (const FCompactStateTreeState& State : RootStateTree.States)
-	{
-		// Tasks
-		if (State.TasksNum)
-		{
-			for (int32 TaskIndex = State.TasksBegin; TaskIndex < (State.TasksBegin + State.TasksNum); TaskIndex++)
-			{
-				const FStateTreeTaskBase& Task = RootStateTree.Nodes[TaskIndex].Get<const FStateTreeTaskBase>();
-				DebugString += FString::Printf(TEXT("  | %-30s | %-30s | %8d | %10s |\n"), *State.Name.ToString(),
-					*Task.Name.ToString(), Task.BindingsBatch.Get(), *Task.InstanceDataHandle.Describe());
-			}
-		}
-	}
-
-	UE_LOG(LogStateTree, Log, TEXT("%s"), *DebugString);
+	UE_LOG(LogStateTree, Log, TEXT("%s"), *RootStateTree.DebugInternalLayoutAsString());
 }
 
 int32 FStateTreeExecutionContext::GetStateChangeCount() const
@@ -4686,7 +4588,7 @@ FString FStateTreeExecutionContext::GetActiveStateName() const
 
 	const FStateTreeExecutionState& Exec = GetExecState();
 
-	FString FullStateName;
+	TStringBuilder<1024> FullStateName;
 	
 	const UStateTree* LastStateTree = &RootStateTree;
 	int32 Indent = 0;
@@ -4699,12 +4601,15 @@ FString FStateTreeExecutionContext::GetActiveStateName() const
 		// Append linked state marker at the end of the previous line.
 		if (Indent > 0)
 		{
-			FullStateName += TEXT(" >");
+			FullStateName << TEXT(" >");
 		}
 		// If tree has changed, append that too.
 		if (CurrentFrame.StateTree != LastStateTree)
 		{
-			FullStateName.Appendf(TEXT(" [%s]"), *GetNameSafe(CurrentFrame.StateTree));
+			FullStateName << TEXT(" [");
+			FullStateName << CurrentFrame.StateTree.GetFName();
+			FullStateName << TEXT(']');
+
 			LastStateTree = CurrentFrame.StateTree;
 		}
 
@@ -4718,8 +4623,8 @@ FString FStateTreeExecutionContext::GetActiveStateName() const
 				{
 					FullStateName += TEXT("\n");
 				}
-				FullStateName += FString::Printf(TEXT("%*s-"), Indent * 3, TEXT("")); // Indent
-				FullStateName += *State.Name.ToString();
+				FullStateName.Appendf(TEXT("%*s-"), Indent * 3, TEXT("")); // Indent
+				FullStateName << State.Name;
 				Indent++;
 			}
 		}
@@ -4728,19 +4633,19 @@ FString FStateTreeExecutionContext::GetActiveStateName() const
 	switch (Exec.TreeRunStatus)
 	{
 	case EStateTreeRunStatus::Failed:
-		FullStateName += TEXT(" FAILED\n");
+		FullStateName << TEXT(" FAILED\n");
 		break;
 	case EStateTreeRunStatus::Succeeded:
-		FullStateName += TEXT(" SUCCEEDED\n");
+		FullStateName << TEXT(" SUCCEEDED\n");
 		break;
 	case EStateTreeRunStatus::Running:
 		// Empty
 		break;
 	default:
-		FullStateName += TEXT("--\n");
+		FullStateName << TEXT("--\n");
 	}
 
-	return FullStateName;
+	return FullStateName.ToString();
 }
 
 TArray<FName> FStateTreeExecutionContext::GetActiveStateNames() const
diff --git a/Engine/Plugins/Runtime/StateTree/Source/StateTreeModule/Private/StateTreePropertyBindings.cpp b/Engine/Plugins/Runtime/StateTree/Source/StateTreeModule/Private/StateTreePropertyBindings.cpp
--- a/Engine/Plugins/Runtime/StateTree/Source/StateTreeModule/Private/StateTreePropertyBindings.cpp
+++ b/Engine/Plugins/Runtime/StateTree/Source/StateTreeModule/Private/StateTreePropertyBindings.cpp
@@ -1256,67 +1256,98 @@ bool FStateTreePropertyBindings::ContainsAnyStruct(const TSet<const UStruct*>& S
 
 void FStateTreePropertyBindings::DebugPrintInternalLayout(FString& OutString) const
 {
+#if WITH_EDITOR || WITH_STATETREE_DEBUG
+	OutString += DebugInternalLayoutAsString();
+#endif
+}
+
+#if WITH_EDITOR || WITH_STATETREE_DEBUG
+FString FStateTreePropertyBindings::DebugInternalLayoutAsString() const
+{
+	FStringBuilderBase DebugString;
 	/** Array of expected source structs. */
-	OutString += FString::Printf(TEXT("\nBindableStructDesc (%d)\n  [ %-40s | %-40s ]\n"), SourceStructs.Num(), TEXT("Type"), TEXT("Name"));
-	for (const FStateTreeBindableStructDesc& BindableStructDesc : SourceStructs)
+	DebugString.Appendf(TEXT("\nBindableStructDesc (%d)\n"), SourceStructs.Num());
+	if (SourceStructs.Num())
 	{
-		OutString += FString::Printf(TEXT("  | %-40s | %-40s |\n"),
-									 BindableStructDesc.Struct ? *BindableStructDesc.Struct->GetName() : TEXT("null"),
-									 *BindableStructDesc.Name.ToString());
+		DebugString.Appendf(TEXT("  [ %-40s | % -40s ]\n"), TEXT("Type"), TEXT("Name"));
+		for (const FStateTreeBindableStructDesc& BindableStructDesc : SourceStructs)
+		{
+			DebugString.Appendf(TEXT("  | %-40s | %-40s |\n"),
+				BindableStructDesc.Struct ? *BindableStructDesc.Struct->GetName() : TEXT("null"),
+				*BindableStructDesc.Name.ToString());
+		}
 	}
 
 	/** Array of copy batches. */
-	OutString += FString::Printf(TEXT("\nCopyBatches (%d)\n  [ %-40s | %-40s | %-8s [%-3s:%-3s[ ]\n"), CopyBatches.Num(),
-		TEXT("Target Type"), TEXT("Target Name"), TEXT("Bindings"), TEXT("Beg"), TEXT("End"));
-	for (const FStateTreePropertyCopyBatch& CopyBatch : CopyBatches)
+	DebugString.Appendf(TEXT("\nCopyBatches (%d)\n"), CopyBatches.Num());
+	if (CopyBatches.Num())
 	{
-		OutString += FString::Printf(TEXT("  | %-40s | %-40s | %8s [%3d:%-3d[ |\n"),
-									 CopyBatch.TargetStruct.Struct ? *CopyBatch.TargetStruct.Struct->GetName() : TEXT("null"),
-									 *CopyBatch.TargetStruct.Name.ToString(),
-									 TEXT(""), CopyBatch.BindingsBegin.Get(), CopyBatch.BindingsEnd.Get());
+		DebugString.Appendf(TEXT("  [ %-40s | %-40s | %-8s [%-3s:%-3s[ | %-8s [%-3s:%-3s[ ]\n"),
+			TEXT("Target Type"), TEXT("Target Name"),
+			TEXT("Bindings"), TEXT("Beg"), TEXT("End"),
+			TEXT("ProFunc"), TEXT("Beg"), TEXT("End"));
+		for (const FStateTreePropertyCopyBatch& CopyBatch : CopyBatches)
+		{
+			DebugString.Appendf(TEXT("  | %-40s | %-40s | %8s [%3d:%-3d[ |\n"),
+				CopyBatch.TargetStruct.Struct ? *CopyBatch.TargetStruct.Struct->GetName() : TEXT("null"),
+				*CopyBatch.TargetStruct.Name.ToString(),
+				TEXT(""), CopyBatch.BindingsBegin.Get(), CopyBatch.BindingsEnd.Get(),
+				TEXT(""), CopyBatch.PropertyFunctionsBegin.Get(), CopyBatch.PropertyFunctionsEnd.Get());
+		}
 	}
 
 	/** Array of property bindings, resolved into arrays of copies before use. */
-	OutString += FString::Printf(TEXT("\nPropertyPathBindings (%d)\n"), PropertyPathBindings.Num());
+	DebugString.Appendf(TEXT("\nPropertyPathBindings (%d)\n"), PropertyPathBindings.Num());
 	for (const FStateTreePropertyPathBinding& PropertyBinding : PropertyPathBindings)
 	{
-		OutString += FString::Printf(TEXT("\n  Source: %s | Target: %s"),
+		DebugString.Appendf(TEXT("  Source: %s | Target: %s\n"),
 					*PropertyBinding.GetSourcePath().ToString(), *PropertyBinding.GetSourcePath().ToString());
 	}
 
 	/** Array of property copies */
-	OutString += FString::Printf(TEXT("\nPropertyCopies (%d)\n  [ %-7s | %-4s | %-4s | %-10s | %-7s | %-4s | %-4s | %-10s | %-10s | %-20s | %-4s ]\n"), PropertyCopies.Num(),
-		TEXT("Src Idx"), TEXT("Off."), TEXT("Next"), TEXT("Type"),
-		TEXT("Tgt Idx"), TEXT("Off."), TEXT("Next"), TEXT("Type"),
-		TEXT("Source"), TEXT("Copy Type"), TEXT("Size"));
-	for (const FStateTreePropertyCopy& PropertyCopy : PropertyCopies)
-	{
-		OutString += FString::Printf(TEXT("  | %7d | %4d | %4d | %-10s | %7d | %4d | %4d | %-10s | %10s | %-20s | %4d |\n"),
-					PropertyCopy.SourceIndirection.ArrayIndex.Get(),
-					PropertyCopy.SourceIndirection.Offset,
-					PropertyCopy.SourceIndirection.NextIndex.Get(),
-					*UEnum::GetDisplayValueAsText(PropertyCopy.SourceIndirection.Type).ToString(),
-					PropertyCopy.TargetIndirection.ArrayIndex.Get(),
-					PropertyCopy.TargetIndirection.Offset,
-					PropertyCopy.TargetIndirection.NextIndex.Get(),
-					*UEnum::GetDisplayValueAsText(PropertyCopy.TargetIndirection.Type).ToString(),
-					*PropertyCopy.SourceDataHandle.Describe(),
-					*UEnum::GetDisplayValueAsText(PropertyCopy.Type).ToString(),
-					PropertyCopy.CopySize);
+	DebugString.Appendf(TEXT("\nPropertyCopies (%d)\n"), PropertyCopies.Num());
+	if (PropertyCopies.Num())
+	{
+		DebugString.Appendf(TEXT("  [ %-7s | %-4s | %-7s | %-10s | %-7s | %-4s | %-7s | %-10s | %-15s | %-20s | %-4s ]\n"),
+			TEXT("Src Idx"), TEXT("Off."), TEXT("Next"), TEXT("Type"),
+			TEXT("Tgt Idx"), TEXT("Off."), TEXT("Next"), TEXT("Type"),
+			TEXT("Source"), TEXT("Copy Type"), TEXT("Size"));
+		for (const FStateTreePropertyCopy& PropertyCopy : PropertyCopies)
+		{
+			DebugString.Appendf(TEXT("  | %7d | %4d | %7d | %-10s | %7d | %4d | %7d | %-10s | %15s | %-20s | %4d |\n"),
+						PropertyCopy.SourceIndirection.ArrayIndex.Get(),
+						PropertyCopy.SourceIndirection.Offset,
+						PropertyCopy.SourceIndirection.NextIndex.Get(),
+						*UEnum::GetDisplayValueAsText(PropertyCopy.SourceIndirection.Type).ToString(),
+						PropertyCopy.TargetIndirection.ArrayIndex.Get(),
+						PropertyCopy.TargetIndirection.Offset,
+						PropertyCopy.TargetIndirection.NextIndex.Get(),
+						*UEnum::GetDisplayValueAsText(PropertyCopy.TargetIndirection.Type).ToString(),
+						*PropertyCopy.SourceDataHandle.Describe(),
+						*UEnum::GetDisplayValueAsText(PropertyCopy.Type).ToString(),
+						PropertyCopy.CopySize);
+		}
 	}
 
 	/** Array of property indirections, indexed by accesses*/
-	OutString += FString::Printf(TEXT("\nPropertyIndirections (%d)\n  [ %-4s | %-4s | %-4s | %-10s ] \n"), PropertyIndirections.Num(),
-		TEXT("Idx"), TEXT("Off."), TEXT("Next"), TEXT("Access Type"));
-	for (const FStateTreePropertyIndirection& PropertyIndirection : PropertyIndirections)
+	DebugString.Appendf(TEXT("\nPropertyIndirections (%d)\n"), PropertyIndirections.Num());
+	if (PropertyIndirections.Num())
 	{
-		OutString += FString::Printf(TEXT("  | %4d | %4d | %4d | %-10s |\n"),
-					PropertyIndirection.ArrayIndex.Get(),
-					PropertyIndirection.Offset,
-					PropertyIndirection.NextIndex.Get(),
-					*UEnum::GetDisplayValueAsText(PropertyIndirection.Type).ToString());
+		DebugString.Appendf(TEXT("[ %-4s | %-4s | %-4s | %-10s ] \n"),
+			TEXT("Idx"), TEXT("Off."), TEXT("Next"), TEXT("Access Type"));
+		for (const FStateTreePropertyIndirection& PropertyIndirection : PropertyIndirections)
+		{
+			DebugString.Appendf(TEXT("  | %4d | %4d | %4d | %-10s |\n"),
+				PropertyIndirection.ArrayIndex.Get(),
+				PropertyIndirection.Offset,
+				PropertyIndirection.NextIndex.Get(),
+				*UEnum::GetDisplayValueAsText(PropertyIndirection.Type).ToString());
+		}
 	}
+
+	return DebugString.ToString();
 }
+#endif // WITH_EDITOR || WITH_STATETREE_DEBUG
 
 
 //----------------------------------------------------------------//
diff --git a/Engine/Plugins/Runtime/StateTree/Source/StateTreeModule/Public/StateTree.h b/Engine/Plugins/Runtime/StateTree/Source/StateTreeModule/Public/StateTree.h
--- a/Engine/Plugins/Runtime/StateTree/Source/StateTreeModule/Public/StateTree.h
+++ b/Engine/Plugins/Runtime/StateTree/Source/StateTreeModule/Public/StateTree.h
@@ -178,6 +178,12 @@ public:
 
 	/** Compile the state trees if the editor hash data as changed since the last compilation. */
 	void CompileIfChanged();
+
+#endif
+
+#if WITH_EDITOR || WITH_STATETREE_DEBUG
+	/** @return the internal content of the state tree compiled asset. */
+	[[nodiscard]] FString DebugInternalLayoutAsString() const;
 #endif
 
 #if WITH_EDITORONLY_DATA
@@ -294,7 +300,6 @@ private:
 	UPROPERTY()
 	FInstancedPropertyBag Parameters;
 
-
 	/** Number of context data, include parameters and all context data. */
 	UPROPERTY()
 	uint16 NumContextData = 0;
diff --git a/Engine/Plugins/Runtime/StateTree/Source/StateTreeModule/Public/StateTreePropertyBindings.h b/Engine/Plugins/Runtime/StateTree/Source/StateTreeModule/Public/StateTreePropertyBindings.h
--- a/Engine/Plugins/Runtime/StateTree/Source/StateTreeModule/Public/StateTreePropertyBindings.h
+++ b/Engine/Plugins/Runtime/StateTree/Source/StateTreeModule/Public/StateTreePropertyBindings.h
@@ -980,8 +980,6 @@ struct STATETREEMODULE_API FStateTreePropertyBindings
 	 * @return true if any of the elements in the property bindings contains any of the structs in the set.
 	 */
 	bool ContainsAnyStruct(const TSet<const UStruct*>& Structs);
-	
-	void DebugPrintInternalLayout(FString& OutString) const;
 
 	/** @return how properties are compatible for copying. */
 	static EStateTreePropertyAccessCompatibility GetPropertyCompatibility(const FProperty* FromProperty, const FProperty* ToProperty);
@@ -1007,6 +1005,13 @@ struct STATETREEMODULE_API FStateTreePropertyBindings
 		return false;
 	}
 
+	UE_DEPRECATED(5.5, "DebugPrintInternalLayout is deprecated. Use GetDebugInfo instead.")
+	void DebugPrintInternalLayout(FString& OutString) const;
+
+#if WITH_EDITOR || WITH_STATETREE_DEBUG
+	[[nodiscard]] FString DebugInternalLayoutAsString() const;
+#endif
+
 private:
 	[[nodiscard]] bool ResolvePath(const UStruct* Struct, const FStateTreePropertyPath& Path, FStateTreePropertyIndirection& OutFirstIndirection, FStateTreePropertyPathIndirection& OutLeafIndirection);
 	const FStateTreeBindableStructDesc* GetSourceDescByHandle(const FStateTreeDataHandle SourceDataHandle);
