StateTree: Set the extension on the DetailView before any object is set.

diff --git a/Engine/Plugins/Runtime/StateTree/Source/StateTreeEditorModule/Private/StateTreeEditorMode.cpp b/Engine/Plugins/Runtime/StateTree/Source/StateTreeEditorModule/Private/StateTreeEditorMode.cpp
--- a/Engine/Plugins/Runtime/StateTree/Source/StateTreeEditorModule/Private/StateTreeEditorMode.cpp
+++ b/Engine/Plugins/Runtime/StateTree/Source/StateTreeEditorModule/Private/StateTreeEditorMode.cpp
@@ -40,7 +40,7 @@ UStateTreeEditorMode::UStateTreeEditorMode()
 void UStateTreeEditorMode::Enter()
 {
 	Super::Enter();
-
+	
 	DetailsViewExtensionHandler = MakeShared<FStateTreeBindingExtension>();
 	DetailsViewChildrenCustomizationHandler = MakeShared<FStateTreeBindingsChildrenCustomization>();
 	
@@ -70,6 +70,7 @@ void UStateTreeEditorMode::Enter()
 				
 				AssetDetailsView->SetExtensionHandler(DetailsViewExtensionHandler);
 				AssetDetailsView->SetChildrenCustomizationHandler(DetailsViewChildrenCustomizationHandler);
+				bForceAssetDetailViewToRefresh = true;
 			}
 		}
 	}
@@ -187,16 +188,17 @@ void UStateTreeEditorMode::Exit()
 			
 			if (TSharedPtr<IDetailsView> DetailsView = GetDetailsView())
 			{
-				DetailsView->OnFinishedChangingProperties().RemoveAll(this);				
+				DetailsView->OnFinishedChangingProperties().RemoveAll(this);
 				DetailsView->SetExtensionHandler(nullptr);
 				DetailsView->SetChildrenCustomizationHandler(nullptr);
 			}
 			
 			if (TSharedPtr<IDetailsView> AssetDetailsView = GetAssetDetailsView())
 			{
-				AssetDetailsView->OnFinishedChangingProperties().RemoveAll(this);				
+				AssetDetailsView->OnFinishedChangingProperties().RemoveAll(this);
 				AssetDetailsView->SetExtensionHandler(nullptr);
 				AssetDetailsView->SetChildrenCustomizationHandler(nullptr);
+				bForceAssetDetailViewToRefresh = true;
 			}
 		}
 	}
@@ -249,16 +251,23 @@ void UStateTreeEditorMode::OnStateTreeChanged()
 					OldViewModel->GetOnSelectionChanged().RemoveAll(this);
 					OldViewModel->GetOnBringNodeToFocus().RemoveAll(this);
 				}
-			}			
+			}
 		}
-		
-		CachedStateTree = Context->EditorHostInterface->GetStateTree();
 
-		if (CachedStateTree.IsValid())
+		UStateTree* StateTree = Context->EditorHostInterface->GetStateTree();
+		CachedStateTree = StateTree;
+
+		if (TSharedPtr<IDetailsView> AssetDetailsView = GetAssetDetailsView())
+		{
+			AssetDetailsView->SetObject(StateTree ? StateTree->EditorData : nullptr, bForceAssetDetailViewToRefresh);
+			bForceAssetDetailViewToRefresh = false;
+		}
+
+		if (StateTree)
 		{
 			if (UStateTreeEditingSubsystem* StateTreeEditingSubsystem = GEditor->GetEditorSubsystem<UStateTreeEditingSubsystem>())
 			{				
-				TSharedRef<FStateTreeViewModel> NewViewModel = StateTreeEditingSubsystem->FindOrAddViewModel(CachedStateTree.Get());
+				TSharedRef<FStateTreeViewModel> NewViewModel = StateTreeEditingSubsystem->FindOrAddViewModel(StateTree);
 				{
 					NewViewModel->GetOnAssetChanged().AddUObject(this, &UStateTreeEditorMode::HandleModelAssetChanged);
 					NewViewModel->GetOnStateAdded().AddUObject(this, &UStateTreeEditorMode::HandleStateAdded);
@@ -384,7 +393,7 @@ bool UStateTreeEditorMode::CanCompile() const
 	{
 		return false;
 	}
-    
+
 	// We can't recompile while in PIE
 	if (GEditor->IsPlaySessionInProgress())
 	{
diff --git a/Engine/Plugins/Runtime/StateTree/Source/StateTreeEditorModule/Public/StateTreeEditorMode.h b/Engine/Plugins/Runtime/StateTree/Source/StateTreeEditorModule/Public/StateTreeEditorMode.h
--- a/Engine/Plugins/Runtime/StateTree/Source/StateTreeEditorModule/Public/StateTreeEditorMode.h
+++ b/Engine/Plugins/Runtime/StateTree/Source/StateTreeEditorModule/Public/StateTreeEditorMode.h
@@ -69,6 +69,7 @@ protected:
 protected:
 	uint32 EditorDataHash = 0;
 	bool bLastCompileSucceeded = true;
+	bool bForceAssetDetailViewToRefresh = false;
 
 	TWeakObjectPtr<UStateTree> CachedStateTree = nullptr;
 
