StateTree: Add developper menu in StateTree editor. Add the option to display the compilation result.#rb Yoan.StAmant
#rnx

diff --git a/Engine/Plugins/Runtime/StateTree/Source/StateTreeEditorModule/Private/StateTreeCompiler.cpp b/Engine/Plugins/Runtime/StateTree/Source/StateTreeEditorModule/Private/StateTreeCompiler.cpp
--- a/Engine/Plugins/Runtime/StateTree/Source/StateTreeEditorModule/Private/StateTreeCompiler.cpp
+++ b/Engine/Plugins/Runtime/StateTree/Source/StateTreeEditorModule/Private/StateTreeCompiler.cpp
@@ -21,7 +21,7 @@
 namespace UE::StateTree::Compiler
 {
 	FAutoConsoleVariable CVarLogCompiledStateTree(
-		TEXT("StateTree.LogCompiledResult"),
+		TEXT("StateTree.Compiler.LogResultOnSuccess"),
 		false,
 		TEXT("After a StateTree compiles, log the internal content of the StateTree.")
 	);
diff --git a/Engine/Plugins/Runtime/StateTree/Source/StateTreeEditorModule/Private/StateTreeEditor.cpp b/Engine/Plugins/Runtime/StateTree/Source/StateTreeEditorModule/Private/StateTreeEditor.cpp
--- a/Engine/Plugins/Runtime/StateTree/Source/StateTreeEditorModule/Private/StateTreeEditor.cpp
+++ b/Engine/Plugins/Runtime/StateTree/Source/StateTreeEditorModule/Private/StateTreeEditor.cpp
@@ -17,6 +17,8 @@
 #include "EdModeInteractiveToolsContext.h"
 #include "IDetailsView.h"
 #include "IContentBrowserSingleton.h"
+#include "ISourceCodeAccessModule.h"
+#include "ISourceCodeAccessor.h"
 #include "Framework/Application/SlateApplication.h"
 #include "Framework/Commands/UICommandList.h"
 #include "IMessageLogListing.h"
@@ -193,6 +195,7 @@ void FStateTreeEditor::InitEditor( const EToolkitMode::Type Mode, const TSharedP
 	constexpr bool bCreateDefaultToolbar = true;
 	FAssetEditorToolkit::InitAssetEditor(Mode, InitToolkitHost, StateTreeEditorAppName, StandaloneDefaultLayout, bCreateDefaultStandaloneMenu, bCreateDefaultToolbar, StateTree);
 
+	RegisterMenu();
 	RegisterToolbar();
 	
 	AddMenuExtender(StateTreeEditorModule.GetMenuExtensibilityManager()->GetAllExtenders(GetToolkitCommands(), GetEditingObjects()));
@@ -346,6 +349,62 @@ void FStateTreeEditor::SaveAsset_Execute()
 	FAssetEditorToolkit::SaveAsset_Execute();
 }
 
+namespace UE::StateTree::Editor::Private
+{
+void FillDeveloperMenu(UToolMenu* InMenu)
+{
+	const FStateTreeEditorCommands& Commands = FStateTreeEditorCommands::Get();
+	{
+		FToolMenuSection& Section = InMenu->AddSection("FileDeveloperCompilerSettings", LOCTEXT("CompileOptionsHeading", "Compiler Settings"));
+		Section.AddMenuEntry(Commands.LogResultOnCompileSuccess);
+	}
+}
+void FillDynamicDeveloperMenu(FToolMenuSection& Section)
+{
+	// Only show the developer menu on machines with the solution (assuming they can build it)
+	ISourceCodeAccessModule* SourceCodeAccessModule = FModuleManager::GetModulePtr<ISourceCodeAccessModule>("SourceCodeAccess");
+	if (SourceCodeAccessModule != nullptr && SourceCodeAccessModule->GetAccessor().CanAccessSourceCode())
+	{
+		Section.AddSubMenu(
+			"DeveloperMenu",
+			LOCTEXT("DeveloperMenu", "Developer"),
+			LOCTEXT("DeveloperMenu_ToolTip", "Open the developer menu"),
+			FNewToolMenuDelegate::CreateStatic(FillDeveloperMenu),
+			false);
+	}
+}
+}
+
+void FStateTreeEditor::RegisterMenu()
+{
+	UToolMenus* ToolMenus = UToolMenus::Get();
+	const FName FileMenuName = *(GetToolMenuName().ToString() + TEXT(".File"));
+	if (!ToolMenus->IsMenuRegistered(FileMenuName))
+	{
+		const FName ParentFileMenuName = TEXT("MainFrame.MainMenu.File");
+		UToolMenus::Get()->RegisterMenu(FileMenuName, ParentFileMenuName, EMultiBoxType::ToolBar);
+		{
+			UToolMenu* FileMenu = UToolMenus::Get()->RegisterMenu(FileMenuName, ParentFileMenuName);
+			const FName FileStateTreeSection = "FileStateTree";
+			FToolMenuSection& Section = FileMenu->AddSection("StateTree", LOCTEXT("StateTreeHeading", "State Tree"));
+			FToolMenuInsert InsertPosition("FileLoadAndSave", EToolMenuInsertType::After);
+			Section.InsertPosition = InsertPosition;
+
+			Section.AddDynamicEntry("FileDeveloper", FNewToolMenuSectionDelegate::CreateStatic(UE::StateTree::Editor::Private::FillDynamicDeveloperMenu));
+		}
+	}
+
+	const FName EditMenuName = *(GetToolMenuName().ToString() + TEXT(".Edit"));
+	if (!UToolMenus::Get()->IsMenuRegistered(EditMenuName))
+	{
+		const FName ParentEditMenuName = TEXT("MainFrame.MainMenu.Edit");
+		UToolMenu* EditMenu = UToolMenus::Get()->RegisterMenu(EditMenuName, ParentEditMenuName);
+		FToolMenuSection& Section = EditMenu->AddSection("StateTree", LOCTEXT("StateTreeHeading", "State Tree"));
+		FToolMenuInsert InsertPosition("Configuration", EToolMenuInsertType::After);
+		Section.InsertPosition = InsertPosition;
+	}
+}
+
 void FStateTreeEditor::RegisterToolbar()
 {
 	UToolMenus* ToolMenus = UToolMenus::Get();
diff --git a/Engine/Plugins/Runtime/StateTree/Source/StateTreeEditorModule/Private/StateTreeEditor.h b/Engine/Plugins/Runtime/StateTree/Source/StateTreeEditorModule/Private/StateTreeEditor.h
--- a/Engine/Plugins/Runtime/StateTree/Source/StateTreeEditorModule/Private/StateTreeEditor.h
+++ b/Engine/Plugins/Runtime/StateTree/Source/StateTreeEditorModule/Private/StateTreeEditor.h
@@ -69,6 +69,7 @@ private:
 	TSharedRef<SDockTab> SpawnTab_AssetDetails(const FSpawnTabArgs& Args);
 	TSharedRef<SDockTab> SpawnTab_CompilerResults(const FSpawnTabArgs& Args) const;
 
+	void RegisterMenu();
 	void RegisterToolbar();
 	
 	/** State Tree being edited */
diff --git a/Engine/Plugins/Runtime/StateTree/Source/StateTreeEditorModule/Private/StateTreeEditorCommands.cpp b/Engine/Plugins/Runtime/StateTree/Source/StateTreeEditorModule/Private/StateTreeEditorCommands.cpp
--- a/Engine/Plugins/Runtime/StateTree/Source/StateTreeEditorModule/Private/StateTreeEditorCommands.cpp
+++ b/Engine/Plugins/Runtime/StateTree/Source/StateTreeEditorModule/Private/StateTreeEditorCommands.cpp
@@ -21,6 +21,7 @@ void FStateTreeEditorCommands::RegisterCommands()
 	UI_COMMAND(SaveOnCompile_Never, "Never", "Sets the save-on-compile option to 'Never', meaning that your StateTree will not be saved when they are compiled", EUserInterfaceActionType::RadioButton, FInputChord());
 	UI_COMMAND(SaveOnCompile_SuccessOnly, "On Success Only", "Sets the save-on-compile option to 'Success Only', meaning that your StateTree will be saved whenever they are successfully compiled", EUserInterfaceActionType::RadioButton, FInputChord());
 	UI_COMMAND(SaveOnCompile_Always, "Always", "Sets the save-on-compile option to 'Always', meaning that your StateTree will be saved whenever they are compiled (even if there were errors)", EUserInterfaceActionType::RadioButton, FInputChord());
+	UI_COMMAND(LogResultOnCompileSuccess, "Log compilation result", "After a StateTree compiles, log the internal content of the StateTree.", EUserInterfaceActionType::ToggleButton, FInputChord());
 
 	UI_COMMAND(AddSiblingState, "Add Sibling State", "Add a Sibling State", EUserInterfaceActionType::Button, FInputChord());
 	UI_COMMAND(AddChildState, "Add Child State", "Add a Child State", EUserInterfaceActionType::Button, FInputChord());
diff --git a/Engine/Plugins/Runtime/StateTree/Source/StateTreeEditorModule/Private/StateTreeEditorCommands.h b/Engine/Plugins/Runtime/StateTree/Source/StateTreeEditorModule/Private/StateTreeEditorCommands.h
--- a/Engine/Plugins/Runtime/StateTree/Source/StateTreeEditorModule/Private/StateTreeEditorCommands.h
+++ b/Engine/Plugins/Runtime/StateTree/Source/StateTreeEditorModule/Private/StateTreeEditorCommands.h
@@ -20,6 +20,7 @@ public:
 	TSharedPtr<FUICommandInfo> SaveOnCompile_Never;
 	TSharedPtr<FUICommandInfo> SaveOnCompile_SuccessOnly;
 	TSharedPtr<FUICommandInfo> SaveOnCompile_Always;
+	TSharedPtr<FUICommandInfo> LogResultOnCompileSuccess;
 
 	TSharedPtr<FUICommandInfo> AddSiblingState;
 	TSharedPtr<FUICommandInfo> AddChildState;
diff --git a/Engine/Plugins/Runtime/StateTree/Source/StateTreeEditorModule/Private/StateTreeEditorMode.cpp b/Engine/Plugins/Runtime/StateTree/Source/StateTreeEditorModule/Private/StateTreeEditorMode.cpp
--- a/Engine/Plugins/Runtime/StateTree/Source/StateTreeEditorModule/Private/StateTreeEditorMode.cpp
+++ b/Engine/Plugins/Runtime/StateTree/Source/StateTreeEditorModule/Private/StateTreeEditorMode.cpp
@@ -5,6 +5,7 @@
 #include "ContextObjectStore.h"
 #include "EditorModeManager.h"
 #include "FileHelpers.h"
+#include "HAL/IConsoleManager.h"
 
 #include "StateTreeEditorModeToolkit.h"
 #include "StateTreeEditorSettings.h"
@@ -301,6 +302,27 @@ static bool IsSaveOnCompileOptionSet(const EStateTreeSaveOnCompile Option)
 	const UStateTreeEditorSettings* Settings = GetDefault<UStateTreeEditorSettings>();
 	return (Settings->SaveOnCompile == Option);
 }
+
+static IConsoleVariable* GetLogResultOnSuccessCVar()
+{
+	static IConsoleVariable* FoundVariable = IConsoleManager::Get().FindConsoleVariable(TEXT("StateTree.Compiler.LogResultOnSuccess"));
+	return FoundVariable;
+}
+
+static void ToggleLogResultOnCompileSuccess()
+{
+	IConsoleVariable* LogResultCVar = GetLogResultOnSuccessCVar();
+	if (ensure(LogResultCVar))
+	{
+		LogResultCVar->Set(!LogResultCVar->GetBool(), ECVF_SetByConsole);
+	}
+}
+
+static bool IsLogResultOnCompileSuccess()
+{
+	IConsoleVariable* LogResultCVar = GetLogResultOnSuccessCVar();
+	return LogResultCVar ? LogResultCVar->GetBool() : false;
+}
 }
 
 void UStateTreeEditorMode::BindToolkitCommands(const TSharedRef<FUICommandList>& ToolkitCommands)
@@ -336,6 +358,12 @@ void UStateTreeEditorMode::BindToolkitCommands(const TSharedRef<FUICommandList>&
 		FIsActionChecked::CreateStatic(&UE::StateTree::Editor::Internal::IsSaveOnCompileOptionSet,  EStateTreeSaveOnCompile::Always),
 		FIsActionButtonVisible::CreateUObject(this, &UStateTreeEditorMode::HasValidStateTree)
 	);
+	ToolkitCommands->MapAction(
+		FStateTreeEditorCommands::Get().LogResultOnCompileSuccess,
+		FExecuteAction::CreateStatic(&UE::StateTree::Editor::Internal::ToggleLogResultOnCompileSuccess),
+		FCanExecuteAction(),
+		FIsActionChecked::CreateStatic(&UE::StateTree::Editor::Internal::IsLogResultOnCompileSuccess)
+	);
 }
 
 void UStateTreeEditorMode::BindCommands()
diff --git a/Engine/Plugins/Runtime/StateTree/Source/StateTreeEditorModule/StateTreeEditorModule.Build.cs b/Engine/Plugins/Runtime/StateTree/Source/StateTreeEditorModule/StateTreeEditorModule.Build.cs
--- a/Engine/Plugins/Runtime/StateTree/Source/StateTreeEditorModule/StateTreeEditorModule.Build.cs
+++ b/Engine/Plugins/Runtime/StateTree/Source/StateTreeEditorModule/StateTreeEditorModule.Build.cs
@@ -44,6 +44,7 @@ namespace UnrealBuildTool.Rules
 				"GraphEditor",
 				"KismetWidgets",
 				"PropertyPath",
+				"SourceCodeAccess",
 				"ToolMenus",
 				"ToolWidgets",
 				"ApplicationCore",
