StateTree: Rework the FStateTreePropertyBindings::DebugInternalLayoutAsString format.

diff --git a/Engine/Plugins/Runtime/StateTree/Source/StateTreeModule/Private/StateTreePropertyBindings.cpp b/Engine/Plugins/Runtime/StateTree/Source/StateTreeModule/Private/StateTreePropertyBindings.cpp
--- a/Engine/Plugins/Runtime/StateTree/Source/StateTreeModule/Private/StateTreePropertyBindings.cpp
+++ b/Engine/Plugins/Runtime/StateTree/Source/StateTreeModule/Private/StateTreePropertyBindings.cpp
@@ -1266,15 +1266,16 @@ FString FStateTreePropertyBindings::DebugInternalLayoutAsString() const
 {
 	FStringBuilderBase DebugString;
 	/** Array of expected source structs. */
-	DebugString.Appendf(TEXT("\nBindableStructDesc (%d)\n"), SourceStructs.Num());
+	DebugString.Appendf(TEXT("\nSourceStructs (%d)\n"), SourceStructs.Num());
 	if (SourceStructs.Num())
 	{
-		DebugString.Appendf(TEXT("  [ %-40s | % -40s ]\n"), TEXT("Type"), TEXT("Name"));
+		DebugString.Appendf(TEXT("  [ %-45s | % -40s | % -55s ]\n"), TEXT("Type"), TEXT("DataSource"), TEXT("Name"));
 		for (const FStateTreeBindableStructDesc& BindableStructDesc : SourceStructs)
 		{
-			DebugString.Appendf(TEXT("  | %-40s | %-40s |\n"),
+			DebugString.Appendf(TEXT("  | %-45s | %-40s | % -55s |\n"),
 				BindableStructDesc.Struct ? *BindableStructDesc.Struct->GetName() : TEXT("null"),
-				*BindableStructDesc.Name.ToString());
+				*UEnum::GetDisplayValueAsText(BindableStructDesc.DataSource).ToString(),
+				*BindableStructDesc.ToString());
 		}
 	}
 
@@ -1282,15 +1283,15 @@ FString FStateTreePropertyBindings::DebugInternalLayoutAsString() const
 	DebugString.Appendf(TEXT("\nCopyBatches (%d)\n"), CopyBatches.Num());
 	if (CopyBatches.Num())
 	{
-		DebugString.Appendf(TEXT("  [ %-40s | %-40s | %-8s [%-3s:%-3s[ | %-8s [%-3s:%-3s[ ]\n"),
+		DebugString.Appendf(TEXT("  [ %-45s | %-45s | %-8s [%-5s:%-5s[ | %-8s [%-5s:%-5s[ ]\n"),
 			TEXT("Target Type"), TEXT("Target Name"),
 			TEXT("Bindings"), TEXT("Beg"), TEXT("End"),
 			TEXT("ProFunc"), TEXT("Beg"), TEXT("End"));
 		for (const FStateTreePropertyCopyBatch& CopyBatch : CopyBatches)
 		{
-			DebugString.Appendf(TEXT("  | %-40s | %-40s | %8s [%3d:%-3d[ |\n"),
+			DebugString.Appendf(TEXT("  | %-45s | %-45s | %8s [%5d:%-5d[ | %8s [%5d:%-5d[ |\n"),
 				CopyBatch.TargetStruct.Struct ? *CopyBatch.TargetStruct.Struct->GetName() : TEXT("null"),
-				*CopyBatch.TargetStruct.Name.ToString(),
+				*CopyBatch.TargetStruct.ToString(),
 				TEXT(""), CopyBatch.BindingsBegin.Get(), CopyBatch.BindingsEnd.Get(),
 				TEXT(""), CopyBatch.PropertyFunctionsBegin.Get(), CopyBatch.PropertyFunctionsEnd.Get());
 		}
@@ -1298,10 +1299,17 @@ FString FStateTreePropertyBindings::DebugInternalLayoutAsString() const
 
 	/** Array of property bindings, resolved into arrays of copies before use. */
 	DebugString.Appendf(TEXT("\nPropertyPathBindings (%d)\n"), PropertyPathBindings.Num());
-	for (const FStateTreePropertyPathBinding& PropertyBinding : PropertyPathBindings)
+	if (PropertyPathBindings.Num())
 	{
-		DebugString.Appendf(TEXT("  Source: %s | Target: %s\n"),
-					*PropertyBinding.GetSourcePath().ToString(), *PropertyBinding.GetSourcePath().ToString());
+		DebugString.Appendf(TEXT("  [ %-45s | %-45s | %-45s ]\n"),
+			TEXT("Source"), TEXT("SrcPath"), TEXT("TargetPath"));
+		for (const FStateTreePropertyPathBinding& PropertyBinding : PropertyPathBindings)
+		{
+			DebugString.Appendf(TEXT("  | %-45s | %-45s | %-45s |\n"),
+				*PropertyBinding.GetSourceDataHandle().Describe(),
+				*PropertyBinding.GetSourcePath().ToString(),
+				*PropertyBinding.GetTargetPath().ToString());
+		}
 	}
 
 	/** Array of property copies */
@@ -1329,7 +1337,7 @@ FString FStateTreePropertyBindings::DebugInternalLayoutAsString() const
 		}
 	}
 
-	/** Array of property indirections, indexed by accesses*/
+	/** Array of property indirections, indexed by accesses */
 	DebugString.Appendf(TEXT("\nPropertyIndirections (%d)\n"), PropertyIndirections.Num());
 	if (PropertyIndirections.Num())
 	{
