[State Tree] Added state weight to scale final utility score

diff --git a/Engine/Plugins/Runtime/StateTree/Source/StateTreeEditorModule/Private/Customizations/StateTreeStateDetails.cpp b/Engine/Plugins/Runtime/StateTree/Source/StateTreeEditorModule/Private/Customizations/StateTreeStateDetails.cpp
--- a/Engine/Plugins/Runtime/StateTree/Source/StateTreeEditorModule/Private/Customizations/StateTreeStateDetails.cpp
+++ b/Engine/Plugins/Runtime/StateTree/Source/StateTreeEditorModule/Private/Customizations/StateTreeStateDetails.cpp
@@ -84,6 +84,7 @@ void FStateTreeStateDetails::CustomizeDetails(IDetailLayoutBuilder& DetailBuilde
 	const TSharedPtr<IPropertyHandle> SelectionBehaviorProperty = DetailBuilder.GetProperty(GET_MEMBER_NAME_CHECKED(UStateTreeState, SelectionBehavior));
 	const TSharedPtr<IPropertyHandle> RequiredEventToEnterProperty = DetailBuilder.GetProperty(GET_MEMBER_NAME_CHECKED(UStateTreeState, RequiredEventToEnter));
 	const TSharedPtr<IPropertyHandle> CheckPrerequisitesProperty = DetailBuilder.GetProperty(GET_MEMBER_NAME_CHECKED(UStateTreeState, bCheckPrerequisitesWhenActivatingChildDirectly));
+	const TSharedPtr<IPropertyHandle> WeightProperty = DetailBuilder.GetProperty(GET_MEMBER_NAME_CHECKED(UStateTreeState, Weight));
 
 	PropUtils = DetailBuilder.GetPropertyUtilities();
 	
@@ -286,6 +287,7 @@ void FStateTreeStateDetails::CustomizeDetails(IDetailLayoutBuilder& DetailBuilde
 		DetailBuilder.EditCategory(EnterConditionsCategoryName).SetCategoryVisibility(false);
 	}
 
+	// Utility
 	ConsiderationsProperty->MarkHiddenByCustomization();
 	if (Schema && Schema->AllowUtilityConsiderations())
 	{
@@ -300,6 +302,10 @@ void FStateTreeStateDetails::CustomizeDetails(IDetailLayoutBuilder& DetailBuilde
 			UE::StateTree::Colors::Orange.WithAlpha(192),
 			LOCTEXT("UtilityAddTooltip", "Add new Utility Consideration"),
 			/*SortOrder*/3);
+
+		// Weight
+		WeightProperty->MarkHiddenByCustomization();
+		UtilityConsiderationsCategory.AddProperty(WeightProperty);
 	}
 
 	// Tasks
diff --git a/Engine/Plugins/Runtime/StateTree/Source/StateTreeEditorModule/Private/StateTreeCompiler.cpp b/Engine/Plugins/Runtime/StateTree/Source/StateTreeEditorModule/Private/StateTreeCompiler.cpp
--- a/Engine/Plugins/Runtime/StateTree/Source/StateTreeEditorModule/Private/StateTreeCompiler.cpp
+++ b/Engine/Plugins/Runtime/StateTree/Source/StateTreeEditorModule/Private/StateTreeCompiler.cpp
@@ -447,6 +447,7 @@ bool FStateTreeCompiler::CreateStateRecursive(UStateTreeState& State, const FSta
 	CompactState.Parent = Parent;
 	CompactState.bEnabled = State.bEnabled;
 	CompactState.bCheckPrerequisitesWhenActivatingChildDirectly = State.bCheckPrerequisitesWhenActivatingChildDirectly;
+	CompactState.Weight = State.Weight;
 
 	CompactState.Type = State.Type;
 	CompactState.SelectionBehavior = State.SelectionBehavior;
diff --git a/Engine/Plugins/Runtime/StateTree/Source/StateTreeEditorModule/Public/StateTreeState.h b/Engine/Plugins/Runtime/StateTree/Source/StateTreeEditorModule/Public/StateTreeState.h
--- a/Engine/Plugins/Runtime/StateTree/Source/StateTreeEditorModule/Public/StateTreeState.h
+++ b/Engine/Plugins/Runtime/StateTree/Source/StateTreeEditorModule/Public/StateTreeState.h
@@ -366,6 +366,10 @@ public:
 	/** Defines the event required to be present during state selection for the state to be selected. */
 	UPROPERTY(EditDefaultsOnly, Category = "Enter Conditions", meta = (EditCondition = "bHasRequiredEventToEnter"))
 	FStateTreeEventDesc RequiredEventToEnter;
+	
+	/** Weight used to scale the normalized final utility score for this state */
+	UPROPERTY(EditDefaultsOnly, Category = "Utility", meta=(ClampMin=0))
+	float Weight = 1.f;
 
 	/** Expression of enter conditions that needs to evaluate true to allow the state to be selected. */
 	UPROPERTY(EditDefaultsOnly, Category = "Enter Conditions", meta = (BaseStruct = "/Script/StateTreeModule.StateTreeConditionBase", BaseClass = "/Script/StateTreeModule.StateTreeConditionBlueprintBase"))
diff --git a/Engine/Plugins/Runtime/StateTree/Source/StateTreeModule/Private/StateTreeExecutionContext.cpp b/Engine/Plugins/Runtime/StateTree/Source/StateTreeModule/Private/StateTreeExecutionContext.cpp
--- a/Engine/Plugins/Runtime/StateTree/Source/StateTreeModule/Private/StateTreeExecutionContext.cpp
+++ b/Engine/Plugins/Runtime/StateTree/Source/StateTreeModule/Private/StateTreeExecutionContext.cpp
@@ -2783,7 +2783,7 @@ bool FStateTreeExecutionContext::TestAllConditions(const FStateTreeExecutionFram
 	return Values[0];
 }
 
-float FStateTreeExecutionContext::EvaluateUtility(const FStateTreeExecutionFrame* CurrentParentFrame, const FStateTreeExecutionFrame& CurrentFrame, const int32 ConsiderationsOffset, const int32 ConsiderationsNum)
+float FStateTreeExecutionContext::EvaluateUtility(const FStateTreeExecutionFrame* CurrentParentFrame, const FStateTreeExecutionFrame& CurrentFrame, const int32 ConsiderationsOffset, const int32 ConsiderationsNum, const float StateWeight)
 {
 	// @todo: Tracing support
 	CSV_SCOPED_TIMING_STAT_EXCLUSIVE(StateTree_EvaluateUtility);
@@ -2860,7 +2860,7 @@ float FStateTreeExecutionContext::EvaluateUtility(const FStateTreeExecutionFrame
 		}
 	}
 
-	return Values[0];
+	return StateWeight * Values[0];
 }
 
 void FStateTreeExecutionContext::EvaluatePropertyFunctionsOnActiveInstances(const FStateTreeExecutionFrame* CurrentParentFrame, const FStateTreeExecutionFrame& CurrentFrame, FStateTreeIndex16 FuncsBegin, uint16 FuncsNum)
@@ -4247,7 +4247,7 @@ bool FStateTreeExecutionContext::SelectStateInternal(
 					{
 						const uint16 CurrentStateIndex = NextLevelChildStates[Index];
 						const FCompactStateTreeState& CurrentState = CurrentStateTree->States[CurrentStateIndex];
-						const float Score = EvaluateUtility(CurrentParentFrame, CurrentFrame, CurrentState.UtilityConsiderationsBegin, CurrentState.UtilityConsiderationsNum);
+						const float Score = EvaluateUtility(CurrentParentFrame, CurrentFrame, CurrentState.UtilityConsiderationsBegin, CurrentState.UtilityConsiderationsNum, CurrentState.Weight);
 						if (Score > HighestScore)
 						{
 							HighestScore = Score;
@@ -4297,7 +4297,7 @@ bool FStateTreeExecutionContext::SelectStateInternal(
 				for (uint16 CurrentStateIndex = NextState.ChildrenBegin; CurrentStateIndex < NextState.ChildrenEnd; CurrentStateIndex = CurrentStateTree->States[CurrentStateIndex].GetNextSibling())
 				{
 					const FCompactStateTreeState& CurrentState = CurrentStateTree->States[CurrentStateIndex];
-					const float CurrentStateScore = EvaluateUtility(CurrentParentFrame, CurrentFrame, CurrentState.UtilityConsiderationsBegin, CurrentState.UtilityConsiderationsNum);
+					const float CurrentStateScore = EvaluateUtility(CurrentParentFrame, CurrentFrame, CurrentState.UtilityConsiderationsBegin, CurrentState.UtilityConsiderationsNum, CurrentState.Weight);
 					NextLevelChildStates.Emplace(CurrentStateIndex, CurrentStateScore);
 					TotalScore += CurrentStateScore;
 				}
diff --git a/Engine/Plugins/Runtime/StateTree/Source/StateTreeModule/Public/StateTreeExecutionContext.h b/Engine/Plugins/Runtime/StateTree/Source/StateTreeModule/Public/StateTreeExecutionContext.h
--- a/Engine/Plugins/Runtime/StateTree/Source/StateTreeModule/Public/StateTreeExecutionContext.h
+++ b/Engine/Plugins/Runtime/StateTree/Source/StateTreeModule/Public/StateTreeExecutionContext.h
@@ -595,7 +595,7 @@ protected:
 	 * Calculate the final score of all considerations at given range
 	 * @return the final score
 	 */
-	float EvaluateUtility(const FStateTreeExecutionFrame* CurrentParentFrame, const FStateTreeExecutionFrame& CurrentFrame, const int32 ConsiderationsOffset, const int32 ConsiderationsNum);
+	float EvaluateUtility(const FStateTreeExecutionFrame* CurrentParentFrame, const FStateTreeExecutionFrame& CurrentFrame, const int32 ConsiderationsOffset, const int32 ConsiderationsNum, const float StateWeight);
 
 	/* Evaluate all function at given range. Should be used only on active instances, assumes valid handles and does not consider temporary instances. */
 	void EvaluatePropertyFunctionsOnActiveInstances(const FStateTreeExecutionFrame* CurrentParentFrame, const FStateTreeExecutionFrame& CurrentFrame, FStateTreeIndex16 FuncsBegin, uint16 FuncsNum);
diff --git a/Engine/Plugins/Runtime/StateTree/Source/StateTreeModule/Public/StateTreeTypes.h b/Engine/Plugins/Runtime/StateTree/Source/StateTreeModule/Public/StateTreeTypes.h
--- a/Engine/Plugins/Runtime/StateTree/Source/StateTreeModule/Public/StateTreeTypes.h
+++ b/Engine/Plugins/Runtime/StateTree/Source/StateTreeModule/Public/StateTreeTypes.h
@@ -763,6 +763,10 @@ struct STATETREEMODULE_API FCompactStateTreeState
 	/** Should state's required event and enter conditions be evaluated when transition leads directly to it's child. */
 	UPROPERTY()
 	uint8 bCheckPrerequisitesWhenActivatingChildDirectly : 1 = false;
+	
+	/** Weight used to scale the normalized final utility score for this state */
+	UPROPERTY()
+	float Weight = 1.f;
 
 	/** True if the state is Enabled (i.e. not explicitly marked as disabled). */
 	UPROPERTY()
