StateTree: Fixed not displaying last property in binding if source is not a property function.
#rnx

diff --git a/Engine/Plugins/Runtime/StateTree/Source/StateTreeEditorModule/Private/Customizations/StateTreeBindingExtension.cpp b/Engine/Plugins/Runtime/StateTree/Source/StateTreeEditorModule/Private/Customizations/StateTreeBindingExtension.cpp
--- a/Engine/Plugins/Runtime/StateTree/Source/StateTreeEditorModule/Private/Customizations/StateTreeBindingExtension.cpp
+++ b/Engine/Plugins/Runtime/StateTree/Source/StateTreeEditorModule/Private/Customizations/StateTreeBindingExtension.cpp
@@ -739,8 +739,8 @@ struct FCachedBindingData : public TSharedFromThis<FCachedBindingData>
 			FStateTreeBindableStructDesc SourceDesc;
 			if (BindingOwner->GetStructByID(SourcePath.GetStructID(), SourceDesc))
 			{
-				// Making first segment of the path invisible for the user if it's struct's single output property.
-				if (UE::StateTree::GetStructSingleOutputProperty(*SourceDesc.Struct))
+				// Making first segment of the path invisible for the user if it's property function's single output property.
+				if (SourceDesc.DataSource == EStateTreeBindableStructSource::PropertyFunction && UE::StateTree::GetStructSingleOutputProperty(*SourceDesc.Struct))
 				{
 					SourcePropertyPathAsString = SourcePath.ToString(/*HighlightedSegment*/ INDEX_NONE, /*HighlightPrefix*/ nullptr, /*HighlightPostfix*/ nullptr, /*bOutputInstances*/ false, 1);
 				}
@@ -1444,10 +1444,11 @@ private:
 bool FStateTreeBindingExtension::IsPropertyExtendable(const UClass* InObjectClass, const IPropertyHandle& PropertyHandle) const
 {
 	const FProperty* Property = PropertyHandle.GetProperty();
-	if (Property->HasAnyPropertyFlags(CPF_PersistentInstance | CPF_EditorOnly | CPF_Config))
+	if (Property->HasAnyPropertyFlags(CPF_PersistentInstance | CPF_EditorOnly | CPF_Config | CPF_Deprecated))
 	{
 		return false;
 	}
+// private
 	
 	FStateTreePropertyPath TargetPath;
 	// Figure out the structs we're editing, and property path relative to current property.
