Add ability to set StateTree schema class used by its asset factory

diff --git a/Engine/Plugins/Runtime/StateTree/Source/StateTreeEditorModule/Private/StateTreeFactory.cpp b/Engine/Plugins/Runtime/StateTree/Source/StateTreeEditorModule/Private/StateTreeFactory.cpp
--- a/Engine/Plugins/Runtime/StateTree/Source/StateTreeEditorModule/Private/StateTreeFactory.cpp
+++ b/Engine/Plugins/Runtime/StateTree/Source/StateTreeEditorModule/Private/StateTreeFactory.cpp
@@ -104,6 +104,11 @@ bool UStateTreeFactory::ConfigureProperties()
 	return bResult;
 }
 
+void UStateTreeFactory::SetSchemaClass(const TObjectPtr<UClass>& InSchemaClass)
+{
+	StateTreeSchemaClass = InSchemaClass;
+}
+
 UObject* UStateTreeFactory::FactoryCreateNew(UClass* Class, UObject* InParent, FName Name, EObjectFlags Flags, UObject* Context, FFeedbackContext* Warn)
 {
 	if (StateTreeSchemaClass == nullptr)
diff --git a/Engine/Plugins/Runtime/StateTree/Source/StateTreeEditorModule/Private/StateTreeFactory.h b/Engine/Plugins/Runtime/StateTree/Source/StateTreeEditorModule/Private/StateTreeFactory.h
deleted file mode 100644
diff --git a/Engine/Plugins/Runtime/StateTree/Source/StateTreeEditorModule/Public/StateTreeFactory.h b/Engine/Plugins/Runtime/StateTree/Source/StateTreeEditorModule/Public/StateTreeFactory.h
new file mode 100644
--- /dev/null
+++ b/Engine/Plugins/Runtime/StateTree/Source/StateTreeEditorModule/Public/StateTreeFactory.h
@@ -0,0 +1,30 @@
+// Copyright Epic Games, Inc. All Rights Reserved.
+
+#pragma once
+
+#include "Factories/Factory.h"
+#include "StateTreeFactory.generated.h"
+
+class UStateTreeSchema;
+
+/**
+ * Factory for UStateTree
+ */
+
+UCLASS()
+class STATETREEEDITORMODULE_API UStateTreeFactory : public UFactory
+{
+	GENERATED_UCLASS_BODY()
+
+	//~ UFactory interface
+	virtual bool ConfigureProperties() override;
+	virtual UObject* FactoryCreateNew(UClass* Class, UObject* InParent, FName Name, EObjectFlags Flags, UObject* Context, FFeedbackContext* Warn) override;
+	//~ End of UFactory interface
+
+	void SetSchemaClass(const TObjectPtr<UClass>& InSchemaClass);  
+
+protected:
+	
+	UPROPERTY(Transient)
+	TObjectPtr<UClass> StateTreeSchemaClass;
+};
