StateTree: Remove all objects that a StateTree compilation creates when recompiling. That helps to keep the object name consistent between recompilation and help reduce the diff cook package size.

diff --git a/Engine/Plugins/Runtime/StateTree/Source/StateTreeModule/Private/StateTree.cpp b/Engine/Plugins/Runtime/StateTree/Source/StateTreeModule/Private/StateTree.cpp
--- a/Engine/Plugins/Runtime/StateTree/Source/StateTreeModule/Private/StateTree.cpp
+++ b/Engine/Plugins/Runtime/StateTree/Source/StateTreeModule/Private/StateTree.cpp
@@ -21,6 +21,7 @@
 #include "StateTreeModuleImpl.h"
 #include "StructUtils/UserDefinedStruct.h"
 #include "StructUtilsDelegates.h"
+#include "UObject/LinkerLoad.h"
 #endif
 
 #include UE_INLINE_GENERATED_CPP_BY_NAME(StateTree)
@@ -152,6 +153,21 @@ bool UStateTree::HasCompatibleContextData(const UStateTree& Other) const
 
 
 #if WITH_EDITOR
+namespace UE::StateTree::Compiler
+{
+	void RenameObjectToTransientPackage(UObject* ObjectToRename)
+	{
+		const ERenameFlags RenFlags = REN_DoNotDirty | REN_DontCreateRedirectors;
+
+		ObjectToRename->SetFlags(RF_Transient);
+		ObjectToRename->ClearFlags(RF_Public | RF_Standalone | RF_ArchetypeObject);
+
+		// Rename will remove the renamed object's linker when moving to a new package so invalidate the export beforehand
+		FLinkerLoad::InvalidateExport(ObjectToRename);
+		ObjectToRename->Rename(nullptr, GetTransientPackage(), RenFlags);
+	}
+}
+
 void UStateTree::ResetCompiled()
 {
 	Schema = nullptr;
@@ -175,6 +191,24 @@ void UStateTree::ResetCompiled()
 	bHasGlobalTransitionTasks = false;
 	
 	ResetLinked();
+
+	// Remove objects created from last compilation.
+	{
+		TArray<UObject*, TInlineAllocator<32>> Children;
+		const bool bIncludeNestedObjects = false;
+		ForEachObjectWithOuter(this, [&Children, EditorData = EditorData.Get()](UObject* Child)
+			{
+				if (Child != EditorData)
+				{
+					Children.Add(Child);
+				}
+			}, bIncludeNestedObjects);
+
+		for (UObject* Child : Children)
+		{
+			UE::StateTree::Compiler::RenameObjectToTransientPackage(Child);
+		}
+	}
 }
 
 void UStateTree::OnObjectsReinstanced(const FReplacementObjectMap& ObjectMap)
