StateTree: Fixed ExitState calls in parent frames.

diff --git a/Engine/Plugins/Runtime/StateTree/Source/StateTreeModule/Private/StateTreeExecutionContext.cpp b/Engine/Plugins/Runtime/StateTree/Source/StateTreeModule/Private/StateTreeExecutionContext.cpp
--- a/Engine/Plugins/Runtime/StateTree/Source/StateTreeModule/Private/StateTreeExecutionContext.cpp
+++ b/Engine/Plugins/Runtime/StateTree/Source/StateTreeModule/Private/StateTreeExecutionContext.cpp
@@ -1851,7 +1851,9 @@ void FStateTreeExecutionContext::ExitState(const FStateTreeTransitionResult& Tra
 			}
 
 			const bool bRemainsActive = NextHandle == CurrentHandle;
-			bOnTargetBranch = bOnTargetBranch || NextHandle == Transition.TargetState;
+			const bool bCurrentIsTarget = CurrentFrame.StateTree == Transition.SourceStateTree && CurrentFrame.RootState == Transition.SourceRootState && CurrentHandle == Transition.TargetState;
+			bOnTargetBranch = bOnTargetBranch || bCurrentIsTarget;
+
 			const EStateTreeStateChangeType ChangeType = bRemainsActive ? EStateTreeStateChangeType::Sustained : EStateTreeStateChangeType::Changed;
 			
 			// Should call ExitState() on this state.
