StateTree: Remove temporary log.
#rnx

diff --git a/Engine/Plugins/Runtime/StateTree/Source/StateTreeEditorModule/Private/StateTreeViewModel.cpp b/Engine/Plugins/Runtime/StateTree/Source/StateTreeEditorModule/Private/StateTreeViewModel.cpp
--- a/Engine/Plugins/Runtime/StateTree/Source/StateTreeEditorModule/Private/StateTreeViewModel.cpp
+++ b/Engine/Plugins/Runtime/StateTree/Source/StateTreeEditorModule/Private/StateTreeViewModel.cpp
@@ -26,7 +26,6 @@ namespace UE::StateTree::Editor
 
 		virtual bool CanCreateClass(UClass* InObjectClass, bool& bOmitSubObjs) const override
 		{
-			UE_LOG(LogTemp, Error, TEXT("*** CanCreateClass: %s"), *GetNameSafe(InObjectClass));
 			return InObjectClass->IsChildOf(UStateTreeState::StaticClass())
 				|| InObjectClass->IsChildOf(UStateTreeClipboardBindings::StaticClass());
 		}
