SlateTree: Fix bad access to the editor.
#rnx

diff --git a/Engine/Plugins/Runtime/StateTree/Source/StateTreeEditorModule/Private/StateTreeEditorModule.cpp b/Engine/Plugins/Runtime/StateTree/Source/StateTreeEditorModule/Private/StateTreeEditorModule.cpp
--- a/Engine/Plugins/Runtime/StateTree/Source/StateTreeEditorModule/Private/StateTreeEditorModule.cpp
+++ b/Engine/Plugins/Runtime/StateTree/Source/StateTreeEditorModule/Private/StateTreeEditorModule.cpp
@@ -43,15 +43,8 @@ namespace UE::StateTree::Editor
 	// @todo Could we make this a IModularFeature?
 	static bool CompileStateTree(UStateTree& StateTree)
 	{
-		if (UStateTreeEditingSubsystem* StateTreeEditingSubsystem = GEditor->GetEditorSubsystem<UStateTreeEditingSubsystem>())
-		{
-			FStateTreeCompilerLog Log;
-			const bool bSuccess = UStateTreeEditingSubsystem::CompileStateTree(&StateTree, Log);
-
-			return bSuccess;
-		}
-
-		return false;
+		FStateTreeCompilerLog Log;
+		return UStateTreeEditingSubsystem::CompileStateTree(&StateTree, Log);
 	}
 
 }; // UE::StateTree::Editor
