StateTree: Remove binding that can overrides the target of another binding. Prevents binding that target a struct field to overrides the result of another binding that target the struct itself.
#rnx

diff --git a/Engine/Plugins/Runtime/StateTree/Source/StateTreeEditorModule/Private/StateTreeCompiler.cpp b/Engine/Plugins/Runtime/StateTree/Source/StateTreeEditorModule/Private/StateTreeCompiler.cpp
--- a/Engine/Plugins/Runtime/StateTree/Source/StateTreeEditorModule/Private/StateTreeCompiler.cpp
+++ b/Engine/Plugins/Runtime/StateTree/Source/StateTreeEditorModule/Private/StateTreeCompiler.cpp
@@ -1996,6 +1996,43 @@ bool FStateTreeCompiler::GetAndValidateBindings(const FStateTreeBindableStructDe
 		}
 	}
 
+	// In the UI a user can add a binding X = A.B and X.Y = C.D. The second overrides data that the first one set. The order is also not guaranteed.
+	// Remove all the CopyBindings that override.
+	for (int32 FirstIndex = OutCopyBindings.Num() - 1; FirstIndex >= 0; --FirstIndex)
+	{
+		for (int32 SecondIndex = FirstIndex - 1; SecondIndex >= 0; --SecondIndex)
+		{
+			const FStateTreePropertyPath& BindingTargetA = OutCopyBindings[FirstIndex].GetTargetPath();
+			const FStateTreePropertyPath& BindingTargetB = OutCopyBindings[SecondIndex].GetTargetPath();
+			if (BindingTargetA.GetStructID() == BindingTargetB.GetStructID())
+			{
+				// X.Y.ZA = A.B and X.Y.ZB = A.B is accepted but X.Y.Z and X.Y is not accepted.
+				bool bAccepted = false;
+				for (int32 SegmentIndex = 0; SegmentIndex < BindingTargetA.NumSegments() && SegmentIndex < BindingTargetB.NumSegments(); ++SegmentIndex)
+				{
+					if (BindingTargetA.GetSegment(SegmentIndex) != BindingTargetB.GetSegment(SegmentIndex))
+					{
+						bAccepted = true;
+						continue;
+					}
+				}
+
+				if (bAccepted == false)
+				{
+					// Remove the longest path or the last path (if both are the same length)
+					const bool bTest = BindingTargetA.NumSegments() >= BindingTargetB.NumSegments();
+					const int32 IndexToRemove = bTest ? FirstIndex : SecondIndex;
+					const int32 IndexToKeep = bTest ? SecondIndex : FirstIndex;
+					Log.Reportf(EMessageSeverity::Warning, TargetStruct,
+						TEXT("The binding for target '%s' overrides the target '%s'."),
+						*OutCopyBindings[IndexToRemove].GetTargetPath().ToString(), *OutCopyBindings[IndexToKeep].GetTargetPath().ToString());
+
+					OutCopyBindings.RemoveAt(IndexToRemove); // Do not swap to make it stable
+					--FirstIndex;
+				}
+			}
+		}
+	}
 
 	auto IsPropertyBound = [](const FName& PropertyName, TConstArrayView<FStateTreePropertyPathBinding> Bindings)
 	{
