StateTree: Add accessor to get the list of active states.

diff --git a/Engine/Plugins/Runtime/GameplayStateTree/Source/GameplayStateTreeModule/Private/Components/StateTreeComponent.cpp b/Engine/Plugins/Runtime/GameplayStateTree/Source/GameplayStateTreeModule/Private/Components/StateTreeComponent.cpp
--- a/Engine/Plugins/Runtime/GameplayStateTree/Source/GameplayStateTreeModule/Private/Components/StateTreeComponent.cpp
+++ b/Engine/Plugins/Runtime/GameplayStateTree/Source/GameplayStateTreeModule/Private/Components/StateTreeComponent.cpp
@@ -343,7 +343,17 @@ FString UStateTreeComponent::GetDebugInfoString() const
 		return FString("No StateTree to run.");
 	}
 
-	return FStateTreeExecutionContext(*GetOwner(), *StateTreeRef.GetStateTree(), const_cast<FStateTreeInstanceData&>(InstanceData)).GetDebugInfoString();
+	return FConstStateTreeExecutionContextView(*GetOwner(), *StateTreeRef.GetStateTree(), InstanceData).Get().GetDebugInfoString();
+}
+
+TArray<FName> UStateTreeComponent::GetActiveStateNames() const
+{
+	if (!StateTreeRef.IsValid())
+	{
+		return TArray<FName>();
+	}
+
+	return FConstStateTreeExecutionContextView(*GetOwner(), *StateTreeRef.GetStateTree(), InstanceData).Get().GetActiveStateNames();
 }
 #endif // WITH_GAMEPLAY_DEBUGGER
 
diff --git a/Engine/Plugins/Runtime/GameplayStateTree/Source/GameplayStateTreeModule/Public/Components/StateTreeComponent.h b/Engine/Plugins/Runtime/GameplayStateTree/Source/GameplayStateTreeModule/Public/Components/StateTreeComponent.h
--- a/Engine/Plugins/Runtime/GameplayStateTree/Source/GameplayStateTreeModule/Public/Components/StateTreeComponent.h
+++ b/Engine/Plugins/Runtime/GameplayStateTree/Source/GameplayStateTreeModule/Public/Components/StateTreeComponent.h
@@ -82,6 +82,13 @@ public:
 
 #if WITH_GAMEPLAY_DEBUGGER
 	virtual FString GetDebugInfoString() const override;
+
+	/**
+	 * @return the list of active states. 
+	 * If the StateTree has linked asset StateTree, then more than one state can have the same name.
+	 * Only used for debugging purposes.
+	 */
+	TArray<FName> GetActiveStateNames() const;
 #endif // WITH_GAMEPLAY_DEBUGGER
 
 protected:
diff --git a/Engine/Plugins/Runtime/StateTree/Source/StateTreeModule/Public/StateTreeExecutionContext.h b/Engine/Plugins/Runtime/StateTree/Source/StateTreeModule/Public/StateTreeExecutionContext.h
--- a/Engine/Plugins/Runtime/StateTree/Source/StateTreeModule/Public/StateTreeExecutionContext.h
+++ b/Engine/Plugins/Runtime/StateTree/Source/StateTreeModule/Public/StateTreeExecutionContext.h
@@ -913,3 +913,31 @@ protected:
 	/** Captured snapshots for transition results that can be used to recreate transitions. This array is only populated if bRecordTransitions is true. */
 	TArray<FRecordedStateTreeTransitionResult> RecordedTransitions;
 };
+
+/**
+ * The const version of a StateTree Execution Context that prevents using the FStateTreeInstanceData with non-const member function.
+ */
+struct FConstStateTreeExecutionContextView
+{
+public:
+	FConstStateTreeExecutionContextView(UObject& InOwner, const UStateTree& InStateTree, const FStateTreeInstanceData& InInstanceData)
+		: ExecutionContext(InOwner, InStateTree, const_cast<FStateTreeInstanceData&>(InInstanceData))
+	{}
+
+	operator const FStateTreeExecutionContext& ()
+	{
+		return ExecutionContext;
+	}
+
+	const FStateTreeExecutionContext& Get() const
+	{
+		return ExecutionContext;
+	}
+
+private:
+	FConstStateTreeExecutionContextView(const FConstStateTreeExecutionContextView&) = delete;
+	FConstStateTreeExecutionContextView& operator=(const FConstStateTreeExecutionContextView&) = delete;
+
+private:
+	FStateTreeExecutionContext ExecutionContext;
+};
