StateTree: Fixed updating InstanceData for overriden LinkedAsset

diff --git a/Engine/Plugins/Runtime/StateTree/Source/StateTreeModule/Private/StateTreeExecutionContext.cpp b/Engine/Plugins/Runtime/StateTree/Source/StateTreeModule/Private/StateTreeExecutionContext.cpp
--- a/Engine/Plugins/Runtime/StateTree/Source/StateTreeModule/Private/StateTreeExecutionContext.cpp
+++ b/Engine/Plugins/Runtime/StateTree/Source/StateTreeModule/Private/StateTreeExecutionContext.cpp
@@ -775,6 +775,8 @@ void FStateTreeExecutionContext::UpdateInstanceData(TConstArrayView<FStateTreeEx
 	TArray<FInstancedStruct*, TConcurrentLinearArrayAllocator<FDefaultBlockAllocationTag>> TempInstanceStructs;
 	TempInstanceStructs.Reserve(EstimatedNumStructs);
 
+	TArray<FCompactStateTreeParameters, TFixedAllocator<FStateSelectionResult::MaxExecutionFrames>> TempParams;
+
 	TArrayView<FStateTreeTemporaryInstanceData> TempInstances = InstanceDataStorage->GetMutableTemporaryInstances();
 	auto FindInstanceTempData = [&TempInstances](const FStateTreeExecutionFrame& Frame, FStateTreeDataHandle DataHandle)
 	{
@@ -918,8 +920,19 @@ void FStateTreeExecutionContext::UpdateInstanceData(TConstArrayView<FStateTreeEx
 					}
 					else
 					{
-						// If not temp data, use the states default values.
-						const FConstStructView ParamsInstanceData = NextFrame.StateTree->DefaultInstanceData.GetStruct(State.ParameterTemplateIndex.Get());
+						// If not temp data, use the states or linked assets default values.
+						FConstStructView ParamsInstanceData;
+						if (State.Type == EStateTreeStateType::LinkedAsset)
+						{
+							if (const FStateTreeReference* Override = GetLinkedStateTreeOverrideForTag(State.Tag))
+							{
+								ParamsInstanceData = FConstStructView::Make(TempParams.Emplace_GetRef(Override->GetParameters()));
+							}
+						}
+						if (!ParamsInstanceData.IsValid())
+						{
+							ParamsInstanceData = NextFrame.StateTree->DefaultInstanceData.GetStruct(State.ParameterTemplateIndex.Get());
+						}
 						InstanceStructs[BaseIndex + State.ParameterDataHandle.GetIndex()] = ParamsInstanceData;
 						Params = ParamsInstanceData.GetPtr<const FCompactStateTreeParameters>();
 						bCanHaveTempData = true;
