StateTree: Refresh the Color picker when the value changes by reaching the current color ref. The DetailView can reset the value on multiple select or when "reset to default".

diff --git a/Engine/Plugins/Runtime/StateTree/Source/StateTreeEditorModule/Private/Customizations/Widgets/SStateTreeEditorColorComboBox.cpp b/Engine/Plugins/Runtime/StateTree/Source/StateTreeEditorModule/Private/Customizations/Widgets/SStateTreeEditorColorComboBox.cpp
--- a/Engine/Plugins/Runtime/StateTree/Source/StateTreeEditorModule/Private/Customizations/Widgets/SStateTreeEditorColorComboBox.cpp
+++ b/Engine/Plugins/Runtime/StateTree/Source/StateTreeEditorModule/Private/Customizations/Widgets/SStateTreeEditorColorComboBox.cpp
@@ -17,6 +17,8 @@ void SStateTreeEditorColorComboBox::Construct(const FArguments& InArgs, TSharedP
 	ColorRefHandle = InColorRefHandle;
 	ColorIDHandle = InColorRefHandle->GetChildHandle(GET_MEMBER_NAME_CHECKED(FStateTreeEditorColorRef, ID));
 
+	ColorRefHandle->SetOnPropertyValueChanged(FSimpleDelegate::CreateSP(this, &SStateTreeEditorColorComboBox::UpdatedSelectedColorWidget));
+
 	ChildSlot
 	[
 		SAssignNew(ColorComboBox, SComboBox<TSharedPtr<FStateTreeEditorColorRef>>)
