StateTree: Add transition when promoting to a parameter.
#rnx

diff --git a/Engine/Plugins/Runtime/StateTree/Source/StateTreeEditorModule/Private/Customizations/StateTreeBindingExtension.cpp b/Engine/Plugins/Runtime/StateTree/Source/StateTreeEditorModule/Private/Customizations/StateTreeBindingExtension.cpp
--- a/Engine/Plugins/Runtime/StateTree/Source/StateTreeEditorModule/Private/Customizations/StateTreeBindingExtension.cpp
+++ b/Engine/Plugins/Runtime/StateTree/Source/StateTreeEditorModule/Private/Customizations/StateTreeBindingExtension.cpp
@@ -20,6 +20,7 @@
 #include "IStructureDataProvider.h"
 #include "Kismet2/BlueprintEditorUtils.h"
 #include "PropertyBagDetails.h"
+#include "ScopedTransaction.h"
 #include "StateTreeEditorNodeUtils.h"
 #include "StateTreeEditorModule.h"
 #include "StateTreeNodeClassCache.h"
@@ -622,6 +623,8 @@ struct FCachedBindingData : public TSharedFromThis<FCachedBindingData>
 
 		const FGuid StructID = InStructDesc.ID;
 
+		const FScopedTransaction Transaction(LOCTEXT("PromoteToParameter", "Promote To Parameter"));
+
 		TArray<FStateTreeEditorPropertyCreationDesc, TFixedAllocator<1>> PropertyCreationDescs;
 		{
 			FStateTreeEditorPropertyCreationDesc& PropertyCreationDesc = PropertyCreationDescs.AddDefaulted_GetRef();
