StateTree: Exapnd after adding a new node item.
#rnx

diff --git a/Engine/Plugins/Runtime/StateTree/Source/StateTreeEditorModule/Private/Customizations/StateTreeEditorNodeUtils.cpp b/Engine/Plugins/Runtime/StateTree/Source/StateTreeEditorModule/Private/Customizations/StateTreeEditorNodeUtils.cpp
--- a/Engine/Plugins/Runtime/StateTree/Source/StateTreeEditorModule/Private/Customizations/StateTreeEditorNodeUtils.cpp
+++ b/Engine/Plugins/Runtime/StateTree/Source/StateTreeEditorModule/Private/Customizations/StateTreeEditorNodeUtils.cpp
@@ -642,14 +642,16 @@ void OnArrayNodePicked(const UStruct* InStruct, TSharedPtr<SComboButton> PickerC
 		ArrayPropertyHandle->NotifyPreChange();
 
 		// Add new item to the end.
-		ArrayHandle->AddItem();
-
-		uint32 NumItems = 0;
-		if (ArrayHandle->GetNumElements(NumItems) && NumItems > 0)
+		if (ArrayHandle->AddItem() == FPropertyAccess::Success)
 		{
-			// Initialize the item
-			TSharedRef<IPropertyHandle> NewNodeHandle = ArrayHandle->GetElement(NumItems - 1);
-			UE::StateTreeEditor::EditorNodeUtils::SetNodeType(NewNodeHandle, InStruct);
+			uint32 NumItems = 0;
+			if (ArrayHandle->GetNumElements(NumItems) && NumItems > 0)
+			{
+				// Initialize the item
+				TSharedRef<IPropertyHandle> NewNodeHandle = ArrayHandle->GetElement(NumItems - 1);
+				UE::StateTreeEditor::EditorNodeUtils::SetNodeType(NewNodeHandle, InStruct);
+				NewNodeHandle->SetExpanded(true);
+			}
 		}
 		
 		ArrayPropertyHandle->NotifyPostChange(EPropertyChangeType::ValueSet);
@@ -740,7 +742,15 @@ TSharedRef<SButton> CreateAddItemButton(const FText& TooltipText, FLinearColor C
 			{
 				if (const TSharedPtr<IPropertyHandleArray> ArrayHandle = ArrayPropertyHandle->AsArray())
 				{
-					ArrayHandle->AddItem();
+					if (ArrayHandle->AddItem() == FPropertyAccess::Success)
+					{
+						uint32 NumElements = 0;
+						if (ArrayHandle->GetNumElements(NumElements) == FPropertyAccess::Success)
+						{
+							TSharedRef<IPropertyHandle> NewPropertyHandle = ArrayHandle->GetElement(NumElements);
+							NewPropertyHandle->SetExpanded(true);
+						}
+					}
 				}
 			}
 			return FReply::Handled();
