SlateTree: Fix localition conflict.

diff --git a/Engine/Plugins/Runtime/StateTree/Source/StateTreeEditorModule/Private/Customizations/StateTreeBindingExtension.cpp b/Engine/Plugins/Runtime/StateTree/Source/StateTreeEditorModule/Private/Customizations/StateTreeBindingExtension.cpp
--- a/Engine/Plugins/Runtime/StateTree/Source/StateTreeEditorModule/Private/Customizations/StateTreeBindingExtension.cpp
+++ b/Engine/Plugins/Runtime/StateTree/Source/StateTreeEditorModule/Private/Customizations/StateTreeBindingExtension.cpp
@@ -623,7 +623,7 @@ struct FCachedBindingData : public TSharedFromThis<FCachedBindingData>
 
 		const FGuid StructID = InStructDesc.ID;
 
-		const FScopedTransaction Transaction(LOCTEXT("PromoteToParameter", "Promote To Parameter"));
+		const FScopedTransaction Transaction(LOCTEXT("PromoteToParameter", "Promote to Parameter"));
 
 		TArray<FStateTreeEditorPropertyCreationDesc, TFixedAllocator<1>> PropertyCreationDescs;
 		{
