StateTree: Fix bad access to the newly created element after adding the element. Introduced with CL 35780931.
#rnx

diff --git a/Engine/Plugins/Runtime/StateTree/Source/StateTreeEditorModule/Private/Customizations/StateTreeEditorNodeUtils.cpp b/Engine/Plugins/Runtime/StateTree/Source/StateTreeEditorModule/Private/Customizations/StateTreeEditorNodeUtils.cpp
--- a/Engine/Plugins/Runtime/StateTree/Source/StateTreeEditorModule/Private/Customizations/StateTreeEditorNodeUtils.cpp
+++ b/Engine/Plugins/Runtime/StateTree/Source/StateTreeEditorModule/Private/Customizations/StateTreeEditorNodeUtils.cpp
@@ -645,7 +645,7 @@ void OnArrayNodePicked(const UStruct* InStruct, TSharedPtr<SComboButton> PickerC
 		if (ArrayHandle->AddItem() == FPropertyAccess::Success)
 		{
 			uint32 NumItems = 0;
-			if (ArrayHandle->GetNumElements(NumItems) && NumItems > 0)
+			if (ArrayHandle->GetNumElements(NumItems) == FPropertyAccess::Success && NumItems > 0)
 			{
 				// Initialize the item
 				TSharedRef<IPropertyHandle> NewNodeHandle = ArrayHandle->GetElement(NumItems - 1);
@@ -745,9 +745,9 @@ TSharedRef<SButton> CreateAddItemButton(const FText& TooltipText, FLinearColor C
 					if (ArrayHandle->AddItem() == FPropertyAccess::Success)
 					{
 						uint32 NumElements = 0;
-						if (ArrayHandle->GetNumElements(NumElements) == FPropertyAccess::Success)
+						if (ArrayHandle->GetNumElements(NumElements) == FPropertyAccess::Success && NumElements > 0)
 						{
-							TSharedRef<IPropertyHandle> NewPropertyHandle = ArrayHandle->GetElement(NumElements);
+							TSharedRef<IPropertyHandle> NewPropertyHandle = ArrayHandle->GetElement(NumElements-1);
 							NewPropertyHandle->SetExpanded(true);
 						}
 					}
