SlateTree: Update the editor data when the state tree changes.

diff --git a/Engine/Plugins/Runtime/StateTree/Source/StateTreeEditorModule/Private/StateTreeEditorMode.cpp b/Engine/Plugins/Runtime/StateTree/Source/StateTreeEditorModule/Private/StateTreeEditorMode.cpp
--- a/Engine/Plugins/Runtime/StateTree/Source/StateTreeEditorModule/Private/StateTreeEditorMode.cpp
+++ b/Engine/Plugins/Runtime/StateTree/Source/StateTreeEditorModule/Private/StateTreeEditorMode.cpp
@@ -208,7 +208,7 @@ void UStateTreeEditorMode::Exit()
 	{
 		if (UStateTreeEditingSubsystem* StateTreeEditingSubsystem = GEditor->GetEditorSubsystem<UStateTreeEditingSubsystem>())
 		{		
-			TSharedRef<FStateTreeViewModel> ViewModel = StateTreeEditingSubsystem->FindOrAddViewModel(CachedStateTree.Get());				
+			TSharedRef<FStateTreeViewModel> ViewModel = StateTreeEditingSubsystem->FindOrAddViewModel(CachedStateTree.Get());
 			{
 				ViewModel->GetOnAssetChanged().RemoveAll(this);
 				ViewModel->GetOnStateAdded().RemoveAll(this);
@@ -217,7 +217,7 @@ void UStateTreeEditorMode::Exit()
 				ViewModel->GetOnSelectionChanged().RemoveAll(this);
 				ViewModel->GetOnBringNodeToFocus().RemoveAll(this);
 			}
-		}			
+		}
 	}
 
 	UE::StateTree::Delegates::OnIdentifierChanged.RemoveAll(this);
@@ -243,7 +243,7 @@ void UStateTreeEditorMode::OnStateTreeChanged()
 		{
 			if (CachedStateTree.IsValid())
 			{
-				TSharedRef<FStateTreeViewModel> OldViewModel = StateTreeEditingSubsystem->FindOrAddViewModel(CachedStateTree.Get());				
+				TSharedRef<FStateTreeViewModel> OldViewModel = StateTreeEditingSubsystem->FindOrAddViewModel(CachedStateTree.Get());
 				{
 					OldViewModel->GetOnAssetChanged().RemoveAll(this);
 					OldViewModel->GetOnStateAdded().RemoveAll(this);
@@ -257,6 +257,7 @@ void UStateTreeEditorMode::OnStateTreeChanged()
 
 		UStateTree* StateTree = Context->EditorHostInterface->GetStateTree();
 		CachedStateTree = StateTree;
+		UpdateAsset();
 
 		if (TSharedPtr<IDetailsView> AssetDetailsView = GetAssetDetailsView())
 		{
diff --git a/Engine/Plugins/Runtime/StateTree/Source/StateTreeEditorModule/Public/StateTreeEditorMode.h b/Engine/Plugins/Runtime/StateTree/Source/StateTreeEditorModule/Public/StateTreeEditorMode.h
--- a/Engine/Plugins/Runtime/StateTree/Source/StateTreeEditorModule/Public/StateTreeEditorMode.h
+++ b/Engine/Plugins/Runtime/StateTree/Source/StateTreeEditorModule/Public/StateTreeEditorMode.h
@@ -54,6 +54,7 @@ protected:
 	void HandleStateAdded(UStateTreeState* , UStateTreeState*) { UpdateAsset(); }
 	void HandleStatesRemoved(const TSet<UStateTreeState*>&) { UpdateAsset(); }
 	void HandleOnStatesMoved(const TSet<UStateTreeState*>&, const TSet<UStateTreeState*>&) { UpdateAsset(); }
+	/** Resolve the internal editor data and fixup the StateTree nodes. */
 	void UpdateAsset();
 
 	TSharedPtr<IDetailsView> GetDetailsView() const;
