StateTree: Modified the compiler to continue on non fatal error. It will still stop at the end of a step if that step fails. That is because some steps require other steps to complete properly. This helps provide a list of errors instead of one error at the time. In the future we can divide the compilation into 2 steps. "Analyze" to check for errors and "compile". We can then break the steps dependency and provide a full list of errors before doing the actual compile work.

diff --git a/Engine/Plugins/Runtime/StateTree/Source/StateTreeEditorModule/Private/StateTreeCompiler.cpp b/Engine/Plugins/Runtime/StateTree/Source/StateTreeEditorModule/Private/StateTreeCompiler.cpp
--- a/Engine/Plugins/Runtime/StateTree/Source/StateTreeEditorModule/Private/StateTreeCompiler.cpp
+++ b/Engine/Plugins/Runtime/StateTree/Source/StateTreeEditorModule/Private/StateTreeCompiler.cpp
@@ -488,6 +488,8 @@ bool FStateTreeCompiler::CreateStateRecursive(UStateTreeState& State, const FSta
 
 bool FStateTreeCompiler::CreateConditions(UStateTreeState& State, const FString& StatePath, TConstArrayView<FStateTreeEditorNode> Conditions)
 {
+	bool bSucceeded = true;
+
 	for (int32 Index = 0; Index < Conditions.Num(); Index++)
 	{
 		const bool bIsFirst = Index == 0;
@@ -503,11 +505,12 @@ bool FStateTreeCompiler::CreateConditions(UStateTreeState& State, const FString&
 
 		if (!CreateCondition(State, StatePath, CondNode, Operand, (int8)DeltaIndent))
 		{
-			return false;
+			bSucceeded = false;
+			continue;
 		}
 	}
 
-	return true;
+	return bSucceeded;
 }
 
 bool FStateTreeCompiler::CreateEvaluators()
@@ -515,6 +518,8 @@ bool FStateTreeCompiler::CreateEvaluators()
 	check(EditorData);
 	check(StateTree);
 
+	bool bSucceeded = true;
+
 	const int32 EvaluatorsBegin = Nodes.Num();
 	if (const auto Validation = UE::StateTree::Compiler::IsValidCount16(EvaluatorsBegin); Validation.DidFail())
 	{
@@ -529,7 +534,8 @@ bool FStateTreeCompiler::CreateEvaluators()
 		const FStateTreeDataHandle EvalDataHandle(EStateTreeDataSourceType::GlobalInstanceData, GlobalInstanceIndex);
 		if (!CreateEvaluator(EvalNode, EvalDataHandle))
 		{
-			return false;
+			bSucceeded = false;
+			continue;
 		}
 	}
 	
@@ -541,7 +547,7 @@ bool FStateTreeCompiler::CreateEvaluators()
 	}
 	StateTree->EvaluatorsNum = uint16(EvaluatorsNum);
 
-	return CreateBindingsForNodes(EditorData->Evaluators, FStateTreeIndex16(EvaluatorsBegin), InstanceStructs);
+	return bSucceeded && CreateBindingsForNodes(EditorData->Evaluators, FStateTreeIndex16(EvaluatorsBegin), InstanceStructs);
 }
 
 bool FStateTreeCompiler::CreateGlobalTasks()
@@ -549,6 +555,8 @@ bool FStateTreeCompiler::CreateGlobalTasks()
 	check(EditorData);
 	check(StateTree);
 
+	bool bSucceeded = true;
+
 	const int32 GlobalTasksBegin = Nodes.Num();
 	if (const auto Validation = UE::StateTree::Compiler::IsValidCount16(GlobalTasksBegin); Validation.DidFail())
 	{
@@ -570,7 +578,8 @@ bool FStateTreeCompiler::CreateGlobalTasks()
 		const FStateTreeDataHandle TaskDataHandle(EStateTreeDataSourceType::GlobalInstanceData, GlobalInstanceIndex);
 		if (!CreateTask(nullptr, TaskNode, TaskDataHandle))
 		{
-			return false;
+			bSucceeded = false;
+			continue;
 		}
 		
 		const FStateTreeTaskBase& LastAddedTask = Nodes.Last().Get<FStateTreeTaskBase>();
@@ -586,13 +595,15 @@ bool FStateTreeCompiler::CreateGlobalTasks()
 	}
 	StateTree->GlobalTasksNum = uint16(GlobalTasksNum);
 
-	return CreateBindingsForNodes(EditorData->GlobalTasks, FStateTreeIndex16(GlobalTasksBegin), InstanceStructs);
+	return bSucceeded && CreateBindingsForNodes(EditorData->GlobalTasks, FStateTreeIndex16(GlobalTasksBegin), InstanceStructs);
 }
 
 bool FStateTreeCompiler::CreateStateTasksAndParameters()
 {
 	check(StateTree);
 
+	bool bSucceeded = true;
+
 	// Index of the first instance data per state. Accumulated depth first.
 	TArray<int32> FirstInstanceDataIndex;
 	FirstInstanceDataIndex.Init(0, StateTree->States.Num());
@@ -658,7 +669,8 @@ bool FStateTreeCompiler::CreateStateTasksAndParameters()
 
 		if (!UE::StateTree::Compiler::ValidateNoLevelActorReferences(Log, LinkedParamsDesc, FStateTreeDataView(), FStateTreeDataView(CompactStateTreeParameters.Parameters.GetMutableValue())))
 		{
-			return false;
+			bSucceeded = false;
+			continue;
 		}
 
 		// Add as binding source.
@@ -687,7 +699,8 @@ bool FStateTreeCompiler::CreateStateTasksAndParameters()
 			{
 				Log.Reportf(EMessageSeverity::Error, Desc,
 					TEXT("Event is marked as required, but isn't set up."));
-				return false;
+				bSucceeded = false;
+				continue;
 			}
 		}
 
@@ -695,7 +708,8 @@ bool FStateTreeCompiler::CreateStateTasksAndParameters()
 		{
 			Log.Reportf(EMessageSeverity::Error, LinkedParamsDesc,
 				TEXT("Exceeds the maximum depth of execution (%u)"), FStateTreeActiveStates::MaxStates);
-			return false;
+			bSucceeded = false;
+			continue;
 		}
 
 		// Subtrees parameters cannot have bindings
@@ -704,7 +718,8 @@ bool FStateTreeCompiler::CreateStateTasksAndParameters()
 			FStateTreeIndex16 PropertyFunctionsBegin(Nodes.Num());
 			if (!CreatePropertyFunctionsForStruct(LinkedParamsDesc.ID))
 			{
-				return false;
+				bSucceeded = false;
+				continue;
 			}
 
 			FStateTreeIndex16 PropertyFunctionsEnd(Nodes.Num());
@@ -717,7 +732,8 @@ bool FStateTreeCompiler::CreateStateTasksAndParameters()
 
 			if (!CreateBindingsForStruct(LinkedParamsDesc, FStateTreeDataView(CompactStateTreeParameters.Parameters.GetMutableValue()), PropertyFunctionsBegin, PropertyFunctionsEnd, CompactState.ParameterBindingsBatch))
 			{
-				return false;
+				bSucceeded = false;
+				continue;
 			}
 		}
 
@@ -752,7 +768,8 @@ bool FStateTreeCompiler::CreateStateTasksAndParameters()
 			const FStateTreeDataHandle TaskDataHandle(EStateTreeDataSourceType::ActiveInstanceData, InstanceDataIndex++, CompactStateHandle);
 			if (!CreateTask(State, TaskNode, TaskDataHandle))
 			{
-				return false;
+				bSucceeded = false;
+				continue;
 			}
 
 			const FStateTreeTaskBase& LastAddedTask = Nodes.Last().Get<FStateTreeTaskBase>();
@@ -781,17 +798,20 @@ bool FStateTreeCompiler::CreateStateTasksAndParameters()
 
 		if (!CreateBindingsForNodes(Tasks, FStateTreeIndex16(TasksBegin), InstanceStructs))
 		{
-			return false;
+			bSucceeded = false;
+			continue;
 		}
 	}
 	
-	return true;
+	return bSucceeded;
 }
 
 bool FStateTreeCompiler::CreateStateTransitions()
 {
 	check(StateTree);
 
+	bool bSucceeded = true;
+
 	for (int32 i = 0; i < StateTree->States.Num(); i++)
 	{
 		FCompactStateTreeState& CompactState = StateTree->States[i];
@@ -816,7 +836,8 @@ bool FStateTreeCompiler::CreateStateTransitions()
 		{
 			Log.Reportf(EMessageSeverity::Error,
 				TEXT("Failed to create state enter condition."));
-			return false;
+			bSucceeded = false;
+			continue;
 		}
 		
 		const int32 EnterConditionsNum = Nodes.Num() - EnterConditionsBegin;
@@ -829,7 +850,8 @@ bool FStateTreeCompiler::CreateStateTransitions()
 
 		if (!CreateBindingsForNodes(SourceState->EnterConditions, FStateTreeIndex16(EnterConditionsBegin), SharedInstanceStructs))
 		{
-			return false;
+			bSucceeded = false;
+			continue;
 		}
 
 		// Check if any of the enter conditions require state completion events, and cache that.
@@ -864,7 +886,8 @@ bool FStateTreeCompiler::CreateStateTransitions()
 				Log.Reportf(EMessageSeverity::Error,
 					TEXT("State is linked to it's parent subtree '%s', which will create infinite loop."),
 					*LinkedParentState->Name.ToString());
-				return false;
+				bSucceeded = false;
+				continue;
 			}
 
 			// The linked state must be a subtree.
@@ -874,7 +897,8 @@ bool FStateTreeCompiler::CreateStateTransitions()
 				Log.Reportf(EMessageSeverity::Error,
 					TEXT("Failed to resolve linked subtree '%s'."),
 					*SourceState->LinkedSubtree.Name.ToString());
-				return false;
+				bSucceeded = false;
+				continue;
 			}
 			
 			if (TargetState->Type != EStateTreeStateType::Subtree)
@@ -882,7 +906,8 @@ bool FStateTreeCompiler::CreateStateTransitions()
 				Log.Reportf(EMessageSeverity::Error,
 					TEXT("State '%s' is linked to subtree '%s', which is not a subtree."),
 					*SourceState->Name.ToString(), *TargetState->Name.ToString());
-				return false;
+				bSucceeded = false;
+				continue;
 			}
 			
 			CompactState.LinkedState = GetStateHandle(SourceState->LinkedSubtree.ID);
@@ -892,7 +917,8 @@ bool FStateTreeCompiler::CreateStateTransitions()
 				Log.Reportf(EMessageSeverity::Error,
 					TEXT("Failed to resolve linked subtree '%s'."),
 					*SourceState->LinkedSubtree.Name.ToString());
-				return false;
+				bSucceeded = false;
+				continue;
 			}
 		}
 		else if (SourceState->Type == EStateTreeStateType::LinkedAsset)
@@ -902,7 +928,8 @@ bool FStateTreeCompiler::CreateStateTransitions()
 			{
 				Log.Reportf(EMessageSeverity::Error,
 					TEXT("It is not allowed to link to the same tree, as it might create infinite loop."));
-				return false;
+				bSucceeded = false;
+				continue;
 			}
 
 			if (SourceState->LinkedAsset)
@@ -914,7 +941,8 @@ bool FStateTreeCompiler::CreateStateTransitions()
 				{
 					Log.Reportf(EMessageSeverity::Error,
 						TEXT("Linked State Tree asset must have valid schema."));
-					return false;
+					bSucceeded = false;
+					continue;
 				}
 			
 				check(Schema);
@@ -926,13 +954,14 @@ bool FStateTreeCompiler::CreateStateTransitions()
 						*LinkedAssetSchema->GetClass()->GetDisplayNameText().ToString(),
 						*Schema->GetClass()->GetDisplayNameText().ToString()
 					);
-					return false;
+					bSucceeded = false;
+					continue;
 				}
 			}
 			
 			CompactState.LinkedAsset = SourceState->LinkedAsset;
 		}
-		
+
 		// Transitions
 		const int32 TransitionsBegin = StateTree->Transitions.Num();
 		if (const auto Validation = UE::StateTree::Compiler::IsValidCount16(TransitionsBegin); Validation.DidFail())
@@ -941,7 +970,8 @@ bool FStateTreeCompiler::CreateStateTransitions()
 			return false;
 		}
 		CompactState.TransitionsBegin = uint16(TransitionsBegin);
-		
+
+		bool bTransitionSucceeded = true;
 		for (FStateTreeTransition& Transition : SourceState->Transitions)
 		{
 			const int32 TransitionIndex = StateTree->Transitions.Num();
@@ -985,14 +1015,16 @@ bool FStateTreeCompiler::CreateStateTransitions()
 					Log.Reportf(EMessageSeverity::Error,
 						TEXT("State completion transition to '%s' must have transition to valid state, 'None' not accepted."),
 						*Transition.State.Name.ToString());
-					return false;
+					bTransitionSucceeded = false;
+					continue;
 				}
 			}
 			
 			CompactTransition.State = FStateTreeStateHandle::Invalid;
 			if (!ResolveTransitionState(SourceState, Transition.State, CompactTransition.State))
 			{
-				return false;
+				bTransitionSucceeded = false;
+				continue;
 			}
 
 			if (CompactTransition.State.IsValid()
@@ -1028,7 +1060,8 @@ bool FStateTreeCompiler::CreateStateTransitions()
 					Log.Reportf(EMessageSeverity::Error, Desc,
 						TEXT("On Event Transition requires at least tag or payload to be set up."),
 						*Transition.State.Name.ToString());
-					return false;
+					bTransitionSucceeded = false;
+					continue;
 				}
 
 				if (CompactTransition.State.IsValid()
@@ -1040,7 +1073,8 @@ bool FStateTreeCompiler::CreateStateTransitions()
 						Log.Reportf(EMessageSeverity::Error, Desc,
 							TEXT("On Event transition to '%s' will never succeed as transition and state required events are incompatible."),
 							*TransitionTargetState.Name.ToString());
-						return false;
+						bTransitionSucceeded = false;
+						continue;
 					}
 				}
 
@@ -1060,7 +1094,8 @@ bool FStateTreeCompiler::CreateStateTransitions()
 				Log.Reportf(EMessageSeverity::Error,
 					TEXT("Failed to create condition for transition to '%s'."),
 					*Transition.State.Name.ToString());
-				return false;
+				bTransitionSucceeded = false;
+				continue;
 			}
 
 			const int32 ConditionsNum = Nodes.Num() - ConditionsBegin;
@@ -1073,9 +1108,16 @@ bool FStateTreeCompiler::CreateStateTransitions()
 
 			if (!CreateBindingsForNodes(Transition.Conditions, FStateTreeIndex16(ConditionsBegin), SharedInstanceStructs))
 			{
-				return false;
+				bTransitionSucceeded = false;
+				continue;
 			}
 		}
+
+		if (!bTransitionSucceeded)
+		{
+			bSucceeded = false;
+			continue;
+		}
 		
 		const int32 TransitionsNum = StateTree->Transitions.Num() - TransitionsBegin;
 		if (const auto Validation = UE::StateTree::Compiler::IsValidCount8(TransitionsNum); Validation.DidFail())
@@ -1088,13 +1130,15 @@ bool FStateTreeCompiler::CreateStateTransitions()
 
 	// @todo: Add test to check that all success/failure transition is possible (see editor).
 	
-	return true;
+	return bSucceeded;
 }
 
 bool FStateTreeCompiler::CreateStateConsiderations()
 {
 	check(StateTree);
 
+	bool bSucceeded = true;
+
 	for (int32 i = 0; i < StateTree->States.Num(); i++)
 	{
 		FCompactStateTreeState& CompactState = StateTree->States[i];
@@ -1109,7 +1153,8 @@ bool FStateTreeCompiler::CreateStateConsiderations()
 		if (const auto Validation = UE::StateTree::Compiler::IsValidCount16(UtilityConsiderationsBegin); Validation.DidFail())
 		{
 			Validation.Log(Log, TEXT("UtilityConsiderationsBegin"));
-			return false;
+			bSucceeded = false;
+			continue;
 		}
 		CompactState.UtilityConsiderationsBegin = uint16(UtilityConsiderationsBegin);
 
@@ -1118,30 +1163,35 @@ bool FStateTreeCompiler::CreateStateConsiderations()
 		{
 			Log.Reportf(EMessageSeverity::Error,
 				TEXT("Failed to create state utility considerations."));
-			return false;
+			bSucceeded = false;
+			continue;
 		}
 
 		const int32 UtilityConsiderationsNum = Nodes.Num() - UtilityConsiderationsBegin;
 		if (const auto Validation = UE::StateTree::Compiler::IsValidCount8(UtilityConsiderationsNum); Validation.DidFail())
 		{
 			Validation.Log(Log, TEXT("UtilityConsiderationsNum"));
-			return false;
+			bSucceeded = false;
+			continue;
 		}
 		CompactState.UtilityConsiderationsNum = uint8(UtilityConsiderationsNum);
 
 		if (!CreateBindingsForNodes(SourceState->Considerations, FStateTreeIndex16(UtilityConsiderationsBegin), SharedInstanceStructs))
 		{
-			return false;
+			bSucceeded = false;
+			continue;
 		}
 	}
 
-	return true;
+	return bSucceeded;
 }
 
 bool FStateTreeCompiler::CreateBindingsForNodes(TConstArrayView<FStateTreeEditorNode> EditorNodes, FStateTreeIndex16 NodesBegin, TArrayView<FInstancedStruct> Instances)
 {
 	check(NodesBegin.IsValid());
 
+	bool bSucceeded = true;
+
 	int32 NodeIndex = NodesBegin.Get();
 	for (const FStateTreeEditorNode& EditorNode : EditorNodes)
 	{
@@ -1159,7 +1209,8 @@ bool FStateTreeCompiler::CreateBindingsForNodes(TConstArrayView<FStateTreeEditor
 		FStateTreeIndex16 PropertyFunctionsBegin(Nodes.Num());
 		if (!CreatePropertyFunctionsForStruct(EditorNode.ID))
 		{
-			return false;
+			bSucceeded = false;
+			continue;
 		}
 		FStateTreeIndex16 PropertyFunctionsEnd(Nodes.Num());
 		
@@ -1186,11 +1237,12 @@ bool FStateTreeCompiler::CreateBindingsForNodes(TConstArrayView<FStateTreeEditor
 
 		if (!CreateBindingsForStruct(*BindableStruct, InstanceView, PropertyFunctionsBegin, PropertyFunctionsEnd, Node.BindingsBatch))
 		{
-			return false;
+			bSucceeded = false;
+			continue;
 		}
 	}
 
-	return true;
+	return bSucceeded;
 }
 
 bool FStateTreeCompiler::CreateBindingsForStruct(const FStateTreeBindableStructDesc& TargetStruct, FStateTreeDataView TargetValue, FStateTreeIndex16 PropertyFuncsBegin, FStateTreeIndex16 PropertyFuncsEnd, FStateTreeIndex16& OutBatchIndex)
@@ -1883,6 +1935,8 @@ bool FStateTreeCompiler::GetAndValidateBindings(const FStateTreeBindableStructDe
 		return true;
 	}
 
+	bool bSucceeded = true;
+
 	for (FStateTreePropertyPathBinding& Binding : EditorData->EditorBindings.GetMutableBindings())
 	{
 		if (Binding.GetTargetPath().GetStructID() != TargetStruct.ID)
@@ -1898,7 +1952,8 @@ bool FStateTreeCompiler::GetAndValidateBindings(const FStateTreeBindableStructDe
 			Log.Reportf(EMessageSeverity::Error, TargetStruct,
 						TEXT("Failed to find binding source property '%s' for target %s."),
 						*Binding.GetSourcePath().ToString(), *UE::StateTree::GetDescAndPathAsString(TargetStruct, Binding.GetTargetPath()));
-			return false;
+			bSucceeded = false;
+			continue;
 		}
 
 		// Update path instance types from latest data. E.g. binding may have been created for instanced object of type FooB, and changed to FooA.
@@ -1909,7 +1964,8 @@ bool FStateTreeCompiler::GetAndValidateBindings(const FStateTreeBindableStructDe
 			Log.Reportf(EMessageSeverity::Error, TargetStruct,
 						TEXT("Malformed target property path for binding source property '%s' for target %s."),
 						*Binding.GetSourcePath().ToString(), *UE::StateTree::GetDescAndPathAsString(TargetStruct, Binding.GetTargetPath()));
-			return false;
+			bSucceeded = false;
+			continue;
 		}
 		
 		// Source must be accessible by the target struct via all execution paths.
@@ -1928,7 +1984,8 @@ bool FStateTreeCompiler::GetAndValidateBindings(const FStateTreeBindableStructDe
 						*UE::StateTree::GetDescAndPathAsString(*SourceStruct, Binding.GetSourcePath()),
 						*UE::StateTree::GetDescAndPathAsString(TargetStruct, Binding.GetTargetPath()),
 						*SourceStruct->ToString(), *TargetStruct.ToString());
-			return false;
+			bSucceeded = false;
+			continue;
 		}
 
 		if (!IDToStructValue.Contains(SourceStructID))
@@ -1936,7 +1993,8 @@ bool FStateTreeCompiler::GetAndValidateBindings(const FStateTreeBindableStructDe
 			Log.Reportf(EMessageSeverity::Error, TargetStruct,
 				TEXT("Failed to find value for binding source property '%s' for target %s."),
 				*Binding.GetSourcePath().ToString(), *UE::StateTree::GetDescAndPathAsString(TargetStruct, Binding.GetTargetPath()));
-			return false;
+			bSucceeded = false;
+			continue;
 		}
 
 		// Update the source structs only if we have value for it. For some sources (e.g. context structs) we know only type, and in that case there are no instance structs.
@@ -1948,7 +2006,8 @@ bool FStateTreeCompiler::GetAndValidateBindings(const FStateTreeBindableStructDe
 				Log.Reportf(EMessageSeverity::Error, TargetStruct,
 					TEXT("Malformed target property path for binding source property '%s' for source %s."),
 					*Binding.GetSourcePath().ToString(), *UE::StateTree::GetDescAndPathAsString(TargetStruct, Binding.GetTargetPath()));
-				return false;
+				bSucceeded = false;
+				continue;
 			}
 		}
 
@@ -1957,7 +2016,8 @@ bool FStateTreeCompiler::GetAndValidateBindings(const FStateTreeBindableStructDe
 			Log.Reportf(EMessageSeverity::Error, TargetStruct,
 				TEXT("Malformed source'%s for property binding property '%s'."),
 				*UE::StateTree::GetDescAndPathAsString(*SourceStruct, Binding.GetSourcePath()), *Binding.GetSourcePath().ToString());
-			return false;
+			bSucceeded = false;
+			continue;
 		}
 		
 		FStateTreePropertyPathBinding BindingCopy(Binding);
@@ -1992,10 +2052,16 @@ bool FStateTreeCompiler::GetAndValidateBindings(const FStateTreeBindableStructDe
 		// Check if the bindings is for struct ref and validate the types.
 		if (!ValidateStructRef(*SourceStruct, Binding.GetSourcePath(), TargetStruct, Binding.GetTargetPath()))
 		{
-			return false;
+			bSucceeded = false;
+			continue;
 		}
 	}
 
+	if (!bSucceeded)
+	{
+		return false;
+	}
+
 	// In the UI a user can add a binding X = A.B and X.Y = C.D. The second overrides data that the first one set. The order is also not guaranteed.
 	// Remove all the CopyBindings that override.
 	for (int32 FirstIndex = OutCopyBindings.Num() - 1; FirstIndex >= 0; --FirstIndex)
@@ -2043,8 +2109,6 @@ bool FStateTreeCompiler::GetAndValidateBindings(const FStateTreeBindableStructDe
 				return Binding.GetTargetPath().GetSegments().Num() >= 1 && Binding.GetTargetPath().GetSegments()[0].GetName() == PropertyName;
 			});
 	};
-
-	bool bResult = true;
 	
 	// Validate that Input and Context bindings
 	for (TFieldIterator<FProperty> It(TargetStruct.Struct); It; ++It)
@@ -2062,7 +2126,8 @@ bool FStateTreeCompiler::GetAndValidateBindings(const FStateTreeBindableStructDe
 				Log.Reportf(EMessageSeverity::Error, TargetStruct,
 						TEXT("Couldn't resolve path to '%s' for target %s."),
 						*PropertyName.ToString(), *TargetStruct.ToString());
-				bResult = false;
+				bSucceeded = false;
+				continue;
 			}
 			else
 			{
@@ -2074,7 +2139,8 @@ bool FStateTreeCompiler::GetAndValidateBindings(const FStateTreeBindableStructDe
 					Log.Reportf(EMessageSeverity::Error, TargetStruct,
 						TEXT("Property reference '%s' on % s is expected to have a binding."),
 						*PropertyName.ToString(), *TargetStruct.ToString());
-					bResult = false;
+					bSucceeded = false;
+					continue;
 				}
 			}
 		}
@@ -2090,7 +2156,8 @@ bool FStateTreeCompiler::GetAndValidateBindings(const FStateTreeBindableStructDe
 					Log.Reportf(EMessageSeverity::Error, TargetStruct,
 						TEXT("Input property '%s' on %s is expected to have a binding."),
 						*PropertyName.ToString(), *TargetStruct.ToString());
-					bResult = false;
+					bSucceeded = false;
+					continue;
 				}
 			}
 			else if (Usage == EStateTreePropertyUsage::Context)
@@ -2111,7 +2178,7 @@ bool FStateTreeCompiler::GetAndValidateBindings(const FStateTreeBindableStructDe
 					Log.Reportf(EMessageSeverity::Error, TargetStruct,
 						TEXT("The type of Context property '%s' on %s is expected to be Object Reference or Struct."),
 						*PropertyName.ToString(), *TargetStruct.ToString());
-					bResult = false;
+					bSucceeded = false;
 					continue;
 				}
 
@@ -2131,14 +2198,15 @@ bool FStateTreeCompiler::GetAndValidateBindings(const FStateTreeBindableStructDe
 						Log.Reportf(EMessageSeverity::Error, TargetStruct,
 							TEXT("Could not find matching Context object for Context property '%s' on '%s'. Property must have manual binding."),
 							*PropertyName.ToString(), *TargetStruct.ToString());
-						bResult = false;
+						bSucceeded = false;
+						continue;
 					}
 				}
 			}
 		}
 	}
 
-	return bResult;
+	return bSucceeded;
 }
 
 void FStateTreeCompiler::InstantiateStructSubobjects(FStructView Struct)
