StateTree: Update the editor subsystem viewmodel cache. The cache is alive forever but the uobject can be GCed or reinstanciated. Update the cache if the found StateTree is not valid anymore.

diff --git a/Engine/Plugins/Runtime/StateTree/Source/StateTreeEditorModule/Private/StateTreeEditingSubsystem.cpp b/Engine/Plugins/Runtime/StateTree/Source/StateTreeEditorModule/Private/StateTreeEditingSubsystem.cpp
--- a/Engine/Plugins/Runtime/StateTree/Source/StateTreeEditorModule/Private/StateTreeEditingSubsystem.cpp
+++ b/Engine/Plugins/Runtime/StateTree/Source/StateTreeEditorModule/Private/StateTreeEditingSubsystem.cpp
@@ -10,6 +10,19 @@
 #include "StateTreeEditorModule.h"
 #include "StateTreeObjectHash.h"
 #include "StateTreeTaskBase.h"
+#include "UObject/UObjectGlobals.h"
+
+
+UStateTreeEditingSubsystem::UStateTreeEditingSubsystem()
+{
+	PostGarbageCollectHandle = FCoreUObjectDelegates::GetPostGarbageCollect().AddUObject(this, &UStateTreeEditingSubsystem::HandlePostGarbageCollect);
+}
+
+void UStateTreeEditingSubsystem::BeginDestroy()
+{
+	FCoreUObjectDelegates::GetPostGarbageCollect().Remove(PostGarbageCollectHandle);
+	Super::BeginDestroy();
+}
 
 bool UStateTreeEditingSubsystem::CompileStateTree(const TNonNullPtr<UStateTree> InStateTree, FStateTreeCompilerLog& InOutLog)
 {
@@ -41,12 +54,24 @@ bool UStateTreeEditingSubsystem::CompileStateTree(const TNonNullPtr<UStateTree>
 
 TSharedRef<FStateTreeViewModel> UStateTreeEditingSubsystem::FindOrAddViewModel(const TNonNullPtr<UStateTree> InStateTree)
 {
-	if(TSharedPtr<FStateTreeViewModel>* ModelPtr = StateTreeViewModels.Find(InStateTree.Get()))
+	const FObjectKey StateTreeKey = InStateTree;
+	TSharedPtr<FStateTreeViewModel> ViewModelPtr = StateTreeViewModels.FindRef(StateTreeKey);
+	if (ViewModelPtr)
 	{
-		return ModelPtr->ToSharedRef();
+		// The StateTree could be re-instantiated. Can occur when the object is destroyed and recreated in a pool or when reloaded in editor.
+		//The object might have the same pointer value or the same path but it's a new object and all weakptr are now invalid.
+		if (ViewModelPtr->GetStateTree() == InStateTree)
+		{
+			return ViewModelPtr.ToSharedRef();
+		}
+		else
+		{
+			StateTreeViewModels.Remove(StateTreeKey);
+			ViewModelPtr = nullptr;
+		}
 	}
 
-	TSharedRef<FStateTreeViewModel> SharedModel = StateTreeViewModels.Add(InStateTree.Get(), MakeShared<FStateTreeViewModel>()).ToSharedRef();
+	TSharedRef<FStateTreeViewModel> SharedModel = StateTreeViewModels.Add(StateTreeKey, MakeShared<FStateTreeViewModel>()).ToSharedRef();
 	UStateTreeEditorData* EditorData = Cast<UStateTreeEditorData>(InStateTree->EditorData);
 	if (EditorData == nullptr)
 	{
@@ -286,3 +311,19 @@ uint32 UStateTreeEditingSubsystem::CalculateStateTreeHash(const TNonNullPtr<cons
 
 	return EditorDataHash;
 }
+
+void UStateTreeEditingSubsystem::HandlePostGarbageCollect()
+{
+	// Remove the stale viewmodels
+	for (TMap<FObjectKey, TSharedPtr<FStateTreeViewModel>>::TIterator It(StateTreeViewModels); It; ++It)
+	{
+		if (!It.Key().ResolveObjectPtr())
+		{
+			It.RemoveCurrent();
+		}
+		else if (!It.Value() || !It.Value()->GetStateTree())
+		{
+			It.RemoveCurrent();
+		}
+	}
+}
diff --git a/Engine/Plugins/Runtime/StateTree/Source/StateTreeEditorModule/Public/StateTreeEditingSubsystem.h b/Engine/Plugins/Runtime/StateTree/Source/StateTreeEditorModule/Public/StateTreeEditingSubsystem.h
--- a/Engine/Plugins/Runtime/StateTree/Source/StateTreeEditorModule/Public/StateTreeEditingSubsystem.h
+++ b/Engine/Plugins/Runtime/StateTree/Source/StateTreeEditorModule/Public/StateTreeEditingSubsystem.h
@@ -5,6 +5,7 @@
 #include "StateTree.h"
 #include "StateTreeViewModel.h"
 #include "EditorSubsystem.h"
+#include "UObject/ObjectKey.h"
 
 #include "StateTreeEditingSubsystem.generated.h"
 
@@ -18,20 +19,29 @@ class STATETREEEDITORMODULE_API UStateTreeEditingSubsystem : public UEditorSubsy
 {
 	GENERATED_BODY()
 public:
-	UStateTreeEditingSubsystem() {}
+	UStateTreeEditingSubsystem();
+	virtual void BeginDestroy() override;
 	
 	TSharedRef<FStateTreeViewModel> FindOrAddViewModel(const TNonNullPtr<UStateTree> InStateTree);
 	
 	static bool CompileStateTree(const TNonNullPtr<UStateTree>InStateTree,  FStateTreeCompilerLog& InOutLog);
 	
+	/** Create a StateTreeView widget for the viewmodel. */
 	static TSharedRef<SWidget> GetStateTreeView(TSharedRef<FStateTreeViewModel> InViewModel, const TSharedRef<FUICommandList>& TreeViewCommandList);
 	
-	// Validates asset state
+	/**
+	 * Validates and applies the schema restrictions on the StateTree.
+	 * Updates state's link, removes the unused node while validating the StateTree asset.
+	 */
 	static void ValidateStateTree(const TNonNullPtr<UStateTree> InStateTree);
 
-	// Calculates editor data hash of the asset.
+	/** Calculates editor data hash of the asset. */
 	static uint32 CalculateStateTreeHash(const TNonNullPtr<const UStateTree> InStateTree);
 	
+private:
+	void HandlePostGarbageCollect();
+
 protected:
-	TMap<TSoftObjectPtr<UStateTree>, TSharedPtr<FStateTreeViewModel>> StateTreeViewModels;
+	TMap<FObjectKey, TSharedPtr<FStateTreeViewModel>> StateTreeViewModels;
+	FDelegateHandle PostGarbageCollectHandle;
 };
