StateTree: Review feedback for CL 35592141

#rnx

diff --git a/Engine/Plugins/Runtime/StateTree/Source/StateTreeEditorModule/Private/StateTreeCompiler.cpp b/Engine/Plugins/Runtime/StateTree/Source/StateTreeEditorModule/Private/StateTreeCompiler.cpp
--- a/Engine/Plugins/Runtime/StateTree/Source/StateTreeEditorModule/Private/StateTreeCompiler.cpp
+++ b/Engine/Plugins/Runtime/StateTree/Source/StateTreeEditorModule/Private/StateTreeCompiler.cpp
@@ -21,7 +21,7 @@
 namespace UE::StateTree::Compiler
 {
 	FAutoConsoleVariable CVarLogCompiledStateTree(
-		TEXT("StateTree.Compiler.LogResultOnSuccess"),
+		TEXT("StateTree.Compiler.LogResultOnCompilationCompleted"),
 		false,
 		TEXT("After a StateTree compiles, log the internal content of the StateTree.")
 	);
diff --git a/Engine/Plugins/Runtime/StateTree/Source/StateTreeEditorModule/Private/StateTreeEditor.cpp b/Engine/Plugins/Runtime/StateTree/Source/StateTreeEditorModule/Private/StateTreeEditor.cpp
--- a/Engine/Plugins/Runtime/StateTree/Source/StateTreeEditorModule/Private/StateTreeEditor.cpp
+++ b/Engine/Plugins/Runtime/StateTree/Source/StateTreeEditorModule/Private/StateTreeEditor.cpp
@@ -373,7 +373,7 @@ void FillDynamicDeveloperMenu(FToolMenuSection& Section)
 			false);
 	}
 }
-}
+} // UE::StateTree::Editor::Private
 
 void FStateTreeEditor::RegisterMenu()
 {
diff --git a/Engine/Plugins/Runtime/StateTree/Source/StateTreeEditorModule/Private/StateTreeEditorMode.cpp b/Engine/Plugins/Runtime/StateTree/Source/StateTreeEditorModule/Private/StateTreeEditorMode.cpp
--- a/Engine/Plugins/Runtime/StateTree/Source/StateTreeEditorModule/Private/StateTreeEditorMode.cpp
+++ b/Engine/Plugins/Runtime/StateTree/Source/StateTreeEditorModule/Private/StateTreeEditorMode.cpp
@@ -306,7 +306,7 @@ static bool IsSaveOnCompileOptionSet(const EStateTreeSaveOnCompile Option)
 
 static IConsoleVariable* GetLogResultOnSuccessCVar()
 {
-	static IConsoleVariable* FoundVariable = IConsoleManager::Get().FindConsoleVariable(TEXT("StateTree.Compiler.LogResultOnSuccess"));
+	static IConsoleVariable* FoundVariable = IConsoleManager::Get().FindConsoleVariable(TEXT("StateTree.Compiler.LogResultOnCompilationCompleted"));
 	return FoundVariable;
 }
 
