StateTree: Prevents incorrect character or too long state name.

diff --git a/Engine/Plugins/Runtime/StateTree/Source/StateTreeEditorModule/Private/SStateTreeViewRow.cpp b/Engine/Plugins/Runtime/StateTree/Source/StateTreeEditorModule/Private/SStateTreeViewRow.cpp
--- a/Engine/Plugins/Runtime/StateTree/Source/StateTreeEditorModule/Private/SStateTreeViewRow.cpp
+++ b/Engine/Plugins/Runtime/StateTree/Source/StateTreeEditorModule/Private/SStateTreeViewRow.cpp
@@ -293,6 +293,7 @@ void SStateTreeViewRow::Construct(const FArguments& InArgs, const TSharedRef<STa
 														return !NewLabel.IsEmptyOrWhitespace();
 													})
 												.OnTextCommitted(this, &SStateTreeViewRow::HandleNodeLabelTextCommitted)
+												.OnVerifyTextChanged(this, &SStateTreeViewRow::HandleVerifyNodeLabelTextChanged)
 												.Text(this, &SStateTreeViewRow::GetStateDesc)
 												.ToolTipText(this, &SStateTreeViewRow::GetStateTypeTooltip)
 												.Clipping(EWidgetClipping::ClipToBounds)
@@ -1964,13 +1965,36 @@ bool SStateTreeViewRow::IsStateSelected() const
 	return false;
 }
 
+bool SStateTreeViewRow::HandleVerifyNodeLabelTextChanged(const FText& InText, FText& OutErrorMessage) const
+{
+	if (StateTreeViewModel)
+	{
+		if (const UStateTreeState* State = WeakState.Get())
+		{
+			const FString NewName = FText::TrimPrecedingAndTrailing(InText).ToString();
+			if (NewName.Len() >= NAME_SIZE)
+			{
+				OutErrorMessage = LOCTEXT("VerifyNodeLabelFailed_MaxLength", "Max length exceeded");
+				return false;
+			}
+			return FName::IsValidXName(NewName, INVALID_NAME_CHARACTERS, &OutErrorMessage);
+		}
+	}
+	OutErrorMessage = LOCTEXT("VerifyNodeLabelFailed", "Invalid State Tree");
+	return false;
+}
+
 void SStateTreeViewRow::HandleNodeLabelTextCommitted(const FText& NewLabel, ETextCommit::Type CommitType) const
 {
 	if (StateTreeViewModel)
 	{
 		if (UStateTreeState* State = WeakState.Get())
 		{
-			StateTreeViewModel->RenameState(State, FName(*FText::TrimPrecedingAndTrailing(NewLabel).ToString()));
+			const FString NewName = FText::TrimPrecedingAndTrailing(NewLabel).ToString();
+			if (FName::IsValidXName(NewName, INVALID_NAME_CHARACTERS) && NewName.Len() < NAME_SIZE)
+			{
+				StateTreeViewModel->RenameState(State, FName(NewName));
+			}
 		}
 	}
 }
diff --git a/Engine/Plugins/Runtime/StateTree/Source/StateTreeEditorModule/Private/SStateTreeViewRow.h b/Engine/Plugins/Runtime/StateTree/Source/StateTreeEditorModule/Private/SStateTreeViewRow.h
--- a/Engine/Plugins/Runtime/StateTree/Source/StateTreeEditorModule/Private/SStateTreeViewRow.h
+++ b/Engine/Plugins/Runtime/StateTree/Source/StateTreeEditorModule/Private/SStateTreeViewRow.h
@@ -109,6 +109,7 @@ private:
 	bool IsLeafState() const;
 	bool IsStateSelected() const;
 
+	bool HandleVerifyNodeLabelTextChanged(const FText& InText, FText& OutErrorMessage) const;
 	void HandleNodeLabelTextCommitted(const FText& NewLabel, ETextCommit::Type CommitType) const;
 
 	FReply HandleDragDetected(const FGeometry&, const FPointerEvent&) const;
