StateTree: Add new basic color.

#rnx

diff --git a/Engine/Plugins/Runtime/StateTree/Source/StateTreeModule/Private/StateTreeTypes.cpp b/Engine/Plugins/Runtime/StateTree/Source/StateTreeModule/Private/StateTreeTypes.cpp
--- a/Engine/Plugins/Runtime/StateTree/Source/StateTreeModule/Private/StateTreeTypes.cpp
+++ b/Engine/Plugins/Runtime/StateTree/Source/StateTreeModule/Private/StateTreeTypes.cpp
@@ -3,6 +3,7 @@
 #include "StateTreeTypes.h"
 #include "StateTreeEvents.h"
 #include "StateTree.h" // FStateTreeCustomVersion
+#include "Math/ColorList.h"
 
 #include UE_INLINE_GENERATED_CPP_BY_NAME(StateTreeTypes)
 
@@ -28,6 +29,7 @@ namespace UE::StateTree::Colors
 	const FColor Blue = FColor::FromHex(TEXT("#649ED3"));
 	const FColor Purple = FColor::FromHex(TEXT("#B397D6"));
 	const FColor Magenta = FColor::FromHex(TEXT("#CE85C7"));
+	const FColor Bronze = FColorList::Bronze;
 
 	constexpr float DarkenLevel = 0.6f;
 	const FColor DarkGrey = Darken(Grey, DarkenLevel);
@@ -39,6 +41,7 @@ namespace UE::StateTree::Colors
 	const FColor DarkBlue = Darken(Blue, DarkenLevel);
 	const FColor DarkPurple = Darken(Purple, DarkenLevel);
 	const FColor DarkMagenta = Darken(Magenta, DarkenLevel);
+	const FColor DarkBronze = Darken(Bronze, DarkenLevel);
 } // UE::StateTree::Colors
 
 
diff --git a/Engine/Plugins/Runtime/StateTree/Source/StateTreeModule/Public/StateTreeTypes.h b/Engine/Plugins/Runtime/StateTree/Source/StateTreeModule/Public/StateTreeTypes.h
--- a/Engine/Plugins/Runtime/StateTree/Source/StateTreeModule/Public/StateTreeTypes.h
+++ b/Engine/Plugins/Runtime/StateTree/Source/StateTreeModule/Public/StateTreeTypes.h
@@ -50,6 +50,8 @@ namespace UE::StateTree
 		extern const STATETREEMODULE_API FColor DarkPurple;
 		extern const STATETREEMODULE_API FColor Magenta;
 		extern const STATETREEMODULE_API FColor DarkMagenta;
+		extern const STATETREEMODULE_API FColor Bronze;
+		extern const STATETREEMODULE_API FColor DarkBronze;
 	} // Colors
 
 }; // UE::StateTree
