StateTree: Prevent crash when an evaluator class doesn't exist.
#rnx

diff --git a/Engine/Plugins/Runtime/StateTree/Source/StateTreeModule/Private/Blueprint/StateTreeEvaluatorBlueprintBase.cpp b/Engine/Plugins/Runtime/StateTree/Source/StateTreeModule/Private/Blueprint/StateTreeEvaluatorBlueprintBase.cpp
--- a/Engine/Plugins/Runtime/StateTree/Source/StateTreeModule/Private/Blueprint/StateTreeEvaluatorBlueprintBase.cpp
+++ b/Engine/Plugins/Runtime/StateTree/Source/StateTreeModule/Private/Blueprint/StateTreeEvaluatorBlueprintBase.cpp
@@ -91,18 +91,24 @@ FText FStateTreeBlueprintEvaluatorWrapper::GetDescription(const FGuid& ID, FStat
 
 FName FStateTreeBlueprintEvaluatorWrapper::GetIconName() const
 {
-	if (const UStateTreeNodeBlueprintBase* NodeCDO = GetDefault<const UStateTreeNodeBlueprintBase>(EvaluatorClass))
+	if (EvaluatorClass)
 	{
-		return NodeCDO->GetIconName();
+		if (const UStateTreeNodeBlueprintBase* NodeCDO = GetDefault<const UStateTreeNodeBlueprintBase>(EvaluatorClass))
+		{
+			return NodeCDO->GetIconName();
+		}
 	}
 	return FStateTreeEvaluatorBase::GetIconName();
 }
 
 FColor FStateTreeBlueprintEvaluatorWrapper::GetIconColor() const
 {
-	if (const UStateTreeNodeBlueprintBase* NodeCDO = GetDefault<const UStateTreeNodeBlueprintBase>(EvaluatorClass))
+	if (EvaluatorClass)
 	{
-		return NodeCDO->GetIconColor();
+		if (const UStateTreeNodeBlueprintBase* NodeCDO = GetDefault<const UStateTreeNodeBlueprintBase>(EvaluatorClass))
+		{
+			return NodeCDO->GetIconColor();
+		}
 	}
 	return FStateTreeEvaluatorBase::GetIconColor();
 }
