StateTree: Add image and color to the StateTree DetailView similar to the State DetailView.
#rnx

diff --git a/Engine/Plugins/Runtime/StateTree/Source/StateTreeEditorModule/Private/Customizations/StateTreeEditorDataDetails.cpp b/Engine/Plugins/Runtime/StateTree/Source/StateTreeEditorModule/Private/Customizations/StateTreeEditorDataDetails.cpp
--- a/Engine/Plugins/Runtime/StateTree/Source/StateTreeEditorModule/Private/Customizations/StateTreeEditorDataDetails.cpp
+++ b/Engine/Plugins/Runtime/StateTree/Source/StateTreeEditorModule/Private/Customizations/StateTreeEditorDataDetails.cpp
@@ -155,9 +155,18 @@ void FStateTreeEditorDataDetails::CustomizeDetails(IDetailLayoutBuilder& DetailB
 
 		const TSharedRef<SHorizontalBox> HeaderContentWidget = SNew(SHorizontalBox)
 			.IsEnabled(PropUtils, &IPropertyUtilities::IsPropertyEditingEnabled)
+			+ SHorizontalBox::Slot()
+			.VAlign(VAlign_Center)
+			.AutoWidth()
+			.Padding(FMargin(4.f, 0.f, 0.f, 0.f))
+			[
+				SNew(SImage)
+					.ColorAndOpacity(UE::StateTree::Colors::Blue)
+					.Image(FStateTreeEditorStyle::Get().GetBrush("StateTreeEditor.Parameters"))
+			]
 			+SHorizontalBox::Slot()
 			.VAlign(VAlign_Center)
-			.Padding(FMargin(4, 0, 0, 0))
+			.Padding(FMargin(4.f, 0.f, 0.f, 0.f))
 			.AutoWidth()
 			[
 				SNew(STextBlock)
@@ -169,7 +178,7 @@ void FStateTreeEditorDataDetails::CustomizeDetails(IDetailLayoutBuilder& DetailB
 			.HAlign(HAlign_Right)
 			.VAlign(VAlign_Center)
 			[
-				FPropertyBagDetails::MakeAddPropertyWidget(PropertyBagParametersProperty, PropUtils, EPropertyBagPropertyType::Bool, FLinearColor(UE::StateTree::Colors::Grey)).ToSharedRef()
+				FPropertyBagDetails::MakeAddPropertyWidget(PropertyBagParametersProperty, PropUtils, EPropertyBagPropertyType::Bool, FLinearColor(UE::StateTree::Colors::Blue)).ToSharedRef()
 			];
 		ParametersCategory.HeaderContent(HeaderContentWidget, /*FullRowContent*/true);
 
@@ -188,9 +197,9 @@ void FStateTreeEditorDataDetails::CustomizeDetails(IDetailLayoutBuilder& DetailB
 			EvaluatorsProperty,
 			EvalCategoryName,
 			LOCTEXT("EditorDataDetailsEvaluators", "Evaluators"),
-			FName(),
-			UE::StateTree::Colors::Grey,
-			UE::StateTree::Colors::Grey,
+			FName("StateTreeEditor.Evaluators"),
+			UE::StateTree::Colors::Bronze,
+			UE::StateTree::Colors::Bronze.WithAlpha(192),
 			LOCTEXT("EditorDataDetailsEvaluatorsAddTooltip", "Add new Evaluator"),
 			/*SortOrder*/3);
 	}
@@ -209,9 +218,9 @@ void FStateTreeEditorDataDetails::CustomizeDetails(IDetailLayoutBuilder& DetailB
 		GlobalTasksProperty,
 		GlobalTasksCategoryName,
 		LOCTEXT("EditorDataDetailsGlobalTasks", "Global Tasks"),
-		FName(),
-		UE::StateTree::Colors::Grey,
-		UE::StateTree::Colors::Grey,
+		FName("StateTreeEditor.Tasks"),
+		UE::StateTree::Colors::Cyan,
+		UE::StateTree::Colors::Cyan.WithAlpha(192),
 		LOCTEXT("EditorDataDetailsGlobalTasksAddTooltip", "Add new Global Task"),
 		/*SortOrder*/4);
 
diff --git a/Engine/Plugins/Runtime/StateTree/Source/StateTreeEditorModule/Private/Customizations/StateTreeStateDetails.cpp b/Engine/Plugins/Runtime/StateTree/Source/StateTreeEditorModule/Private/Customizations/StateTreeStateDetails.cpp
--- a/Engine/Plugins/Runtime/StateTree/Source/StateTreeEditorModule/Private/Customizations/StateTreeStateDetails.cpp
+++ b/Engine/Plugins/Runtime/StateTree/Source/StateTreeEditorModule/Private/Customizations/StateTreeStateDetails.cpp
@@ -223,7 +223,16 @@ void FStateTreeStateDetails::CustomizeDetails(IDetailLayoutBuilder& DetailBuilde
 
 		HeaderContentWidget->AddSlot()
 			.VAlign(VAlign_Center)
-			.Padding(FMargin(4, 0, 0, 0))
+			.AutoWidth()
+			.Padding(FMargin(4.f, 0.f, 0.f, 0.f))
+			[
+				SNew(SImage)
+					.ColorAndOpacity(UE::StateTree::Colors::Blue)
+					.Image(FStateTreeEditorStyle::Get().GetBrush("StateTreeEditor.Parameters"))
+			];
+		HeaderContentWidget->AddSlot()
+			.VAlign(VAlign_Center)
+			.Padding(FMargin(4.0f, 0.0f, 0.0f, 0.0f))
 			.AutoWidth()
 			[
 				SNew(STextBlock)
@@ -242,7 +251,7 @@ void FStateTreeStateDetails::CustomizeDetails(IDetailLayoutBuilder& DetailBuilde
 						ParametersParametersProperty,
 						PropUtils.ToSharedRef(),
 						EPropertyBagPropertyType::Bool,
-						FLinearColor(UE::StateTree::Colors::Grey)).ToSharedRef()
+						FLinearColor(UE::StateTree::Colors::Blue)).ToSharedRef()
 				];
 		}
 		
diff --git a/Engine/Plugins/Runtime/StateTree/Source/StateTreeEditorModule/Private/StateTreeEditorStyle.cpp b/Engine/Plugins/Runtime/StateTree/Source/StateTreeEditorModule/Private/StateTreeEditorStyle.cpp
--- a/Engine/Plugins/Runtime/StateTree/Source/StateTreeEditorModule/Private/StateTreeEditorStyle.cpp
+++ b/Engine/Plugins/Runtime/StateTree/Source/StateTreeEditorModule/Private/StateTreeEditorStyle.cpp
@@ -362,6 +362,8 @@ FStateTreeEditorStyle::FStateTreeEditorStyle()
 		Set("StateTreeEditor.StateConditions", new IMAGE_BRUSH_SVG("Icons/State_Conditions", CoreStyleConstants::Icon16x16));
 
 		Set("StateTreeEditor.Conditions", new IMAGE_BRUSH_SVG("Icons/Conditions", CoreStyleConstants::Icon16x16));
+		Set("StateTreeEditor.Evaluators", new IMAGE_BRUSH_SVG("Icons/Evaluators", CoreStyleConstants::Icon16x16));
+		Set("StateTreeEditor.Parameters", new IMAGE_BRUSH_SVG("Icons/Parameters", CoreStyleConstants::Icon16x16));
 		Set("StateTreeEditor.Utility", new IMAGE_BRUSH_SVG("Icons/Utility", CoreStyleConstants::Icon16x16));
 		Set("StateTreeEditor.Tasks", new IMAGE_BRUSH_SVG("Icons/Tasks", CoreStyleConstants::Icon16x16));
 		Set("StateTreeEditor.Transitions", new IMAGE_BRUSH_SVG("Icons/Transitions", CoreStyleConstants::Icon16x16));
