StateTree: Delay Link after the package is loaded or when first used.

diff --git a/Engine/Plugins/Runtime/StateTree/Source/StateTreeModule/Private/StateTree.cpp b/Engine/Plugins/Runtime/StateTree/Source/StateTreeModule/Private/StateTree.cpp
--- a/Engine/Plugins/Runtime/StateTree/Source/StateTreeModule/Private/StateTree.cpp
+++ b/Engine/Plugins/Runtime/StateTree/Source/StateTreeModule/Private/StateTree.cpp
@@ -37,7 +37,7 @@ bool UStateTree::IsReadyToRun() const
 
 FConstStructView UStateTree::GetNode(const int32 NodeIndex) const
 {
-	return Nodes.IsValidIndex(NodeIndex) ? Nodes[NodeIndex] : FConstStructView();	
+	return Nodes.IsValidIndex(NodeIndex) ? Nodes[NodeIndex] : FConstStructView();
 }
 
 FStateTreeIndex16 UStateTree::GetNodeIndexFromId(const FGuid Id) const
@@ -218,27 +218,45 @@ void UStateTree::OnObjectsReinstanced(const FReplacementObjectMap& ObjectMap)
 		return;
 	}
 
+	bool bShouldRelink = false;
+
+	// Relink if one of the out of date objects got reinstanced.
+	if (OutOfDateStructs.Num() > 0)
+	{
+		for (FObjectKey OutOfDateObjectKey : OutOfDateStructs)
+		{
+			if (const UObject* OutOfDateObject = OutOfDateObjectKey.ResolveObjectPtr())
+			{
+				if (ObjectMap.Contains(OutOfDateObject))
+				{
+					bShouldRelink = true;
+					break;
+				}
+			}
+		}
+	}
+
 	// If the asset is not linked yet (or has failed), no need to link.
-	if (!bIsLinked)
+	if (!bShouldRelink && !bIsLinked)
 	{
 		return;
 	}
 
 	// Relink only if the reinstantiated object belongs to this asset,
 	// or anything from the property binding refers to the classes of the reinstantiated object.
-
-	bool bShouldRelink = false;
-
-	for (TMap<UObject*, UObject*>::TConstIterator It(ObjectMap); It; ++It)
+	if (!bShouldRelink)
 	{
-		if (const UObject* ObjectToBeReplaced = It->Value)
+		for (TMap<UObject*, UObject*>::TConstIterator It(ObjectMap); It; ++It)
 		{
-			if (ObjectToBeReplaced->IsInOuter(this))
+			if (const UObject* ObjectToBeReplaced = It->Value)
 			{
-				bShouldRelink = true;
-				break;
+				if (ObjectToBeReplaced->IsInOuter(this))
+				{
+					bShouldRelink = true;
+					break;
+				}
 			}
-		}			
+		}
 	}
 
 	if (!bShouldRelink)
@@ -248,7 +266,15 @@ void UStateTree::OnObjectsReinstanced(const FReplacementObjectMap& ObjectMap)
 		{
 			if (const UObject* ObjectToBeReplaced = It->Value)
 			{
-				Structs.Add(ObjectToBeReplaced->GetClass());
+				// It's a UClass or a UScriptStruct
+				if (const UStruct* StructToReplaced = Cast<const UStruct>(ObjectToBeReplaced))
+				{
+					Structs.Add(StructToReplaced);
+				}
+				else
+				{
+					Structs.Add(ObjectToBeReplaced->GetClass());
+				}
 			}
 		}
 
@@ -447,7 +473,7 @@ void UStateTree::PostLoad()
 
 	if (!Link())
 	{
-		UE_LOG(LogStateTree, Log, TEXT("%s failed to link. Asset will not be usable at runtime."), *GetFullName());	
+		UE_LOG(LogStateTree, Log, TEXT("%s failed to link. Asset will not be usable at runtime."), *GetFullName());
 	}
 }
 
@@ -489,23 +515,26 @@ void UStateTree::ResetLinked()
 	bIsLinked = false;
 	ExternalDataDescs.Reset();
 
+#if WITH_EDITOR
+	OutOfDateStructs.Reset();
+#endif
+
 	FWriteScopeLock WriteLock(PerThreadSharedInstanceDataLock);
 	PerThreadSharedInstanceData.Reset();
 }
 
-bool UStateTree::ValidateInstanceData() const
+bool UStateTree::ValidateInstanceData()
 {
+	bool bResult = true;
 	for (FConstStructView NodeView : Nodes)
 	{
 		const FStateTreeNodeBase* Node = NodeView.GetPtr<const FStateTreeNodeBase>();
 		if (Node && Node->InstanceTemplateIndex.IsValid())
 		{
-			const UStruct* CurrentInstanceDataType;
+			const UStruct* CurrentInstanceDataType = nullptr;
 			{
 				const bool bUseSharedInstanceData = NodeView.GetPtr<const FStateTreeConditionBase>() || NodeView.GetPtr<const FStateTreeConsiderationBase>() || NodeView.GetPtr<const FStateTreePropertyFunctionBase>();
-
 				const FStateTreeInstanceData& SourceInstanceData = bUseSharedInstanceData ? SharedInstanceData : DefaultInstanceData;
-
 				if (SourceInstanceData.IsObject(Node->InstanceTemplateIndex.Get()))
 				{
 					const UObject* InstanceObject = SourceInstanceData.GetObject(Node->InstanceTemplateIndex.Get());
@@ -517,18 +546,66 @@ bool UStateTree::ValidateInstanceData() const
 				}
 			}
 
-			const UStruct* DesiredInstanceDataType = Node->GetInstanceDataType();
+			{
+				// Is the class/scriptstruct a blueprint that got replaced by another class.
+				bool bHasNewerVersionExists = false;
+				if (const UClass* CurrentInstanceDataClass = Cast<UClass>(CurrentInstanceDataType))
+				{
+					bHasNewerVersionExists = CurrentInstanceDataClass->HasAnyClassFlags(CLASS_NewerVersionExists);
+				}
+				else if (const UScriptStruct* CurrentInstanceDataStruct = Cast<UScriptStruct>(CurrentInstanceDataType))
+				{
+					bHasNewerVersionExists = (CurrentInstanceDataStruct->StructFlags & STRUCT_NewerVersionExists) != 0;
+				}
+
+				if (bHasNewerVersionExists)
+				{
+					bool bLogError = true;
+#if WITH_EDITOR
+					OutOfDateStructs.Add(CurrentInstanceDataType);
+					bLogError = false;
+#endif
+
+					if (bLogError)
+					{
+						UE_LOG(LogStateTree, Error, TEXT("%s: node '%s' failed. The source Instance Data type '%s' has a newer version."), *GetFullName(), *WriteToString<64>(Node->StaticStruct()->GetFName()), *WriteToString<64>(CurrentInstanceDataType->GetFName()));
+					}
+
+					bResult = false;
+				}
+			}
 
-			// Use strict testing so that the users will have option to initialize data mismatch if the type changes (even if potentially compatible).
-			if (CurrentInstanceDataType != DesiredInstanceDataType)
 			{
-				UE_LOG(LogStateTree, Error, TEXT("%s: node '%s' failed. The source Instance Data type '%s' does not match '%s'"), *GetFullName(), *Node->StaticStruct()->GetName(), *GetNameSafe(CurrentInstanceDataType), *GetNameSafe(DesiredInstanceDataType));
-				return false;
+				const UStruct* DesiredInstanceDataType = Node->GetInstanceDataType();
+
+				// Use strict testing so that the users will have option to initialize data mismatch if the type changes (even if potentially compatible).
+				if (CurrentInstanceDataType != DesiredInstanceDataType)
+				{
+					bool bLogError = true;
+#if WITH_EDITOR
+					const UClass* CurrentInstanceDataClass = Cast<UClass>(CurrentInstanceDataType);
+					const UClass* DesiredInstanceDataClass = Cast<UClass>(DesiredInstanceDataType);
+					if (CurrentInstanceDataClass && DesiredInstanceDataClass)
+					{
+						// Because of the loading order.It's possible that the OnObjectsReinstanced did complete.
+						if (CurrentInstanceDataClass->ClassGeneratedBy == DesiredInstanceDataClass->ClassGeneratedBy)
+						{
+							OutOfDateStructs.Add(CurrentInstanceDataType);
+							bLogError = false;
+						}
+					}
+#endif
+					if (bLogError)
+					{
+						UE_LOG(LogStateTree, Error, TEXT("%s: node '%s' failed. The source Instance Data type '%s' does not match '%s'"), *GetFullName(), *WriteToString<64>(Node->StaticStruct()->GetFName()), *GetNameSafe(CurrentInstanceDataType), *GetNameSafe(DesiredInstanceDataType));
+					}
+					bResult = false;
+				}
 			}
 		}
 	}
 
-	return true;
+	return bResult;
 }
 
 bool UStateTree::Link()
diff --git a/Engine/Plugins/Runtime/StateTree/Source/StateTreeModule/Public/StateTree.h b/Engine/Plugins/Runtime/StateTree/Source/StateTreeModule/Public/StateTree.h
--- a/Engine/Plugins/Runtime/StateTree/Source/StateTreeModule/Public/StateTree.h
+++ b/Engine/Plugins/Runtime/StateTree/Source/StateTreeModule/Public/StateTree.h
@@ -6,6 +6,7 @@
 #include "StateTreeSchema.h"
 #include "StateTreePropertyBindings.h"
 #include "StateTreeInstanceData.h"
+#include "UObject/ObjectKey.h"
 #include "StateTree.generated.h"
 
 class UUserDefinedStruct;
@@ -240,7 +241,7 @@ private:
 	void ResetLinked();
 
 	/** @return true if all the source instance data types match with the node's instance data types */
-	bool ValidateInstanceData() const;
+	bool ValidateInstanceData();
 
 	bool PatchBindings();
 
@@ -337,6 +338,11 @@ private:
 	/** True if the StateTree was linked successfully. */
 	bool bIsLinked = false;
 
+#if WITH_EDITORONLY_DATA
+	/** List of Struct that are out of date and waiting to be replaced with the new instance. */
+	TSet<FObjectKey> OutOfDateStructs;
+#endif
+
 	friend struct FStateTreeInstance;
 	friend struct FStateTreeExecutionContext;
 #if WITH_EDITORONLY_DATA
