StateTree: Add "reset to default" on bindings. Hide parent when a member has a binding. Change the test of the parent and color of the parent binding.

diff --git a/Engine/Plugins/Runtime/StateTree/Source/StateTreeEditorModule/Private/Customizations/StateTreeBindingExtension.cpp b/Engine/Plugins/Runtime/StateTree/Source/StateTreeEditorModule/Private/Customizations/StateTreeBindingExtension.cpp
--- a/Engine/Plugins/Runtime/StateTree/Source/StateTreeEditorModule/Private/Customizations/StateTreeBindingExtension.cpp
+++ b/Engine/Plugins/Runtime/StateTree/Source/StateTreeEditorModule/Private/Customizations/StateTreeBindingExtension.cpp
@@ -465,7 +465,30 @@ struct FCachedBindingData : public TSharedFromThis<FCachedBindingData>
 		UE::StateTree::PropertyBinding::OnStateTreePropertyBindingChanged.Broadcast(SourcePath, TargetPath);
 	}
 
-	void RemoveBinding()
+	bool HasBinding(FStateTreeEditorPropertyBindings::ESearchMode SearchMode) const
+	{
+		UObject* OwnerObject = WeakOwnerObject.Get();
+		if (!OwnerObject)
+		{
+			return false;
+		}
+
+		IStateTreeEditorPropertyBindingsOwner* BindingOwner = Cast<IStateTreeEditorPropertyBindingsOwner>(OwnerObject);
+		if (!BindingOwner)
+		{
+			return false;
+		}
+
+		FStateTreeEditorPropertyBindings* EditorBindings = BindingOwner->GetPropertyEditorBindings();
+		if (!EditorBindings)
+		{
+			return false;
+		}
+
+		return EditorBindings && EditorBindings->HasPropertyBinding(TargetPath, SearchMode);
+	}
+
+	void RemoveBinding(FStateTreeEditorPropertyBindings::ESearchMode RemoveMode)
 	{
 		UObject* OwnerObject = WeakOwnerObject.Get();
 		if (!OwnerObject)
@@ -486,7 +509,7 @@ struct FCachedBindingData : public TSharedFromThis<FCachedBindingData>
 		}
 
 		OwnerObject->Modify();
-		EditorBindings->RemovePropertyBindings(TargetPath);
+		EditorBindings->RemovePropertyBindings(TargetPath, RemoveMode);
 
 		UpdateData();
 		
@@ -1551,17 +1574,19 @@ void FStateTreeBindingExtension::ExtendWidgetRow(FDetailWidgetRow& InWidgetRow,
 				}
 			}
 		}
+	}
 
-		// Wrap value widget 
-		if (EditorBindings)
-		{
-			auto IsValueVisible = TAttribute<EVisibility>::Create([TargetPath, EditorBindings]() -> EVisibility
+	TSharedPtr<UE::StateTree::PropertyBinding::FCachedBindingData> CachedBindingData = MakeShared<UE::StateTree::PropertyBinding::FCachedBindingData>(OwnerObject, TargetPath, InPropertyHandle, AccessibleStructs);
+
+	// Wrap value widget 
+	{
+		auto IsValueVisible = TAttribute<EVisibility>::Create([CachedBindingData]() -> EVisibility
 			{
-				return EditorBindings->HasPropertyBinding(TargetPath) ? EVisibility::Collapsed : EVisibility::Visible;
+				return CachedBindingData->HasBinding(FStateTreeEditorPropertyBindings::ESearchMode::Exact) ? EVisibility::Collapsed : EVisibility::Visible;
 			});
 
-			TSharedPtr<SWidget> ValueWidget = InWidgetRow.ValueContent().Widget;
-			InWidgetRow.ValueContent()
+		TSharedPtr<SWidget> ValueWidget = InWidgetRow.ValueContent().Widget;
+		InWidgetRow.ValueContent()
 			[
 				SNew(SBox)
 				.Visibility(IsValueVisible)
@@ -1569,11 +1594,7 @@ void FStateTreeBindingExtension::ExtendWidgetRow(FDetailWidgetRow& InWidgetRow,
 					ValueWidget.ToSharedRef()
 				]
 			];
-		}
-		
 	}
-
-	TSharedPtr<UE::StateTree::PropertyBinding::FCachedBindingData> CachedBindingData = MakeShared<UE::StateTree::PropertyBinding::FCachedBindingData>(OwnerObject, TargetPath, InPropertyHandle, AccessibleStructs);
 	
 	FPropertyBindingWidgetArgs Args;
 	Args.Property = InPropertyHandle->GetProperty();
@@ -1606,13 +1627,13 @@ void FStateTreeBindingExtension::ExtendWidgetRow(FDetailWidgetRow& InWidgetRow,
 
 	Args.OnRemoveBinding = FOnRemoveBinding::CreateLambda([CachedBindingData, &InDetailBuilder](FName InPropertyName)
 		{
-			CachedBindingData->RemoveBinding();
+			CachedBindingData->RemoveBinding(FStateTreeEditorPropertyBindings::ESearchMode::Exact);
 			InDetailBuilder.GetPropertyUtilities()->RequestForceRefresh();
 		});
 
-	Args.OnCanRemoveBinding = FOnCanRemoveBinding::CreateLambda([EditorBindings, TargetPath](FName InPropertyName)
+	Args.OnCanRemoveBinding = FOnCanRemoveBinding::CreateLambda([CachedBindingData](FName InPropertyName)
 		{
-			return EditorBindings && EditorBindings->HasPropertyBinding(TargetPath);
+			return CachedBindingData->HasBinding(FStateTreeEditorPropertyBindings::ESearchMode::Exact);
 		});
 
 	Args.CurrentBindingText = MakeAttributeLambda([CachedBindingData]()
@@ -1705,6 +1726,28 @@ void FStateTreeBindingExtension::ExtendWidgetRow(FDetailWidgetRow& InWidgetRow,
 		);
 	}
 
+	// ResetToDefault
+	{
+		InWidgetRow.CustomResetToDefault = FResetToDefaultOverride::Create(
+			MakeAttributeLambda([CachedBindingData, InPropertyHandle]()
+			{
+				return InPropertyHandle->CanResetToDefault() || CachedBindingData->HasBinding(FStateTreeEditorPropertyBindings::ESearchMode::Includes);
+			}),
+			FSimpleDelegate::CreateLambda([CachedBindingData, &InDetailBuilder, InPropertyHandle]()
+				{
+					if (CachedBindingData->HasBinding(FStateTreeEditorPropertyBindings::ESearchMode::Includes))
+					{
+						CachedBindingData->RemoveBinding(FStateTreeEditorPropertyBindings::ESearchMode::Includes);
+						InDetailBuilder.GetPropertyUtilities()->RequestForceRefresh();
+					}
+					if (InPropertyHandle->CanResetToDefault())
+					{
+						InPropertyHandle->ResetToDefault();
+					}
+				}),
+			false);
+	}
+
 	InWidgetRow.ExtensionContent()
 	[
 		PropertyAccessEditor.MakePropertyBindingWidget(BindingContextStructs, Args)
@@ -1796,8 +1839,10 @@ bool FStateTreeBindingsChildrenCustomization::ShouldCustomizeChildren(TSharedRef
 		IStateTreeEditorPropertyBindingsOwner* BindingOwner = Cast<IStateTreeEditorPropertyBindingsOwner>(UE::StateTree::PropertyBinding::FindEditorBindingsOwner(OuterObjects[0]));
 		if (!TargetPath.IsPathEmpty() && BindingOwner)
 		{
-			FStateTreeEditorPropertyBindings* EditorBindings = BindingOwner->GetPropertyEditorBindings();
-			return EditorBindings->HasPropertyBinding(TargetPath);
+			if (FStateTreeEditorPropertyBindings* EditorBindings = BindingOwner->GetPropertyEditorBindings())
+			{
+				return EditorBindings->HasPropertyBinding(TargetPath);
+			}
 		}
 	}
 
diff --git a/Engine/Plugins/Runtime/StateTree/Source/StateTreeEditorModule/Private/Customizations/StateTreeEditorNodeDetails.cpp b/Engine/Plugins/Runtime/StateTree/Source/StateTreeEditorModule/Private/Customizations/StateTreeEditorNodeDetails.cpp
--- a/Engine/Plugins/Runtime/StateTree/Source/StateTreeEditorModule/Private/Customizations/StateTreeEditorNodeDetails.cpp
+++ b/Engine/Plugins/Runtime/StateTree/Source/StateTreeEditorModule/Private/Customizations/StateTreeEditorNodeDetails.cpp
@@ -160,12 +160,9 @@ namespace UE::StateTreeEditor::Internal
 			FDetailWidgetRow Row;
 			ChildRow.GetDefaultWidgets(NameWidget, ValueWidget, Row);
 
-			auto IsValueVisible = TAttribute<EVisibility>::Create([Path, EditorPropBindings]() -> EVisibility
-				{
-					return EditorPropBindings->HasPropertyBinding(Path) ? EVisibility::Collapsed : EVisibility::Visible;
-				});
-
-			if (Usage == EStateTreePropertyUsage::Input || Usage == EStateTreePropertyUsage::Output || Usage == EStateTreePropertyUsage::Context)
+			const bool bHasChildPropertyBinding = EditorPropBindings->HasPropertyBinding(Path, FStateTreeEditorPropertyBindings::ESearchMode::Includes);
+			const bool bValidUsage = Usage == EStateTreePropertyUsage::Input || Usage == EStateTreePropertyUsage::Output || Usage == EStateTreePropertyUsage::Context;
+			if (bHasChildPropertyBinding || bValidUsage)
 			{				
 				FEdGraphPinType PinType;
 				const UEdGraphSchema_K2* Schema = GetDefault<UEdGraphSchema_K2>();
@@ -189,7 +186,12 @@ namespace UE::StateTreeEditor::Internal
 				{
 					Schema->ConvertPropertyToPinType(Property, PinType);
 				}
-				
+
+				auto IsValueVisible = TAttribute<EVisibility>::Create([Path, EditorPropBindings]() -> EVisibility
+					{
+						return EditorPropBindings->HasPropertyBinding(Path, FStateTreeEditorPropertyBindings::ESearchMode::Exact) ? EVisibility::Collapsed : EVisibility::Visible;
+					});
+
 				const FSlateBrush* Icon = FBlueprintEditorUtils::GetIconFromPin(PinType, true);
 				FText Text = GetPinTypeText(PinType);
 				
@@ -197,6 +199,7 @@ namespace UE::StateTreeEditor::Internal
 				FLinearColor IconColor = Schema->GetPinTypeColor(PinType);
 				FText Label;
 				FText LabelToolTip;
+				FSlateColor TextColor = FSlateColor::UseForeground();
 
 				if (Usage == EStateTreePropertyUsage::Input)
 				{
@@ -242,6 +245,12 @@ namespace UE::StateTreeEditor::Internal
 						IconColor = FLinearColor::White;
 					}
 				}
+				else if (ensure(bHasChildPropertyBinding))
+				{
+					Icon = FCoreStyle::Get().GetBrush("Icons.Link");
+					Text = LOCTEXT("PropertyMemberIsBoundTooltip", "{Member(s) Bound_B}");
+					TextColor = FSlateColor::UseSubduedForeground();
+				}
 				
 				ChildRow
 					.CustomWidget(true)
@@ -262,7 +271,7 @@ namespace UE::StateTreeEditor::Internal
 						.Padding(4.0f, 0.0f)
 						[
 							SNew(SBorder)
-							.Padding(FMargin(6, 1))
+							.Padding(FMargin(6.0f, 1.0f))
 							.BorderImage(FStateTreeEditorStyle::Get().GetBrush("StateTree.Param.Background"))
 							.Visibility(Label.IsEmpty() ? EVisibility::Collapsed : EVisibility::Visible)
 							[
@@ -297,6 +306,7 @@ namespace UE::StateTreeEditor::Internal
 						[
 							SNew(STextBlock)
 							.Font(IDetailLayoutBuilder::GetDetailFont())
+							.ColorAndOpacity(TextColor)
 							.Text(Text)
 							.ToolTipText(ToolTip)
 						]
diff --git a/Engine/Plugins/Runtime/StateTree/Source/StateTreeEditorModule/Private/StateTreeEditorPropertyBindings.cpp b/Engine/Plugins/Runtime/StateTree/Source/StateTreeEditorModule/Private/StateTreeEditorPropertyBindings.cpp
--- a/Engine/Plugins/Runtime/StateTree/Source/StateTreeEditorModule/Private/StateTreeEditorPropertyBindings.cpp
+++ b/Engine/Plugins/Runtime/StateTree/Source/StateTreeEditorModule/Private/StateTreeEditorPropertyBindings.cpp
@@ -19,14 +19,14 @@ UStateTreeEditorPropertyBindingsOwner::UStateTreeEditorPropertyBindingsOwner(con
 
 //////////////////////////////////////////////////////////////////////////
 
-void FStateTreeEditorPropertyBindings::SetBindingsOwner(IStateTreeEditorPropertyBindingsOwner* InBindingsOwner)
+void FStateTreeEditorPropertyBindings::SetBindingsOwner(gsl::not_null<TScriptInterface<IStateTreeEditorPropertyBindingsOwner>> InBindingsOwner)
 {
 	BindingsOwner = InBindingsOwner;
 }
 
 void FStateTreeEditorPropertyBindings::AddPropertyBinding(const FStateTreePropertyPath& SourcePath, const FStateTreePropertyPath& TargetPath)
 {
-	RemovePropertyBindings(TargetPath);
+	RemovePropertyBindings(TargetPath, ESearchMode::Exact);
 
 	FStateTreePropertyPathBinding Binding(SourcePath, TargetPath);
 
@@ -63,7 +63,7 @@ FStateTreePropertyPath FStateTreeEditorPropertyBindings::AddFunctionPropertyBind
 		PropertyFunction.Instance.InitializeAs(InstanceType);
 	}
 
-	RemovePropertyBindings(TargetPath);
+	RemovePropertyBindings(TargetPath, ESearchMode::Exact);
 	FStateTreePropertyPath SourcePath = FStateTreePropertyPath(NodeID, SourcePathSegments);
 	PropertyBindings.Emplace(MoveTemp(PropertyFunctionNode), SourcePath, TargetPath);
 	return SourcePath;
@@ -71,16 +71,26 @@ FStateTreePropertyPath FStateTreeEditorPropertyBindings::AddFunctionPropertyBind
 
 void FStateTreeEditorPropertyBindings::AddPropertyBinding(const FStateTreePropertyPathBinding& Binding)
 {
-	RemovePropertyBindings(Binding.GetTargetPath());
+	RemovePropertyBindings(Binding.GetTargetPath(), ESearchMode::Exact);
 	PropertyBindings.Add(Binding);
 }
 
-void FStateTreeEditorPropertyBindings::RemovePropertyBindings(const FStateTreePropertyPath& TargetPath)
+void FStateTreeEditorPropertyBindings::RemovePropertyBindings(const FStateTreePropertyPath& TargetPath, ESearchMode SearchMode)
 {
-	PropertyBindings.RemoveAll([&TargetPath](const FStateTreePropertyPathBinding& Binding)
-		{
-			return Binding.GetTargetPath() == TargetPath;
-		});
+	if (SearchMode == ESearchMode::Exact)
+	{
+		PropertyBindings.RemoveAllSwap([&TargetPath](const FStateTreePropertyPathBinding& Binding)
+			{
+				return Binding.GetTargetPath() == TargetPath;
+			});
+	}
+	else
+	{
+		PropertyBindings.RemoveAllSwap([&TargetPath](const FStateTreePropertyPathBinding& Binding)
+			{
+				return Binding.GetTargetPath().Includes(TargetPath);
+			});
+	}
 }
 
 void FStateTreeEditorPropertyBindings::CopyBindings(const FGuid FromStructID, const FGuid ToStructID)
@@ -98,12 +108,40 @@ void FStateTreeEditorPropertyBindings::CopyBindings(const FGuid FromStructID, co
 	PropertyBindings.Append(NewBindings);
 }
 
-bool FStateTreeEditorPropertyBindings::HasPropertyBinding(const FStateTreePropertyPath& TargetPath) const
+bool FStateTreeEditorPropertyBindings::HasPropertyBinding(const FStateTreePropertyPath& TargetPath, ESearchMode RemoveMode) const
 {
-	return PropertyBindings.ContainsByPredicate([&TargetPath](const FStateTreePropertyPathBinding& Binding)
-		{
-			return Binding.GetTargetPath() == TargetPath;
-		});
+	if (RemoveMode == ESearchMode::Exact)
+	{
+		return PropertyBindings.ContainsByPredicate([&TargetPath](const FStateTreePropertyPathBinding& Binding)
+			{
+				return Binding.GetTargetPath() == TargetPath;
+			});
+	}
+	else
+	{
+		return PropertyBindings.ContainsByPredicate([&TargetPath](const FStateTreePropertyPathBinding& Binding)
+			{
+				return Binding.GetTargetPath().Includes(TargetPath);
+			});
+	}
+}
+
+const FStateTreePropertyPathBinding* FStateTreeEditorPropertyBindings::FindPropertyBinding(const FStateTreePropertyPath& TargetPath, ESearchMode RemoveMode) const
+{
+	if (RemoveMode == ESearchMode::Exact)
+	{
+		return PropertyBindings.FindByPredicate([&TargetPath](const FStateTreePropertyPathBinding& Binding)
+			{
+				return Binding.GetTargetPath() == TargetPath;
+			});
+	}
+	else
+	{
+		return PropertyBindings.FindByPredicate([&TargetPath](const FStateTreePropertyPathBinding& Binding)
+			{
+				return Binding.GetTargetPath().Includes(TargetPath);
+			});
+	}
 }
 
 const FStateTreePropertyPath* FStateTreeEditorPropertyBindings::GetPropertyBindingSource(const FStateTreePropertyPath& TargetPath) const
@@ -128,7 +166,7 @@ void FStateTreeEditorPropertyBindings::GetPropertyBindingsFor(const FGuid Struct
 
 void FStateTreeEditorPropertyBindings::RemoveUnusedBindings(const TMap<FGuid, const FStateTreeDataView>& ValidStructs)
 {
-	PropertyBindings.RemoveAll([ValidStructs](FStateTreePropertyPathBinding& Binding)
+	PropertyBindings.RemoveAllSwap([ValidStructs](FStateTreePropertyPathBinding& Binding)
 		{
 			// Remove binding if it's target struct has been removed
 			if (!ValidStructs.Contains(Binding.GetTargetPath().GetStructID()))
@@ -217,12 +255,12 @@ void FStateTreeEditorPropertyBindings::AddPropertyBinding(const FStateTreeEditor
 
 void FStateTreeEditorPropertyBindings::RemovePropertyBindings(const FStateTreeEditorPropertyPath& TargetPath)
 {
-	RemovePropertyBindings(UE::StateTree::Private::ConvertEditorPath(TargetPath));
+	RemovePropertyBindings(UE::StateTree::Private::ConvertEditorPath(TargetPath), ESearchMode::Exact);
 }
 
 bool FStateTreeEditorPropertyBindings::HasPropertyBinding(const FStateTreeEditorPropertyPath& TargetPath) const
 {
-	return HasPropertyBinding(UE::StateTree::Private::ConvertEditorPath(TargetPath));
+	return HasPropertyBinding(UE::StateTree::Private::ConvertEditorPath(TargetPath), ESearchMode::Exact);
 }
 
 const FStateTreeEditorPropertyPath* FStateTreeEditorPropertyBindings::GetPropertyBindingSource(const FStateTreeEditorPropertyPath& TargetPath) const
@@ -247,7 +285,7 @@ PRAGMA_ENABLE_DEPRECATION_WARNINGS
 
 void FStateTreeEditorPropertyBindings::RemoveUnusedBindings(const TMap<FGuid, const UStruct*>& ValidStructs)
 {
-	PropertyBindings.RemoveAll([ValidStructs](const FStateTreePropertyPathBinding& Binding)
+	PropertyBindings.RemoveAllSwap([ValidStructs](const FStateTreePropertyPathBinding& Binding)
 		{
 			// Remove binding if it's target struct has been removed
 			if (!ValidStructs.Contains(Binding.GetTargetPath().GetStructID()))
diff --git a/Engine/Plugins/Runtime/StateTree/Source/StateTreeEditorModule/Public/StateTreeEditorPropertyBindings.h b/Engine/Plugins/Runtime/StateTree/Source/StateTreeEditorModule/Public/StateTreeEditorPropertyBindings.h
--- a/Engine/Plugins/Runtime/StateTree/Source/StateTreeEditorModule/Public/StateTreeEditorPropertyBindings.h
+++ b/Engine/Plugins/Runtime/StateTree/Source/StateTreeEditorModule/Public/StateTreeEditorPropertyBindings.h
@@ -20,7 +20,7 @@ struct STATETREEEDITORMODULE_API FStateTreeEditorPropertyBindings
 	GENERATED_BODY()
 
 	/** Sets associated bindings owner, used to validate added property paths. */
-	void SetBindingsOwner(IStateTreeEditorPropertyBindingsOwner* BindingsOwner);
+	void SetBindingsOwner(gsl::not_null<TScriptInterface<IStateTreeEditorPropertyBindingsOwner>> BindingsOwner);
 
 	/**
 	 * Adds binding between source and destination paths. Removes any bindings to TargetPath before adding the new one.
@@ -44,17 +44,28 @@ struct STATETREEEDITORMODULE_API FStateTreeEditorPropertyBindings
 	 */
 	FStateTreePropertyPath AddFunctionPropertyBinding(const UScriptStruct* PropertyFunctionNodeStruct, TConstArrayView<FStateTreePropertyPathSegment> SourcePathSegments, const FStateTreePropertyPath& TargetPath);
 
+	enum class ESearchMode
+	{
+		Exact,				// Binding with exact matching path.
+		Includes,			// Binding with path that matches but the binding path can be longer.
+	};
+
 	/**
 	 * Removes all bindings to target path.
 	 * @param TargetPath Target property path.
 	 */ 
-	void RemovePropertyBindings(const FStateTreePropertyPath& TargetPath);
+	void RemovePropertyBindings(const FStateTreePropertyPath& TargetPath, ESearchMode SearchMode = ESearchMode::Exact);
 	
 	/**
-	 * @param TargetPath Target property path.
+	 * Has any binding to the target path.
 	 * @return True of the target path has any bindings.
 	 */
-	bool HasPropertyBinding(const FStateTreePropertyPath& TargetPath) const;
+	bool HasPropertyBinding(const FStateTreePropertyPath& TargetPath, ESearchMode SearchMode = ESearchMode::Exact) const;
+
+	/**
+	 * @return binding to the target path.
+	 */
+	const FStateTreePropertyPathBinding* FindPropertyBinding(const FStateTreePropertyPath& TargetPath, ESearchMode SearchMode = ESearchMode::Exact) const;
 
 	/**
 	 * Copies property bindings from an existing struct to another.
@@ -86,9 +97,15 @@ struct STATETREEEDITORMODULE_API FStateTreeEditorPropertyBindings
 	bool ContainsAnyStruct(const TSet<const UStruct*>& Structs);
 
 	/** @return array view to all bindings. */
-	TConstArrayView<FStateTreePropertyPathBinding> GetBindings() const { return PropertyBindings; }
+	const TConstArrayView<FStateTreePropertyPathBinding> GetBindings() const
+	{
+		return PropertyBindings;
+	}
 
-	TArrayView<FStateTreePropertyPathBinding> GetMutableBindings() { return PropertyBindings; }
+	const TArrayView<FStateTreePropertyPathBinding> GetMutableBindings()
+	{
+		return PropertyBindings;
+	}
 
 PRAGMA_DISABLE_DEPRECATION_WARNINGS
 	UE_DEPRECATED(5.3, "Use version with FStateTreePropertyPath instead.")
@@ -111,11 +128,11 @@ PRAGMA_DISABLE_DEPRECATION_WARNINGS
 PRAGMA_ENABLE_DEPRECATION_WARNINGS
 	
 private:
-
 	UPROPERTY()
 	TArray<FStateTreePropertyPathBinding> PropertyBindings;
 
-	IStateTreeEditorPropertyBindingsOwner* BindingsOwner = nullptr; 
+	UPROPERTY(Transient)
+	TScriptInterface<IStateTreeEditorPropertyBindingsOwner> BindingsOwner = nullptr; 
 };
 
 
diff --git a/Engine/Plugins/Runtime/StateTree/Source/StateTreeModule/Private/StateTreePropertyBindings.cpp b/Engine/Plugins/Runtime/StateTree/Source/StateTreeModule/Private/StateTreePropertyBindings.cpp
--- a/Engine/Plugins/Runtime/StateTree/Source/StateTreeModule/Private/StateTreePropertyBindings.cpp
+++ b/Engine/Plugins/Runtime/StateTree/Source/StateTreeModule/Private/StateTreePropertyBindings.cpp
@@ -1863,6 +1863,33 @@ bool FStateTreePropertyPath::ResolveIndirectionsWithValue(const FStateTreeDataVi
 	return true;
 }
 
+bool FStateTreePropertyPath::Includes(const FStateTreePropertyPath& Other) const
+{
+#if WITH_EDITORONLY_DATA
+	if (StructID != Other.StructID)
+	{
+		return false;
+	}
+#endif // WITH_EDITORONLY_DATA
+
+	if (Segments.Num() < Other.Segments.Num())
+	{
+		return false;
+	}
+
+	for (TEnumerateRef<const FStateTreePropertyPathSegment> OtherSegment : EnumerateRange(Other.Segments))
+	{
+		if (*OtherSegment != Segments[OtherSegment.GetIndex()])
+		{
+			return false;
+		}
+	}
+
+	return true;
+}
+
+
+
 bool FStateTreePropertyPath::operator==(const FStateTreePropertyPath& RHS) const
 {
 #if WITH_EDITORONLY_DATA
@@ -1871,6 +1898,7 @@ bool FStateTreePropertyPath::operator==(const FStateTreePropertyPath& RHS) const
 		return false;
 	}
 #endif // WITH_EDITORONLY_DATA
+
 	if (Segments.Num() != RHS.Segments.Num())
 	{
 		return false;
diff --git a/Engine/Plugins/Runtime/StateTree/Source/StateTreeModule/Public/StateTreePropertyBindings.h b/Engine/Plugins/Runtime/StateTree/Source/StateTreeModule/Public/StateTreePropertyBindings.h
--- a/Engine/Plugins/Runtime/StateTree/Source/StateTreeModule/Public/StateTreePropertyBindings.h
+++ b/Engine/Plugins/Runtime/StateTree/Source/StateTreeModule/Public/StateTreePropertyBindings.h
@@ -493,6 +493,12 @@ struct STATETREEMODULE_API FStateTreePropertyPath
 		Segments.Add(PathSegment);
 	}
 
+	/**
+	 * Test if this paths includes the provided path.
+	 * A path includes another one when they are == but this path can be longer.
+	 */
+	bool Includes(const FStateTreePropertyPath& Other) const;
+
 	/** Test if paths are equal. */
 	bool operator==(const FStateTreePropertyPath& RHS) const;
 
