[State Tree] Simplified Namings

diff --git a/Engine/Plugins/Runtime/StateTree/Source/StateTreeEditorModule/Private/Customizations/StateTreeStateDetails.cpp b/Engine/Plugins/Runtime/StateTree/Source/StateTreeEditorModule/Private/Customizations/StateTreeStateDetails.cpp
--- a/Engine/Plugins/Runtime/StateTree/Source/StateTreeEditorModule/Private/Customizations/StateTreeStateDetails.cpp
+++ b/Engine/Plugins/Runtime/StateTree/Source/StateTreeEditorModule/Private/Customizations/StateTreeStateDetails.cpp
@@ -300,12 +300,12 @@ void FStateTreeStateDetails::CustomizeDetails(IDetailLayoutBuilder& DetailBuilde
 	ConsiderationsProperty->MarkHiddenByCustomization();
 	if (Schema && Schema->AllowUtilityConsiderations())
 	{
-		const FName UtilityCategoryName(TEXT("Utility"));
+		const FName UtilityCategoryName(TEXT("Selection Utility"));
 		IDetailCategoryBuilder& UtilityConsiderationsCategory = UE::StateTreeEditor::EditorNodeUtils::MakeArrayCategory(
 			DetailBuilder,
 			ConsiderationsProperty,
 			UtilityCategoryName,
-			LOCTEXT("StateDetailsUtility", "Utility"),
+			LOCTEXT("StateDetailsSelectionUtility", "Selection Utility"),
 			FName("StateTreeEditor.Utility"),
 			UE::StateTree::Colors::Orange,
 			UE::StateTree::Colors::Orange.WithAlpha(192),
diff --git a/Engine/Plugins/Runtime/StateTree/Source/StateTreeEditorModule/Private/SStateTreeViewRow.cpp b/Engine/Plugins/Runtime/StateTree/Source/StateTreeEditorModule/Private/SStateTreeViewRow.cpp
--- a/Engine/Plugins/Runtime/StateTree/Source/StateTreeEditorModule/Private/SStateTreeViewRow.cpp
+++ b/Engine/Plugins/Runtime/StateTree/Source/StateTreeEditorModule/Private/SStateTreeViewRow.cpp
@@ -1027,9 +1027,9 @@ FText SStateTreeViewRow::GetSelectorTooltip() const
 			case EStateTreeStateSelectionBehavior::TryFollowTransitions:
 				return Enum->GetToolTipTextByIndex(Index);
 			case EStateTreeStateSelectionBehavior::TrySelectChildrenInOrder:
-			case EStateTreeStateSelectionBehavior::TrySelectChildrenAtUniformRandom:
+			case EStateTreeStateSelectionBehavior::TrySelectChildrenAtRandom:
 			case EStateTreeStateSelectionBehavior::TrySelectChildrenWithHighestUtility:
-			case EStateTreeStateSelectionBehavior::TrySelectChildrenBasedOnRelativeUtility:
+			case EStateTreeStateSelectionBehavior::TrySelectChildrenAtRandomWeightedByUtility:
 				if (State->Children.IsEmpty()
 					|| State->Type == EStateTreeStateType::Linked
 					|| State->Type == EStateTreeStateType::LinkedAsset)
@@ -1332,14 +1332,14 @@ bool SStateTreeViewRow::GetStateWarnings(FText* OutText) const
 	{
 		if (!State->Parent 
 			|| (State->Parent->SelectionBehavior != EStateTreeStateSelectionBehavior::TrySelectChildrenWithHighestUtility
-				&& State->Parent->SelectionBehavior != EStateTreeStateSelectionBehavior::TrySelectChildrenBasedOnRelativeUtility))
+				&& State->Parent->SelectionBehavior != EStateTreeStateSelectionBehavior::TrySelectChildrenAtRandomWeightedByUtility))
 		{
 			if (OutText)
 			{
 				*OutText = LOCTEXT("ChildStateUtilityConsiderationWarning", 
 					"State has Utility Considerations but they don't have effect."
 					"The Utility Considerations are used only when parent State's Selection Behavior is:"
-					"\"Try Select Children with Highest Utility\" or \"Try Select Children Based on Relative Utility.");
+					"\"Try Select Children with Highest Utility\" or \"Try Select Children At Random Weighted By Utility.");
 			}
 			bHasWarnings = true;
 		}
diff --git a/Engine/Plugins/Runtime/StateTree/Source/StateTreeEditorModule/Private/StateTreeCompiler.cpp b/Engine/Plugins/Runtime/StateTree/Source/StateTreeEditorModule/Private/StateTreeCompiler.cpp
--- a/Engine/Plugins/Runtime/StateTree/Source/StateTreeEditorModule/Private/StateTreeCompiler.cpp
+++ b/Engine/Plugins/Runtime/StateTree/Source/StateTreeEditorModule/Private/StateTreeCompiler.cpp
@@ -1531,11 +1531,11 @@ bool FStateTreeCompiler::CreateConsiderations(UStateTreeState& State, const FStr
 	{
 		if (!State.Parent
 			|| (State.Parent->SelectionBehavior != EStateTreeStateSelectionBehavior::TrySelectChildrenWithHighestUtility
-				&& State.Parent->SelectionBehavior != EStateTreeStateSelectionBehavior::TrySelectChildrenBasedOnRelativeUtility))
+				&& State.Parent->SelectionBehavior != EStateTreeStateSelectionBehavior::TrySelectChildrenAtRandomWeightedByUtility))
 		{
 			Log.Reportf(EMessageSeverity::Warning, TEXT("State's Utility Considerations data are compiled but they don't have effect."
 					"The Utility Considerations are used only when parent State's Selection Behavior is:"
-					"\"Try Select Children with Highest Utility\" or \"Try Select Children Based on Relative Utility\"."));
+					"\"Try Select Children with Highest Utility\" or \"Try Select Children At Random Weighted By Utility\"."));
 		}
 	}
 
diff --git a/Engine/Plugins/Runtime/StateTree/Source/StateTreeEditorModule/Private/StateTreeEditorStyle.cpp b/Engine/Plugins/Runtime/StateTree/Source/StateTreeEditorModule/Private/StateTreeEditorStyle.cpp
--- a/Engine/Plugins/Runtime/StateTree/Source/StateTreeEditorModule/Private/StateTreeEditorStyle.cpp
+++ b/Engine/Plugins/Runtime/StateTree/Source/StateTreeEditorModule/Private/StateTreeEditorStyle.cpp
@@ -420,7 +420,7 @@ const FSlateBrush* FStateTreeEditorStyle::GetBrushForSelectionBehaviorType(EStat
 	}
 	else if (InBehaviour == EStateTreeStateSelectionBehavior::TrySelectChildrenInOrder
 		|| InBehaviour == EStateTreeStateSelectionBehavior::TrySelectChildrenWithHighestUtility
-		|| InBehaviour == EStateTreeStateSelectionBehavior::TrySelectChildrenBasedOnRelativeUtility)
+		|| InBehaviour == EStateTreeStateSelectionBehavior::TrySelectChildrenAtRandomWeightedByUtility)
 	{
 		if (!bHasChildren
 			|| StateType == EStateTreeStateType::Linked
@@ -433,7 +433,7 @@ const FSlateBrush* FStateTreeEditorStyle::GetBrushForSelectionBehaviorType(EStat
 			return Get().GetBrush("StateTreeEditor.TrySelectChildrenInOrder");
 		}
 	}
-	else if (InBehaviour == EStateTreeStateSelectionBehavior::TrySelectChildrenAtUniformRandom)
+	else if (InBehaviour == EStateTreeStateSelectionBehavior::TrySelectChildrenAtRandom)
 	{
 		return Get().GetBrush("StateTreeEditor.TrySelectChildrenAtRandom");
 	}
diff --git a/Engine/Plugins/Runtime/StateTree/Source/StateTreeModule/Private/StateTreeExecutionContext.cpp b/Engine/Plugins/Runtime/StateTree/Source/StateTreeModule/Private/StateTreeExecutionContext.cpp
--- a/Engine/Plugins/Runtime/StateTree/Source/StateTreeModule/Private/StateTreeExecutionContext.cpp
+++ b/Engine/Plugins/Runtime/StateTree/Source/StateTreeModule/Private/StateTreeExecutionContext.cpp
@@ -4201,7 +4201,7 @@ bool FStateTreeExecutionContext::SelectStateInternal(
 				break;
 			}
 		}
-		else if (NextState.SelectionBehavior == EStateTreeStateSelectionBehavior::TrySelectChildrenAtUniformRandom)
+		else if (NextState.SelectionBehavior == EStateTreeStateSelectionBehavior::TrySelectChildrenAtRandom)
 		{
 			if (NextState.HasChildren())
 			{
@@ -4303,7 +4303,7 @@ bool FStateTreeExecutionContext::SelectStateInternal(
 				break;
 			}
 		}
-		else if (NextState.SelectionBehavior == EStateTreeStateSelectionBehavior::TrySelectChildrenBasedOnRelativeUtility)
+		else if (NextState.SelectionBehavior == EStateTreeStateSelectionBehavior::TrySelectChildrenAtRandomWeightedByUtility)
 		{
 			if (NextState.HasChildren())
 			{
diff --git a/Engine/Plugins/Runtime/StateTree/Source/StateTreeModule/Public/StateTreeTypes.h b/Engine/Plugins/Runtime/StateTree/Source/StateTreeModule/Public/StateTreeTypes.h
--- a/Engine/Plugins/Runtime/StateTree/Source/StateTreeModule/Public/StateTreeTypes.h
+++ b/Engine/Plugins/Runtime/StateTree/Source/StateTreeModule/Public/StateTreeTypes.h
@@ -133,17 +133,20 @@ enum class EStateTreeStateSelectionBehavior : uint8
 	TrySelectChildrenInOrder UMETA(DisplayName = "Try Select Children In Order"),
 	
 	/** When state is considered for selection, shuffle the order of child states and try to select the first one. If no child states are present, behaves like SelectState. */
-	TrySelectChildrenAtUniformRandom UMETA(DisplayName = "Try Select Children At Uniform Random"),
-
+	TrySelectChildrenAtRandom UMETA(DisplayName = "Try Select Children At Random"),
+	
 	/** When state is considered for selection, try to select the child state with highest utility score. If there is a tie, it will try to select in order. */
-	TrySelectChildrenWithHighestUtility UMETA(DisplayName = "Try Select Children With Highest Utility Score"),
+	TrySelectChildrenWithHighestUtility UMETA(DisplayName = "Try Select Children With Highest Utility"),
 
 	/** When state is considered for selection, randomly pick one of its child states. The probability of selecting each child state is its normalized utility score */
-	TrySelectChildrenBasedOnRelativeUtility UMETA(DisplayName = "Try Select Children Based On Relative Utility"),
-	
+	TrySelectChildrenAtRandomWeightedByUtility UMETA(DisplayName = "Try Select Children At Random Weighted By Utility"),
+
 	/** When state is considered for selection, try to trigger the transitions instead. */
 	TryFollowTransitions UMETA(DisplayName = "Try Follow Transitions"),
 
+	// Olds names that needs to be kept forever to ensure asset serialization to work correctly when UENUM() switched from serializing int to names.
+	TrySelectChildrenAtUniformRandom UE_DEPRECATED(5.5, "Use TrySelectChildrenAtRandom Instead") = TrySelectChildrenAtRandom UMETA(Hidden),
+	TrySelectChildrenBasedOnRelativeUtility UE_DEPRECATED(5.5, "Use TrySelectChildrenAtRandomWeightedByUtility Instead") = TrySelectChildrenAtRandomWeightedByUtility UMETA(Hidden)
 };
 
 
