[State Tree Editor] Added icons for nodes creation buttons

diff --git a/Engine/Plugins/Runtime/StateTree/Source/StateTreeEditorModule/Private/StateTreeEditorModeToolkit.cpp b/Engine/Plugins/Runtime/StateTree/Source/StateTreeEditorModule/Private/StateTreeEditorModeToolkit.cpp
--- a/Engine/Plugins/Runtime/StateTree/Source/StateTreeEditorModule/Private/StateTreeEditorModeToolkit.cpp
+++ b/Engine/Plugins/Runtime/StateTree/Source/StateTreeEditorModule/Private/StateTreeEditorModeToolkit.cpp
@@ -281,7 +281,7 @@ void FStateTreeEditorModeToolkit::ExtendSecondaryModeToolbar(UToolMenu* ToolBar)
 						 FOnGetContent::CreateSP(this, &FStateTreeEditorModeToolkit::GenerateTaskBPBaseClassesMenu),
 						 LOCTEXT("CreateNewTask_Title", "New Task"),
 						 LOCTEXT("CreateNewTask_ToolbarTooltip", "Create a new Blueprint State Tree Task"),
-						 TAttribute<FSlateIcon>::CreateSP(this, &FStateTreeEditorModeToolkit::GetNewTaskButtonImage)
+						 GetNewTaskButtonImage()
 					 ));
 						
 					InSection.AddEntry(FToolMenuEntry::InitComboButton(
@@ -290,7 +290,7 @@ void FStateTreeEditorModeToolkit::ExtendSecondaryModeToolbar(UToolMenu* ToolBar)
 						 FOnGetContent::CreateSP(this, &FStateTreeEditorModeToolkit::GenerateConditionBPBaseClassesMenu),
 						 LOCTEXT("CreateNewCondition_Title", "New Condition"),
 						 LOCTEXT("CreateNewCondition_ToolbarTooltip", "Create a new Blueprint State Tree Condition"),
-						 TAttribute<FSlateIcon>::CreateSP(this, &FStateTreeEditorModeToolkit::GetNewConditionButtonImage)
+						 GetNewConditionButtonImage()
 					 ));
 						
 					 InSection.AddEntry(FToolMenuEntry::InitComboButton(
@@ -299,7 +299,7 @@ void FStateTreeEditorModeToolkit::ExtendSecondaryModeToolbar(UToolMenu* ToolBar)
 						 FOnGetContent::CreateSP(this, &FStateTreeEditorModeToolkit::GenerateConsiderationBPBaseClassesMenu),
 						 LOCTEXT("CreateNewConsideration_Title", "New Consideration"),
 						 LOCTEXT("CreateNewConsideration_ToolbarTooltip", "Create a new Blueprint State Tree Utility Consideration"),
-						 TAttribute<FSlateIcon>::CreateSP(this, &FStateTreeEditorModeToolkit::GetNewConsiderationButtonImage)
+						 GetNewConsiderationButtonImage()
 					 ));
 				}
 			}
@@ -343,8 +343,7 @@ using FStateTreeConsiderationBPClassFilter = FEditorNodeClassFilter<UStateTreeCo
 
 FSlateIcon FStateTreeEditorModeToolkit::GetNewTaskButtonImage() const
 {
-	//placeholder
-	return FSlateIcon(FAppStyle::GetAppStyleSetName(), NAME_None, NAME_None, NAME_None);
+	return FSlateIcon("StateTreeEditorStyle", "StateTreeEditor.Tasks.Large");
 }
 
 TSharedRef<SWidget> FStateTreeEditorModeToolkit::GenerateTaskBPBaseClassesMenu() const
@@ -360,8 +359,7 @@ TSharedRef<SWidget> FStateTreeEditorModeToolkit::GenerateTaskBPBaseClassesMenu()
 
 FSlateIcon FStateTreeEditorModeToolkit::GetNewConditionButtonImage() const
 {
-	//placeholder
-    return FSlateIcon(FAppStyle::GetAppStyleSetName(), NAME_None, NAME_None, NAME_None);
+	return FSlateIcon("StateTreeEditorStyle", "StateTreeEditor.Conditions.Large");
 }
 
 TSharedRef<SWidget> FStateTreeEditorModeToolkit::GenerateConditionBPBaseClassesMenu() const
@@ -377,8 +375,7 @@ TSharedRef<SWidget> FStateTreeEditorModeToolkit::GenerateConditionBPBaseClassesM
 
 FSlateIcon FStateTreeEditorModeToolkit::GetNewConsiderationButtonImage() const
 {
-	//placeholder
-    return FSlateIcon(FAppStyle::GetAppStyleSetName(), NAME_None, NAME_None, NAME_None);
+    return FSlateIcon("StateTreeEditorStyle", "StateTreeEditor.Utility.Large");
 }
 
 TSharedRef<SWidget> FStateTreeEditorModeToolkit::GenerateConsiderationBPBaseClassesMenu() const
diff --git a/Engine/Plugins/Runtime/StateTree/Source/StateTreeEditorModule/Private/StateTreeEditorStyle.cpp b/Engine/Plugins/Runtime/StateTree/Source/StateTreeEditorModule/Private/StateTreeEditorStyle.cpp
--- a/Engine/Plugins/Runtime/StateTree/Source/StateTreeEditorModule/Private/StateTreeEditorStyle.cpp
+++ b/Engine/Plugins/Runtime/StateTree/Source/StateTreeEditorModule/Private/StateTreeEditorStyle.cpp
@@ -362,10 +362,13 @@ FStateTreeEditorStyle::FStateTreeEditorStyle()
 		Set("StateTreeEditor.StateConditions", new IMAGE_BRUSH_SVG("Icons/State_Conditions", CoreStyleConstants::Icon16x16));
 
 		Set("StateTreeEditor.Conditions", new IMAGE_BRUSH_SVG("Icons/Conditions", CoreStyleConstants::Icon16x16));
+		Set("StateTreeEditor.Conditions.Large", new IMAGE_BRUSH_SVG("Icons/Conditions", CoreStyleConstants::Icon24x24));
 		Set("StateTreeEditor.Evaluators", new IMAGE_BRUSH_SVG("Icons/Evaluators", CoreStyleConstants::Icon16x16));
 		Set("StateTreeEditor.Parameters", new IMAGE_BRUSH_SVG("Icons/Parameters", CoreStyleConstants::Icon16x16));
 		Set("StateTreeEditor.Utility", new IMAGE_BRUSH_SVG("Icons/Utility", CoreStyleConstants::Icon16x16));
+		Set("StateTreeEditor.Utility.Large", new IMAGE_BRUSH_SVG("Icons/Utility", CoreStyleConstants::Icon24x24));
 		Set("StateTreeEditor.Tasks", new IMAGE_BRUSH_SVG("Icons/Tasks", CoreStyleConstants::Icon16x16));
+		Set("StateTreeEditor.Tasks.Large", new IMAGE_BRUSH_SVG("Icons/Tasks", CoreStyleConstants::Icon24x24));
 		Set("StateTreeEditor.Transitions", new IMAGE_BRUSH_SVG("Icons/Transitions", CoreStyleConstants::Icon16x16));
 
 		Set("StateTreeEditor.StateSubtree", new IMAGE_BRUSH_SVG("Icons/State_Subtree", CoreStyleConstants::Icon16x16));
