[State Tree] Utility Selector Rework
- some Refactorings and API changes
- Constant Consideration
- New Float Input Consideration
- Enum Input Consideration, and its editor customization
- Fixed Consideration nodes not picking up Property Function Nodes
- TInterval property binding function

diff --git a/Engine/Plugins/Runtime/StateTree/Source/StateTreeEditorModule/Private/Customizations/StateTreeEnumValueScorePairArrayBuilder.cpp b/Engine/Plugins/Runtime/StateTree/Source/StateTreeEditorModule/Private/Customizations/StateTreeEnumValueScorePairArrayBuilder.cpp
new file mode 100644
--- /dev/null
+++ b/Engine/Plugins/Runtime/StateTree/Source/StateTreeEditorModule/Private/Customizations/StateTreeEnumValueScorePairArrayBuilder.cpp
@@ -0,0 +1,121 @@
+// Copyright Epic Games, Inc. All Rights Reserved.
+
+#include "StateTreeEnumValueScorePairArrayBuilder.h"
+#include "DetailLayoutBuilder.h"
+#include "IDetailChildrenBuilder.h"
+#include "Considerations/StateTreeCommonConsiderations.h"
+#include "Framework/MultiBox/MultiBoxBuilder.h"
+#include "StateTreePropertyHelpers.h"
+#include "Widgets/Input/SComboButton.h"
+
+#define LOCTEXT_NAMESPACE "StateTreeEditor"
+
+FStateTreeEnumValueScorePairArrayBuilder::FStateTreeEnumValueScorePairArrayBuilder(TSharedRef<IPropertyHandle> InBasePropertyHandle, const UEnum* InEnumType, bool InGenerateHeader, bool InDisplayResetToDefault, bool InDisplayElementNum)
+	: FDetailArrayBuilder(InBasePropertyHandle, InGenerateHeader, InDisplayResetToDefault, InDisplayElementNum)
+	, EnumType(InEnumType)
+	, PairArrayProperty(InBasePropertyHandle->AsArray())
+{
+}
+
+void FStateTreeEnumValueScorePairArrayBuilder::GenerateChildContent(IDetailChildrenBuilder& ChildrenBuilder)
+{
+	uint32 NumChildren = 0;
+	PairArrayProperty->GetNumElements(NumChildren);
+
+	for (uint32 ChildIndex = 0; ChildIndex < NumChildren; ++ChildIndex)
+	{
+		TSharedRef<IPropertyHandle> PairPropertyHandle = PairArrayProperty->GetElement(ChildIndex);
+
+		CustomizePairRowWidget(PairPropertyHandle, ChildrenBuilder);
+	}
+}
+
+void FStateTreeEnumValueScorePairArrayBuilder::CustomizePairRowWidget(TSharedRef<IPropertyHandle> PairPropertyHandle, IDetailChildrenBuilder& ChildrenBuilder)
+{
+	TSharedPtr<IPropertyHandle> EnumValuePropertyHandle = PairPropertyHandle->GetChildHandle(GET_MEMBER_NAME_CHECKED(FStateTreeEnumValueScorePair, EnumValue));
+	TSharedPtr<IPropertyHandle> EnumNamePropertyHandle = PairPropertyHandle->GetChildHandle(GET_MEMBER_NAME_CHECKED(FStateTreeEnumValueScorePair, EnumName));
+	TSharedPtr<IPropertyHandle> ScorePropertyHandle = PairPropertyHandle->GetChildHandle(GET_MEMBER_NAME_CHECKED(FStateTreeEnumValueScorePair, Score));
+
+	IDetailPropertyRow& PropertyRow = ChildrenBuilder.AddProperty(PairPropertyHandle);
+
+	PropertyRow.CustomWidget(false)
+		.NameContent()
+		.VAlign(VAlign_Center)
+		.HAlign(HAlign_Left)
+		[
+			SNew(SComboButton)
+				.OnGetMenuContent(this, &FStateTreeEnumValueScorePairArrayBuilder::GetEnumEntryComboContent, EnumValuePropertyHandle, EnumNamePropertyHandle)
+				.ContentPadding(.0f)
+				.ButtonContent()
+				[
+					SNew(STextBlock)
+						.Text(this, &FStateTreeEnumValueScorePairArrayBuilder::GetEnumEntryDescription, PairPropertyHandle)
+						.Font(IDetailLayoutBuilder::GetDetailFontBold())
+				]
+		]
+		.ValueContent()
+		[
+			ScorePropertyHandle->CreatePropertyValueWidget()
+		];
+}
+
+FText FStateTreeEnumValueScorePairArrayBuilder::GetEnumEntryDescription(TSharedRef<IPropertyHandle> PairPropertyHandle) const
+{
+	FStateTreeEnumValueScorePair EnumValueScorePair;
+	FPropertyAccess::Result Result = UE::StateTree::PropertyHelpers::GetStructValue<FStateTreeEnumValueScorePair>(PairPropertyHandle, EnumValueScorePair);
+	if (Result == FPropertyAccess::Success)
+	{
+		if (EnumType)
+		{
+			return EnumType->GetDisplayNameTextByValue(int64(EnumValueScorePair.EnumValue));
+		}
+	}
+	else if (Result == FPropertyAccess::MultipleValues)
+	{
+		return LOCTEXT("MultipleSelected", "Multiple Selected");
+	}
+
+	return LOCTEXT("None", "None");
+}
+
+TSharedRef<SWidget> FStateTreeEnumValueScorePairArrayBuilder::GetEnumEntryComboContent(TSharedPtr<IPropertyHandle> EnumValuePropertyHandle, TSharedPtr<IPropertyHandle> EnumNamePropertyHandle) const
+{
+	FMenuBuilder MenuBuilder(/*bInShouldCloseWindowAfterMenuSelection*/true, /*InCommandList*/nullptr);
+
+	if (EnumType)
+	{
+		const bool bHasMaxValue = EnumType->ContainsExistingMax();
+		const int32 NumEnums = bHasMaxValue ? EnumType->NumEnums() - 1 : EnumType->NumEnums();
+
+		for (int32 Index = 0; Index < NumEnums; Index++)
+		{
+
+#if WITH_METADATA
+			if (EnumType->HasMetaData(TEXT("Hidden"), Index))
+			{
+				continue;
+			}
+#endif //WITH_METADATA
+
+			const int64 Value = EnumType->GetValueByIndex(Index);
+			MenuBuilder.AddMenuEntry(EnumType->GetDisplayNameTextByIndex(Index), TAttribute<FText>(), FSlateIcon(),
+				FUIAction(FExecuteAction::CreateLambda(
+					[=, this]() {
+						if (EnumValuePropertyHandle)
+						{
+							EnumValuePropertyHandle->SetValue(Value);
+							EnumNamePropertyHandle->SetValue(EnumType->GetNameByIndex(Index));
+						}
+					})
+				));
+		}
+	}
+	else
+	{
+		MenuBuilder.AddMenuEntry(LOCTEXT("None", "None"), TAttribute<FText>(), FSlateIcon(), FUIAction());
+	}
+
+	return MenuBuilder.MakeWidget();
+}
+
+#undef LOCTEXT_NAMESPACE
diff --git a/Engine/Plugins/Runtime/StateTree/Source/StateTreeEditorModule/Private/Customizations/StateTreeEnumValueScorePairArrayBuilder.h b/Engine/Plugins/Runtime/StateTree/Source/StateTreeEditorModule/Private/Customizations/StateTreeEnumValueScorePairArrayBuilder.h
new file mode 100644
--- /dev/null
+++ b/Engine/Plugins/Runtime/StateTree/Source/StateTreeEditorModule/Private/Customizations/StateTreeEnumValueScorePairArrayBuilder.h
@@ -0,0 +1,28 @@
+// Copyright Epic Games, Inc. All Rights Reserved.
+
+#pragma once
+
+#include "PropertyCustomizationHelpers.h"
+
+class SWidget;
+class IDetailChildrenBuilder;
+class IPropertyHandle;
+class IPropertyHandleArray;
+
+/** Details customization for arrays composed of <EnumValue, Score> properties with a provided Enum Type.*/
+class FStateTreeEnumValueScorePairArrayBuilder : public FDetailArrayBuilder
+	, public TSharedFromThis<FStateTreeEnumValueScorePairArrayBuilder>
+{
+public:
+	FStateTreeEnumValueScorePairArrayBuilder(TSharedRef<IPropertyHandle> InBasePropertyHandle, const UEnum* InEnumType, bool InGenerateHeader = true, bool InDisplayResetToDefault = true, bool InDisplayElementNum = true);
+
+	virtual void GenerateChildContent(IDetailChildrenBuilder& ChildrenBuilder) override;
+private:
+	void CustomizePairRowWidget(TSharedRef<IPropertyHandle> PairPropertyHandle, IDetailChildrenBuilder& ChildrenBuilder);
+	FText GetEnumEntryDescription(TSharedRef<IPropertyHandle> PairPropertyHandle) const;
+	TSharedRef<SWidget> GetEnumEntryComboContent(TSharedPtr<IPropertyHandle> EnumValuePropertyHandle, TSharedPtr<IPropertyHandle> EnumNamePropertyHandle) const;
+
+private:
+	TStrongObjectPtr<const UEnum> EnumType;
+	TSharedPtr<IPropertyHandleArray> PairArrayProperty;
+};
diff --git a/Engine/Plugins/Runtime/StateTree/Source/StateTreeEditorModule/Private/Customizations/StateTreeEnumValueScorePairsDetails.cpp b/Engine/Plugins/Runtime/StateTree/Source/StateTreeEditorModule/Private/Customizations/StateTreeEnumValueScorePairsDetails.cpp
new file mode 100644
--- /dev/null
+++ b/Engine/Plugins/Runtime/StateTree/Source/StateTreeEditorModule/Private/Customizations/StateTreeEnumValueScorePairsDetails.cpp
@@ -0,0 +1,37 @@
+// Copyright Epic Games, Inc. All Rights Reserved.
+
+#include "StateTreeEnumValueScorePairsDetails.h"
+#include "Considerations/StateTreeCommonConsiderations.h"
+#include "DetailWidgetRow.h"
+#include "DetailLayoutBuilder.h"
+#include "IDetailChildrenBuilder.h"
+#include "StateTreeEnumValueScorePairArrayBuilder.h"
+
+#define LOCTEXT_NAMESPACE "StateTreeEditor"
+
+TSharedRef<IPropertyTypeCustomization> FStateTreeEnumValueScorePairsDetails::MakeInstance()
+{
+	return MakeShareable(new FStateTreeEnumValueScorePairsDetails);
+}
+
+void FStateTreeEnumValueScorePairsDetails::CustomizeHeader(TSharedRef<IPropertyHandle> StructPropertyHandle, FDetailWidgetRow& HeaderRow, IPropertyTypeCustomizationUtils& StructCustomizationUtils)
+{
+}
+
+void FStateTreeEnumValueScorePairsDetails::CustomizeChildren(TSharedRef<IPropertyHandle> StructPropertyHandle, IDetailChildrenBuilder& StructBuilder, IPropertyTypeCustomizationUtils& StructCustomizationUtils)
+{
+	EnumProperty = StructPropertyHandle->GetChildHandle(GET_MEMBER_NAME_CHECKED(FStateTreeEnumValueScorePairs, Enum));
+	PairsProperty = StructPropertyHandle->GetChildHandle(GET_MEMBER_NAME_CHECKED(FStateTreeEnumValueScorePairs, Data));
+
+	UObject* Object = nullptr; 
+	if (EnumProperty->GetValue(Object) == FPropertyAccess::Success)
+	{
+		const UEnum* EnumType = static_cast<UEnum*>(Object);
+		TSharedRef<FStateTreeEnumValueScorePairArrayBuilder> Builder = MakeShareable(
+			new FStateTreeEnumValueScorePairArrayBuilder(
+				PairsProperty.ToSharedRef(), EnumType, /*InGenerateHeader*/ true, /*InDisplayResetToDefault*/ false, /*InDisplayElementNum*/ true));
+			
+		StructBuilder.AddCustomBuilder(Builder);
+	}
+}
+#undef LOCTEXT_NAMESPACE
diff --git a/Engine/Plugins/Runtime/StateTree/Source/StateTreeEditorModule/Private/Customizations/StateTreeEnumValueScorePairsDetails.h b/Engine/Plugins/Runtime/StateTree/Source/StateTreeEditorModule/Private/Customizations/StateTreeEnumValueScorePairsDetails.h
new file mode 100644
--- /dev/null
+++ b/Engine/Plugins/Runtime/StateTree/Source/StateTreeEditorModule/Private/Customizations/StateTreeEnumValueScorePairsDetails.h
@@ -0,0 +1,32 @@
+// Copyright Epic Games, Inc. All Rights Reserved.
+
+#pragma once
+
+#include "IPropertyTypeCustomization.h"
+
+class SWidget;
+
+class IPropertyHandle;
+
+/**
+ * Type customization for FStateTreeEnumValueScorePairs.
+ */
+
+class FStateTreeEnumValueScorePairsDetails : public IPropertyTypeCustomization
+{
+public:
+	/** Makes a new instance of this detail layout class for a specific detail view requesting it */
+	static TSharedRef<IPropertyTypeCustomization> MakeInstance();
+
+private:
+
+	/** IPropertyTypeCustomization interface */
+	virtual void CustomizeHeader(TSharedRef<IPropertyHandle> StructPropertyHandle, class FDetailWidgetRow& HeaderRow, IPropertyTypeCustomizationUtils& StructCustomizationUtils) override;
+	virtual void CustomizeChildren(TSharedRef<IPropertyHandle> StructPropertyHandle, class IDetailChildrenBuilder& StructBuilder, IPropertyTypeCustomizationUtils& StructCustomizationUtils) override;
+
+	FText GetDescription() const;
+	TSharedRef<SWidget> OnGetComboContent() const;
+
+	TSharedPtr<IPropertyHandle> EnumProperty;
+	TSharedPtr<IPropertyHandle> PairsProperty;
+};
diff --git a/Engine/Plugins/Runtime/StateTree/Source/StateTreeEditorModule/Private/StateTreeEditorData.cpp b/Engine/Plugins/Runtime/StateTree/Source/StateTreeEditorModule/Private/StateTreeEditorData.cpp
--- a/Engine/Plugins/Runtime/StateTree/Source/StateTreeEditorModule/Private/StateTreeEditorData.cpp
+++ b/Engine/Plugins/Runtime/StateTree/Source/StateTreeEditorModule/Private/StateTreeEditorData.cpp
@@ -1217,7 +1217,12 @@ EStateTreeVisitor UStateTreeEditorData::VisitStateNodes(const UStateTreeState& S
 		// Utility Considerations
 		for (const FStateTreeEditorNode& Node : State.Considerations)
 		{
-			if (const FStateTreeConsiderationBase* Consideration = Node.Node.GetPtr<FStateTreeConsiderationBase>())
+			if (VisitFuncNodes(Node.ID, Node.GetName()) == EStateTreeVisitor::Break)
+			{
+				bContinue = false;
+				break;
+			}
+			else if (const FStateTreeConsiderationBase* Consideration = Node.Node.GetPtr<FStateTreeConsiderationBase>())
 			{
 				FStateTreeBindableStructDesc Desc;
 				Desc.StatePath = StatePathWithConsiderations;
diff --git a/Engine/Plugins/Runtime/StateTree/Source/StateTreeEditorModule/Private/StateTreeEditorModule.cpp b/Engine/Plugins/Runtime/StateTree/Source/StateTreeEditorModule/Private/StateTreeEditorModule.cpp
--- a/Engine/Plugins/Runtime/StateTree/Source/StateTreeEditorModule/Private/StateTreeEditorModule.cpp
+++ b/Engine/Plugins/Runtime/StateTree/Source/StateTreeEditorModule/Private/StateTreeEditorModule.cpp
@@ -7,6 +7,7 @@
 #include "Blueprint/StateTreeTaskBlueprintBase.h"
 #include "StateTreePropertyFunctionBase.h"
 #include "Customizations/StateTreeAnyEnumDetails.h"
+#include "Customizations/StateTreeEnumValueScorePairsDetails.h"
 #include "Customizations/StateTreeEditorColorDetails.h"
 #include "Customizations/StateTreeEditorDataDetails.h"
 #include "Customizations/StateTreeEditorNodeDetails.h"
@@ -77,6 +78,7 @@ void FStateTreeEditorModule::StartupModule()
 	PropertyModule.RegisterCustomPropertyTypeLayout("StateTreeEditorColorRef", FOnGetPropertyTypeCustomizationInstance::CreateStatic(&FStateTreeEditorColorRefDetails::MakeInstance));
 	PropertyModule.RegisterCustomPropertyTypeLayout("StateTreeEditorColor", FOnGetPropertyTypeCustomizationInstance::CreateStatic(&FStateTreeEditorColorDetails::MakeInstance));
 	PropertyModule.RegisterCustomPropertyTypeLayout("StateTreeBlueprintPropertyRef", FOnGetPropertyTypeCustomizationInstance::CreateStatic(&FStateTreeBlueprintPropertyRefDetails::MakeInstance));
+	PropertyModule.RegisterCustomPropertyTypeLayout("StateTreeEnumValueScorePairs", FOnGetPropertyTypeCustomizationInstance::CreateStatic(&FStateTreeEnumValueScorePairsDetails::MakeInstance));
 	PropertyModule.RegisterCustomClassLayout("StateTreeState", FOnGetDetailCustomizationInstance::CreateStatic(&FStateTreeStateDetails::MakeInstance));
 	PropertyModule.RegisterCustomClassLayout("StateTreeEditorData", FOnGetDetailCustomizationInstance::CreateStatic(&FStateTreeEditorDataDetails::MakeInstance));
 
diff --git a/Engine/Plugins/Runtime/StateTree/Source/StateTreeModule/Private/Blueprint/StateTreeConsiderationBlueprintBase.cpp b/Engine/Plugins/Runtime/StateTree/Source/StateTreeModule/Private/Blueprint/StateTreeConsiderationBlueprintBase.cpp
--- a/Engine/Plugins/Runtime/StateTree/Source/StateTreeModule/Private/Blueprint/StateTreeConsiderationBlueprintBase.cpp
+++ b/Engine/Plugins/Runtime/StateTree/Source/StateTreeModule/Private/Blueprint/StateTreeConsiderationBlueprintBase.cpp
@@ -13,21 +13,21 @@
 UStateTreeConsiderationBlueprintBase::UStateTreeConsiderationBlueprintBase(const FObjectInitializer& ObjectInitializer)
 	: Super(ObjectInitializer)
 {
-	bHasComputeRawScore = BlueprintNodeHelpers::HasBlueprintFunction(TEXT("ReceiveComputeRawScore"), *this, *StaticClass());
+	bHasGetScore = BlueprintNodeHelpers::HasBlueprintFunction(GET_FUNCTION_NAME_CHECKED(UStateTreeConsiderationBlueprintBase, ReceiveGetScore), *this, *StaticClass());
 }
 
-float UStateTreeConsiderationBlueprintBase::ComputeRawScore(FStateTreeExecutionContext& Context) const
+float UStateTreeConsiderationBlueprintBase::GetScore(FStateTreeExecutionContext& Context) const
 {
-	if (bHasComputeRawScore)
+	if (bHasGetScore)
 	{
 		// Cache the owner and event queue for the duration the consideration is evaluated.
 		SetCachedInstanceDataFromContext(Context);
 
-		const float RawScore = ReceiveComputeRawScore();
+		const float Score = ReceiveGetScore();
 
 		ClearCachedInstanceData();
 
-		return RawScore;
+		return Score;
 	}
 
 	return .0f;
@@ -79,9 +79,9 @@ FColor FStateTreeBlueprintConsiderationWrapper::GetIconColor() const
 }
 #endif //WITH_EDITOR
 
-float FStateTreeBlueprintConsiderationWrapper::ComputeRawScore(FStateTreeExecutionContext& Context) const
+float FStateTreeBlueprintConsiderationWrapper::GetScore(FStateTreeExecutionContext& Context) const
 {
 	UStateTreeConsiderationBlueprintBase* Consideration = Context.GetInstanceDataPtr<UStateTreeConsiderationBlueprintBase>(*this);
 	check(Consideration);
-	return Consideration->ComputeRawScore(Context);
+	return Consideration->GetScore(Context);
 }
diff --git a/Engine/Plugins/Runtime/StateTree/Source/StateTreeModule/Private/Considerations/StateTreeCommonConsiderations.cpp b/Engine/Plugins/Runtime/StateTree/Source/StateTreeModule/Private/Considerations/StateTreeCommonConsiderations.cpp
--- a/Engine/Plugins/Runtime/StateTree/Source/StateTreeModule/Private/Considerations/StateTreeCommonConsiderations.cpp
+++ b/Engine/Plugins/Runtime/StateTree/Source/StateTreeModule/Private/Considerations/StateTreeCommonConsiderations.cpp
@@ -1,6 +1,7 @@
 // Copyright Epic Games, Inc. All Rights Reserved.
 
 #include "Considerations/StateTreeCommonConsiderations.h"
+#include "Algo/Sort.h"
 #include "StateTreeExecutionContext.h"
 #include "StateTreeNodeDescriptionHelpers.h"
 
@@ -9,38 +10,183 @@
 #define LOCTEXT_NAMESPACE "StateTree"
 
 #if WITH_EDITOR
-FText FStateTreeFloatConsideration::GetDescription(const FGuid& ID, FStateTreeDataView InstanceDataView, const IStateTreeBindingLookup& BindingLookup, EStateTreeNodeFormatting Formatting /*= EStateTreeNodeFormatting::Text*/) const
+FText FStateTreeFloatInputConsideration::GetDescription(const FGuid& ID, FStateTreeDataView InstanceDataView, const IStateTreeBindingLookup& BindingLookup, EStateTreeNodeFormatting Formatting /*= EStateTreeNodeFormatting::Text*/) const
 {
 	const FInstanceDataType* InstanceData = InstanceDataView.GetPtr<FInstanceDataType>();
 	check(InstanceData);
 
-	FText OptionalConstantKeywordText = FText::GetEmpty();
-	FText RawScoreOrPropertyPathText = BindingLookup.GetBindingSourceDisplayName(FStateTreePropertyPath(ID, GET_MEMBER_NAME_CHECKED(FInstanceDataType, RawScore)), Formatting);
-	if (RawScoreOrPropertyPathText.IsEmpty())
+	FText InputText = BindingLookup.GetBindingSourceDisplayName(FStateTreePropertyPath(ID, GET_MEMBER_NAME_CHECKED(FInstanceDataType, Input)), Formatting);
+	if (InputText.IsEmpty())
 	{
-		//using a constant value for the float param
-		OptionalConstantKeywordText = (Formatting == EStateTreeNodeFormatting::RichText)
-			? LOCTEXT("ConstantConsiderationRich", "<s>Constant</>")
-			: LOCTEXT("ConstantConsideration", "Constant");
+		FNumberFormattingOptions Options;
+		Options.MinimumFractionalDigits = 1;
+		Options.MaximumFractionalDigits = 2;
 
-		RawScoreOrPropertyPathText = FText::AsNumber(InstanceData->RawScore);
+		InputText = FText::AsNumber(InstanceData->Input, &Options);
+	}
+	
+	const FFloatInterval& Interval = InstanceData->Interval;
+	FText IntervalText = BindingLookup.GetBindingSourceDisplayName(FStateTreePropertyPath(ID, GET_MEMBER_NAME_CHECKED(FInstanceDataType, Interval)), Formatting);
+	if (IntervalText.IsEmpty())
+	{
+		IntervalText = UE::StateTree::DescHelpers::GetIntervalText(Interval, Formatting);
 	}
-
-	FText WithinValueRangeText = UE::StateTree::DescHelpers::GetWithinValueRangeText(ResponseCurve.RawScoreLowerBound, ResponseCurve.RawScoreUpperBound, Formatting);
 	
-	return FText::Format(LOCTEXT("FloatParam", "{OptionalConstantKeyword} {RawScoreOrPropertyPath} {WithinValueRange}"),
-		OptionalConstantKeywordText,
-		RawScoreOrPropertyPathText,
-		WithinValueRangeText
+	if (Formatting == EStateTreeNodeFormatting::RichText)
+	{
+		return FText::FormatNamed(LOCTEXT("InputInIntervalRich", "{Input} <s>in</> {Interval}"),
+			TEXT("Input"), InputText,
+			TEXT("Interval"), IntervalText
+		);
+	}
+	else //EStateTreeNodeFormatting::Text
+	{
+		return FText::FormatNamed(LOCTEXT("InputInInterval", "{Input} in {Interval}"),
+			TEXT("Input"), InputText,
+			TEXT("Interval"), IntervalText
 		);
+	}
 }
 #endif //WITH_EDITOR
 
-float FStateTreeFloatConsideration::ComputeRawScore(FStateTreeExecutionContext& Context) const
+float FStateTreeFloatInputConsideration::GetScore(FStateTreeExecutionContext& Context) const
 {
 	const FInstanceDataType& InstanceData = Context.GetInstanceData(*this);
+	const float Input = InstanceData.Input;
+	const FFloatInterval& Interval = InstanceData.Interval;
+	const float NormalizedInput = FMath::Clamp(Interval.GetRangePct(Input), 0.f, 1.f);
+
+	return ResponseCurve.Evaluate(NormalizedInput);
+}
+
+#if WITH_EDITOR
+FText FStateTreeConstantConsideration::GetDescription(const FGuid& ID, FStateTreeDataView InstanceDataView, const IStateTreeBindingLookup& BindingLookup, EStateTreeNodeFormatting Formatting /*= EStateTreeNodeFormatting::Text*/) const
+{
+	const FInstanceDataType* InstanceData = InstanceDataView.GetPtr<FInstanceDataType>();
+	check(InstanceData);
+
+	FText ConstantText = BindingLookup.GetBindingSourceDisplayName(FStateTreePropertyPath(ID, GET_MEMBER_NAME_CHECKED(FInstanceDataType, Constant)), Formatting);
+	if (ConstantText.IsEmpty())
+	{
+		FNumberFormattingOptions Options;
+		Options.MinimumFractionalDigits = 1;
+		Options.MaximumFractionalDigits = 2;
+
+		ConstantText = FText::AsNumber(InstanceData->Constant, &Options);
+	}
+
+	return FText::FormatNamed(LOCTEXT("Constant", "{ConstantValue}"), 
+		TEXT("ConstantValue"), ConstantText);
+}
+#endif //WITH_EDITOR
+
+float FStateTreeConstantConsideration::GetScore(FStateTreeExecutionContext& Context) const
+{
+	const FInstanceDataType& InstanceData = Context.GetInstanceData(*this);
+
+	return InstanceData.Constant;
+}
+
+float FStateTreeEnumInputConsideration::GetScore(FStateTreeExecutionContext& Context) const
+{
+	const FStateTreeAnyEnum& Input = Context.GetInstanceData(*this).Input;
+	const uint32 EnumValue = Input.Value;
+	const TArray<FStateTreeEnumValueScorePair>& Data = EnumValueScorePairs.Data;
+
+	for (int32 Index = 0; Index < Data.Num(); ++Index)
+	{
+		const FStateTreeEnumValueScorePair& Pair = Data[Index];
+		if (EnumValue == Pair.EnumValue)
+		{
+			return Pair.Score;
+		}
+	}
+
+	return 0.f;
+}
+
+#if WITH_EDITORONLY_DATA
+EDataValidationResult FStateTreeEnumInputConsideration::Compile(FStateTreeDataView InstanceDataView, TArray<FText>& ValidationMessages)
+{
+	const FInstanceDataType& InstanceData = InstanceDataView.Get<FInstanceDataType>();
+	check(InstanceData.Input.Enum == EnumValueScorePairs.Enum);
+	const int32 NumPairs = EnumValueScorePairs.Data.Num();
+
+	//Validate uniqueness of keys
+	const auto SortedPairs = EnumValueScorePairs.Data;
+	Algo::Sort(SortedPairs, 
+		[](const FStateTreeEnumValueScorePair& A, const FStateTreeEnumValueScorePair& B) {
+		return A.EnumValue < B.EnumValue;
+		});
+	for (int32 Idx = 1; Idx < NumPairs; ++Idx)
+	{
+		if (SortedPairs[Idx].EnumValue == SortedPairs[Idx - 1].EnumValue)
+		{
+			ValidationMessages.Add(LOCTEXT("DuplicateEnumValues", "Duplicate Enum Values found."));
 
-	return InstanceData.RawScore;
+			return EDataValidationResult::Invalid;
+		}
+	}
+
+	return EDataValidationResult::Valid;
 }
 
+void FStateTreeEnumInputConsideration::OnBindingChanged(const FGuid& ID, FStateTreeDataView InstanceData, const FStateTreePropertyPath& SourcePath, const FStateTreePropertyPath& TargetPath, const IStateTreeBindingLookup& BindingLookup)
+{
+	if (!TargetPath.GetStructID().IsValid())
+	{
+		return;
+	}
+
+	FInstanceDataType& Instance = InstanceData.GetMutable<FInstanceDataType>();
+
+	// Left has changed, update enums from the leaf property.
+	if (!TargetPath.IsPathEmpty()
+		&& TargetPath.GetSegments().Last().GetName() == GET_MEMBER_NAME_CHECKED(FInstanceDataType, Input))
+	{
+		if (const FProperty* LeafProperty = BindingLookup.GetPropertyPathLeafProperty(SourcePath))
+		{
+			// Handle both old type namespace enums and new class enum properties.
+			UEnum* NewEnum = nullptr;
+			if (const FByteProperty* ByteProperty = CastField<FByteProperty>(LeafProperty))
+			{
+				NewEnum = ByteProperty->GetIntPropertyEnum();
+			}
+			else if (const FEnumProperty* EnumProperty = CastField<FEnumProperty>(LeafProperty))
+			{
+				NewEnum = EnumProperty->GetEnum();
+			}
+
+			if (Instance.Input.Enum != NewEnum)
+			{
+				Instance.Input.Initialize(NewEnum);
+			}
+		}
+		else
+		{
+			Instance.Input.Initialize(nullptr);
+		}
+
+		if (EnumValueScorePairs.Enum != Instance.Input.Enum)
+		{
+			EnumValueScorePairs.Initialize(Instance.Input.Enum);
+		}
+	}
+}
+#endif //WITH_EDITORONLY_DATA
+
+#if WITH_EDITOR
+FText FStateTreeEnumInputConsideration::GetDescription(const FGuid& ID, FStateTreeDataView InstanceDataView, const IStateTreeBindingLookup& BindingLookup, EStateTreeNodeFormatting Formatting /*= EStateTreeNodeFormatting::Text*/) const
+{
+	const FInstanceDataType* InstanceData = InstanceDataView.GetPtr<FInstanceDataType>();
+	check(InstanceData);
+
+	FText InputText = BindingLookup.GetBindingSourceDisplayName(FStateTreePropertyPath(ID, GET_MEMBER_NAME_CHECKED(FInstanceDataType, Input)), Formatting);
+	
+	return FText::FormatNamed(LOCTEXT("EnumInput", "{Input}"),
+		TEXT("Input"), InputText);
+}
+
+#endif //WITH_EDITOR
+
 #undef LOCTEXT_NAMESPACE
diff --git a/Engine/Plugins/Runtime/StateTree/Source/StateTreeModule/Private/PropertyFunctions/StateTreeIntervalPropertyFunctions.cpp b/Engine/Plugins/Runtime/StateTree/Source/StateTreeModule/Private/PropertyFunctions/StateTreeIntervalPropertyFunctions.cpp
new file mode 100644
--- /dev/null
+++ b/Engine/Plugins/Runtime/StateTree/Source/StateTreeModule/Private/PropertyFunctions/StateTreeIntervalPropertyFunctions.cpp
@@ -0,0 +1,37 @@
+// Copyright Epic Games, Inc. All Rights Reserved.
+
+#include "PropertyFunctions/StateTreeIntervalPropertyFunctions.h"
+
+#include "StateTreeExecutionContext.h"
+#include "StateTreeNodeDescriptionHelpers.h"
+
+#define LOCTEXT_NAMESPACE "StateTree"
+
+void FStateTreeMakeIntervalPropertyFunction::Execute(FStateTreeExecutionContext& Context) const
+{
+	FInstanceDataType& InstanceData = Context.GetInstanceData(*this);
+	InstanceData.Result = FFloatInterval(InstanceData.Min, InstanceData.Max);
+}
+
+#if WITH_EDITOR
+FText FStateTreeMakeIntervalPropertyFunction::GetDescription(const FGuid& ID, FStateTreeDataView InstanceDataView, const IStateTreeBindingLookup& BindingLookup, EStateTreeNodeFormatting Formatting) const
+{
+	const FInstanceDataType& InstanceData = InstanceDataView.Get<FInstanceDataType>();
+	
+	FText MinValueText = BindingLookup.GetBindingSourceDisplayName(FStateTreePropertyPath(ID, GET_MEMBER_NAME_CHECKED(FInstanceDataType, Min)), Formatting);
+	if (MinValueText.IsEmpty())
+	{
+		MinValueText = UE::StateTree::DescHelpers::GetText(InstanceData.Min, Formatting);
+	}
+
+	FText MaxValueText = BindingLookup.GetBindingSourceDisplayName(FStateTreePropertyPath(ID, GET_MEMBER_NAME_CHECKED(FInstanceDataType, Max)), Formatting);
+	if (MaxValueText.IsEmpty())
+	{
+		MaxValueText = UE::StateTree::DescHelpers::GetText(InstanceData.Max, Formatting);
+	}
+
+	return UE::StateTree::DescHelpers::GetIntervalText(MinValueText, MaxValueText, Formatting);
+}
+#endif //WITH_EDITOR
+
+#undef LOCTEXT_NAMESPACE
diff --git a/Engine/Plugins/Runtime/StateTree/Source/StateTreeModule/Private/PropertyFunctions/StateTreeIntervalPropertyFunctions.h b/Engine/Plugins/Runtime/StateTree/Source/StateTreeModule/Private/PropertyFunctions/StateTreeIntervalPropertyFunctions.h
new file mode 100644
--- /dev/null
+++ b/Engine/Plugins/Runtime/StateTree/Source/StateTreeModule/Private/PropertyFunctions/StateTreeIntervalPropertyFunctions.h
@@ -0,0 +1,42 @@
+// Copyright Epic Games, Inc. All Rights Reserved.
+
+#pragma once
+
+#include "StateTreePropertyFunctionBase.h"
+#include "StateTreeIntervalPropertyFunctions.generated.h"
+
+struct FStateTreeExecutionContext;
+
+USTRUCT()
+struct STATETREEMODULE_API FStateTreeMakeIntervalPropertyFunctionInstanceData
+{
+	GENERATED_BODY()
+
+	UPROPERTY(EditAnywhere, Category = Parameter)
+	float Min = 0.f;
+
+	UPROPERTY(EditAnywhere, Category = Parameter)
+	float Max = 1.f;
+
+	UPROPERTY(EditAnywhere, Category = Output)
+	FFloatInterval Result = FFloatInterval(0.f, 1.f);
+};
+
+/**
+ * Make an Interval from two floats.
+ */
+USTRUCT(DisplayName = "Make Interval")
+struct STATETREEMODULE_API FStateTreeMakeIntervalPropertyFunction : public FStateTreePropertyFunctionCommonBase
+{
+	GENERATED_BODY()
+
+	using FInstanceDataType = FStateTreeMakeIntervalPropertyFunctionInstanceData;
+
+	virtual const UStruct* GetInstanceDataType() const override { return FInstanceDataType::StaticStruct(); }
+
+	virtual void Execute(FStateTreeExecutionContext& Context) const override;
+
+#if WITH_EDITOR
+	virtual FText GetDescription(const FGuid& ID, FStateTreeDataView InstanceDataView, const IStateTreeBindingLookup& BindingLookup, EStateTreeNodeFormatting Formatting) const;
+#endif
+};
diff --git a/Engine/Plugins/Runtime/StateTree/Source/StateTreeModule/Private/StateTreeConsiderationBase.cpp b/Engine/Plugins/Runtime/StateTree/Source/StateTreeModule/Private/StateTreeConsiderationBase.cpp
--- a/Engine/Plugins/Runtime/StateTree/Source/StateTreeModule/Private/StateTreeConsiderationBase.cpp
+++ b/Engine/Plugins/Runtime/StateTree/Source/StateTreeModule/Private/StateTreeConsiderationBase.cpp
@@ -1,33 +1,17 @@
 // Copyright Epic Games, Inc. All Rights Reserved.
 
 #include "StateTreeConsiderationBase.h"
-#include "AlphaBlend.h"
 #include "StateTreeTypes.h"
 
 #include UE_INLINE_GENERATED_CPP_BY_NAME(StateTreeConsiderationBase)
 
-FStateTreeConsiderationResponseCurve::FStateTreeConsiderationResponseCurve()
-	: BlendOption(EAlphaBlendOption::Linear)
-	, RawScoreLowerBound(.0f)
-	, RawScoreUpperBound(.1f)
-{
-}
-
 FStateTreeConsiderationBase::FStateTreeConsiderationBase()
 	: Operand(EStateTreeExpressionOperand::And)
 	, DeltaIndent(0)
 {
 }
 
-float FStateTreeConsiderationBase::ComputeNormalizedScore(FStateTreeExecutionContext& Context) const
+float FStateTreeConsiderationBase::GetNormalizedScore(FStateTreeExecutionContext& Context) const
 {
-	const float RawScore = ComputeRawScore(Context);
-	const float NormalizedScore = FMath::Clamp<float>(
-		FMath::GetRangePct(ResponseCurve.RawScoreLowerBound, 
-			ResponseCurve.RawScoreUpperBound, 
-			RawScore), 
-		0.f, 1.f);
-	
-	constexpr class UCurveFloat* OptionalCurve = nullptr;
-	return FAlphaBlend::AlphaToBlendOption(NormalizedScore, ResponseCurve.BlendOption, OptionalCurve);
+	return FMath::Clamp(GetScore(Context), 0.f, 1.f);
 }
diff --git a/Engine/Plugins/Runtime/StateTree/Source/StateTreeModule/Private/StateTreeExecutionContext.cpp b/Engine/Plugins/Runtime/StateTree/Source/StateTreeModule/Private/StateTreeExecutionContext.cpp
--- a/Engine/Plugins/Runtime/StateTree/Source/StateTreeModule/Private/StateTreeExecutionContext.cpp
+++ b/Engine/Plugins/Runtime/StateTree/Source/StateTreeModule/Private/StateTreeExecutionContext.cpp
@@ -2832,7 +2832,7 @@ float FStateTreeExecutionContext::EvaluateUtility(const FStateTreeExecutionFrame
 			}
 		}
 
-		Value = Consideration.ComputeNormalizedScore(*this);
+		Value = Consideration.GetNormalizedScore(*this);
 
 		// Reset copied properties that might contain object references.
 		if (Consideration.BindingsBatch.IsValid())
diff --git a/Engine/Plugins/Runtime/StateTree/Source/StateTreeModule/Private/StateTreeNodeDescriptionHelpers.cpp b/Engine/Plugins/Runtime/StateTree/Source/StateTreeModule/Private/StateTreeNodeDescriptionHelpers.cpp
--- a/Engine/Plugins/Runtime/StateTree/Source/StateTreeModule/Private/StateTreeNodeDescriptionHelpers.cpp
+++ b/Engine/Plugins/Runtime/StateTree/Source/StateTreeModule/Private/StateTreeNodeDescriptionHelpers.cpp
@@ -78,7 +78,12 @@ FText GetBoolText(bool bValue, EStateTreeNodeFormatting Formatting)
 	return bValue ? LOCTEXT("True", "True") : LOCTEXT("False", "False");
 }
 
-FText GetWithinValueRangeText(float Min, float Max, EStateTreeNodeFormatting Formatting)
+FText GetIntervalText(const FFloatInterval& Interval, EStateTreeNodeFormatting Formatting)
+{
+	return GetIntervalText(Interval.Min, Interval.Max, Formatting);
+}
+
+FText GetIntervalText(float Min, float Max, EStateTreeNodeFormatting Formatting)
 {
 	FNumberFormattingOptions Options;
 	Options.MinimumFractionalDigits = 1;
@@ -87,21 +92,26 @@ FText GetWithinValueRangeText(float Min, float Max, EStateTreeNodeFormatting For
 	FText MinValueText = FText::AsNumber(Min, &Options);
 	FText MaxValueText = FText::AsNumber(Max, &Options);
 
-	FText WithinValueRangeText;
+	return GetIntervalText(MinValueText, MaxValueText, Formatting);
+}
+
+FText GetIntervalText(const FText& MinValueText, const FText& MaxValueText, EStateTreeNodeFormatting Formatting)
+{
+	FText IntervalText;
 	if (Formatting == EStateTreeNodeFormatting::RichText)
 	{
-		WithinValueRangeText = FText::Format(LOCTEXT("WithinValueRangeRich", "<s>in</> [{Min}<s>,</> {Max}]"),
-			MinValueText,
-			MaxValueText);
+		IntervalText = FText::FormatNamed(LOCTEXT("IntervalRich", "[{Min}<s>,</> {Max}]"),
+			TEXT("Min"), MinValueText,
+			TEXT("Max"), MaxValueText);
 	}
 	else //EStateTreeNodeFormatting::Text
 	{
-		WithinValueRangeText = FText::Format(LOCTEXT("WithinValueRange", "in [{Min}, {Max}]"),
-			MinValueText,
-			MaxValueText);
+		IntervalText = FText::FormatNamed(LOCTEXT("Interval", "[{Min}, {Max}]"),
+			TEXT("Min"), MinValueText,
+			TEXT("Max"), MaxValueText);
 	}
 
-	return WithinValueRangeText;
+	return IntervalText;
 }
 
 FText GetGameplayTagContainerAsText(const FGameplayTagContainer& TagContainer, const int ApproxMaxLength)
@@ -207,7 +217,6 @@ FText GetSingleParamFunctionText(const FText& FunctionText, const FText& ParamTe
 		TEXT("Function"), FunctionText,
 		TEXT("Input"), ParamText);
 }
-
 } // UE::StateTree::Helpers
 
 #undef LOCTEXT_NAMESPACE
diff --git a/Engine/Plugins/Runtime/StateTree/Source/StateTreeModule/Public/Blueprint/StateTreeConsiderationBlueprintBase.h b/Engine/Plugins/Runtime/StateTree/Source/StateTreeModule/Public/Blueprint/StateTreeConsiderationBlueprintBase.h
--- a/Engine/Plugins/Runtime/StateTree/Source/StateTreeModule/Public/Blueprint/StateTreeConsiderationBlueprintBase.h
+++ b/Engine/Plugins/Runtime/StateTree/Source/StateTreeModule/Public/Blueprint/StateTreeConsiderationBlueprintBase.h
@@ -20,15 +20,15 @@ class STATETREEMODULE_API UStateTreeConsiderationBlueprintBase : public UStateTr
 public:
 	UStateTreeConsiderationBlueprintBase(const FObjectInitializer& ObjectInitializer);
 
-	UFUNCTION(BlueprintImplementableEvent, meta = (DisplayName = "ComputeRawScore"))
-	float ReceiveComputeRawScore() const;
+	UFUNCTION(BlueprintImplementableEvent, meta = (DisplayName = "GetScore"))
+	float ReceiveGetScore() const;
 
 protected:
-	virtual float ComputeRawScore(FStateTreeExecutionContext& Context) const;
+	virtual float GetScore(FStateTreeExecutionContext& Context) const;
 
 	friend struct FStateTreeBlueprintConsiderationWrapper;
 
-	uint8 bHasComputeRawScore : 1;
+	uint8 bHasGetScore : 1;
 };
 
 /**
@@ -50,7 +50,7 @@ struct STATETREEMODULE_API FStateTreeBlueprintConsiderationWrapper : public FSta
 
 protected:
 	//~ Begin FStateTreeConsiderationBase Interface
-	virtual float ComputeRawScore(FStateTreeExecutionContext& Context) const override;
+	virtual float GetScore(FStateTreeExecutionContext& Context) const override;
 	//~ End FStateTreeConsiderationBase Interface
 
 public:
diff --git a/Engine/Plugins/Runtime/StateTree/Source/StateTreeModule/Public/Considerations/StateTreeCommonConsiderations.h b/Engine/Plugins/Runtime/StateTree/Source/StateTreeModule/Public/Considerations/StateTreeCommonConsiderations.h
--- a/Engine/Plugins/Runtime/StateTree/Source/StateTreeModule/Public/Considerations/StateTreeCommonConsiderations.h
+++ b/Engine/Plugins/Runtime/StateTree/Source/StateTreeModule/Public/Considerations/StateTreeCommonConsiderations.h
@@ -2,35 +2,197 @@
 
 #pragma once
 
-#include "StateTreeConsiderationBase.h"
+#include "Curves/CurveFloat.h"
 #include "StateTreeAnyEnum.h"
+#include "StateTreeConsiderationBase.h"
 #include "StateTreeCommonConsiderations.generated.h"
 
 USTRUCT()
-struct STATETREEMODULE_API FStateTreeFloatConsiderationInstanceData
+struct STATETREEMODULE_API FStateTreeConstantConsiderationInstanceData
 {
 	GENERATED_BODY()
-
-	UPROPERTY(EditAnywhere, Category = "Parameter")
-	float RawScore = .0f;
+	
+	UPROPERTY(EditAnywhere, meta = (ClampMin = 0, ClampMax = 1, UIMin = 0, UIMax = 1), Category = "Default")
+	float Constant = 0.f;
 };
 
 /**
- * Consideration using a Float Parameter as raw score.
- * You can use the parameter as a float constant or bind it to a property that changes value over time.
+ * Consideration using a constant as its score.
  */
-USTRUCT(DisplayName = "Float")
-struct STATETREEMODULE_API FStateTreeFloatConsideration : public FStateTreeConsiderationCommonBase
+USTRUCT(DisplayName = "Constant")
+struct STATETREEMODULE_API FStateTreeConstantConsideration: public FStateTreeConsiderationCommonBase
 {
 	GENERATED_BODY()
 
-	using FInstanceDataType = FStateTreeFloatConsiderationInstanceData;
+	using FInstanceDataType = FStateTreeConstantConsiderationInstanceData;
 
+	//~ Begin FStateTreeNodeBase Interface
 	virtual const UStruct* GetInstanceDataType() const override { return FInstanceDataType::StaticStruct(); }
 #if WITH_EDITOR
 	virtual FText GetDescription(const FGuid& ID, FStateTreeDataView InstanceDataView, const IStateTreeBindingLookup& BindingLookup, EStateTreeNodeFormatting Formatting = EStateTreeNodeFormatting::Text) const override;
+	//~ End FStateTreeNodeBase Interface
 #endif
 
 protected:
-	virtual float ComputeRawScore(FStateTreeExecutionContext& Context) const override;
+	//~ Begin FStateTreeConsiderationBase Interface
+	virtual float GetScore(FStateTreeExecutionContext& Context) const override;
+	//~ End FStateTreeConsiderationBase Interface
+};
+
+USTRUCT()
+struct STATETREEMODULE_API FStateTreeConsiderationResponseCurve
+{
+	GENERATED_BODY()
+
+	/**
+	 * Evaluate the output value from curve
+	 * @param NormalizedInput the normalized input value to the response curve
+	 * @return The output value. If the curve is not set, will simply return NormalizedInput.
+	 */
+	float Evaluate(float NormalizedInput) const
+	{
+		if (const FRichCurve* Curve = CurveInfo.GetRichCurveConst())
+		{
+			if (!Curve->IsEmpty())
+			{
+				return Curve->Eval(NormalizedInput);
+			}
+		}
+
+		return NormalizedInput;
+	}
+
+	/* Curve used to output the final score for the Consideration. */
+	UPROPERTY(EditAnywhere, Category = Default, DisplayName = "Curve")
+	FRuntimeFloatCurve CurveInfo;
+};
+
+USTRUCT()
+struct STATETREEMODULE_API FStateTreeFloatInputConsiderationInstanceData
+{
+	GENERATED_BODY()
+
+	UPROPERTY(EditAnywhere, Category = "Input")
+	float Input = 0.f;
+
+	UPROPERTY(EditAnywhere, DisplayName = "InputRange", Category = "Parameter")
+	FFloatInterval Interval = FFloatInterval(0.f, 1.f);
+};
+
+/**
+ * Consideration using a Float as input to the response curve.
+ */
+USTRUCT(DisplayName = "Float Input")
+struct STATETREEMODULE_API FStateTreeFloatInputConsideration : public FStateTreeConsiderationCommonBase
+{
+	GENERATED_BODY()
+
+	using FInstanceDataType = FStateTreeFloatInputConsiderationInstanceData;
+	
+	//~ Begin FStateTreeNodeBase Interface
+	virtual const UStruct* GetInstanceDataType() const override { return FInstanceDataType::StaticStruct(); }
+#if WITH_EDITOR
+	virtual FText GetDescription(const FGuid& ID, FStateTreeDataView InstanceDataView, const IStateTreeBindingLookup& BindingLookup, EStateTreeNodeFormatting Formatting = EStateTreeNodeFormatting::Text) const override;
+	//~ End FStateTreeNodeBase Interface
+#endif
+
+protected:
+	//~ Begin FStateTreeConsiderationBase Interface
+	virtual float GetScore(FStateTreeExecutionContext& Context) const override;
+	//~ End FStateTreeConsiderationBase Interface
+
+public:
+	UPROPERTY(EditAnywhere, Category = "Default")
+	FStateTreeConsiderationResponseCurve ResponseCurve;
+};
+
+USTRUCT()
+struct STATETREEMODULE_API FStateTreeEnumValueScorePair
+{
+	GENERATED_BODY()
+
+	bool operator==(const FStateTreeEnumValueScorePair& RHS) const
+	{
+		return EnumValue == RHS.EnumValue && Score == RHS.Score;
+	}
+
+	bool operator!=(const FStateTreeEnumValueScorePair& RHS) const
+	{
+		return EnumValue != RHS.EnumValue || Score != RHS.Score;
+	}
+
+#if WITH_EDITORONLY_DATA
+	UPROPERTY(EditAnywhere, Category = Default)
+	FName EnumName;
+#endif
+
+	UPROPERTY(EditAnywhere, Category = Default)
+	int64 EnumValue = 0;
+
+	UPROPERTY(EditAnywhere, Category = Default, meta = (UIMin = 0.0, ClampMin = 0.0, UIMax = 1.0, ClampMax = 1.0))
+	float Score = 0.f;
+};
+
+USTRUCT()
+struct STATETREEMODULE_API FStateTreeEnumValueScorePairs
+{
+	GENERATED_BODY()
+	
+#if WITH_EDITORONLY_DATA
+	UPROPERTY(EditAnywhere, Category = Default)
+	TObjectPtr<const UEnum> Enum;
+#endif //WITH_EDITORONLY_DATA
+
+	UPROPERTY(EditAnywhere, Category = Default, DisplayName = EnumValueScorePairs)
+	TArray<FStateTreeEnumValueScorePair> Data;
+
+#if WITH_EDITORONLY_DATA
+	/** Initializes the class to specific enum.*/
+	void Initialize(const UEnum* NewEnum)
+	{
+		if (Enum != NewEnum || NewEnum == nullptr)
+		{
+			Data.Empty();
+		}
+
+		Enum = NewEnum;
+	}
+#endif //WITH_EDITORONLY_DATA
+};
+
+USTRUCT()
+struct STATETREEMODULE_API FStateTreeEnumInputConsiderationInstanceData
+{
+	GENERATED_BODY()
+
+	UPROPERTY(EditAnywhere, Category = "Input", meta=(AllowAnyBinding))
+	FStateTreeAnyEnum Input;
+};
+
+USTRUCT(DisplayName = "Enum Input")
+struct STATETREEMODULE_API FStateTreeEnumInputConsideration : public FStateTreeConsiderationCommonBase
+{
+	GENERATED_BODY()
+
+	using FInstanceDataType = FStateTreeEnumInputConsiderationInstanceData;
+
+	//~ Begin FStateTreeNodeBase Interface
+	virtual const UStruct* GetInstanceDataType() const override { return FInstanceDataType::StaticStruct(); }
+#if WITH_EDITORONLY_DATA
+	virtual EDataValidationResult Compile(FStateTreeDataView InstanceDataView, TArray<FText>& ValidationMessages) override;
+	virtual void OnBindingChanged(const FGuid& ID, FStateTreeDataView InstanceData, const FStateTreePropertyPath& SourcePath, const FStateTreePropertyPath& TargetPath, const IStateTreeBindingLookup& BindingLookup) override;
+#endif //WITH_EDITORONLY_DATA
+
+#if WITH_EDITOR
+	virtual FText GetDescription(const FGuid& ID, FStateTreeDataView InstanceDataView, const IStateTreeBindingLookup& BindingLookup, EStateTreeNodeFormatting Formatting = EStateTreeNodeFormatting::Text) const override;
+	//~ End FStateTreeNodeBase Interface
+#endif //WITH_EDITOR
+
+protected:
+	//~ Begin FStateTreeConsiderationBase Interface
+	virtual float GetScore(FStateTreeExecutionContext& Context) const override;
+	//~ End FStateTreeConsiderationBase Interface
+	
+	UPROPERTY(EditAnywhere, Category = "Default")
+	FStateTreeEnumValueScorePairs EnumValueScorePairs;
 };
diff --git a/Engine/Plugins/Runtime/StateTree/Source/StateTreeModule/Public/StateTreeConsiderationBase.h b/Engine/Plugins/Runtime/StateTree/Source/StateTreeModule/Public/StateTreeConsiderationBase.h
--- a/Engine/Plugins/Runtime/StateTree/Source/StateTreeModule/Public/StateTreeConsiderationBase.h
+++ b/Engine/Plugins/Runtime/StateTree/Source/StateTreeModule/Public/StateTreeConsiderationBase.h
@@ -6,29 +6,8 @@
 #include "StateTreeConsiderationBase.generated.h"
 
 struct FStateTreeExecutionContext;
-enum class EAlphaBlendOption : uint8;
 enum class EStateTreeExpressionOperand : uint8;
 
-USTRUCT()
-struct STATETREEMODULE_API FStateTreeConsiderationResponseCurve
-{
-	GENERATED_BODY()
-
-	FStateTreeConsiderationResponseCurve();
-
-	/* Optional Curve used to output the final normalized score.If it is set empty, the final value will be raw score normalized by bounds */
-	UPROPERTY(EditAnywhere, Category = Default)
-	EAlphaBlendOption BlendOption;
-
-	/* Lower Bound used to normalize the raw score */
-	UPROPERTY(EditAnywhere, Category = Default)
-	float RawScoreLowerBound;
-
-	/* Upper Bound used to normalize the raw score */
-	UPROPERTY(EditAnywhere, Category = Default)
-	float RawScoreUpperBound;
-};
-
 /**
  * This feature is experimental and the API is expected to change. 
  * Base struct for all utility considerations.
@@ -40,10 +19,11 @@ struct STATETREEMODULE_API FStateTreeConsiderationBase : public FStateTreeNodeBa
 
 	FStateTreeConsiderationBase();
 
-	float ComputeNormalizedScore(FStateTreeExecutionContext& Context) const;
+public:
+	float GetNormalizedScore(FStateTreeExecutionContext& Context) const;
 
 protected:
-	virtual float ComputeRawScore(FStateTreeExecutionContext& Context) const { return .0f; }
+	virtual float GetScore(FStateTreeExecutionContext& Context) const { return 0.f; };
 
 public:
 	UPROPERTY()
@@ -51,10 +31,6 @@ public:
 
 	UPROPERTY()
 	int8 DeltaIndent;
-
-	//Response Curve used to output the final normalized score.
-	UPROPERTY(EditAnywhere, Category = Default)
-	FStateTreeConsiderationResponseCurve ResponseCurve;
 };
 
 /**
diff --git a/Engine/Plugins/Runtime/StateTree/Source/StateTreeModule/Public/StateTreeNodeDescriptionHelpers.h b/Engine/Plugins/Runtime/StateTree/Source/StateTreeModule/Public/StateTreeNodeDescriptionHelpers.h
--- a/Engine/Plugins/Runtime/StateTree/Source/StateTreeModule/Public/StateTreeNodeDescriptionHelpers.h
+++ b/Engine/Plugins/Runtime/StateTree/Source/StateTreeModule/Public/StateTreeNodeDescriptionHelpers.h
@@ -24,8 +24,14 @@ extern STATETREEMODULE_API FText GetInvertText(bool bInvert, EStateTreeNodeForma
 /** @return description of a boolean value. */
 extern STATETREEMODULE_API FText GetBoolText(bool bValue, EStateTreeNodeFormatting Formatting);
 
-/** @return description of being within a value range. */
-extern STATETREEMODULE_API FText GetWithinValueRangeText(float Min, float Max, EStateTreeNodeFormatting Formatting);
+/** @return description of a float interval. */
+extern STATETREEMODULE_API FText GetIntervalText(const FFloatInterval& Interval, EStateTreeNodeFormatting Formatting);
+
+/** @return description of a float interval. */
+extern STATETREEMODULE_API FText GetIntervalText(float Min, float Max, EStateTreeNodeFormatting Formatting);
+
+/** @return description of a float interval. */
+extern STATETREEMODULE_API FText GetIntervalText(const FText& MinValueText, const FText& MaxValueText, EStateTreeNodeFormatting Formatting);
 
 /** @return description for a Gameplay Tag Container. If the length of container description is longer than ApproxMaxLength, the it truncated and ... as added to the end. */
 extern STATETREEMODULE_API FText GetGameplayTagContainerAsText(const FGameplayTagContainer& TagContainer, const int ApproxMaxLength = 60);
