StateTree: Fix empty states that do not transition when completed. Fix bad memory access when state transitioning.

diff --git a/Engine/Plugins/Runtime/StateTree/Source/StateTreeModule/Private/StateTreeExecutionContext.cpp b/Engine/Plugins/Runtime/StateTree/Source/StateTreeModule/Private/StateTreeExecutionContext.cpp
--- a/Engine/Plugins/Runtime/StateTree/Source/StateTreeModule/Private/StateTreeExecutionContext.cpp
+++ b/Engine/Plugins/Runtime/StateTree/Source/StateTreeModule/Private/StateTreeExecutionContext.cpp
@@ -2700,7 +2700,7 @@ EStateTreeRunStatus FStateTreeExecutionContext::TickTasks(const float DeltaTime)
 		// No tasks, done ticking.
 		Result = EStateTreeRunStatus::Succeeded;
 		Exec.CompletedFrameIndex = FStateTreeIndex16(0);
-		Exec.CompletedStateHandle = Exec.ActiveFrames[0].ActiveStates.GetStateSafe(0);
+		Exec.CompletedStateHandle = Exec.ActiveFrames.Num() > 0 ? Exec.ActiveFrames[0].ActiveStates.Last() : FStateTreeStateHandle::Invalid;
 	}
 
 	return Result;
@@ -3639,11 +3639,12 @@ bool FStateTreeExecutionContext::SelectState(const FStateTreeExecutionFrame& Cur
 			}
 		}
 	}
-	// Existing state's data is safe to access during select.
-	LastFrame.NumCurrentlyActiveStates = static_cast<uint8>(LastFrame.ActiveStates.Num());
 
 	LastFrame.ActiveStates.SetNum(FirstNewStateIndex);
 
+	// Existing state's data is safe to access during select.
+	LastFrame.NumCurrentlyActiveStates = static_cast<uint8>(LastFrame.ActiveStates.Num());
+
 	FStateSelectionResult InitialSelection;
 
 	if (Fallback == EStateTreeSelectionFallback::NextSelectableSibling)
