StateTree: Prevent long name and invalid characters for StateTree nodes.

diff --git a/Engine/Plugins/Runtime/StateTree/Source/StateTreeEditorModule/Private/Customizations/StateTreeEditorNodeDetails.cpp b/Engine/Plugins/Runtime/StateTree/Source/StateTreeEditorModule/Private/Customizations/StateTreeEditorNodeDetails.cpp
--- a/Engine/Plugins/Runtime/StateTree/Source/StateTreeEditorModule/Private/Customizations/StateTreeEditorNodeDetails.cpp
+++ b/Engine/Plugins/Runtime/StateTree/Source/StateTreeEditorModule/Private/Customizations/StateTreeEditorNodeDetails.cpp
@@ -11,8 +11,8 @@
 #include "StateTreePropertyRef.h"
 #include "Widgets/Input/SButton.h"
 #include "Widgets/Input/SComboButton.h"
-#include "Widgets/Input/SEditableText.h"
 #include "Widgets/Notifications/SNotificationList.h"
+#include "Widgets/Text/SInlineEditableTextBlock.h"
 #include "Widgets/Text/SRichTextBlock.h"
 #include "Widgets/Layout/SWidgetSwitcher.h"
 #include "Widgets/SStateTreeNodeTypePicker.h"
@@ -538,12 +538,11 @@ void FStateTreeEditorNodeDetails::CustomizeHeader(TSharedRef<class IPropertyHand
 						]
 						+ SWidgetSwitcher::Slot()
 						[
-							SAssignNew(NameEdit, SEditableText)
-							.Style(&FStateTreeEditorStyle::Get().GetWidgetStyle<FEditableTextStyle>("StateTree.Node.Editable"))
+							SAssignNew(NameEdit, SInlineEditableTextBlock)
+							.Style(FStateTreeEditorStyle::Get(), "StateTree.Node.TitleInlineEditableText")
 							.Text(this, &FStateTreeEditorNodeDetails::GetName)
-							.OnTextCommitted(this, &FStateTreeEditorNodeDetails::OnNameCommitted)
-							.SelectAllTextWhenFocused(true)
-							.RevertTextOnEscape(true)
+							.OnTextCommitted(this, &FStateTreeEditorNodeDetails::HandleNameCommitted)
+							.OnVerifyTextChanged(this, &FStateTreeEditorNodeDetails::HandleVerifyNameChanged)
 							.Visibility(this, &FStateTreeEditorNodeDetails::IsNodeDescriptionVisible)
 						]
 					]
@@ -1337,6 +1336,7 @@ FReply FStateTreeEditorNodeDetails::OnDescriptionClicked(const FGeometry& MyGeom
 			// Focus on name edit.
 			FReply Reply = FReply::Handled();
 			Reply.SetUserFocus(NameEdit.ToSharedRef());
+			NameEdit->EnterEditingMode();
 			return Reply;
 		}
 	}
@@ -1465,49 +1465,62 @@ FText FStateTreeEditorNodeDetails::GetName() const
 	return LOCTEXT("MultipleSelected", "Multiple Selected");
 }
 
-void FStateTreeEditorNodeDetails::OnNameCommitted(const FText& NewText, ETextCommit::Type InTextCommit) const
+bool FStateTreeEditorNodeDetails::HandleVerifyNameChanged(const FText& InText, FText& OutErrorMessage) const
+{
+	const FString NewName = FText::TrimPrecedingAndTrailing(InText).ToString();
+	if (NewName.Len() >= NAME_SIZE)
+	{
+		OutErrorMessage = LOCTEXT("VerifyNodeLabelFailed_MaxLength", "Max length exceeded");
+		return false;
+	}
+	return NewName.Len() > 0 && FName::IsValidXName(NewName, INVALID_NAME_CHARACTERS, &OutErrorMessage);
+}
+
+void FStateTreeEditorNodeDetails::HandleNameCommitted(const FText& NewText, ETextCommit::Type InTextCommit) const
 {
 	check(StructProperty);
 
 	if (InTextCommit == ETextCommit::OnEnter || InTextCommit == ETextCommit::OnUserMovedFocus)
 	{
 		// Remove excess whitespace and prevent categories with just spaces
-		FString NewName = FText::TrimPrecedingAndTrailing(NewText).ToString();
-
-		if (GEditor)
+		const FString NewName = FText::TrimPrecedingAndTrailing(NewText).ToString();
+		if (NewName.Len() > 0 && FName::IsValidXName(NewName, INVALID_NAME_CHARACTERS) && NewName.Len() < NAME_SIZE)
 		{
-			GEditor->BeginTransaction(LOCTEXT("SetName", "Set Name"));
-		}
-		StructProperty->NotifyPreChange();
+			if (GEditor)
+			{
+				GEditor->BeginTransaction(LOCTEXT("SetName", "Set Name"));
+			}
+			StructProperty->NotifyPreChange();
 
-		TArray<void*> RawNodeData;
-		StructProperty->AccessRawData(RawNodeData);
-		
-		for (void* Data : RawNodeData)
-		{
-			// Set Name
-			if (FStateTreeEditorNode* Node = static_cast<FStateTreeEditorNode*>(Data))
+			TArray<void*> RawNodeData;
+			StructProperty->AccessRawData(RawNodeData);
+
+			for (void* Data : RawNodeData)
 			{
-				if (FStateTreeNodeBase* BaseNode = Node->Node.GetMutablePtr<FStateTreeNodeBase>())
+				// Set Name
+				if (FStateTreeEditorNode* Node = static_cast<FStateTreeEditorNode*>(Data))
 				{
-					BaseNode->Name = FName(NewName);
+					if (FStateTreeNodeBase* BaseNode = Node->Node.GetMutablePtr<FStateTreeNodeBase>())
+					{
+						BaseNode->Name = FName(NewName);
+					}
 				}
 			}
-		}
 
-		StructProperty->NotifyPostChange(EPropertyChangeType::ValueSet);
+			StructProperty->NotifyPostChange(EPropertyChangeType::ValueSet);
 
-		if (StateTree)
-		{
-			UE::StateTree::Delegates::OnIdentifierChanged.Broadcast(*StateTree);
-		}
+			if (StateTree)
+			{
+				UE::StateTree::Delegates::OnIdentifierChanged.Broadcast(*StateTree);
+			}
 
-		if (GEditor)
-		{
-			GEditor->EndTransaction();
-		}
+			if (GEditor)
+			{
+				GEditor->EndTransaction();
+			}
 
-		StructProperty->NotifyFinishedChangingProperties();
+			StructProperty->NotifyFinishedChangingProperties();
+		}
 	}
 
 	// Switch back to rich view.
@@ -1814,6 +1827,7 @@ void FStateTreeEditorNodeDetails::OnRenameNode() const
 
 			FSlateApplication::Get().SetKeyboardFocus(NameEdit);
 			FSlateApplication::Get().SetUserFocus(0, NameEdit);
+			NameEdit->EnterEditingMode();
 		}
 	}
 }
diff --git a/Engine/Plugins/Runtime/StateTree/Source/StateTreeEditorModule/Private/Customizations/StateTreeEditorNodeDetails.h b/Engine/Plugins/Runtime/StateTree/Source/StateTreeEditorModule/Private/Customizations/StateTreeEditorNodeDetails.h
--- a/Engine/Plugins/Runtime/StateTree/Source/StateTreeEditorModule/Private/Customizations/StateTreeEditorNodeDetails.h
+++ b/Engine/Plugins/Runtime/StateTree/Source/StateTreeEditorModule/Private/Customizations/StateTreeEditorNodeDetails.h
@@ -17,7 +17,7 @@ class SComboButton;
 class SWidget;
 class SSearchBox;
 class SWidgetSwitcher;
-class SEditableText;
+class SInlineEditableTextBlock;
 class SBorder;
 class SMenuAnchor;
 class UStateTree;
@@ -94,7 +94,8 @@ private:
 	FReply OnRowMouseUp(const FGeometry& MyGeometry, const FPointerEvent& MouseEvent);
 
 	FText GetName() const;
-	void OnNameCommitted(const FText& NewText, ETextCommit::Type InTextCommit) const;
+	bool HandleVerifyNameChanged(const FText& InText, FText& OutErrorMessage) const;
+	void HandleNameCommitted(const FText& NewLabel, ETextCommit::Type CommitType) const;
 
 	FText GetNodePickerTooltip() const;
 	void OnNodePicked(const UStruct* InStruct) const;
@@ -118,7 +119,7 @@ private:
 	UScriptStruct* BaseScriptStruct = nullptr;
 	UClass* BaseClass = nullptr;
 	TSharedPtr<SWidgetSwitcher> NameSwitcher;
-	TSharedPtr<SEditableText> NameEdit;
+	TSharedPtr<SInlineEditableTextBlock> NameEdit;
 	TSharedPtr<SBorder> RowBorder; 
 
 	UStateTreeEditorData* EditorData = nullptr;
diff --git a/Engine/Plugins/Runtime/StateTree/Source/StateTreeEditorModule/Private/SStateTreeViewRow.cpp b/Engine/Plugins/Runtime/StateTree/Source/StateTreeEditorModule/Private/SStateTreeViewRow.cpp
--- a/Engine/Plugins/Runtime/StateTree/Source/StateTreeEditorModule/Private/SStateTreeViewRow.cpp
+++ b/Engine/Plugins/Runtime/StateTree/Source/StateTreeEditorModule/Private/SStateTreeViewRow.cpp
@@ -288,10 +288,6 @@ void SStateTreeViewRow::Construct(const FArguments& InArgs, const TSharedRef<STa
 											[
 												SAssignNew(NameTextBlock, SInlineEditableTextBlock)
 												.Style(FStateTreeEditorStyle::Get(), "StateTree.State.TitleInlineEditableText")
-												.OnVerifyTextChanged_Lambda([](const FText& NewLabel, FText& OutErrorMessage)
-													{
-														return !NewLabel.IsEmptyOrWhitespace();
-													})
 												.OnTextCommitted(this, &SStateTreeViewRow::HandleNodeLabelTextCommitted)
 												.OnVerifyTextChanged(this, &SStateTreeViewRow::HandleVerifyNodeLabelTextChanged)
 												.Text(this, &SStateTreeViewRow::GetStateDesc)
@@ -1977,7 +1973,7 @@ bool SStateTreeViewRow::HandleVerifyNodeLabelTextChanged(const FText& InText, FT
 				OutErrorMessage = LOCTEXT("VerifyNodeLabelFailed_MaxLength", "Max length exceeded");
 				return false;
 			}
-			return FName::IsValidXName(NewName, INVALID_NAME_CHARACTERS, &OutErrorMessage);
+			return NewName.Len() > 0 && FName::IsValidXName(NewName, INVALID_NAME_CHARACTERS, &OutErrorMessage);
 		}
 	}
 	OutErrorMessage = LOCTEXT("VerifyNodeLabelFailed", "Invalid State Tree");
@@ -1991,7 +1987,7 @@ void SStateTreeViewRow::HandleNodeLabelTextCommitted(const FText& NewLabel, ETex
 		if (UStateTreeState* State = WeakState.Get())
 		{
 			const FString NewName = FText::TrimPrecedingAndTrailing(NewLabel).ToString();
-			if (FName::IsValidXName(NewName, INVALID_NAME_CHARACTERS) && NewName.Len() < NAME_SIZE)
+			if (NewName.Len() > 0 && FName::IsValidXName(NewName, INVALID_NAME_CHARACTERS) && NewName.Len() < NAME_SIZE)
 			{
 				StateTreeViewModel->RenameState(State, FName(NewName));
 			}
diff --git a/Engine/Plugins/Runtime/StateTree/Source/StateTreeEditorModule/Private/StateTreeEditorStyle.cpp b/Engine/Plugins/Runtime/StateTree/Source/StateTreeEditorModule/Private/StateTreeEditorStyle.cpp
--- a/Engine/Plugins/Runtime/StateTree/Source/StateTreeEditorModule/Private/StateTreeEditorStyle.cpp
+++ b/Engine/Plugins/Runtime/StateTree/Source/StateTreeEditorModule/Private/StateTreeEditorStyle.cpp
@@ -256,9 +256,15 @@ FStateTreeEditorStyle::FStateTreeEditorStyle()
 		EditableTextStyle.Font.Size = 10.0f;
 		Set("StateTree.Node.Editable", EditableTextStyle);
 
-		Set("StateTree.Node.Normal", FTextBlockStyle(NormalText)
+		FEditableTextBoxStyle EditableTextBlockStyle = FEditableTextBoxStyle(FAppStyle::GetWidgetStyle<FEditableTextBoxStyle>("NormalEditableTextBox"));
+		EditableTextStyle.Font = FAppStyle::GetFontStyle(TEXT("PropertyWindow.NormalFont"));
+		EditableTextStyle.Font.Size = 10.0f;
+		Set("StateTree.Node.EditableTextBlock", EditableTextBlockStyle);
+
+		const FTextBlockStyle StateNodeNormalText = FTextBlockStyle(NormalText)
 			.SetFont(FAppStyle::GetFontStyle(TEXT("PropertyWindow.NormalFont")))
-			.SetFontSize(10));
+			.SetFontSize(10);
+		Set("StateTree.Node.Normal", StateNodeNormalText);
 
 		Set("StateTree.Node.Bold", FTextBlockStyle(NormalText)
 			.SetFont(FAppStyle::GetFontStyle(TEXT("PropertyWindow.BoldFont")))
@@ -268,6 +274,10 @@ FStateTreeEditorStyle::FStateTreeEditorStyle()
 			.SetColorAndOpacity(FSlateColor::UseSubduedForeground())
 			.SetFont(FAppStyle::GetFontStyle(TEXT("PropertyWindow.NormalFont")))
 			.SetFontSize(10));
+
+		Set("StateTree.Node.TitleInlineEditableText", FInlineEditableTextBlockStyle()
+			.SetTextStyle(StateNodeNormalText)
+			.SetEditableTextBoxStyle(EditableTextBlockStyle));
 	}
 
 	
