[StateTree] diff tools

diff --git a/Engine/Plugins/Runtime/StateTree/Source/StateTreeEditorModule/Private/AssetDefinition_StateTree.cpp b/Engine/Plugins/Runtime/StateTree/Source/StateTreeEditorModule/Private/AssetDefinition_StateTree.cpp
--- a/Engine/Plugins/Runtime/StateTree/Source/StateTreeEditorModule/Private/AssetDefinition_StateTree.cpp
+++ b/Engine/Plugins/Runtime/StateTree/Source/StateTreeEditorModule/Private/AssetDefinition_StateTree.cpp
@@ -1,9 +1,10 @@
 // Copyright Epic Games, Inc. All Rights Reserved.
 
 #include "AssetDefinition_StateTree.h"
-#include "StateTreeEditorModule.h"
-#include "StateTree.h"
 #include "Modules/ModuleManager.h"
+#include "SStateTreeDiff.h"
+#include "StateTree.h"
+#include "StateTreeEditorModule.h"
 
 #define LOCTEXT_NAMESPACE "AssetTypeActions"
 
@@ -42,4 +43,18 @@ EAssetCommandResult UAssetDefinition_StateTree::OpenAssets(const FAssetOpenArgs&
 	return EAssetCommandResult::Handled;
 }
 
+EAssetCommandResult UAssetDefinition_StateTree::PerformAssetDiff(const FAssetDiffArgs& DiffArgs) const
+{
+	const UStateTree* OldStateTree = Cast<UStateTree>(DiffArgs.OldAsset);
+	const UStateTree* NewStateTree = Cast<UStateTree>(DiffArgs.NewAsset);
+
+	if (OldStateTree == nullptr || NewStateTree == nullptr)
+	{
+		return EAssetCommandResult::Unhandled;
+	}
+
+	UE::StateTree::Diff::SDiffWidget::CreateDiffWindow(OldStateTree, NewStateTree, DiffArgs.OldRevision, DiffArgs.NewRevision, UStateTree::StaticClass());
+	return EAssetCommandResult::Handled;
+}
+
 #undef LOCTEXT_NAMESPACE
diff --git a/Engine/Plugins/Runtime/StateTree/Source/StateTreeEditorModule/Private/AssetDefinition_StateTree.h b/Engine/Plugins/Runtime/StateTree/Source/StateTreeEditorModule/Private/AssetDefinition_StateTree.h
--- a/Engine/Plugins/Runtime/StateTree/Source/StateTreeEditorModule/Private/AssetDefinition_StateTree.h
+++ b/Engine/Plugins/Runtime/StateTree/Source/StateTreeEditorModule/Private/AssetDefinition_StateTree.h
@@ -17,5 +17,6 @@ protected:
 	virtual TSoftClassPtr<UObject> GetAssetClass() const override;
 	virtual TConstArrayView<FAssetCategoryPath> GetAssetCategories() const override;
 	virtual EAssetCommandResult OpenAssets(const FAssetOpenArgs& OpenArgs) const override;
+	virtual EAssetCommandResult PerformAssetDiff(const FAssetDiffArgs& DiffArgs) const override;
 	// UAssetDefinition End
 };
\ No newline at end of file
diff --git a/Engine/Plugins/Runtime/StateTree/Source/StateTreeEditorModule/Private/AsyncStateTreeDiff.cpp b/Engine/Plugins/Runtime/StateTree/Source/StateTreeEditorModule/Private/AsyncStateTreeDiff.cpp
new file mode 100644
--- /dev/null
+++ b/Engine/Plugins/Runtime/StateTree/Source/StateTreeEditorModule/Private/AsyncStateTreeDiff.cpp
@@ -0,0 +1,542 @@
+// Copyright Epic Games, Inc. All Rights Reserved.
+
+#include "AsyncStateTreeDiff.h"
+#include "DiffUtils.h"
+#include "SStateTreeView.h"
+#include "StateTreeDiffHelper.h"
+#include "StateTreeEditorData.h"
+#include "StateTreeState.h"
+#include "StateTreeViewModel.h"
+
+namespace UE::StateTree::Diff
+{
+
+static bool AreObjectsEqual(const UObject* ObjectA, const UObject* ObjectB)
+{
+	if (!ObjectA || !ObjectB)
+	{
+		return ObjectA == ObjectB;
+	}
+
+	if (ObjectA->GetClass() != ObjectB->GetClass())
+	{
+		return false;
+	}
+
+	if (ObjectA != ObjectB)
+	{
+		const FProperty* ClassProperty = ObjectA->GetClass()->PropertyLink;
+		while (ClassProperty)
+		{
+			if (!ClassProperty->Identical_InContainer(ObjectA, ObjectB, /*ArrayIndex*/ 0, PPF_DeepComparison))
+			{
+				return false;
+			}
+			ClassProperty = ClassProperty->PropertyLinkNext;
+		}
+	}
+
+	return true;
+}
+
+static bool AreNodesEqual(const FStateTreeEditorNode& NodeA, const FStateTreeEditorNode& NodeB)
+{
+	return AreObjectsEqual(NodeA.InstanceObject.Get(), NodeB.InstanceObject.Get())
+		&& NodeA.Node.Identical(&NodeB.Node, PPF_DeepComparison)
+		&& NodeA.Instance.Identical(&NodeB.Instance, PPF_DeepComparison)
+		&& NodeA.ExpressionIndent == NodeB.ExpressionIndent
+		&& NodeA.ExpressionOperand == NodeB.ExpressionOperand;
+}
+
+static bool AreNodeArraysEqual(const TArray<FStateTreeEditorNode>& ArrayA, const TArray<FStateTreeEditorNode>& ArrayB)
+{
+	const int32 Count = ArrayA.Num();
+	if (Count != ArrayB.Num())
+	{
+		return false;
+	}
+	for (int32 Index = 0; Index < Count; Index++)
+	{
+		if (!AreNodesEqual(ArrayA[Index], ArrayB[Index]))
+		{
+			return false;
+		}
+	}
+	return true;
+}
+
+static bool AreStateTreeStateParametersEqual(const FStateTreeStateParameters& ParametersA, const FStateTreeStateParameters& ParametersB)
+{
+	if (ParametersA.Parameters.GetNumPropertiesInBag() != ParametersB.Parameters.GetNumPropertiesInBag())
+	{
+		return false;
+	}
+	if (ParametersA.PropertyOverrides != ParametersB.PropertyOverrides)
+	{
+		return false;
+	}
+	const FInstancedPropertyBag ParamsA = ParametersA.Parameters;
+	const FInstancedPropertyBag ParamsB = ParametersB.Parameters;
+	const UPropertyBag* BagA = ParamsA.GetPropertyBagStruct();
+	const UPropertyBag* BagB = ParamsB.GetPropertyBagStruct();
+	if (!BagA || !BagB)
+	{
+		return BagA == BagB;
+	}
+
+	const TConstArrayView<FPropertyBagPropertyDesc> DescsA = BagA->GetPropertyDescs();
+	const TConstArrayView<FPropertyBagPropertyDesc> DescsB = BagB->GetPropertyDescs();
+	const int32 Count = DescsA.Num();
+	for (int32 Index = 0; Index < Count; Index++)
+	{
+		if (DescsA[Index].Name != DescsB[Index].Name
+			|| !DescsA[Index].CompatibleType(DescsB[Index]))
+		{
+			return false;
+		}
+
+		const FName Name = DescsA[Index].Name;
+		TValueOrError<FString, EPropertyBagResult> SerializedA = ParamsA.GetValueSerializedString(Name);
+		TValueOrError<FString, EPropertyBagResult> SerializedB = ParamsB.GetValueSerializedString(Name);
+		if (SerializedA.HasError() || SerializedB.HasError())
+		{
+			return false;
+		}
+
+		if (SerializedA.HasValue() != SerializedB.HasValue())
+		{
+			return false;
+		}
+
+		if (SerializedA.GetValue() != SerializedB.GetValue())
+		{
+			return false;
+		}
+	}
+	return true;
+}
+
+static bool ArePropertiesEqual(const UStateTreeState* StateA, const UStateTreeState* StateB)
+{
+	if (!StateA || !StateB)
+	{
+		return StateA == StateB;
+	}
+
+	return StateA->Name == StateB->Name
+		&& StateA->Tag == StateB->Tag
+		&& StateA->ColorRef == StateB->ColorRef
+		&& StateA->Type == StateB->Type
+		&& StateA->SelectionBehavior == StateB->SelectionBehavior;
+}
+
+static bool AreParametersEqual(const UStateTreeState* StateA, const UStateTreeState* StateB)
+{
+	return AreStateTreeStateParametersEqual(StateA->Parameters, StateB->Parameters);
+}
+
+static bool AreConditionsEqual(const UStateTreeState* StateA, const UStateTreeState* StateB)
+{
+	return AreNodeArraysEqual(StateA->EnterConditions, StateB->EnterConditions);
+}
+
+static bool AreConsiderationsEqual(const UStateTreeState* StateA, const UStateTreeState* StateB)
+{
+	return AreNodeArraysEqual(StateA->Considerations, StateB->Considerations);
+}
+
+static bool AreTasksEqual(const UStateTreeState* StateA, const UStateTreeState* StateB)
+{
+	return AreNodeArraysEqual(StateA->Tasks, StateB->Tasks);
+}
+
+static bool AreTransitionsEqual(const UStateTreeState* StateA, const UStateTreeState* StateB)
+{
+	if (StateA->Transitions.Num() != StateB->Transitions.Num())
+	{
+		return false;
+	}
+
+	for (int32 Index = 0; Index < StateA->Transitions.Num(); Index++)
+	{
+		const FStateTreeTransition* TransitionA = &StateA->Transitions[Index];
+		const FStateTreeTransition* TransitionB = &StateB->Transitions[Index];
+		if (!TransitionA || !TransitionB)
+		{
+			return TransitionA == TransitionB;
+		}
+		// Not checking transitions on IDs
+		const bool bEqual = TransitionA->Trigger == TransitionB->Trigger
+			&& TransitionA->RequiredEvent == TransitionB->RequiredEvent
+			&& TransitionA->State.Name == TransitionB->State.Name
+			&& TransitionA->State.ID == TransitionB->State.ID
+			&& TransitionA->State.LinkType == TransitionB->State.LinkType
+			&& TransitionA->State.StateHandle == TransitionB->State.StateHandle
+			&& TransitionA->Priority == TransitionB->Priority
+			&& TransitionA->bDelayTransition == TransitionB->bDelayTransition
+			&& TransitionA->DelayDuration == TransitionB->DelayDuration
+			&& TransitionA->DelayRandomVariance == TransitionB->DelayRandomVariance
+			&& AreNodeArraysEqual(TransitionA->Conditions, TransitionB->Conditions)
+			&& TransitionA->bTransitionEnabled == TransitionB->bTransitionEnabled;
+		if (!bEqual)
+		{
+			return false;
+		}
+	}
+	return true;
+}
+
+static bool AreStateTreePropertiesEqual(const UStateTreeEditorData* StateTreeDataA, const UStateTreeEditorData* StateTreeDataB)
+{
+	// Check the differences in Bindings
+	bool bBindingsEqual = StateTreeDataA->EditorBindings.GetBindings().Num() == StateTreeDataB->EditorBindings.GetBindings().Num();
+	if (bBindingsEqual)
+	{
+		for (const FStateTreePropertyPathBinding& PropertyPathBinding : StateTreeDataA->EditorBindings.GetBindings())
+		{
+			const FStateTreePropertyPath& PropertyPathTarget = PropertyPathBinding.GetTargetPath();
+			if (StateTreeDataB->EditorBindings.HasPropertyBinding(PropertyPathBinding.GetTargetPath()))
+			{
+				if (*StateTreeDataA->EditorBindings.GetPropertyBindingSource(PropertyPathTarget) != *StateTreeDataB->EditorBindings.GetPropertyBindingSource(PropertyPathTarget))
+				{
+					bBindingsEqual = false;
+					break;
+				}
+			}
+			else
+			{
+				bBindingsEqual = false;
+				break;
+			}
+		}
+	}
+	if (!bBindingsEqual)
+	{
+		return false;
+	}
+
+	// Check the differences in Evaluators and Tasks
+	if (!AreNodeArraysEqual(StateTreeDataA->Evaluators, StateTreeDataB->Evaluators))
+	{
+		return false;
+	}
+	if (!AreNodeArraysEqual(StateTreeDataA->GlobalTasks, StateTreeDataB->GlobalTasks))
+	{
+		return false;
+	}
+
+	// Check the differences in Parameters
+	if (!AreStateTreeStateParametersEqual(StateTreeDataA->RootParameters, StateTreeDataB->RootParameters))
+	{
+		return false;
+	}
+	return true;
+}
+
+static FPropertySoftPath GetPropertyPath(const FStateTreePropertyPath& StateTreePropertyPath, const UStateTreeState* StateTreeState)
+{
+	TArray<FName> Path;
+	auto CheckNodes = [&StateTreePropertyPath, &Path](const TArray<FStateTreeEditorNode>& List, FName PathSegmentName)
+	{
+		for (int i = 0; i < List.Num(); i++)
+		{
+			if (List[i].ID == StateTreePropertyPath.GetStructID())
+			{
+				Path.Add(PathSegmentName);
+				Path.Add(FName(FString::FromInt(i)));
+				if (List[i].InstanceObject)
+				{
+					Path.Add(FName("InstanceObject"));
+				}
+				else
+				{
+					Path.Add(FName("Instance"));
+				}
+				return true;
+			}
+		}
+		return false;
+	};
+	auto CheckTransitions = [StateTreePropertyPath, &Path](const TArray<FStateTreeTransition>& List, FName PathSegmentName)
+	{
+		for (int i = 0; i < List.Num(); i++)
+		{
+			if (List[i].ID == StateTreePropertyPath.GetStructID())
+			{
+				Path.Add(PathSegmentName);
+				Path.Add(FName(FString::FromInt(i)));
+				return true;
+			}
+		}
+		return false;
+	};
+	if (CheckNodes(StateTreeState->EnterConditions, "EnterConditions")
+		|| CheckNodes(StateTreeState->Tasks, "Tasks")
+		|| CheckTransitions(StateTreeState->Transitions, "Transitions"))
+	{
+		for (const FStateTreePropertyPathSegment& PropertySegment : StateTreePropertyPath.GetSegments())
+		{
+			Path.Add(PropertySegment.GetName());
+		}
+	}
+
+	return FPropertySoftPath(Path);
+}
+
+static void GetBindingsDifferences(UStateTreeEditorData* StateTreeDataA, UStateTreeEditorData* StateTreeDataB, TArray<FSingleDiffEntry>& OutDiffEntries)
+{
+	struct FBindingDiff
+	{
+		FStateTreePropertyPath TargetPath;
+		FStateTreePropertyPath SourcePathA;
+		FStateTreePropertyPath SourcePathB;
+	};
+
+	TArray<FBindingDiff> BindingDiffs;
+	// Check the differences in Bindings
+	for (const FStateTreePropertyPathBinding& PropertyPathBinding : StateTreeDataA->EditorBindings.GetBindings())
+	{
+		FStateTreePropertyPath PropertyPathTarget = PropertyPathBinding.GetTargetPath();
+		FStateTreePropertyPath PropertyPathSource = PropertyPathBinding.GetSourcePath();
+
+		FBindingDiff Entry;
+		Entry.TargetPath = PropertyPathTarget;
+		Entry.SourcePathA = PropertyPathSource;
+		BindingDiffs.Add(Entry);
+	}
+	for (const FStateTreePropertyPathBinding& PropertyPathBinding : StateTreeDataB->EditorBindings.GetBindings())
+	{
+		FStateTreePropertyPath PropertyPathTarget = PropertyPathBinding.GetTargetPath();
+		FStateTreePropertyPath PropertyPathSource = PropertyPathBinding.GetSourcePath();
+		bool bFound = false;
+		for (FBindingDiff& Diff : BindingDiffs)
+		{
+			if (Diff.TargetPath == PropertyPathTarget)
+			{
+				Diff.SourcePathB = PropertyPathSource;
+				bFound = true;
+				break;
+			}
+		}
+		if (!bFound)
+		{
+			FBindingDiff Entry;
+			Entry.TargetPath = PropertyPathTarget;
+			Entry.SourcePathB = PropertyPathSource;
+			BindingDiffs.Add(Entry);
+		}
+	}
+
+	for (const FBindingDiff& DiffEntry : BindingDiffs)
+	{
+		if (DiffEntry.SourcePathA != DiffEntry.SourcePathB)
+		{
+			const UStateTreeState* TargetStateA = StateTreeDataA->GetStateByStructID(DiffEntry.TargetPath.GetStructID());
+			const UStateTreeState* TargetStateB = StateTreeDataB->GetStateByStructID(DiffEntry.TargetPath.GetStructID());
+
+			if (TargetStateA && TargetStateB)
+			{
+				FStateSoftPath StatePathA(TargetStateA);
+				FStateSoftPath StatePathB(TargetStateB);
+				FPropertySoftPath PropertyPath = GetPropertyPath(DiffEntry.TargetPath, TargetStateA);
+
+				EStateDiffType StateTreeDiffType = EStateDiffType::BindingChanged;
+				if (DiffEntry.SourcePathA.IsPathEmpty())
+				{
+					StateTreeDiffType = EStateDiffType::BindingAddedToB;
+				}
+				else if (DiffEntry.SourcePathB.IsPathEmpty())
+				{
+					StateTreeDiffType = EStateDiffType::BindingAddedToA;
+				}
+				OutDiffEntries.Add(FSingleDiffEntry(StatePathA, StatePathB, StateTreeDiffType, PropertyPath));
+			}
+		}
+	}
+}
+
+FAsyncDiff::FAsyncDiff(const TSharedRef<SStateTreeView>& LeftTree, const TSharedRef<SStateTreeView>& RightTree)
+	: TAsyncTreeDifferences(RootNodesAttribute(LeftTree), RootNodesAttribute(RightTree))
+	, LeftView(LeftTree)
+	, RightView(RightTree)
+{}
+
+TAttribute<TArray<TWeakObjectPtr<UStateTreeState>>> FAsyncDiff::RootNodesAttribute(TWeakPtr<SStateTreeView> StateTreeView)
+{
+	return TAttribute<TArray<TWeakObjectPtr<UStateTreeState>>>::CreateLambda([StateTreeView]()
+	{
+		if (const TSharedPtr<SStateTreeView> TreeView = StaticCastSharedPtr<SStateTreeView>(StateTreeView.Pin()))
+		{
+			TArray<TWeakObjectPtr<UStateTreeState>> SubTrees;
+			TreeView->GetViewModel()->GetSubTrees(SubTrees);
+			return SubTrees;
+		}
+		return TArray<TWeakObjectPtr<UStateTreeState>>();
+	});
+}
+
+void FAsyncDiff::GetStatesDifferences(TArray<FSingleDiffEntry>& OutDiffEntries) const
+{
+	TArray<FString> RemovedStates;
+	TArray<FString> AddedStates;
+	ForEach(ETreeTraverseOrder::PreOrder, [&](const TUniquePtr<DiffNodeType>& Node) -> ETreeTraverseControl
+	{
+		FStateSoftPath StatePath;
+		UStateTreeState* LeftState = Node->ValueA.Get();
+		UStateTreeState* RightState = Node->ValueB.Get();
+		if (LeftState)
+		{
+			StatePath = FStateSoftPath(LeftState);
+		}
+		else if (RightState)
+		{
+			StatePath = FStateSoftPath(RightState);
+		}
+
+		EStateDiffType StateTreeDiffType = EStateDiffType::Invalid;
+		bool bSkipChildren = false;
+		switch (Node->DiffResult)
+		{
+		case ETreeDiffResult::MissingFromTree1:
+			StateTreeDiffType = EStateDiffType::StateAddedToB;
+			AddedStates.Add(StatePath.ToDisplayName(true));
+			if (RemovedStates.Contains(StatePath.ToDisplayName(true)))
+			{
+				StateTreeDiffType = EStateDiffType::StateMoved;
+			}
+			bSkipChildren = true;
+			break;
+		case ETreeDiffResult::MissingFromTree2:
+			StateTreeDiffType = EStateDiffType::StateAddedToA;
+			RemovedStates.Add(StatePath.ToDisplayName(true));
+			if (AddedStates.Contains(StatePath.ToDisplayName(true)))
+			{
+				StateTreeDiffType = EStateDiffType::StateMoved;
+			}
+			bSkipChildren = true;
+			break;
+		case ETreeDiffResult::DifferentValues:
+			StateTreeDiffType = EStateDiffType::StateChanged;
+			break;
+		case ETreeDiffResult::Identical:
+			StateTreeDiffType = EStateDiffType::Identical;
+			if (LeftState && RightState)
+			{
+				if (LeftState->bEnabled != RightState->bEnabled)
+				{
+					StateTreeDiffType = RightState->bEnabled ? EStateDiffType::StateEnabled : EStateDiffType::StateDisabled;
+				}
+			}
+			break;
+		default:
+			return ETreeTraverseControl::Continue;
+		}
+
+		if (StateTreeDiffType == EStateDiffType::Identical)
+		{
+			return ETreeTraverseControl::Continue;
+		}
+
+		if (StateTreeDiffType == EStateDiffType::StateMoved)
+		{
+			for (FSingleDiffEntry& DiffEntry : OutDiffEntries)
+			{
+				if (DiffEntry.Identifier.ToDisplayName(true) == StatePath.ToDisplayName(true))
+				{
+					if (DiffEntry.DiffType == EStateDiffType::StateAddedToA)
+					{
+						DiffEntry.SecondaryIdentifier = StatePath;
+					}
+					else
+					{
+						DiffEntry.SecondaryIdentifier = DiffEntry.Identifier;
+						DiffEntry.Identifier = StatePath;
+					}
+					DiffEntry.DiffType = EStateDiffType::StateMoved;
+
+					// For now, we are skipping children, we may need to revisit that
+					return ETreeTraverseControl::SkipChildren;
+				}
+			}
+		}
+
+		OutDiffEntries.Add(FSingleDiffEntry(StatePath, StateTreeDiffType));
+		
+		return bSkipChildren ? ETreeTraverseControl::SkipChildren : ETreeTraverseControl::Continue;
+	});
+}
+
+void FAsyncDiff::GetStateTreeDifferences(TArray<FSingleDiffEntry>& OutDiffEntries) const
+{
+	if (LeftView && RightView)
+	{
+		const FStateTreeViewModel* LeftViewModel = LeftView->GetViewModel().Get();
+		const FStateTreeViewModel* RightViewModel = RightView->GetViewModel().Get();
+		if (LeftViewModel && RightViewModel)
+		{
+			UStateTreeEditorData* LeftEditorData = Cast<UStateTreeEditorData>(LeftViewModel->GetStateTree()->EditorData);
+			UStateTreeEditorData* RightEditorData = Cast<UStateTreeEditorData>(RightViewModel->GetStateTree()->EditorData);
+			if (!AreStateTreePropertiesEqual(LeftEditorData, RightEditorData))
+			{
+				OutDiffEntries.Add(FSingleDiffEntry(
+					/*Identifier*/FStateSoftPath(),
+					EStateDiffType::StateTreePropertiesChanged));
+			}
+
+			GetStatesDifferences(OutDiffEntries);
+
+			GetBindingsDifferences(LeftEditorData, RightEditorData, OutDiffEntries);
+		}
+	}
+}
+
+} // UE::AsyncStateTreeDiff
+
+
+bool TTreeDiffSpecification<TWeakObjectPtr<UStateTreeState>>::AreValuesEqual(const TWeakObjectPtr<UStateTreeState>& StateTreeNodeA, const TWeakObjectPtr<UStateTreeState>& StateTreeNodeB, TArray<FPropertySoftPath>*) const
+{
+	const TStrongObjectPtr<UStateTreeState> StrongStateA = StateTreeNodeA.Pin();
+	const TStrongObjectPtr<UStateTreeState> StrongStateB = StateTreeNodeB.Pin();
+	const UStateTreeState* StateA = StrongStateA.Get();
+	const UStateTreeState* StateB = StrongStateB.Get();
+
+	if (!StateA || !StateB)
+	{
+		return StateA == StateB;
+	}
+
+	using namespace UE::StateTree::Diff;
+	return ArePropertiesEqual(StateA, StateB)
+		&& AreParametersEqual(StateA, StateB)
+		&& AreConditionsEqual(StateA, StateB)
+		&& AreTasksEqual(StateA, StateB)
+		&& AreTransitionsEqual(StateA, StateB)
+		&& AreConsiderationsEqual(StateA, StateB);
+}
+
+bool TTreeDiffSpecification<TWeakObjectPtr<UStateTreeState>>::AreMatching(const TWeakObjectPtr<UStateTreeState>& StateTreeNodeA, const TWeakObjectPtr<UStateTreeState>& StateTreeNodeB, TArray<FPropertySoftPath>*) const
+{
+	const TStrongObjectPtr<UStateTreeState> StrongStateA = StateTreeNodeA.Pin();
+	const TStrongObjectPtr<UStateTreeState> StrongStateB = StateTreeNodeB.Pin();
+	const UStateTreeState* StateA = StrongStateA.Get();
+	const UStateTreeState* StateB = StrongStateB.Get();
+	if (!StateA || !StateB)
+	{
+		return StateA == StateB;
+	}
+
+	return StateA->ID == StateB->ID;
+}
+
+void TTreeDiffSpecification<TWeakObjectPtr<UStateTreeState>>::GetChildren(const TWeakObjectPtr<UStateTreeState>& InParent, TArray<TWeakObjectPtr<UStateTreeState>>& OutChildren) const
+{
+	const TStrongObjectPtr<UStateTreeState> StrongParent = InParent.Pin();
+	if (UStateTreeState* InParentPtr = StrongParent.Get())
+	{
+		OutChildren.Reserve(InParentPtr->Children.Num());
+		for (const TObjectPtr<UStateTreeState>& Child : InParentPtr->Children)
+		{
+			OutChildren.Add(Child);
+		}
+	}
+}
\ No newline at end of file
diff --git a/Engine/Plugins/Runtime/StateTree/Source/StateTreeEditorModule/Private/AsyncStateTreeDiff.h b/Engine/Plugins/Runtime/StateTree/Source/StateTreeEditorModule/Private/AsyncStateTreeDiff.h
new file mode 100644
--- /dev/null
+++ b/Engine/Plugins/Runtime/StateTree/Source/StateTreeEditorModule/Private/AsyncStateTreeDiff.h
@@ -0,0 +1,50 @@
+// Copyright Epic Games, Inc. All Rights Reserved.
+#pragma once
+
+#include "AsyncTreeDifferences.h"
+
+class SStateTreeView;
+class UStateTreeState;
+
+namespace UE::StateTree::Diff
+{
+struct FSingleDiffEntry;
+
+class STATETREEEDITORMODULE_API FAsyncDiff : public TAsyncTreeDifferences<TWeakObjectPtr<UStateTreeState>>
+{
+public:
+	FAsyncDiff(const TSharedRef<SStateTreeView>& LeftTree, const TSharedRef<SStateTreeView>& RightTree);
+
+	void GetStateTreeDifferences(TArray<FSingleDiffEntry>& OutDiffEntries) const;
+
+private:
+	void GetStatesDifferences(TArray<FSingleDiffEntry>& OutDiffEntries) const;
+
+	static TAttribute<TArray<TWeakObjectPtr<UStateTreeState>>> RootNodesAttribute(TWeakPtr<SStateTreeView> StateTreeView);
+
+	TSharedPtr<SStateTreeView> LeftView;
+	TSharedPtr<SStateTreeView> RightView;
+};
+
+} // UE::StateTree::Diff
+
+template <>
+class STATETREEEDITORMODULE_API TTreeDiffSpecification<TWeakObjectPtr<UStateTreeState>>
+{
+public:
+	bool AreValuesEqual(const TWeakObjectPtr<UStateTreeState>& StateTreeNodeA, const TWeakObjectPtr<UStateTreeState>& StateTreeNodeB, TArray<FPropertySoftPath>* OutDifferingProperties = nullptr) const;
+
+	bool AreMatching(const TWeakObjectPtr<UStateTreeState>& StateTreeNodeA, const TWeakObjectPtr<UStateTreeState>& StateTreeNodeB, TArray<FPropertySoftPath>* OutDifferingProperties = nullptr) const;
+
+	void GetChildren(const TWeakObjectPtr<UStateTreeState>& InParent, TArray<TWeakObjectPtr<UStateTreeState>>& OutChildren) const;
+
+	bool ShouldMatchByValue(const TWeakObjectPtr<UStateTreeState>&) const
+	{
+		return false;
+	}
+
+	bool ShouldInheritEqualFromChildren(const TWeakObjectPtr<UStateTreeState>&, const TWeakObjectPtr<UStateTreeState>&) const
+	{
+		return false;
+	}
+};
\ No newline at end of file
diff --git a/Engine/Plugins/Runtime/StateTree/Source/StateTreeEditorModule/Private/SStateTreeDiff.cpp b/Engine/Plugins/Runtime/StateTree/Source/StateTreeEditorModule/Private/SStateTreeDiff.cpp
new file mode 100644
--- /dev/null
+++ b/Engine/Plugins/Runtime/StateTree/Source/StateTreeEditorModule/Private/SStateTreeDiff.cpp
@@ -0,0 +1,663 @@
+// Copyright Epic Games, Inc. All Rights Reserved.
+
+#include "SStateTreeDiff.h"
+#include "Customizations/StateTreeBindingExtension.h"
+#include "DetailsDiff.h"
+#include "DetailTreeNode.h"
+#include "Editor.h"
+#include "Framework/Application/SlateApplication.h"
+#include "Framework/MultiBox/MultiBoxBuilder.h"
+#include "IDetailsView.h"
+#include "SDetailsDiff.h"
+#include "SlateOptMacros.h"
+#include "SStateTreeSplitter.h"
+#include "StateTreeDiffControl.h"
+#include "StateTreeDiffHelper.h"
+#include "StateTreeState.h"
+#include "Widgets/Layout/SBox.h"
+#include "Widgets/Layout/SSpacer.h"
+
+#define LOCTEXT_NAMESPACE "SStateTreeDif"
+
+namespace UE::StateTree::Diff
+{
+
+SDiffWidget::SDiffWidget()
+{
+}
+
+SDiffWidget::~SDiffWidget()
+{
+	if (AssetEditorCloseHandle.IsValid())
+	{
+		GEditor->GetEditorSubsystem<UAssetEditorSubsystem>()->OnAssetEditorRequestClose().Remove(AssetEditorCloseHandle);
+	}
+}
+
+BEGIN_SLATE_FUNCTION_BUILD_OPTIMIZATION
+
+void SDiffWidget::Construct(const FArguments& InArgs)
+{
+	check(InArgs._OldAsset || InArgs._NewAsset);
+	OldAssetPanel.StateTree = TStrongObjectPtr(InArgs._OldAsset);
+	NewAssetPanel.StateTree = TStrongObjectPtr(InArgs._NewAsset);
+	OldAssetPanel.RevisionInfo = InArgs._OldRevision;
+	NewAssetPanel.RevisionInfo = InArgs._NewRevision;
+
+	// sometimes we want to clearly identify the assets being diffed (when it's
+	// not the same asset in each panel)
+	OldAssetPanel.bShowAssetName = InArgs._ShowAssetNames;
+	NewAssetPanel.bShowAssetName = InArgs._ShowAssetNames;
+
+	if (InArgs._ParentWindow.IsValid())
+	{
+		WeakParentWindow = InArgs._ParentWindow;
+
+		AssetEditorCloseHandle = GEditor->GetEditorSubsystem<UAssetEditorSubsystem>()->OnAssetEditorRequestClose().AddSP(this, &SDiffWidget::HandleAssetEditorRequestClose);
+	}
+
+	FToolBarBuilder NavToolBarBuilder(TSharedPtr<const FUICommandList>(), FMultiBoxCustomization::None);
+	NavToolBarBuilder.AddToolBarButton(
+		FUIAction(
+			FExecuteAction::CreateSP(this, &SDiffWidget::PrevDiff),
+			FCanExecuteAction::CreateSP(this, &SDiffWidget::HasPrevDiff)),
+		NAME_None, LOCTEXT("PrevDiffLabel", "Prev"), LOCTEXT("PrevDiffTooltip", "Go to previous difference"), FSlateIcon(FAppStyle::GetAppStyleSetName(), "BlueprintDif.PrevDiff"));
+	NavToolBarBuilder.AddToolBarButton(
+		FUIAction(
+			FExecuteAction::CreateSP(this, &SDiffWidget::NextDiff),
+			FCanExecuteAction::CreateSP(this, &SDiffWidget::HasNextDiff)),
+		NAME_None, LOCTEXT("NextDiffLabel", "Next"), LOCTEXT("NextDiffTooltip", "Go to next difference"), FSlateIcon(FAppStyle::GetAppStyleSetName(), "BlueprintDif.NextDiff"));
+
+	DifferencesTreeView = DiffTreeView::CreateTreeView(&Differences);
+
+	GenerateDifferencesList();
+
+	const auto TextBlock = [](const FText Text) -> TSharedRef<SWidget>
+	{
+		return SNew(SBox)
+			.Padding(FMargin(4.0f, 10.0f))
+			.VAlign(VAlign_Center)
+			.HAlign(HAlign_Left)
+			[
+				SNew(STextBlock)
+				.Visibility(EVisibility::HitTestInvisible)
+				.TextStyle(FAppStyle::Get(), "DetailsView.CategoryTextStyle")
+				.Text(Text)
+			];
+	};
+
+	TopRevisionInfoWidget =
+		SNew(SSplitter)
+		.Visibility(EVisibility::HitTestInvisible)
+		+ SSplitter::Slot()
+		.Value(.2f)
+		[
+			SNew(SBox)
+		]
+		+ SSplitter::Slot()
+		.Value(.8f)
+		[
+			SNew(SSplitter)
+			.PhysicalSplitterHandleSize(10.0f)
+			+ SSplitter::Slot()
+			.Value(.5f)
+			[
+				TextBlock(DiffViewUtils::GetPanelLabel(OldAssetPanel.StateTree.Get(), OldAssetPanel.RevisionInfo, FText()))
+			]
+			+ SSplitter::Slot()
+			.Value(.5f)
+			[
+				TextBlock(DiffViewUtils::GetPanelLabel(NewAssetPanel.StateTree.Get(), NewAssetPanel.RevisionInfo, FText()))
+			]
+		];
+
+	this->ChildSlot
+		[
+			SNew(SBorder)
+			.BorderImage(FAppStyle::GetBrush("Docking.Tab", ".ContentAreaBrush"))
+			[
+				SNew(SOverlay)
+				+ SOverlay::Slot()
+				.VAlign(VAlign_Top)
+				[
+					TopRevisionInfoWidget.ToSharedRef()
+				]
+				+ SOverlay::Slot()
+				[
+					SNew(SSplitter)
+					.Orientation(EOrientation::Orient_Vertical)
+					+ SSplitter::Slot()
+					.Value(.55f)
+					[
+						SNew(SVerticalBox)
+						+ SVerticalBox::Slot()
+						.AutoHeight()
+						.Padding(0.0f, 2.0f, 0.0f, 2.0f)
+						[
+							SNew(SHorizontalBox)
+							+ SHorizontalBox::Slot()
+							.Padding(4.f)
+							.AutoWidth()
+							[
+								NavToolBarBuilder.MakeWidget()
+							]
+							+ SHorizontalBox::Slot()
+							[
+								SNew(SSpacer)
+							]
+						]
+						+ SVerticalBox::Slot()
+						[
+							SNew(SSplitter)
+							+ SSplitter::Slot()
+							.Value(.2f)
+							[
+								SNew(SBorder)
+								.BorderImage(FAppStyle::GetBrush("ToolPanel.GroupBorder"))
+								[
+									DifferencesTreeView.ToSharedRef()
+								]
+							]
+							+ SSplitter::Slot()
+							.Value(.8f)
+							[
+								StateTreePanel.Splitter.ToSharedRef()
+							]
+						]
+					]
+					+ SSplitter::Slot()
+					.Value(.45f)
+					[
+						SAssignNew(DetailsViewContents, SBox)
+					]
+				]
+			]
+		];
+	SetDetailsDiff();
+}
+
+END_SLATE_FUNCTION_BUILD_OPTIMIZATION
+
+void SDiffWidget::HandleAssetEditorRequestClose(UObject* Asset, const EAssetEditorCloseReason CloseReason)
+{
+	if (OldAssetPanel.StateTree.Get()== Asset || NewAssetPanel.StateTree.Get() == Asset || CloseReason == EAssetEditorCloseReason::CloseAllAssetEditors)
+	{
+		// Tell our window to close and set our selves to collapsed to try and stop it from ticking
+		SetVisibility(EVisibility::Collapsed);
+
+		if (AssetEditorCloseHandle.IsValid())
+		{
+			GEditor->GetEditorSubsystem<UAssetEditorSubsystem>()->OnAssetEditorRequestClose().Remove(AssetEditorCloseHandle);
+		}
+
+		if (const TSharedPtr<SWindow> ParentWindow = WeakParentWindow.Pin())
+		{
+			ParentWindow->RequestDestroyWindow();
+		}
+	}
+}
+
+TSharedRef<SDiffWidget> SDiffWidget::CreateDiffWindow(const FText WindowTitle, gsl::not_null<const UStateTree*> OldStateTree, gsl::not_null<const UStateTree*> NewStateTree, const FRevisionInfo& OldRevision, const FRevisionInfo& NewRevision)
+{
+	// sometimes we're comparing different revisions of one single asset (other
+	// times we're comparing two completely separate assets altogether)
+	const bool bIsSingleAsset = (NewStateTree->GetName() == OldStateTree->GetName());
+
+	const TSharedPtr<SWindow> Window = SNew(SWindow)
+		.Title(WindowTitle)
+		.ClientSize(FVector2D(1000.f, 800.f));
+
+	TSharedRef<SDiffWidget> StateTreeDiff = SNew(SDiffWidget)
+		.OldAsset(OldStateTree)
+		.NewAsset(NewStateTree)
+		.OldRevision(OldRevision)
+		.NewRevision(NewRevision)
+		.ShowAssetNames(!bIsSingleAsset)
+		.ParentWindow(Window);
+
+	Window->SetContent(StateTreeDiff);
+
+	// Make this window a child of the modal window if we've been spawned while one is active.
+	const TSharedPtr<SWindow> ActiveModal = FSlateApplication::Get().GetActiveTopLevelWindow();
+	if (ActiveModal.IsValid())
+	{
+		FSlateApplication::Get().AddWindowAsNativeChild(Window.ToSharedRef(), ActiveModal.ToSharedRef());
+	}
+	else
+	{
+		FSlateApplication::Get().AddWindow(Window.ToSharedRef());
+	}
+
+	return StateTreeDiff;
+}
+
+TSharedRef<SDiffWidget> SDiffWidget::CreateDiffWindow(gsl::not_null<const UStateTree*> OldStateTree, gsl::not_null<const UStateTree*> NewStateTree,
+	const FRevisionInfo& OldRevision, const FRevisionInfo& NewRevision, const UClass* StateTreeClass)
+{
+	check(OldStateTree || NewStateTree);
+
+	// sometimes we're comparing different revisions of one single asset (other
+	// times we're comparing two completely separate assets altogether)
+	const bool bIsSingleAsset = !OldStateTree || !NewStateTree || (NewStateTree->GetName() == OldStateTree->GetName());
+
+	FText WindowTitle = FText::Format(LOCTEXT("NamelesSDiffWidget", "{0} Diff"), StateTreeClass->GetDisplayNameText());
+	// if we're diffing one asset against itself
+	if (bIsSingleAsset)
+	{
+		// identify the assumed single asset in the window's title
+		const FString STName = NewStateTree ? NewStateTree->GetName() : OldStateTree->GetName();
+		WindowTitle = FText::Format(LOCTEXT("NamedStateTreeDiff", "{0} - {1} Diff"), FText::FromString(STName), StateTreeClass->GetDisplayNameText());
+	}
+
+	return CreateDiffWindow(WindowTitle, OldStateTree, NewStateTree, OldRevision, NewRevision);
+}
+
+void SDiffWidget::NextDiff() const
+{
+	DiffTreeView::HighlightNextDifference(DifferencesTreeView.ToSharedRef(), Differences, Differences);
+}
+
+void SDiffWidget::PrevDiff() const
+{
+	DiffTreeView::HighlightPrevDifference(DifferencesTreeView.ToSharedRef(), Differences, Differences);
+}
+
+bool SDiffWidget::HasNextDiff() const
+{
+	return DiffTreeView::HasNextDifference(DifferencesTreeView.ToSharedRef(), Differences);
+}
+
+bool SDiffWidget::HasPrevDiff() const
+{
+	return DiffTreeView::HasPrevDifference(DifferencesTreeView.ToSharedRef(), Differences);
+}
+
+void SDiffWidget::GenerateDifferencesList()
+{
+	Differences.Empty();
+
+	GenerateDiffPanel();
+
+	DifferencesTreeView->RebuildList();
+}
+
+void SDiffWidget::GenerateDiffPanel()
+{
+	const UStateTree* OldStateTree = OldAssetPanel.StateTree.Get();
+	const UStateTree* NewStateTree = NewAssetPanel.StateTree.Get();
+	StateTreePanel.DiffControl = MakeShared<FDiffControl>(
+		OldStateTree,
+		NewStateTree,
+		FOnDiffEntryFocused{});
+	StateTreePanel.DiffControl->GenerateTreeEntries(Differences);
+	StateTreePanel.DiffControl->GetOnStateDiffEntryFocused().AddSP(this, &SDiffWidget::HandleStateDiffEntryFocused);
+
+	TSharedPtr<SDiffSplitter> DiffSplitter = SNew(SDiffSplitter);
+	if (OldAssetPanel.StateTree)
+	{
+		DiffSplitter->AddSlot(
+			SDiffSplitter::Slot()
+			.Value(0.5f)
+			.StateTreeView(StateTreePanel.DiffControl->GetDetailsWidget(OldStateTree))
+			.StateTree(OldStateTree));
+	}
+	if (NewAssetPanel.StateTree)
+	{
+		DiffSplitter->AddSlot(
+			SDiffSplitter::Slot()
+			.Value(0.5f)
+			.StateTreeView(StateTreePanel.DiffControl->GetDetailsWidget(NewStateTree))
+			.StateTree(NewStateTree));
+	}
+	StateTreePanel.Splitter = DiffSplitter;
+}
+
+void SDiffWidget::HandleStateDiffEntryFocused(const FSingleDiffEntry& StateDiff)
+{
+	const FStateSoftPath LeftStatePath = StateDiff.Identifier;
+	const FStateSoftPath RightStatePath = StateDiff.SecondaryIdentifier ? StateDiff.SecondaryIdentifier : StateDiff.Identifier;
+	StateTreePanel.Splitter->HandleSelectionChanged(LeftStatePath, RightStatePath);
+
+	const UStateTree* OldStateTree = OldAssetPanel.StateTree.Get();
+	const UStateTree* NewStateTree = NewAssetPanel.StateTree.Get();
+	const UStateTreeState* OldState = OldStateTree != nullptr ? LeftStatePath.ResolvePath(OldStateTree) : nullptr;
+	const UStateTreeState* NewState = NewStateTree != nullptr ? RightStatePath.ResolvePath(NewStateTree) : nullptr;
+
+	// If comparing states that exist in both state trees display them in the details diff view
+	if (OldState && NewState)
+	{
+		SetDetailsDiff(OldState, NewState);
+
+	}
+	// If we clear selection on both state trees we can display an empty details diff view
+	else if (!OldState && !NewState)
+	{
+		SetDetailsDiff();
+	}
+	// If the state only exists in one of the state trees (either added or removed), details diff view will not work.
+	else
+	{
+		// So the states are put into separate details views
+		const TSharedPtr<SBox> LeftWidget = SNew(SBox);
+		const TSharedPtr<SBox> RightWidget = SNew(SBox);
+		if (OldState)
+		{
+			const FDetailsDiff DetailsDiff(OldState, FDetailsDiff::FOnDisplayedPropertiesChanged(), true);
+			LeftWidget->SetContent(DetailsDiff.DetailsWidget());
+		}
+		if (NewState)
+		{
+			const FDetailsDiff DetailsDiff(NewState, FDetailsDiff::FOnDisplayedPropertiesChanged(), false);
+			RightWidget->SetContent(DetailsDiff.DetailsWidget());
+		}
+		// And displayed in a way that resembles the details diff view
+		DetailsViewContents->SetContent(
+			SNew(SBorder)
+			.BorderImage(FAppStyle::GetBrush("Docking.Tab", ".ContentAreaBrush"))
+			[
+				SNew(SVerticalBox)
+				+ SVerticalBox::Slot()
+				.AutoHeight()
+				.Padding(0.0f, 2.0f, 0.0f, 2.0f)
+				+ SVerticalBox::Slot()
+				[
+					SNew(SSplitter)
+					+ SSplitter::Slot()
+					.Value(.2f)
+					+ SSplitter::Slot()
+					.Value(.8f)
+					[
+						SNew(SSplitter)
+						.PhysicalSplitterHandleSize(5.f)
+						+ SSplitter::Slot()
+						.Value(.5f)
+						[
+							SNew(SBox)
+							.Padding(15.f, 0.f, 15.f, 0.f)
+							[
+								LeftWidget.ToSharedRef()
+							]
+						]
+						+ SSplitter::Slot()
+						.Value(.5f)
+						[
+							SNew(SBox)
+							.Padding(15.f, 0.f, 15.f, 0.f)
+							[
+								RightWidget.ToSharedRef()
+							]
+						]
+					]
+				]
+			]
+
+		);
+	}
+}
+
+void SDiffWidget::SetDetailsDiff(const UStateTreeState* OldState, const UStateTreeState* NewState)
+{
+	const UObject* OldAsset = OldState ? OldState : OldAssetPanel.StateTree ? static_cast<const UObject*>(OldAssetPanel.StateTree->EditorData) : nullptr;
+	const UObject* NewAsset = NewState ? NewState : NewAssetPanel.StateTree ? static_cast<const UObject*>(NewAssetPanel.StateTree->EditorData) : nullptr;
+
+	if (const bool bIsState = OldState || NewState)
+	{
+		StateBindingDiffs.Reset(StateTreePanel.DiffControl->GetBindingDifferences().Num());
+		for (const FSingleDiffEntry& BindingDiff : StateTreePanel.DiffControl->GetBindingDifferences())
+		{
+			const UStateTree* OldStateTree = OldAssetPanel.StateTree.Get();
+			const UStateTree* NewStateTree = NewAssetPanel.StateTree.Get();
+			if (OldStateTree != nullptr
+				&& NewStateTree != nullptr
+				&& BindingDiff.Identifier.ResolvePath(OldStateTree) == OldState
+				&& BindingDiff.SecondaryIdentifier.ResolvePath(NewStateTree) == NewState)
+			{
+				StateBindingDiffs.Push(BindingDiff);
+			}
+		}
+	}
+	else
+	{
+		StateBindingDiffs.Reset();
+	}
+
+	TArray<FSingleObjectDiffEntry> Entries;
+	const TSharedRef<SDetailsDiff> DetailsDiff = SNew(SDetailsDiff)
+		.OldAsset(OldAsset)
+		.NewAsset(NewAsset)
+		.OldRevision(OldAssetPanel.RevisionInfo)
+		.NewRevision(NewAssetPanel.RevisionInfo)
+		.ShowAssetNames(false)
+		.OnCustomizeDetailsWidget_Static(&SDiffWidget::AddStateTreeExtensionToDetailsView)
+		.OnGenerateCustomDiffEntries(this, &SDiffWidget::AddBindingDiffToDiffEntries)
+		.OnOrganizeDiffEntries_Static(&SDiffWidget::OrganizeDiffEntries, OldState, NewState)
+		.OnGenerateCustomDiffEntryWidget_Static(&SDiffWidget::GenerateCustomDiffEntryWidget, OldState, NewState)
+		.RowHighlightColor_Static(&SDiffWidget::GetRowHighlightColor)
+		.ShouldHighlightRow(this, &SDiffWidget::ShouldHighlightRow);
+
+	DetailsViewContents->SetContent(DetailsDiff);
+}
+
+void SDiffWidget::AddBindingDiffToDiffEntries(TArray<FSingleObjectDiffEntry>& OutEntries)
+{
+	OutEntries.Reserve(StateBindingDiffs.Num());
+
+	for (const FSingleDiffEntry& BindingDiff : StateBindingDiffs)
+	{
+		EPropertyDiffType::Type DiffType = EPropertyDiffType::Invalid;
+		switch (BindingDiff.DiffType)
+		{
+		case EStateDiffType::BindingAddedToA:
+		case EStateDiffType::BindingAddedToB:
+		case EStateDiffType::BindingChanged:
+			DiffType = EPropertyDiffType::PropertyValueChanged;
+			break;
+		}
+
+		if (DiffType != EPropertyDiffType::Invalid)
+		{
+			FSingleObjectDiffEntry Entry(BindingDiff.BindingPath, DiffType);
+			OutEntries.Add(Entry);
+		}
+	}
+}
+
+void SDiffWidget::OrganizeDiffEntries(
+	TArray<TSharedPtr<FBlueprintDifferenceTreeEntry>>& OutDiffTreeEntries,
+	const TArray<FSingleObjectDiffEntry>& DiffEntries,
+	TFunctionRef<TSharedPtr<FBlueprintDifferenceTreeEntry>(const FSingleObjectDiffEntry&)> GenerateDiffTreeEntry,
+	TFunctionRef<TSharedPtr<FBlueprintDifferenceTreeEntry>(FText&)> GenerateCategoryEntry, const UStateTreeState* OldState, const UStateTreeState* NewState)
+{
+	static FText RightRevision = LOCTEXT("NewRevisionIdentifier", "Right Revision");
+	static FText StateText = LOCTEXT("StateText", "State");
+	static FText ParameterText = LOCTEXT("ParametersText", "Parameters");
+	static FText ConditionText = LOCTEXT("EnterConditionsText", "Enter Conditions");
+	static FText TaskText = LOCTEXT("TasksText", "Tasks");
+	static FText TransitionText = LOCTEXT("TransitionsText", "Transitions");
+	static FText ConsiderationText = LOCTEXT("ConsiderationText", "Utility");
+
+	TSet<int32> ConditionIndices;
+	TSet<int32> TaskIndices;
+	TSet<int32> TransitionIndices;
+	TSet<int32> ConsiderationIndices;
+	TArray<TSharedPtr<FBlueprintDifferenceTreeEntry>> ParametersEntries;
+	TArray<TSharedPtr<FBlueprintDifferenceTreeEntry>> ConditionEntries;
+	TArray<TSharedPtr<FBlueprintDifferenceTreeEntry>> TaskEntries;
+	TArray<TSharedPtr<FBlueprintDifferenceTreeEntry>> TransitionEntries;
+	TArray<TSharedPtr<FBlueprintDifferenceTreeEntry>> ConsiderationEntries;
+	TArray<TSharedPtr<FBlueprintDifferenceTreeEntry>> StateEntries;
+
+	TArray<FSingleObjectDiffEntry> SortedEntries = DiffEntries;
+	SortedEntries.Sort([](const FSingleObjectDiffEntry& A, const FSingleObjectDiffEntry& B)
+	{
+		return A.Identifier.TryReadIndex(0) < B.Identifier.TryReadIndex(0);
+	});
+
+	for (const FSingleObjectDiffEntry& Difference : SortedEntries)
+	{
+		constexpr int32 PropertyCountFromRoot = 2; // 2 levels down from the root; first level being the category/parent collection and the second level the property that changed
+		FSingleObjectDiffEntry SimplifiedEntry(Difference.Identifier.GetRootProperty(PropertyCountFromRoot), Difference.DiffType);
+		TSharedPtr<FBlueprintDifferenceTreeEntry> Entry = GenerateDiffTreeEntry(SimplifiedEntry);
+		if (Difference.Identifier.IsSubPropertyMatch(ConditionName))
+		{
+			const int32 Index = Difference.Identifier.TryReadIndex(0);
+			if (!ConditionIndices.Contains(Index))
+			{
+				ConditionIndices.Add(Index);
+				ConditionEntries.Add(Entry);
+			}
+		}
+		else if (Difference.Identifier.IsSubPropertyMatch(TaskName))
+		{
+			const int32 Index = Difference.Identifier.TryReadIndex(0);
+			if (!TaskIndices.Contains(Index))
+			{
+				TaskIndices.Add(Index);
+				TaskEntries.Add(Entry);
+			}
+		}
+		else if (Difference.Identifier.IsSubPropertyMatch(TransitionName))
+		{
+			const int32 Index = Difference.Identifier.TryReadIndex(0);
+			if (!TransitionIndices.Contains(Index))
+			{
+				TransitionIndices.Add(Index);
+				TransitionEntries.Add(Entry);
+			}
+		}
+		else if (Difference.Identifier.IsSubPropertyMatch(ConsiderationName))
+		{
+			const int32 Index = Difference.Identifier.TryReadIndex(0);
+			if (!ConsiderationIndices.Contains(Index))
+			{
+				ConsiderationIndices.Add(Index);
+				ConsiderationEntries.Add(Entry);
+			}
+		}
+		else if (Difference.Identifier.IsSubPropertyMatch(ParameterName))
+		{
+			// @todo investigate: currently unable to resolve full property path (FInstancedPropertyBag issue?)
+			ParametersEntries.Add(GenerateDiffTreeEntry(Difference));
+		}
+		else
+		{
+			StateEntries.Add(Entry);
+		}
+	}
+
+	OutDiffTreeEntries.Append(StateEntries);
+
+	if (ParametersEntries.Num() > 0)
+	{
+		TSharedPtr<FBlueprintDifferenceTreeEntry> ParametersEntry = GenerateCategoryEntry(ParameterText);
+		ParametersEntry->Children = ParametersEntries;
+		OutDiffTreeEntries.Push(ParametersEntry);
+	}
+
+	if (ConditionEntries.Num() > 0)
+	{
+		TSharedPtr<FBlueprintDifferenceTreeEntry> ConditionEntry = GenerateCategoryEntry(ConditionText);
+		ConditionEntry->Children = ConditionEntries;
+		OutDiffTreeEntries.Push(ConditionEntry);
+	}
+
+	if (ConsiderationEntries.Num() > 0)
+	{
+		TSharedPtr<FBlueprintDifferenceTreeEntry> ConsiderationEntry = GenerateCategoryEntry(ConsiderationText);
+		ConsiderationEntry->Children = ConsiderationEntries;
+		OutDiffTreeEntries.Push(ConsiderationEntry);
+	}
+
+	if (TaskEntries.Num() > 0)
+	{
+		TSharedPtr<FBlueprintDifferenceTreeEntry> TaskEntry = GenerateCategoryEntry(TaskText);
+		TaskEntry->Children = TaskEntries;
+		OutDiffTreeEntries.Push(TaskEntry);
+	}
+
+	if (TransitionEntries.Num() > 0)
+	{
+		TSharedPtr<FBlueprintDifferenceTreeEntry> TransitionEntry = GenerateCategoryEntry(TransitionText);
+		TransitionEntry->Children = TransitionEntries;
+		OutDiffTreeEntries.Push(TransitionEntry);
+	}
+}
+
+TSharedRef<SWidget> SDiffWidget::GenerateCustomDiffEntryWidget(const FSingleObjectDiffEntry& DiffEntry, FText&, const UStateTreeState* OldState, const UStateTreeState* NewState)
+{
+	const UStateTreeState* SourceState = DiffEntry.DiffType == EPropertyDiffType::PropertyAddedToB ? NewState : OldState;
+	FText PropertyName = FText::FromString(DiffEntry.Identifier.ToDisplayName());
+	if (DiffEntry.Identifier.IsSubPropertyMatch(ConditionName))
+	{
+		const int32 ConditionIndex = DiffEntry.Identifier.TryReadIndex(0);
+		const FStateTreeEditorNode& ConditionEntry = SourceState->EnterConditions[ConditionIndex];
+		PropertyName = FText::Format(FText::FromString(TEXT("[{0}]")), FText::FromName(ConditionEntry.GetName()));
+	}
+	else if (DiffEntry.Identifier.IsSubPropertyMatch(TaskName))
+	{
+		const int32 TaskIndex = DiffEntry.Identifier.TryReadIndex(0);
+		const FStateTreeEditorNode& TaskEntry = SourceState->Tasks[TaskIndex];
+		PropertyName = FText::Format(FText::FromString(TEXT("[{0}]")), FText::FromName(TaskEntry.GetName()));
+	}
+	else if (DiffEntry.Identifier.IsSubPropertyMatch(ConsiderationName))
+	{
+		const int32 ConsiderationIndex = DiffEntry.Identifier.TryReadIndex(0);
+		const FStateTreeEditorNode& ConsiderationEntry = SourceState->Considerations[ConsiderationIndex];
+		PropertyName = FText::Format(FText::FromString(TEXT("[{0}]")), FText::FromName(ConsiderationEntry.GetName()));
+	}
+	else if (DiffEntry.Identifier.IsSubPropertyMatch(ParameterName))
+	{
+		constexpr int32 NumberOfPathElements = 1;
+		PropertyName = FText::Format(FText::FromString(TEXT("[{0}]")), FText::FromString(DiffEntry.Identifier.ToDisplayName(NumberOfPathElements)));
+	}
+
+	return SNew(STextBlock)
+		.Text(GetStateDiffMessage(DiffEntry, PropertyName))
+		.ToolTipText(GetStateDiffMessage(DiffEntry, PropertyName))
+		.ColorAndOpacity(GetStateDiffMessageColor(DiffEntry));
+}
+
+bool SDiffWidget::ShouldHighlightRow(const TUniquePtr<FAsyncDetailViewDiff::DiffNodeType>& DiffNode)
+{
+	if (DiffNode->DiffResult != ETreeDiffResult::Identical)
+	{
+		return true;
+	}
+
+	TSharedPtr<FDetailTreeNode> DetailNode = DiffNode->ValueA.Pin();
+	DetailNode = DetailNode ? DetailNode : DiffNode->ValueB.Pin();
+	const FPropertySoftPath PropertySoftPath(DetailNode->GetPropertyPath());
+	if (PropertySoftPath.ToDisplayName().Len() == 0)
+	{
+		return false;
+	}
+
+	for (const FSingleDiffEntry& BindingDiff : StateBindingDiffs)
+	{
+		return BindingDiff.BindingPath.IsSubPropertyMatch(PropertySoftPath) || BindingDiff.BindingPath == PropertySoftPath;
+	}
+
+	return false;
+}
+
+FLinearColor SDiffWidget::GetRowHighlightColor(const TUniquePtr<FAsyncDetailViewDiff::DiffNodeType>& DiffNode)
+{
+	switch (DiffNode->DiffResult)
+	{
+	case ETreeDiffResult::MissingFromTree1:
+		return FLinearColor(0.f, 1.f, 0.f, .7f);
+	case ETreeDiffResult::MissingFromTree2:
+		return FLinearColor(1.f, 0.f, 0.f, .7f);
+	default:
+		return FLinearColor(1.f, 1.f, 0.f, .7f);
+	}
+}
+
+void SDiffWidget::AddStateTreeExtensionToDetailsView(const TSharedRef<IDetailsView>& DetailsView)
+{
+	DetailsView->SetExtensionHandler(MakeShared<FStateTreeBindingExtension>().ToSharedPtr());
+}
+
+} // UE::StateTree::Diff
+
+#undef LOCTEXT_NAMESPACE
\ No newline at end of file
diff --git a/Engine/Plugins/Runtime/StateTree/Source/StateTreeEditorModule/Private/SStateTreeDiff.h b/Engine/Plugins/Runtime/StateTree/Source/StateTreeEditorModule/Private/SStateTreeDiff.h
new file mode 100644
--- /dev/null
+++ b/Engine/Plugins/Runtime/StateTree/Source/StateTreeEditorModule/Private/SStateTreeDiff.h
@@ -0,0 +1,143 @@
+// Copyright Epic Games, Inc. All Rights Reserved.
+#pragma once
+
+#include "StateTree.h"
+#include "AsyncDetailViewDiff.h"
+#include "IAssetTypeActions.h"
+#include "Subsystems/AssetEditorSubsystem.h"
+#include "Widgets/SCompoundWidget.h"
+
+class FSpawnTabArgs;
+class FTabManager;
+class FUICommandList;
+
+namespace UE::StateTree::Diff
+{
+	class SDiffSplitter;
+	class FDiffControl;
+	struct FSingleDiffEntry;
+	struct FStateSoftPath;
+
+
+/** Panel used to display the state tree */
+struct STATETREEEDITORMODULE_API FDiffPanel
+{
+	/** The asset that owns the state tree view we are showing */
+	TStrongObjectPtr<const UStateTree> StateTree = nullptr;
+
+	/** Revision information for this asset */
+	FRevisionInfo RevisionInfo;
+
+	/** True if we should show a name identifying which asset this panel is displaying */
+	bool bShowAssetName = true;
+
+	/** The widget that contains the revision info in graph mode */
+	TSharedPtr<SWidget> OverlayRevisionInfo;
+
+private:
+	/** Command list for this diff panel */
+	TSharedPtr<FUICommandList> GraphEditorCommands;
+};
+
+/** Visual Diff between two StateTree assets */
+class STATETREEEDITORMODULE_API SDiffWidget : public SCompoundWidget
+{
+public:
+	SLATE_BEGIN_ARGS(SDiffWidget) {}
+		SLATE_ARGUMENT(const UStateTree*, OldAsset)
+		SLATE_ARGUMENT(const UStateTree*, NewAsset)
+		SLATE_ARGUMENT(struct FRevisionInfo, OldRevision)
+		SLATE_ARGUMENT(struct FRevisionInfo, NewRevision)
+		SLATE_ARGUMENT(bool, ShowAssetNames)
+		SLATE_ARGUMENT(TSharedPtr<SWindow>, ParentWindow)
+	SLATE_END_ARGS()
+
+	SDiffWidget();
+	void Construct(const FArguments& InArgs);
+	virtual ~SDiffWidget();
+
+	/** Helper function to create a window that holds a diff widget */
+	static TSharedRef<SDiffWidget> CreateDiffWindow(FText WindowTitle, gsl::not_null<const UStateTree*> OldStateTree, gsl::not_null<const UStateTree*> NewStateTree, const FRevisionInfo& OldRevision, const FRevisionInfo& NewRevision);
+
+	/** Helper function to create a window that holds a diff widget (default window title) */
+	static TSharedRef<SDiffWidget> CreateDiffWindow(gsl::not_null<const UStateTree*> OldStateTree, gsl::not_null<const UStateTree*> NewStateTree, const FRevisionInfo& OldRevision, const FRevisionInfo& NewRevision, const UClass* StateTreeClass);
+
+protected:
+	/** Called when user clicks button to go to next difference */
+	void NextDiff() const;
+
+	/** Called when user clicks button to go to prev difference */
+	void PrevDiff() const;
+
+	/** Called to determine whether we have a list of differences to cycle through */
+	bool HasNextDiff() const;
+	bool HasPrevDiff() const;
+
+	/** Function used to generate the list of differences and the widgets needed to calculate that list */
+	void GenerateDifferencesList();
+
+	/** Called when editor may need to be closed */
+	void HandleAssetEditorRequestClose(UObject* Asset, EAssetEditorCloseReason CloseReason);
+
+	struct FStateTreePanel
+	{
+		TSharedPtr<SDiffSplitter> Splitter;
+		TSharedPtr<FDiffControl> DiffControl;
+	};
+
+	void GenerateDiffPanel();
+
+	void HandleStateDiffEntryFocused(const FSingleDiffEntry& StateDiff);
+	
+	void SetDetailsDiff(const UStateTreeState* OldState = nullptr, const UStateTreeState* NewState = nullptr);
+
+	static void AddStateTreeExtensionToDetailsView(const TSharedRef<IDetailsView>& DetailsView);
+
+	void AddBindingDiffToDiffEntries(TArray<FSingleObjectDiffEntry>& OutEntries);
+
+	static TSharedRef<SWidget> GenerateCustomDiffEntryWidget(const FSingleObjectDiffEntry& DiffEntry, FText& ObjectName, const UStateTreeState* OldState, const UStateTreeState* NewState);
+
+	static void OrganizeDiffEntries(TArray<TSharedPtr<FBlueprintDifferenceTreeEntry>>& OutDiffTreeEntries, const TArray<FSingleObjectDiffEntry>& DiffEntries,
+		TFunctionRef<TSharedPtr<FBlueprintDifferenceTreeEntry>(const FSingleObjectDiffEntry&)> GenerateDiffTreeEntry,
+		TFunctionRef<TSharedPtr<FBlueprintDifferenceTreeEntry>(FText&)> GenerateCategoryEntry,
+		const UStateTreeState* OldState, const UStateTreeState* NewState);
+
+	bool ShouldHighlightRow(const TUniquePtr<FAsyncDetailViewDiff::DiffNodeType>& DiffNode);
+
+	static FLinearColor GetRowHighlightColor(const TUniquePtr<FAsyncDetailViewDiff::DiffNodeType>& DiffNode);
+
+	/** The panel used to show the old revision */
+	FDiffPanel OldAssetPanel;
+
+	/** The panel used to show the new revision */
+	FDiffPanel NewAssetPanel;
+
+	TSharedPtr<SBox> DetailsViewContents;
+
+	TSharedPtr<SSplitter> TopRevisionInfoWidget;
+
+	TArray<FSingleDiffEntry> StateBindingDiffs;
+
+	/** List of all differences, cached so that we can iterate only the differences and not labels, etc. */
+	TArray<TSharedPtr<class FBlueprintDifferenceTreeEntry>> Differences;
+
+	/** Tree view that displays the differences, cached for the buttons that iterate the differences: */
+	TSharedPtr<STreeView<TSharedPtr<FBlueprintDifferenceTreeEntry>>> DifferencesTreeView;
+
+	/** Stored reference to widget used to display the StateTree */
+	FStateTreePanel StateTreePanel;
+
+	/** A pointer to the window holding this */
+	TWeakPtr<SWindow> WeakParentWindow;
+
+	FDelegateHandle AssetEditorCloseHandle;
+
+private:
+	static inline const FLazyName ConditionName = FLazyName("EnterConditions");
+	static inline const FLazyName TaskName = FLazyName("Tasks");
+	static inline const FLazyName TransitionName = FLazyName("Transitions");
+	static inline const FLazyName ConsiderationName = FLazyName("Considerations");
+	static inline const FLazyName ParameterName = FLazyName("Parameters");
+};
+
+} // UE::StateTree::Diff
diff --git a/Engine/Plugins/Runtime/StateTree/Source/StateTreeEditorModule/Private/SStateTreeSplitter.cpp b/Engine/Plugins/Runtime/StateTree/Source/StateTreeEditorModule/Private/SStateTreeSplitter.cpp
new file mode 100644
--- /dev/null
+++ b/Engine/Plugins/Runtime/StateTree/Source/StateTreeEditorModule/Private/SStateTreeSplitter.cpp
@@ -0,0 +1,99 @@
+// Copyright Epic Games, Inc. All Rights Reserved.
+#include "SStateTreeSplitter.h"
+#include "SStateTreeView.h"
+#include "StateTree.h"
+#include "StateTreeDiffHelper.h"
+#include "StateTreeEditorData.h"
+#include "StateTreeViewModel.h"
+
+namespace UE::StateTree::Diff
+{
+void SDiffSplitter::Construct(const FArguments& InArgs)
+{
+	Splitter = SNew(SSplitter).PhysicalSplitterHandleSize(5.f).Orientation(EOrientation::Orient_Horizontal);
+
+	for (const FSlot::FSlotArguments& SlotArgs : InArgs._Slots)
+	{
+		AddSlot(SlotArgs);
+	}
+
+	ChildSlot
+		[
+			Splitter.ToSharedRef()
+		];
+}
+
+void SDiffSplitter::AddSlot(const FSlot::FSlotArguments& SlotArgs, int32 Index)
+{
+	if (Index == INDEX_NONE)
+	{
+		Index = Panels.Num();
+	}
+
+	Splitter->AddSlot(Index)
+		.Value(SlotArgs._Value)
+		[
+			SNew(SBox).Padding(15.f, 0.f, 15.f, 0.f)
+			[
+				SlotArgs._StateTreeView.ToSharedRef()
+			]
+		];
+
+	Panels.Insert(
+		{
+			SlotArgs._StateTreeView,
+			SlotArgs._StateTree,
+			SlotArgs._IsReadonly,
+			SlotArgs._DifferencesWithRightPanel,
+		},
+		Index);
+	if (SlotArgs._StateTreeView)
+	{
+		SlotArgs._StateTreeView->GetViewModel()->GetOnSelectionChanged().AddSP(this, &SDiffSplitter::HandleSelectionChanged);
+	}
+}
+
+void SDiffSplitter::HandleSelectionChanged(const FStateSoftPath& StatePath, const FStateSoftPath& SecondaryStatePath)
+{
+	if (StatePath != SelectedState)
+	{
+		SelectedState = StatePath;
+
+		for (const FPanel& Panel : Panels)
+		{
+			const FStateTreeViewModel* ViewModel = Panel.StateTreeView->GetViewModel().Get();
+			const UStateTreeEditorData* EditorData = Cast<UStateTreeEditorData>(Panel.StateTree->EditorData);
+			if (EditorData != nullptr && ViewModel != nullptr)
+			{
+				UStateTreeState* PanelState = SelectedState.ResolvePath(EditorData);
+				if (!PanelState)
+				{
+					PanelState = SecondaryStatePath.ResolvePath(EditorData);
+				}
+
+				TArray<UStateTreeState*> CurSelectedStates;
+				ViewModel->GetSelectedStates(CurSelectedStates);
+				if (CurSelectedStates.Num() != 1 || CurSelectedStates[0] != PanelState)
+				{
+					Panel.StateTreeView->SetSelection({PanelState});
+				}
+			}
+		}
+	}
+}
+
+void SDiffSplitter::HandleSelectionChanged(const TArray<TWeakObjectPtr<UStateTreeState>>& SelectedStates)
+{
+	if (const UStateTreeState* State = SelectedStates.Num() == 1 ? SelectedStates[0].Get() : nullptr)
+	{
+		const FStateSoftPath StatePath = FStateSoftPath(State);
+		HandleSelectionChanged(StatePath, FStateSoftPath());
+	}
+}
+
+SDiffSplitter::FSlot::FSlotArguments SDiffSplitter::Slot()
+{
+	return FSlot::FSlotArguments(MakeUnique<FSlot>());
+}
+
+} // UE::StateTree::Diff
\ No newline at end of file
diff --git a/Engine/Plugins/Runtime/StateTree/Source/StateTreeEditorModule/Private/SStateTreeSplitter.h b/Engine/Plugins/Runtime/StateTree/Source/StateTreeEditorModule/Private/SStateTreeSplitter.h
new file mode 100644
--- /dev/null
+++ b/Engine/Plugins/Runtime/StateTree/Source/StateTreeEditorModule/Private/SStateTreeSplitter.h
@@ -0,0 +1,71 @@
+// Copyright Epic Games, Inc. All Rights Reserved.
+#pragma once
+
+#include "StateTreeDiffHelper.h"
+#include "Widgets/SCompoundWidget.h"
+
+class SStateTreeView;
+
+namespace UE::StateTree::Diff
+{
+class FAsyncDiff;
+
+/** Splitter that allows you to provide an FAsyncStateTreeViewDiff to connect like-properties between two or more state tree panels */
+class STATETREEEDITORMODULE_API SDiffSplitter : public SCompoundWidget
+{
+public:
+	class FSlot : public TSlotBase<FSlot>
+	{
+	public:
+		FSlot() = default;
+
+		SLATE_SLOT_BEGIN_ARGS(FSlot, TSlotBase<FSlot>)
+			/** When the RuleSize is set to FractionOfParent, the size of the slot is the Value percentage of its parent size. */
+			SLATE_ATTRIBUTE(float, Value)
+			SLATE_ARGUMENT(TSharedPtr<SStateTreeView>, StateTreeView)
+			SLATE_ARGUMENT(const UStateTree*, StateTree)
+			SLATE_ATTRIBUTE(bool, IsReadonly)
+			SLATE_ATTRIBUTE(TSharedPtr<FAsyncDiff>, DifferencesWithRightPanel)
+		SLATE_SLOT_END_ARGS()
+	};
+	static FSlot::FSlotArguments Slot();
+
+	SLATE_BEGIN_ARGS(SDiffSplitter)
+	{
+	}
+		SLATE_SLOT_ARGUMENT(FSlot, Slots)
+	SLATE_END_ARGS()
+
+	void Construct(const FArguments& InArgs);
+
+	void AddSlot(const FSlot::FSlotArguments& SlotArgs, int32 Index = INDEX_NONE);
+
+	void HandleSelectionChanged(const FStateSoftPath& StatePath, const FStateSoftPath& SecondaryStatePath);
+
+protected:
+	void HandleSelectionChanged(const TArray<TWeakObjectPtr<UStateTreeState>>& SelectedStates);
+
+private:
+	struct FPanel
+	{
+		FPanel(const TSharedPtr<SStateTreeView>& StateTreeView, const UStateTree* StateTree, const TAttribute<bool>& bIsReadonly,
+			const TAttribute<TSharedPtr<FAsyncDiff>>& DiffRight)
+			: StateTreeView(StateTreeView)
+			, StateTree(StateTree)
+			, IsReadonly(bIsReadonly)
+			, DiffRight(DiffRight)
+		{
+		}
+
+		TSharedPtr<SStateTreeView> StateTreeView;
+		const UStateTree* StateTree = nullptr;
+		TAttribute<bool> IsReadonly;
+		TAttribute<TSharedPtr<FAsyncDiff>> DiffRight;
+	};
+
+	TArray<FPanel> Panels;
+	TSharedPtr<SSplitter> Splitter;
+
+	FStateSoftPath SelectedState;
+};
+} // UE::StateTree::Diff
\ No newline at end of file
diff --git a/Engine/Plugins/Runtime/StateTree/Source/StateTreeEditorModule/Private/SStateTreeView.cpp b/Engine/Plugins/Runtime/StateTree/Source/StateTreeEditorModule/Private/SStateTreeView.cpp
--- a/Engine/Plugins/Runtime/StateTree/Source/StateTreeEditorModule/Private/SStateTreeView.cpp
+++ b/Engine/Plugins/Runtime/StateTree/Source/StateTreeEditorModule/Private/SStateTreeView.cpp
@@ -682,4 +682,27 @@ void SStateTreeView::HandleDisableSelectedStates()
 	}
 }
 
+TSharedPtr<FStateTreeViewModel> SStateTreeView::GetViewModel() const
+{
+	return StateTreeViewModel;
+}
+
+void SStateTreeView::SetSelection(const TArray<TWeakObjectPtr<UStateTreeState>>& SelectedStates) const
+{
+	for (const TWeakObjectPtr<UStateTreeState>& WeakState : SelectedStates)
+	{
+		if (const UStateTreeState* SelectedState = WeakState.Get())
+		{
+			UStateTreeState* ParentState = SelectedState->Parent;
+			while (ParentState)
+			{
+				constexpr bool bShouldExpandItem(true);
+				TreeView->SetItemExpansion(ParentState, bShouldExpandItem);
+				ParentState = ParentState->Parent;
+			}
+		}
+	}
+	StateTreeViewModel->SetSelection(SelectedStates);
+}
+
 #undef LOCTEXT_NAMESPACE
diff --git a/Engine/Plugins/Runtime/StateTree/Source/StateTreeEditorModule/Private/SStateTreeView.h b/Engine/Plugins/Runtime/StateTree/Source/StateTreeEditorModule/Private/SStateTreeView.h
--- a/Engine/Plugins/Runtime/StateTree/Source/StateTreeEditorModule/Private/SStateTreeView.h
+++ b/Engine/Plugins/Runtime/StateTree/Source/StateTreeEditorModule/Private/SStateTreeView.h
@@ -35,6 +35,10 @@ public:
 
 	void SavePersistentExpandedStates();
 
+	TSharedPtr<FStateTreeViewModel> GetViewModel() const;
+
+	void SetSelection(const TArray<TWeakObjectPtr<UStateTreeState>>& SelectedStates) const;
+
 private:
 	virtual FReply OnKeyDown(const FGeometry& MyGeometry, const FKeyEvent& InKeyEvent) override;
 	virtual void Tick(const FGeometry& AllottedGeometry, const double InCurrentTime, const float InDeltaTime) override;
diff --git a/Engine/Plugins/Runtime/StateTree/Source/StateTreeEditorModule/Private/StateTreeDiffControl.cpp b/Engine/Plugins/Runtime/StateTree/Source/StateTreeEditorModule/Private/StateTreeDiffControl.cpp
new file mode 100644
--- /dev/null
+++ b/Engine/Plugins/Runtime/StateTree/Source/StateTreeEditorModule/Private/StateTreeDiffControl.cpp
@@ -0,0 +1,148 @@
+// Copyright Epic Games, Inc. All Rights Reserved.
+
+#include "StateTreeDiffControl.h"
+#include "SStateTreeView.h"
+#include "StateTree.h"
+#include "StateTreeEditorData.h"
+#include "StateTreeViewModel.h"
+#include "StateTreeDiffHelper.h"
+#include "Framework/Commands/UICommandList.h"
+
+#define LOCTEXT_NAMESPACE "SStateTreeDif"
+
+namespace UE::StateTree::Diff
+{
+
+FText FDiffControl::RightRevision = LOCTEXT("OlderRevisionIdentifier", "Right Revision");
+
+TSharedRef<SWidget> FDiffControl::GenerateSingleEntryWidget(FSingleDiffEntry DiffEntry, const FText ObjectName)
+{
+	return SNew(STextBlock)
+		.Text(GetStateTreeDiffMessage(DiffEntry, ObjectName, true))
+		.ToolTipText(GetStateTreeDiffMessage(DiffEntry, ObjectName))
+		.ColorAndOpacity(GetStateTreeDiffMessageColor(DiffEntry));
+}
+
+void FDiffControl::OnSelectDiffEntry(const FSingleDiffEntry StateDiff)
+{
+	OnDiffEntryFocused.ExecuteIfBound();
+	OnStateDiffEntryFocused.Broadcast(StateDiff);
+}
+
+void FDiffControl::GenerateTreeEntries(TArray<TSharedPtr<FBlueprintDifferenceTreeEntry>>& OutDifferences)
+{
+	TArray<FSingleDiffEntry> DifferingProperties;
+
+	for (int32 LeftIndex = 0; LeftIndex < DisplayedAssets.Num() - 1; ++LeftIndex)
+	{
+		const UStateTree* LeftStateTree = DisplayedAssets[LeftIndex].Get();
+		if (!ensure(LeftStateTree))
+		{
+			continue;
+		}
+		const TSharedPtr<FAsyncDiff> Diff = StateTreeDifferences[LeftStateTree].Right;
+		Diff->FlushQueue(); // make sure differences are fully up-to-date
+		Diff->GetStateTreeDifferences(DifferingProperties);
+	}
+
+	BindingDiffs.Empty();
+
+	TSet<FString> ExistingEntryPaths;
+
+	for (const FSingleDiffEntry& Difference : DifferingProperties)
+	{
+		FString DiffPath = Difference.Identifier.ToDisplayName();
+		bool bGenerateNewEntry = true;
+		if (IsBindingDiff(Difference.DiffType))
+		{
+			BindingDiffs.Add(Difference);
+			bGenerateNewEntry = !ExistingEntryPaths.Contains(DiffPath);
+		}
+
+		if (bGenerateNewEntry)
+		{
+			TSharedPtr<FBlueprintDifferenceTreeEntry> Entry = MakeShared<FBlueprintDifferenceTreeEntry>(
+				FOnDiffEntryFocused::CreateSP(AsShared(), &FDiffControl::OnSelectDiffEntry, Difference),
+				FGenerateDiffEntryWidget::CreateStatic(&GenerateSingleEntryWidget, Difference, RightRevision));
+			OutDifferences.Push(Entry);
+			ExistingEntryPaths.Add(Difference.Identifier.ToDisplayName());
+		}
+	}
+}
+
+FDiffControl::FDiffControl(
+	const UStateTree* InOldObject,
+	const UStateTree* InNewObject,
+	const FOnDiffEntryFocused& InSelectionCallback)
+	: OnDiffEntryFocused(InSelectionCallback)
+{
+	if (InOldObject)
+	{
+		InsertObject(InOldObject);
+	}
+
+	if (InNewObject)
+	{
+		InsertObject(InNewObject);
+	}
+}
+
+FDiffControl::~FDiffControl()
+{
+}
+
+TSharedRef<SStateTreeView> FDiffControl::InsertObject(gsl::not_null<const UStateTree*> StateTree)
+{
+	const FDiffWidgets DiffWidgets(StateTree);
+	TSharedRef<SStateTreeView> TreeView = DiffWidgets.GetStateTreeWidget();
+
+	const int32 Index = DisplayedAssets.Num();
+	DisplayedAssets.Insert(TStrongObjectPtr<const UStateTree>(StateTree), Index);
+
+	StateTreeDifferences.Add(StateTree, {});
+	StateTreeDiffWidgets.Add(StateTree, DiffWidgets);
+
+	// set up interaction with left panel
+	if (DisplayedAssets.IsValidIndex(Index - 1))
+	{
+		const UStateTree* OtherStateTree = DisplayedAssets[Index - 1].Get();
+		const FDiffWidgets& OtherStateTreeDiff = StateTreeDiffWidgets[OtherStateTree];
+		const TSharedRef<SStateTreeView> OtherTreeView = OtherStateTreeDiff.GetStateTreeWidget();
+
+		StateTreeDifferences[OtherStateTree].Right = MakeShared<FAsyncDiff>(OtherTreeView, TreeView);
+		StateTreeDifferences[StateTree].Left = StateTreeDifferences[OtherStateTree].Right;
+	}
+	// Set up interaction with right panel
+	if (DisplayedAssets.IsValidIndex(Index + 1))
+	{
+		const UStateTree* OtherStateTree = DisplayedAssets[Index + 1].Get();
+		const FDiffWidgets& OtherStateTreeDiff = StateTreeDiffWidgets[OtherStateTree];
+		const TSharedRef<SStateTreeView> OtherTreeView = OtherStateTreeDiff.GetStateTreeWidget();
+
+		StateTreeDifferences[OtherStateTree].Left = MakeShared<FAsyncDiff>(TreeView, OtherTreeView);
+		StateTreeDifferences[StateTree].Right = StateTreeDifferences[OtherStateTree].Left;
+	}
+
+	return TreeView;
+}
+
+TSharedRef<SStateTreeView> FDiffControl::GetDetailsWidget(const UStateTree* Object) const
+{
+	return StateTreeDiffWidgets[Object].GetStateTreeWidget();
+}
+
+FDiffWidgets::FDiffWidgets(const UStateTree* InStateTree)
+{
+	UStateTreeEditorData* EditorData = Cast<UStateTreeEditorData>(InStateTree->EditorData);
+	StateTreeViewModel = MakeShareable(new FStateTreeViewModel());
+	StateTreeViewModel->Init(EditorData);
+	SAssignNew(StateTreeTreeView, SStateTreeView, StateTreeViewModel.ToSharedRef(), MakeShared<FUICommandList>());
+}
+
+TSharedRef<SStateTreeView> FDiffWidgets::GetStateTreeWidget() const
+{
+	return StateTreeTreeView.ToSharedRef();
+}
+
+} // UE::StateTree::Diff
+#undef LOCTEXT_NAMESPACE
diff --git a/Engine/Plugins/Runtime/StateTree/Source/StateTreeEditorModule/Private/StateTreeDiffControl.h b/Engine/Plugins/Runtime/StateTree/Source/StateTreeEditorModule/Private/StateTreeDiffControl.h
new file mode 100644
--- /dev/null
+++ b/Engine/Plugins/Runtime/StateTree/Source/StateTreeEditorModule/Private/StateTreeDiffControl.h
@@ -0,0 +1,84 @@
+// Copyright Epic Games, Inc. All Rights Reserved.
+
+#pragma once
+
+#include "AsyncStateTreeDiff.h"
+#include "DiffUtils.h"
+#include "Editor/Kismet/Private/DiffControl.h"
+#include "Internationalization/Text.h"
+#include "UObject/ObjectKey.h"
+#include "UObject/StrongObjectPtr.h"
+
+class FBlueprintDifferenceTreeEntry;
+class FStateTreeViewModel;
+class FUICommandList;
+class SLinkableScrollBar;
+class SStateTreeView;
+class SWidget;
+class UStateTree;
+
+namespace UE::StateTree::Diff
+{
+struct FSingleDiffEntry;
+
+class STATETREEEDITORMODULE_API FDiffWidgets
+{
+public:
+	explicit FDiffWidgets(const UStateTree* InStateTree);
+
+	/** Returns actual widget that is used to display trees */
+	TSharedRef<SStateTreeView> GetStateTreeWidget() const;
+
+private:
+	TSharedPtr<SStateTreeView> StateTreeTreeView;
+	TSharedPtr<FStateTreeViewModel> StateTreeViewModel;
+};
+
+DECLARE_MULTICAST_DELEGATE_OneParam(FOnStateDiffEntryFocused, const FSingleDiffEntry&)
+
+class STATETREEEDITORMODULE_API FDiffControl : public TSharedFromThis<FDiffControl>
+{
+public:
+	FDiffControl() = delete;
+	FDiffControl(const FDiffControl& Other) = delete;
+	FDiffControl(const FDiffControl&& Other) = delete;
+	FDiffControl(const UStateTree* InOldObject, const UStateTree* InNewObject, const FOnDiffEntryFocused& InSelectionCallback);
+	~FDiffControl();
+
+	TSharedRef<SStateTreeView> GetDetailsWidget(const UStateTree* Object) const;
+
+	FOnStateDiffEntryFocused& GetOnStateDiffEntryFocused()
+	{
+		return OnStateDiffEntryFocused;
+	};
+
+	void GenerateTreeEntries(TArray<TSharedPtr<FBlueprintDifferenceTreeEntry>>& OutDifferences);
+	TConstArrayView<FSingleDiffEntry> GetBindingDifferences() const
+	{
+		return BindingDiffs;
+	}
+
+protected:
+	static FText RightRevision;
+	static TSharedRef<SWidget> GenerateSingleEntryWidget(FSingleDiffEntry DiffEntry, FText ObjectName);
+
+	TSharedRef<SStateTreeView> InsertObject(gsl::not_null<const UStateTree*> StateTree);
+
+	void OnSelectDiffEntry(const FSingleDiffEntry StateDiff);
+
+	FOnDiffEntryFocused OnDiffEntryFocused;
+	FOnStateDiffEntryFocused OnStateDiffEntryFocused;
+
+	TArray<TStrongObjectPtr<const UStateTree>> DisplayedAssets;
+
+	struct FStateTreeTreeDiffPairs
+	{
+		TSharedPtr<FAsyncDiff> Left;
+		TSharedPtr<FAsyncDiff> Right;
+	};
+	TMap<FObjectKey, FStateTreeTreeDiffPairs> StateTreeDifferences;
+	TArray<FSingleDiffEntry> BindingDiffs;
+	TMap<FObjectKey, FDiffWidgets> StateTreeDiffWidgets;
+};
+
+} // UE::StateTree::Diff
\ No newline at end of file
diff --git a/Engine/Plugins/Runtime/StateTree/Source/StateTreeEditorModule/Private/StateTreeDiffHelper.cpp b/Engine/Plugins/Runtime/StateTree/Source/StateTreeEditorModule/Private/StateTreeDiffHelper.cpp
new file mode 100644
--- /dev/null
+++ b/Engine/Plugins/Runtime/StateTree/Source/StateTreeEditorModule/Private/StateTreeDiffHelper.cpp
@@ -0,0 +1,229 @@
+// Copyright Epic Games, Inc. All Rights Reserved.
+
+#include "StateTreeDiffHelper.h"
+#include "StateTreeEditorData.h"
+#include "StateTreeEditorStyle.h"
+
+#define LOCTEXT_NAMESPACE "StateTreeDiffUtils"
+
+namespace UE::StateTree::Diff
+{
+bool IsBindingDiff(EStateDiffType DiffType)
+{
+	return DiffType == EStateDiffType::BindingAddedToA || DiffType == EStateDiffType::BindingAddedToB || DiffType == EStateDiffType::BindingChanged;
+}
+
+void FStateSoftPath::PostOrderChain(const UStateTreeState* CurBackwardsState)
+{
+	if (CurBackwardsState)
+	{
+		PostOrderChain(CurBackwardsState->Parent);
+		StateChain.Push(FChainElement(CurBackwardsState));
+	}
+}
+
+FStateSoftPath::FChainElement::FChainElement(const UStateTreeState* State)
+{
+	if (State)
+	{
+		StateName = State->Name;
+		DisplayString = State->Name.ToString();
+		ID = State->ID;
+	}
+}
+
+FStateSoftPath::FStateSoftPath(gsl::not_null<const UStateTreeState*> InState)
+{
+	int32 ChainCount = 1;
+	const UStateTreeState* CurState = InState;
+	while (CurState->Parent)
+	{
+		ChainCount++;
+		CurState = CurState->Parent;
+	}
+	StateChain.Reserve(ChainCount);
+	PostOrderChain(InState);
+}
+
+FString FStateSoftPath::ToDisplayName(bool Short) const
+{
+	if (Short && StateChain.Num() > 0)
+	{
+		return StateChain.Last().DisplayString;
+	}
+	FString Ret;
+	for (FChainElement Element : StateChain)
+	{
+		FString PropertyAsString = Element.DisplayString;
+		if (Ret.IsEmpty())
+		{
+			Ret.Append(PropertyAsString);
+		}
+		else
+		{
+			Ret.AppendChar('.');
+			Ret.Append(PropertyAsString);
+		}
+	}
+	return Ret;
+}
+
+UStateTreeState* FStateSoftPath::ResolvePath(gsl::not_null<const UStateTree*> StateTree) const
+{
+	if (StateChain.Num() == 0)
+	{
+		return nullptr;
+	}
+	return StateTree->EditorData != nullptr ? ResolvePath(Cast<UStateTreeEditorData>(StateTree->EditorData)) : nullptr;
+
+}
+
+UStateTreeState* FStateSoftPath::ResolvePath(gsl::not_null<const UStateTreeEditorData*> EditorData) const
+{
+	if (StateChain.Num() == 0)
+	{
+		return nullptr;
+	}
+
+	UStateTreeState* CurState = nullptr;
+	for (FChainElement Element : StateChain)
+	{
+		// First State in the chain should look for States in the SubTrees
+		if (CurState == nullptr)
+		{
+			for (UStateTreeState* RootState : EditorData->SubTrees)
+			{
+				if (RootState->ID == Element.ID || RootState->Name == Element.StateName)
+				{
+					CurState = RootState;
+					break;
+				}
+			}
+
+			// Can't find root State, stop parsing the chain
+			if (CurState == nullptr)
+			{
+				break;
+			}
+			continue;
+		}
+
+		bool bHasChild = false;
+		for (UStateTreeState* ChildState : CurState->Children)
+		{
+			if (ChildState->ID == Element.ID || ChildState->Name == Element.StateName)
+			{
+				CurState = ChildState;
+				bHasChild = true;
+				break;
+			}
+		}
+
+		// Can't find child State, stop parsing the chain
+		if (!bHasChild)
+		{
+			CurState = nullptr;
+			break;
+		}
+	}
+
+	return CurState;
+}
+
+FText GetStateTreeDiffMessage(const FSingleDiffEntry& Difference, const FText ObjectName, const bool bShort)
+{
+	FText Message;
+	const FString StateName = Difference.Identifier.ToDisplayName(bShort);
+
+	switch (Difference.DiffType)
+	{
+	case EStateDiffType::StateAddedToA:
+		Message = FText::Format(LOCTEXT("StateChange_Removed", "{0} removed from {1}"), FText::FromString(StateName), ObjectName);
+		break;
+	case EStateDiffType::StateAddedToB:
+		Message = FText::Format(LOCTEXT("StateChange_Added", "{0} added to {1}"), FText::FromString(StateName), ObjectName);
+		break;
+	case EStateDiffType::StateChanged:
+	case EStateDiffType::BindingChanged:
+	case EStateDiffType::BindingAddedToB:
+	case EStateDiffType::BindingAddedToA:
+		Message = FText::Format(LOCTEXT("StateChange", "{0} changed values"), FText::FromString(StateName));
+		break;
+	case EStateDiffType::StateEnabled:
+		Message = FText::Format(LOCTEXT("StateEnabled", "{0} enabled in {1}"), FText::FromString(StateName), ObjectName);
+		break;
+	case EStateDiffType::StateDisabled:
+		Message = FText::Format(LOCTEXT("StateDisabled", "{0} disabled in {1}"), FText::FromString(StateName), ObjectName);
+		break;
+	case EStateDiffType::StateMoved:
+		Message = FText::Format(LOCTEXT("StateMoved", "{0} moved"), FText::FromString(StateName));
+		break;
+	case EStateDiffType::StateTreePropertiesChanged:
+		Message = FText::Format(LOCTEXT("StateTreeChange", "State Tree Properties changed value in {0}"), ObjectName);
+		break;
+	case EStateDiffType::Invalid:
+	case EStateDiffType::Identical:
+		break;
+	}
+	return Message;
+}
+
+FLinearColor GetStateTreeDiffMessageColor(const FSingleDiffEntry& Difference)
+{
+	const static FLinearColor ForegroundColor = FAppStyle::GetColor("Graph.ForegroundColor");
+	switch (Difference.DiffType)
+	{
+	case EStateDiffType::StateAddedToA:
+		return FStateTreeEditorStyle::Get().GetColor("DiffTools.Removed");
+	case EStateDiffType::StateAddedToB:
+		return FStateTreeEditorStyle::Get().GetColor("DiffTools.Added");
+	case EStateDiffType::StateChanged:
+	case EStateDiffType::BindingChanged:
+	case EStateDiffType::BindingAddedToB:
+	case EStateDiffType::BindingAddedToA:
+		return FStateTreeEditorStyle::Get().GetColor("DiffTools.Changed");
+	case EStateDiffType::StateMoved:
+		return FStateTreeEditorStyle::Get().GetColor("DiffTools.Moved");
+	case EStateDiffType::StateEnabled:
+		return FStateTreeEditorStyle::Get().GetColor("DiffTools.Enabled");
+	case EStateDiffType::StateDisabled:
+	return FStateTreeEditorStyle::Get().GetColor("DiffTools.Disabled");
+	case EStateDiffType::StateTreePropertiesChanged:
+		return FStateTreeEditorStyle::Get().GetColor("DiffTools.Properties");
+	default:
+		return ForegroundColor;
+	}
+}
+
+FText GetStateDiffMessage(const FSingleObjectDiffEntry& Difference, FText PropertyName)
+{
+	switch (Difference.DiffType)
+	{
+	case EPropertyDiffType::PropertyAddedToA:
+		return FText::Format(LOCTEXT("StatePropertyChange_Removed", "{0} removed"), PropertyName);
+	case EPropertyDiffType::PropertyAddedToB:
+		return FText::Format(LOCTEXT("StatePropertyChange_Added", "{0} added"), PropertyName);
+	case EPropertyDiffType::PropertyValueChanged:
+		return FText::Format(LOCTEXT("StatePropertyChange", "{0} changed value"), PropertyName);
+	default:
+		return PropertyName;
+	}
+}
+
+FLinearColor GetStateDiffMessageColor(const FSingleObjectDiffEntry& Difference)
+{
+	switch (Difference.DiffType)
+	{
+	case EPropertyDiffType::PropertyAddedToA:
+		return FStateTreeEditorStyle::Get().GetColor("DiffTools.Removed");
+	case EPropertyDiffType::PropertyAddedToB:
+		return FStateTreeEditorStyle::Get().GetColor("DiffTools.Added");
+	case EPropertyDiffType::PropertyValueChanged:
+		return FStateTreeEditorStyle::Get().GetColor("DiffTools.Changed");
+	default:
+		return FStateTreeEditorStyle::Get().GetColor("DiffTools.Properties");
+	}
+}
+
+} // namespace UE::StateTree::Diff
+#undef LOCTEXT_NAMESPACE
\ No newline at end of file
diff --git a/Engine/Plugins/Runtime/StateTree/Source/StateTreeEditorModule/Private/StateTreeDiffHelper.h b/Engine/Plugins/Runtime/StateTree/Source/StateTreeEditorModule/Private/StateTreeDiffHelper.h
new file mode 100644
--- /dev/null
+++ b/Engine/Plugins/Runtime/StateTree/Source/StateTreeEditorModule/Private/StateTreeDiffHelper.h
@@ -0,0 +1,143 @@
+// Copyright Epic Games, Inc. All Rights Reserved.
+#pragma once
+
+#include "DiffUtils.h"
+
+class UStateTreeEditorData;
+class UStateTree;
+class UStateTreeState;
+
+namespace UE::StateTree::Diff
+{
+
+enum class EStateDiffType : uint8
+{
+	Invalid,
+
+	Identical,
+
+	StateAddedToA,
+	StateAddedToB,
+	StateChanged,
+
+	StateEnabled,
+	StateDisabled,
+
+	StateMoved,
+
+	StateTreePropertiesChanged,
+
+	BindingAddedToA,
+	BindingAddedToB,
+	BindingChanged,
+};
+
+struct FStateSoftPath
+{
+	FStateSoftPath() = default;
+	explicit FStateSoftPath(gsl::not_null<const UStateTreeState*> InState);
+
+	UStateTreeState* ResolvePath(gsl::not_null<const UStateTree*> StateTree) const;
+	UStateTreeState* ResolvePath(gsl::not_null<const UStateTreeEditorData*> EditorData) const;
+	FString ToDisplayName(bool Short = false) const;
+
+	/**
+	 * Indicates whether a given path is a base path of the current path.
+	 * @param BaseStatePath The path to look for as a base common path
+	 * @return Whether the path has the provided path as base path or not
+	 */
+	bool IsSubStateMatch(const FStateSoftPath& BaseStatePath) const
+	{
+		if (StateChain.Num() <= BaseStatePath.StateChain.Num())
+		{
+			return false;
+		}
+
+		for (int32 ElementIndex = 0; ElementIndex < BaseStatePath.StateChain.Num(); ElementIndex++)
+		{
+			if (BaseStatePath.StateChain[ElementIndex] != StateChain[ElementIndex])
+			{
+				return false;
+			}
+		}
+
+		return true;
+	}
+
+	bool operator==(FStateSoftPath const& RHS) const
+	{
+		return StateChain == RHS.StateChain;
+	}
+
+	bool operator!=(FStateSoftPath const& RHS) const
+	{
+		return !(*this == RHS);
+	}
+
+	explicit operator bool() const
+	{
+		return StateChain.Num() > 0;
+	}
+
+private:
+	void PostOrderChain(const UStateTreeState* CurBackwardsState);
+
+	struct FChainElement
+	{
+		FName StateName;
+		FString DisplayString;
+		FGuid ID;
+
+		FChainElement() = default;
+		explicit FChainElement(const UStateTreeState* State);
+
+		bool operator==(FChainElement const& RHS) const
+		{
+			return ID == RHS.ID || (StateName == RHS.StateName);
+		}
+
+		bool operator!=(FChainElement const& RHS) const
+		{
+			return !(*this == RHS);
+		}
+	};
+	TArray<FChainElement> StateChain;
+};
+
+struct FSingleDiffEntry
+{
+	FSingleDiffEntry(const FStateSoftPath& InIdentifier, const FStateSoftPath& InSecondaryIdentifier, const EStateDiffType InDiffType, const FPropertySoftPath& InBindingPath)
+		: Identifier(InIdentifier)
+		, SecondaryIdentifier(InSecondaryIdentifier)
+		, BindingPath(InBindingPath)
+		, DiffType(InDiffType)
+	{
+	}
+
+	FSingleDiffEntry(const FStateSoftPath& InIdentifier, const FStateSoftPath& InSecondaryIdentifier, const EStateDiffType InDiffType)
+		: FSingleDiffEntry(InIdentifier, InSecondaryIdentifier, InDiffType, FPropertySoftPath())
+	{
+	}
+
+	FSingleDiffEntry(const FStateSoftPath& InIdentifier, const EStateDiffType InDiffType)
+		: FSingleDiffEntry(InIdentifier, FStateSoftPath(), InDiffType)
+	{
+	}
+
+	FStateSoftPath Identifier;
+	FStateSoftPath SecondaryIdentifier;
+	FPropertySoftPath BindingPath;
+	EStateDiffType DiffType = EStateDiffType::Invalid;
+};
+
+FText GetStateTreeDiffMessage(const FSingleDiffEntry& Difference, FText ObjectName, bool bShort = false);
+
+FLinearColor GetStateTreeDiffMessageColor(const FSingleDiffEntry& Difference);
+
+FText GetStateDiffMessage(const FSingleObjectDiffEntry& Difference, FText PropertyName);
+
+FLinearColor GetStateDiffMessageColor(const FSingleObjectDiffEntry& Difference);
+
+bool IsBindingDiff(EStateDiffType DiffType);
+
+} // namespace UE::StateTree::Diff
\ No newline at end of file
diff --git a/Engine/Plugins/Runtime/StateTree/Source/StateTreeEditorModule/Private/StateTreeEditorStyle.cpp b/Engine/Plugins/Runtime/StateTree/Source/StateTreeEditorModule/Private/StateTreeEditorStyle.cpp
--- a/Engine/Plugins/Runtime/StateTree/Source/StateTreeEditorModule/Private/StateTreeEditorStyle.cpp
+++ b/Engine/Plugins/Runtime/StateTree/Source/StateTreeEditorModule/Private/StateTreeEditorStyle.cpp
@@ -177,6 +177,17 @@ FStateTreeEditorStyle::FStateTreeEditorStyle()
 			.SetFont(DEFAULT_FONT("Regular", 11)));
 	}
 
+	// Diff tool
+	{
+		Set("DiffTools.Added", FLinearColor(0.3f, 1.f, 0.3f)); // green
+		Set("DiffTools.Removed", FLinearColor(1.0f, 0.2f, 0.3f)); // red
+		Set("DiffTools.Changed", FLinearColor(0.85f, 0.71f, 0.25f)); // yellow
+		Set("DiffTools.Moved", FLinearColor(0.5f, 0.8f, 1.f)); // light blue
+		Set("DiffTools.Enabled", FLinearColor(0.7f, 1.f, 0.7f)); // light green
+		Set("DiffTools.Disabled", FLinearColor(1.0f, 0.6f, 0.5f)); // light red
+		Set("DiffTools.Properties", FLinearColor(0.2f, 0.4f, 1.f)); // blue
+	}
+
 	// Debugger
 	{
 		Set("StateTreeDebugger.Element.Normal",
diff --git a/Engine/Plugins/Runtime/StateTree/Source/StateTreeEditorModule/StateTreeEditorModule.Build.cs b/Engine/Plugins/Runtime/StateTree/Source/StateTreeEditorModule/StateTreeEditorModule.Build.cs
--- a/Engine/Plugins/Runtime/StateTree/Source/StateTreeEditorModule/StateTreeEditorModule.Build.cs
+++ b/Engine/Plugins/Runtime/StateTree/Source/StateTreeEditorModule/StateTreeEditorModule.Build.cs
@@ -52,6 +52,7 @@ namespace UnrealBuildTool.Rules
 				"RewindDebuggerInterface",
 				"DetailCustomizations",
 				"AppFramework",
+				"Kismet",
 				"KismetCompiler",
 				"EditorInteractiveToolsFramework",
 				"InteractiveToolsFramework",
