StateTree: Compile the StateTree after each load in the editor. That ensures the data is valid when cooking and that if the cpp code changed, the new properties are properly converted. Do not mark the asset as dirty when it compiles. It should only mark the asset dirty when a user does the manipulation. Any automatic conversion should be transparent for the user.

diff --git a/Engine/Plugins/Runtime/StateTree/Source/StateTreeEditorModule/Private/StateTreeEditingSubsystem.cpp b/Engine/Plugins/Runtime/StateTree/Source/StateTreeEditorModule/Private/StateTreeEditingSubsystem.cpp
--- a/Engine/Plugins/Runtime/StateTree/Source/StateTreeEditorModule/Private/StateTreeEditingSubsystem.cpp
+++ b/Engine/Plugins/Runtime/StateTree/Source/StateTreeEditorModule/Private/StateTreeEditingSubsystem.cpp
@@ -24,7 +24,7 @@ void UStateTreeEditingSubsystem::BeginDestroy()
 	Super::BeginDestroy();
 }
 
-bool UStateTreeEditingSubsystem::CompileStateTree(const TNonNullPtr<UStateTree> InStateTree, FStateTreeCompilerLog& InOutLog)
+bool UStateTreeEditingSubsystem::CompileStateTree(const gsl::not_null<UStateTree*> InStateTree, FStateTreeCompilerLog& InOutLog)
 {
 	ValidateStateTree(InStateTree);
 	const uint32 EditorDataHash = CalculateStateTreeHash(InStateTree);
@@ -52,7 +52,7 @@ bool UStateTreeEditingSubsystem::CompileStateTree(const TNonNullPtr<UStateTree>
 	return bCompilationResult;
 }
 
-TSharedRef<FStateTreeViewModel> UStateTreeEditingSubsystem::FindOrAddViewModel(const TNonNullPtr<UStateTree> InStateTree)
+TSharedRef<FStateTreeViewModel> UStateTreeEditingSubsystem::FindOrAddViewModel(const gsl::not_null<UStateTree*> InStateTree)
 {
 	const FObjectKey StateTreeKey = InStateTree;
 	TSharedPtr<FStateTreeViewModel> ViewModelPtr = StateTreeViewModels.FindRef(StateTreeKey);
@@ -112,7 +112,7 @@ TSharedRef<SWidget> UStateTreeEditingSubsystem::GetStateTreeView(TSharedRef<FSta
 	return SNew(SStateTreeView, InViewModel, TreeViewCommandList);
 }
 
-void UStateTreeEditingSubsystem::ValidateStateTree(const TNonNullPtr<UStateTree> InStateTree)
+void UStateTreeEditingSubsystem::ValidateStateTree(const gsl::not_null<UStateTree*> InStateTree)
 {
 	auto FixChangedStateLinkName = [](FStateTreeStateLink& StateLink, const TMap<FGuid, FName>& IDToName) -> bool
 	{
@@ -142,7 +142,8 @@ void UStateTreeEditingSubsystem::ValidateStateTree(const TNonNullPtr<UStateTree>
 			return;
 		}
 
-		TreeData->Modify();
+		constexpr bool bMarkDirty = false;
+		TreeData->Modify(bMarkDirty);
 
 		// Make sure all state links are valid and update the names if needed.
 
@@ -158,7 +159,8 @@ void UStateTreeEditingSubsystem::ValidateStateTree(const TNonNullPtr<UStateTree>
 		// Fix changed names.
 		TreeData->VisitHierarchy([&IDToName, FixChangedStateLinkName](UStateTreeState& State, UStateTreeState* /*ParentState*/)
 		{
-			State.Modify();
+			constexpr bool bMarkDirty = false;
+			State.Modify(bMarkDirty);
 			if (State.Type == EStateTreeStateType::Linked)
 			{
 				FixChangedStateLinkName(State.LinkedSubtree, IDToName);
@@ -181,7 +183,8 @@ void UStateTreeEditingSubsystem::ValidateStateTree(const TNonNullPtr<UStateTree>
 			return;
 		}
 
-		TreeData->Modify();
+		constexpr bool bMarkDirty = false;
+		TreeData->Modify(bMarkDirty);
 		TreeData->ReparentStates();
 	};
 
@@ -199,7 +202,8 @@ void UStateTreeEditingSubsystem::ValidateStateTree(const TNonNullPtr<UStateTree>
 			return;
 		}
 
-		TreeData->Modify();
+		constexpr bool bMarkDirty = false;
+		TreeData->Modify(bMarkDirty);
 		
 		// Clear evaluators if not allowed.
 		if (Schema->AllowEvaluators() == false && TreeData->Evaluators.Num() > 0)
@@ -211,7 +215,8 @@ void UStateTreeEditingSubsystem::ValidateStateTree(const TNonNullPtr<UStateTree>
 
 		TreeData->VisitHierarchy([&StateTree, Schema](UStateTreeState& State, UStateTreeState* /*ParentState*/)
 		{
-			State.Modify();
+			constexpr bool bMarkDirty = false;
+			State.Modify(bMarkDirty);
 
 			// Clear enter conditions if not allowed.
 			if (Schema->AllowEnterConditions() == false && State.EnterConditions.Num() > 0)
@@ -266,7 +271,8 @@ void UStateTreeEditingSubsystem::ValidateStateTree(const TNonNullPtr<UStateTree>
 
 		TMap<FGuid, const FStateTreeDataView> AllStructValues;
 		TreeData->GetAllStructValues(AllStructValues);
-		TreeData->Modify();
+		constexpr bool bMarkDirty = false;
+		TreeData->Modify(bMarkDirty);
 		TreeData->GetPropertyEditorBindings()->RemoveUnusedBindings(AllStructValues);
 	};
 
@@ -278,7 +284,8 @@ void UStateTreeEditingSubsystem::ValidateStateTree(const TNonNullPtr<UStateTree>
 			return;
 		}
 
-		TreeData->Modify();
+		constexpr bool bMarkDirty = false;
+		TreeData->Modify(bMarkDirty);
 
 		const EStateTreeVisitor Result = TreeData->VisitHierarchy([](UStateTreeState& State, UStateTreeState* /*ParentState*/)
 		{
@@ -300,7 +307,7 @@ void UStateTreeEditingSubsystem::ValidateStateTree(const TNonNullPtr<UStateTree>
 	UpdateLinkedStateParameters(*InStateTree);
 }
 
-uint32 UStateTreeEditingSubsystem::CalculateStateTreeHash(const TNonNullPtr<const UStateTree> InStateTree)
+uint32 UStateTreeEditingSubsystem::CalculateStateTreeHash(const gsl::not_null<const UStateTree*> InStateTree)
 {
 	uint32 EditorDataHash = 0;
 	if (InStateTree->EditorData != nullptr)
diff --git a/Engine/Plugins/Runtime/StateTree/Source/StateTreeEditorModule/Private/StateTreeFactory.cpp b/Engine/Plugins/Runtime/StateTree/Source/StateTreeEditorModule/Private/StateTreeFactory.cpp
--- a/Engine/Plugins/Runtime/StateTree/Source/StateTreeEditorModule/Private/StateTreeFactory.cpp
+++ b/Engine/Plugins/Runtime/StateTree/Source/StateTreeEditorModule/Private/StateTreeFactory.cpp
@@ -126,9 +126,6 @@ UObject* UStateTreeFactory::FactoryCreateNew(UClass* Class, UObject* InParent, F
 	
 	NewStateTree->EditorData = EditorData;
 
-	// Compile the StateTree asset, so that it is valid initially.
-	UStateTreeEditingSubsystem::ValidateStateTree(NewStateTree);
-
 	FStateTreeCompilerLog Log;
 	FStateTreeCompiler Compiler(Log);
 
diff --git a/Engine/Plugins/Runtime/StateTree/Source/StateTreeEditorModule/Public/StateTreeEditingSubsystem.h b/Engine/Plugins/Runtime/StateTree/Source/StateTreeEditorModule/Public/StateTreeEditingSubsystem.h
--- a/Engine/Plugins/Runtime/StateTree/Source/StateTreeEditorModule/Public/StateTreeEditingSubsystem.h
+++ b/Engine/Plugins/Runtime/StateTree/Source/StateTreeEditorModule/Public/StateTreeEditingSubsystem.h
@@ -22,9 +22,9 @@ public:
 	UStateTreeEditingSubsystem();
 	virtual void BeginDestroy() override;
 	
-	TSharedRef<FStateTreeViewModel> FindOrAddViewModel(const TNonNullPtr<UStateTree> InStateTree);
+	TSharedRef<FStateTreeViewModel> FindOrAddViewModel(const gsl::not_null<UStateTree*> InStateTree);
 	
-	static bool CompileStateTree(const TNonNullPtr<UStateTree>InStateTree,  FStateTreeCompilerLog& InOutLog);
+	static bool CompileStateTree(const gsl::not_null<UStateTree*> InStateTree,  FStateTreeCompilerLog& InOutLog);
 	
 	/** Create a StateTreeView widget for the viewmodel. */
 	static TSharedRef<SWidget> GetStateTreeView(TSharedRef<FStateTreeViewModel> InViewModel, const TSharedRef<FUICommandList>& TreeViewCommandList);
@@ -33,10 +33,10 @@ public:
 	 * Validates and applies the schema restrictions on the StateTree.
 	 * Updates state's link, removes the unused node while validating the StateTree asset.
 	 */
-	static void ValidateStateTree(const TNonNullPtr<UStateTree> InStateTree);
+	static void ValidateStateTree(const gsl::not_null<UStateTree*> InStateTree);
 
 	/** Calculates editor data hash of the asset. */
-	static uint32 CalculateStateTreeHash(const TNonNullPtr<const UStateTree> InStateTree);
+	static uint32 CalculateStateTreeHash(const gsl::not_null<const UStateTree*> InStateTree);
 	
 private:
 	void HandlePostGarbageCollect();
diff --git a/Engine/Plugins/Runtime/StateTree/Source/StateTreeModule/Private/StateTree.cpp b/Engine/Plugins/Runtime/StateTree/Source/StateTreeModule/Private/StateTree.cpp
--- a/Engine/Plugins/Runtime/StateTree/Source/StateTreeModule/Private/StateTree.cpp
+++ b/Engine/Plugins/Runtime/StateTree/Source/StateTreeModule/Private/StateTree.cpp
@@ -21,6 +21,7 @@
 #include "StateTreeModuleImpl.h"
 #include "StructUtils/UserDefinedStruct.h"
 #include "StructUtilsDelegates.h"
+#include "Templates/GuardValueAccessors.h"
 #include "UObject/LinkerLoad.h"
 #endif
 
@@ -443,25 +444,9 @@ void UStateTree::PostLoad()
 		// Make sure all the fix up logic in the editor data has had chance to happen.
 		EditorData->ConditionalPostLoad();
 	}
-
-	if (CurrentVersion < FStateTreeCustomVersion::LatestVersion)
-	{
-		// Compiled data is in older format, try to compile the StateTree.
-		if (UE::StateTree::Delegates::OnRequestCompile.IsBound())
-		{
-			LOG_SCOPE_VERBOSITY_OVERRIDE(LogStateTree, ELogVerbosity::Log);
-			UE_LOG(LogStateTree, Log, TEXT("%s: compiled data is in older format. Trying to compile the asset..."), *GetFullName());
-			UE::StateTree::Delegates::OnRequestCompile.Execute(*this);
-		}
-		else
-		{
-			ResetCompiled();
-			UE_LOG(LogStateTree, Warning, TEXT("%s: compiled data is in older format. Please resave the StateTree asset."), *GetFullName());
-		}
-	}
-	else
 	{
-		CompileIfChanged();
+		TGuardValueAccessors<bool> IsEditorLoadingPackageGuard(UE::GetIsEditorLoadingPackage, UE::SetIsEditorLoadingPackage, true);
+		Compile();
 	}
 #else
 	if (CurrentVersion < FStateTreeCustomVersion::LatestVersion)
@@ -1179,6 +1164,25 @@ void UStateTree::CompileIfChanged()
 			UE::StateTree::Delegates::OnRequestCompile.Execute(*this);
 		}
 	}
+	else
+	{
+		ResetCompiled();
+		UE_LOG(LogStateTree, Warning, TEXT("%s: could not compile. Please resave the StateTree asset."), *GetFullName());
+	}
+}
+
+void UStateTree::Compile()
+{
+	if (UE::StateTree::Delegates::OnRequestCompile.IsBound())
+	{
+		UE_LOG(LogStateTree, Log, TEXT("%s: Editor data has changed. Recompiling state tree."), *GetFullName());
+		UE::StateTree::Delegates::OnRequestCompile.Execute(*this);
+	}
+	else
+	{
+		ResetCompiled();
+		UE_LOG(LogStateTree, Warning, TEXT("%s: could not compile. Please resave the StateTree asset."), *GetFullName());
+	}
 }
 #endif //WITH_EDITOR
 
diff --git a/Engine/Plugins/Runtime/StateTree/Source/StateTreeModule/Public/StateTree.h b/Engine/Plugins/Runtime/StateTree/Source/StateTreeModule/Public/StateTree.h
--- a/Engine/Plugins/Runtime/StateTree/Source/StateTreeModule/Public/StateTree.h
+++ b/Engine/Plugins/Runtime/StateTree/Source/StateTreeModule/Public/StateTree.h
@@ -233,10 +233,12 @@ public:
 	static void AddReferencedObjects(UObject* InThis, FReferenceCollector& Collector);
 	
 private:
+	/** Compile the state trees. */
+	void Compile();
 
 	/**
-     * Reset the data generated by Link(), this in turn will cause IsReadyToRun() to return false.
-     * Used during linking, or to invalidate the linked data when data version is old (requires recompile). 
+	 * Reset the data generated by Link(), this in turn will cause IsReadyToRun() to return false.
+	 * Used during linking, or to invalidate the linked data when data version is old (requires recompile). 
 	 */
 	void ResetLinked();
 
