Prevent infinite loop listening to selection clears when clearing an empty selection

#tests editor

diff --git a/Engine/Plugins/Runtime/StateTree/Source/StateTreeEditorModule/Private/StateTreeViewModel.cpp b/Engine/Plugins/Runtime/StateTree/Source/StateTreeEditorModule/Private/StateTreeViewModel.cpp
--- a/Engine/Plugins/Runtime/StateTree/Source/StateTreeEditorModule/Private/StateTreeViewModel.cpp
+++ b/Engine/Plugins/Runtime/StateTree/Source/StateTreeEditorModule/Private/StateTreeViewModel.cpp
@@ -588,6 +588,11 @@ void FStateTreeViewModel::PostRedo(bool bSuccess)
 
 void FStateTreeViewModel::ClearSelection()
 {
+	if (SelectedStates.IsEmpty())
+	{
+		return;
+	}
+
 	SelectedStates.Reset();
 
 	const TArray<TWeakObjectPtr<UStateTreeState>> SelectedStatesArr;
@@ -596,6 +601,11 @@ void FStateTreeViewModel::ClearSelection()
 
 void FStateTreeViewModel::SetSelection(UStateTreeState* Selected)
 {
+	if (SelectedStates.Num() == 1 && SelectedStates.Contains(Selected))
+	{
+		return;
+	}
+
 	SelectedStates.Reset();
 
 	SelectedStates.Add(Selected);
@@ -607,6 +617,11 @@ void FStateTreeViewModel::SetSelection(UStateTreeState* Selected)
 
 void FStateTreeViewModel::SetSelection(const TArray<TWeakObjectPtr<UStateTreeState>>& InSelectedStates)
 {
+	if (SelectedStates.Num() == InSelectedStates.Num() && SelectedStates.Array() == InSelectedStates)
+	{
+		return;
+	}
+
 	SelectedStates.Reset();
 
 	for (const TWeakObjectPtr<UStateTreeState>& State : InSelectedStates)
@@ -617,7 +632,6 @@ void FStateTreeViewModel::SetSelection(const TArray<TWeakObjectPtr<UStateTreeSta
 		}
 	}
 
-	TArray<FGuid> SelectedTaskIDArr;
 	OnSelectionChanged.Broadcast(InSelectedStates);
 }
 
