StateTree: Revert part of CL36113278 name check for State and Node. The FName test is obsolete.

diff --git a/Engine/Plugins/Runtime/StateTree/Source/StateTreeEditorModule/Private/Customizations/StateTreeEditorNodeDetails.cpp b/Engine/Plugins/Runtime/StateTree/Source/StateTreeEditorModule/Private/Customizations/StateTreeEditorNodeDetails.cpp
--- a/Engine/Plugins/Runtime/StateTree/Source/StateTreeEditorModule/Private/Customizations/StateTreeEditorNodeDetails.cpp
+++ b/Engine/Plugins/Runtime/StateTree/Source/StateTreeEditorModule/Private/Customizations/StateTreeEditorNodeDetails.cpp
@@ -1473,7 +1473,7 @@ bool FStateTreeEditorNodeDetails::HandleVerifyNameChanged(const FText& InText, F
 		OutErrorMessage = LOCTEXT("VerifyNodeLabelFailed_MaxLength", "Max length exceeded");
 		return false;
 	}
-	return NewName.Len() > 0 && FName::IsValidXName(NewName, INVALID_NAME_CHARACTERS, &OutErrorMessage);
+	return NewName.Len() > 0;
 }
 
 void FStateTreeEditorNodeDetails::HandleNameCommitted(const FText& NewText, ETextCommit::Type InTextCommit) const
@@ -1484,7 +1484,7 @@ void FStateTreeEditorNodeDetails::HandleNameCommitted(const FText& NewText, ETex
 	{
 		// Remove excess whitespace and prevent categories with just spaces
 		const FString NewName = FText::TrimPrecedingAndTrailing(NewText).ToString();
-		if (NewName.Len() > 0 && FName::IsValidXName(NewName, INVALID_NAME_CHARACTERS) && NewName.Len() < NAME_SIZE)
+		if (NewName.Len() > 0 && NewName.Len() < NAME_SIZE)
 		{
 			if (GEditor)
 			{
diff --git a/Engine/Plugins/Runtime/StateTree/Source/StateTreeEditorModule/Private/SStateTreeViewRow.cpp b/Engine/Plugins/Runtime/StateTree/Source/StateTreeEditorModule/Private/SStateTreeViewRow.cpp
--- a/Engine/Plugins/Runtime/StateTree/Source/StateTreeEditorModule/Private/SStateTreeViewRow.cpp
+++ b/Engine/Plugins/Runtime/StateTree/Source/StateTreeEditorModule/Private/SStateTreeViewRow.cpp
@@ -1973,7 +1973,7 @@ bool SStateTreeViewRow::HandleVerifyNodeLabelTextChanged(const FText& InText, FT
 				OutErrorMessage = LOCTEXT("VerifyNodeLabelFailed_MaxLength", "Max length exceeded");
 				return false;
 			}
-			return NewName.Len() > 0 && FName::IsValidXName(NewName, INVALID_NAME_CHARACTERS, &OutErrorMessage);
+			return NewName.Len() > 0;
 		}
 	}
 	OutErrorMessage = LOCTEXT("VerifyNodeLabelFailed", "Invalid State Tree");
@@ -1987,7 +1987,7 @@ void SStateTreeViewRow::HandleNodeLabelTextCommitted(const FText& NewLabel, ETex
 		if (UStateTreeState* State = WeakState.Get())
 		{
 			const FString NewName = FText::TrimPrecedingAndTrailing(NewLabel).ToString();
-			if (NewName.Len() > 0 && FName::IsValidXName(NewName, INVALID_NAME_CHARACTERS) && NewName.Len() < NAME_SIZE)
+			if (NewName.Len() > 0 && NewName.Len() < NAME_SIZE)
 			{
 				StateTreeViewModel->RenameState(State, FName(NewName));
 			}
