[StateTree] added mention '(experimental)' to the diff window

diff --git a/Engine/Plugins/Runtime/StateTree/Source/StateTreeEditorModule/Private/SStateTreeDiff.cpp b/Engine/Plugins/Runtime/StateTree/Source/StateTreeEditorModule/Private/SStateTreeDiff.cpp
--- a/Engine/Plugins/Runtime/StateTree/Source/StateTreeEditorModule/Private/SStateTreeDiff.cpp
+++ b/Engine/Plugins/Runtime/StateTree/Source/StateTreeEditorModule/Private/SStateTreeDiff.cpp
@@ -240,13 +240,13 @@ TSharedRef<SDiffWidget> SDiffWidget::CreateDiffWindow(gsl::not_null<const UState
 	// times we're comparing two completely separate assets altogether)
 	const bool bIsSingleAsset = !OldStateTree || !NewStateTree || (NewStateTree->GetName() == OldStateTree->GetName());
 
-	FText WindowTitle = FText::Format(LOCTEXT("NamelesSDiffWidget", "{0} Diff"), StateTreeClass->GetDisplayNameText());
+	FText WindowTitle = FText::Format(LOCTEXT("NamelessStateTreeDiff", "{0} Diff (experimental)"), StateTreeClass->GetDisplayNameText());
 	// if we're diffing one asset against itself
 	if (bIsSingleAsset)
 	{
 		// identify the assumed single asset in the window's title
 		const FString STName = NewStateTree ? NewStateTree->GetName() : OldStateTree->GetName();
-		WindowTitle = FText::Format(LOCTEXT("NamedStateTreeDiff", "{0} - {1} Diff"), FText::FromString(STName), StateTreeClass->GetDisplayNameText());
+		WindowTitle = FText::Format(LOCTEXT("NamedStateTreeDiff", "{0} - {1} Diff (experimental)"), FText::FromString(STName), StateTreeClass->GetDisplayNameText());
 	}
 
 	return CreateDiffWindow(WindowTitle, OldStateTree, NewStateTree, OldRevision, NewRevision);
