StateTree: Add Debug callback and runtime check framework. Check if the context is the same between each execution context. Check if every call to StartState is met with a StopState.

diff --git a/Engine/Plugins/Runtime/StateTree/Source/StateTreeModule/Private/Debugger/StateTreeDebug.cpp b/Engine/Plugins/Runtime/StateTree/Source/StateTreeModule/Private/Debugger/StateTreeDebug.cpp
new file mode 100644
--- /dev/null
+++ b/Engine/Plugins/Runtime/StateTree/Source/StateTreeModule/Private/Debugger/StateTreeDebug.cpp
@@ -0,0 +1,52 @@
+// Copyright Epic Games, Inc. All Rights Reserved.
+
+#include "Debugger/StateTreeDebug.h"
+#include "Debugger/StateTreeRuntimeValidationInstanceData.h"
+#include "StateTreeExecutionContext.h"
+#include "StateTreeInstanceData.h"
+
+#if WITH_STATETREE_DEBUG
+
+namespace UE::StateTree::Debug
+{
+
+FStateTreeDebug::FNodeDelegate FStateTreeDebug::OnNodeEnterState_AnyThread;
+FStateTreeDebug::FNodeDelegate FStateTreeDebug::OnNodeExitState_AnyThread;
+
+FStateTreeDebug::FNodeReference::FNodeReference(gsl::not_null<const UStateTree*> InStateTree, FStateTreeIndex16 InNodeIndex)
+	: StateTree(InStateTree)
+	, Index(InNodeIndex)
+{
+}
+
+void FStateTreeDebug::NodeEnterState(const FStateTreeExecutionContext& ExecutionContext, FNodeReference Node)
+{
+	FStateTreeInstanceData* InstanceData = ExecutionContext.GetMutableInstanceData();
+	FGuid NodeId = Node.StateTree->GetNodeIdFromIndex(Node.Index);
+	if (ensure(NodeId.IsValid()))
+	{
+		if (FRuntimeValidationInstanceData* RuntimeValidation = InstanceData->GetRuntimeValidation().GetInstanceData())
+		{
+			RuntimeValidation->NodeEnterState(NodeId);
+		}
+		OnNodeEnterState_AnyThread.Broadcast(ExecutionContext, FNodeDelegateArgs{.Node = Node, .NodeId = NodeId});
+	}
+}
+
+void FStateTreeDebug::NodeExitState(const FStateTreeExecutionContext& ExecutionContext, FNodeReference Node)
+{
+	FStateTreeInstanceData* InstanceData = ExecutionContext.GetMutableInstanceData();
+	FGuid NodeId = Node.StateTree->GetNodeIdFromIndex(Node.Index);
+	if (ensure(NodeId.IsValid()))
+	{
+		if (FRuntimeValidationInstanceData* RuntimeValidation = InstanceData->GetRuntimeValidation().GetInstanceData())
+		{
+			RuntimeValidation->NodeExitState(NodeId);
+		}
+		OnNodeExitState_AnyThread.Broadcast(ExecutionContext, FNodeDelegateArgs{.Node = Node, .NodeId = NodeId});
+	}
+}
+
+}//namespace UE::StateTree::Debug
+
+#endif // WITH_STATETREE_DEBUG
diff --git a/Engine/Plugins/Runtime/StateTree/Source/StateTreeModule/Private/Debugger/StateTreeRuntimeValidation.cpp b/Engine/Plugins/Runtime/StateTree/Source/StateTreeModule/Private/Debugger/StateTreeRuntimeValidation.cpp
new file mode 100644
--- /dev/null
+++ b/Engine/Plugins/Runtime/StateTree/Source/StateTreeModule/Private/Debugger/StateTreeRuntimeValidation.cpp
@@ -0,0 +1,171 @@
+// Copyright Epic Games, Inc. All Rights Reserved.
+
+#include "Debugger/StateTreeRuntimeValidation.h"
+#include "Debugger/StateTreeRuntimeValidationInstanceData.h"
+
+#if WITH_STATETREE_DEBUG
+#include "HAL/IConsoleManager.h"
+#include "StateTree.h"
+#endif
+
+namespace UE::StateTree::Debug
+{
+#if WITH_STATETREE_DEBUG
+
+namespace Private
+{
+bool RuntimeValidationContext = true;
+static FAutoConsoleVariableRef CVarRuntimeValidationContext(
+	TEXT("StateTree.RuntimeValidation.Context"),
+	RuntimeValidationContext,
+	TEXT("Test if the context creation parameters are the same between each creation of StateTreeExecutionContext.")
+);
+
+bool RuntimeValidationEnterExitState = true;
+static FAutoConsoleVariableRef CVarRuntimeValidationEnterExitState(
+	TEXT("StateTree.RuntimeValidation.EnterExitState"),
+	RuntimeValidationEnterExitState,
+	TEXT("Test that if a node get a EnterState, it will receive an ExitState.\n"
+		"Test that if a node get a ExitState, it did receive an EnterState before.")
+);
+
+FString NodeToString(const UObject* Obj, FGuid Id)
+{
+	TStringBuilderWithBuffer<TCHAR, 128> Buffer;
+	if (Obj)
+	{
+		Obj->GetPathName(nullptr, Buffer);
+	}
+	Buffer << TEXT(':');
+	Buffer << Id;
+	return Buffer.ToString();
+}
+} // namespace Private
+
+FRuntimeValidationInstanceData::~FRuntimeValidationInstanceData()
+{
+	if (Private::RuntimeValidationEnterExitState)
+	{
+		for (const FNodeStatePair& Pair : NodeStates)
+		{
+			if (EnumHasAllFlags(Pair.State, EState::BetweenEnterExitState))
+			{
+				ensureAlwaysMsgf(false, TEXT("Tree exited. Missing ExitState on %s."), *Pair.Id.ToString());
+				Private::RuntimeValidationEnterExitState = false;
+			}
+		}
+	}
+}
+
+void FRuntimeValidationInstanceData::SetContext(const UObject* InNewOwner, const UStateTree* InNewStateTree)
+{
+	TWeakObjectPtr<const UStateTree> NewStateTree = InNewStateTree;
+	TWeakObjectPtr<const UObject> NewOwner = InNewOwner;
+	if (Private::RuntimeValidationContext)
+	{
+		if (StateTree.IsValid() && StateTree != NewStateTree)
+		{
+			ensureAlwaysMsgf(false, TEXT("StateTree runtime check failed: The StateTree '%s' is different from the previously set '%s'.\n"
+				"Make sure you initialize FStateTreeExecutionContext with the same value every time.\n"
+				"Auto deactivate Runtime check StateTree.RuntimeValidation.Context to prevent reporting the same error multiple times.")
+				, InNewStateTree ? *InNewStateTree->GetFullName() : TEXT("StateTree"), *StateTree.Get()->GetFullName());
+			Private::RuntimeValidationContext = false;
+		}
+		if (Owner.IsValid() && Owner != NewOwner)
+		{
+			ensureAlwaysMsgf(false, TEXT("StateTree runtime check failed: The owner '%s' is different from the previously set '%s'.\n"
+				"Make sure you initialize FStateTreeExecutionContext with the same values every time.\n"
+				"Auto deactivate Runtime check StateTree.RuntimeValidation.Context to prevent reporting the same error multiple times.")
+				, InNewOwner ? *InNewOwner->GetFullName() : TEXT("owner"), *Owner.Get()->GetFullName());
+			Private::RuntimeValidationContext = false;
+		}
+	}
+	StateTree = NewStateTree;
+	Owner = NewOwner;
+}
+
+void FRuntimeValidationInstanceData::NodeEnterState(FGuid Id)
+{
+	FNodeStatePair* Found = NodeStates.FindByPredicate([Id](const FNodeStatePair& Other)
+	{
+		return Other.Id == Id;
+	});
+	if (Found)
+	{
+		if (Private::RuntimeValidationEnterExitState && EnumHasAllFlags(Found->State, EState::BetweenEnterExitState))
+		{
+			ensureAlwaysMsgf(false, TEXT("StateTree runtime check failed: EnterState executed on node %s without an ExitState.\n"
+				"Auto deactivate Runtime check StateTree.RuntimeValidation.EnterExitState to prevent reporting the same error multiple times.")
+				, *Private::NodeToString(Owner.Get(), Id));
+			Private::RuntimeValidationEnterExitState = false;
+		}
+		EnumAddFlags(Found->State, EState::BetweenEnterExitState);
+	}
+	else
+	{
+		NodeStates.Add(FNodeStatePair{.Id = Id, .State = EState::BetweenEnterExitState });
+	}
+}
+
+void FRuntimeValidationInstanceData::NodeExitState(FGuid Id)
+{
+	FNodeStatePair* Found = NodeStates.FindByPredicate([Id](const FNodeStatePair& Other)
+	{
+		return Other.Id == Id;
+	});
+	if (Found)
+	{
+		if (Private::RuntimeValidationEnterExitState && !EnumHasAllFlags(Found->State, EState::BetweenEnterExitState))
+		{
+			ensureAlwaysMsgf(false, TEXT("StateTree runtime check failed: ExitState executed on node %s without an EnterState.\n"
+				"Auto deactivate Runtime check StateTree.RuntimeValidation.EnterExitState to prevent reporting the same error multiple times.")
+				, *Private::NodeToString(Owner.Get(), Id));
+			Private::RuntimeValidationEnterExitState = false;
+		}
+		EnumRemoveFlags(Found->State, EState::BetweenEnterExitState);
+	}
+	else if (Private::RuntimeValidationEnterExitState)
+	{
+		ensureAlwaysMsgf(false, TEXT("StateTree runtime check failed: ExitState executed on node %s without an EnterState.\n"
+			"Auto deactivate Runtime check StateTree.RuntimeValidation.EnterExitState to prevent reporting the same error multiple times.")
+			, *Private::NodeToString(Owner.Get(), Id));
+		Private::RuntimeValidationEnterExitState = false;
+	}
+}
+#endif //WITH_STATETREE_DEBUG
+
+/**
+ * 
+ */
+
+
+#if WITH_STATETREE_DEBUG
+FRuntimeValidation::FRuntimeValidation(gsl::not_null<FRuntimeValidationInstanceData*> Instance)
+	: RuntimeValidationData(Instance)
+{
+}
+
+FRuntimeValidationInstanceData* FRuntimeValidation::GetInstanceData() const
+{
+	return RuntimeValidationData;
+}
+
+void FRuntimeValidation::SetContext(const UObject* Owner, const UStateTree* StateTree) const
+{
+	if (RuntimeValidationData)
+	{
+		RuntimeValidationData->SetContext(Owner, StateTree);
+	}
+}
+
+#else
+
+void FRuntimeValidation::SetContext(const UObject*, const UStateTree*) const
+{
+}
+
+#endif //WITH_STATETREE_DEBUG
+
+} // UE::StateTree::Debug
+
+
diff --git a/Engine/Plugins/Runtime/StateTree/Source/StateTreeModule/Private/Debugger/StateTreeRuntimeValidationInstanceData.h b/Engine/Plugins/Runtime/StateTree/Source/StateTreeModule/Private/Debugger/StateTreeRuntimeValidationInstanceData.h
new file mode 100644
--- /dev/null
+++ b/Engine/Plugins/Runtime/StateTree/Source/StateTreeModule/Private/Debugger/StateTreeRuntimeValidationInstanceData.h
@@ -0,0 +1,50 @@
+// Copyright Epic Games, Inc. All Rights Reserved.
+
+#pragma once
+
+#include "StateTreeTypes.h"
+#include "Debugger/StateTreeDebug.h"
+
+class UStateTree;
+class UObject;
+
+#if WITH_STATETREE_DEBUG
+
+namespace UE::StateTree::Debug
+{
+/**
+ * For debugging purposes. Data used for runtime check.
+*/
+class FRuntimeValidationInstanceData
+{
+public:
+	FRuntimeValidationInstanceData() = default;
+	~FRuntimeValidationInstanceData();
+
+	void SetContext(const UObject* Owner, const UStateTree* StateTree);
+	void NodeEnterState(FGuid Id);
+	void NodeExitState(FGuid Id);
+
+private:
+	enum class EState : uint8
+	{
+		None = 0x00,
+		BetweenEnterExitState = 0x01,
+	};
+	FRIEND_ENUM_CLASS_FLAGS(EState);
+
+	struct FNodeStatePair
+	{
+		FGuid Id;
+		EState State = EState::None;
+	};
+	TArray<FNodeStatePair> NodeStates;
+	TWeakObjectPtr<const UStateTree> StateTree;
+	TWeakObjectPtr<const UObject> Owner;
+};
+
+ENUM_CLASS_FLAGS(FRuntimeValidationInstanceData::EState);
+
+} // UE::StateTree::Debug
+
+#endif // WITH_STATETREE_DEBUG
\ No newline at end of file
diff --git a/Engine/Plugins/Runtime/StateTree/Source/StateTreeModule/Private/StateTreeExecutionContext.cpp b/Engine/Plugins/Runtime/StateTree/Source/StateTreeModule/Private/StateTreeExecutionContext.cpp
--- a/Engine/Plugins/Runtime/StateTree/Source/StateTreeModule/Private/StateTreeExecutionContext.cpp
+++ b/Engine/Plugins/Runtime/StateTree/Source/StateTreeModule/Private/StateTreeExecutionContext.cpp
@@ -8,6 +8,7 @@
 #include "StateTreePropertyFunctionBase.h"
 #include "StateTreeReference.h"
 #include "Containers/StaticArray.h"
+#include "Debugger/StateTreeDebug.h"
 #include "Debugger/StateTreeTrace.h"
 #include "Debugger/StateTreeTraceTypes.h"
 #include "Misc/ScopeExit.h"
@@ -117,6 +118,8 @@ FStateTreeExecutionContext::FStateTreeExecutionContext(UObject& InOwner, const U
 
 		InstanceDataStorage = &InstanceData.GetMutableStorage();
 		check(InstanceDataStorage);
+
+		InstanceDataStorage->GetRuntimeValidation().SetContext(&InOwner, &InStateTree);
 		
 		EventQueue = InstanceData.GetSharedMutableEventQueue();
 
@@ -265,6 +268,8 @@ EStateTreeRunStatus FStateTreeExecutionContext::Start(const FInstancedPropertyBa
 	// Initialize instance data. No active states yet, so we'll initialize the evals and global tasks.
 	InstanceData.Reset();
 
+	InstanceDataStorage->GetRuntimeValidation().SetContext(&Owner, &RootStateTree);
+
 	if (!InitialParameters || !SetGlobalParameters(*InitialParameters))
 	{
 		SetGlobalParameters(RootStateTree.GetDefaultParameters());
@@ -434,7 +439,7 @@ EStateTreeRunStatus FStateTreeExecutionContext::Stop(EStateTreeRunStatus Complet
 		// Transition to Succeeded state.
 		FStateTreeTransitionResult Transition;
 		Transition.TargetState = FStateTreeStateHandle::FromCompletionStatus(CompletionStatus);
-		Transition.CurrentRunStatus = CompletionStatus;		
+		Transition.CurrentRunStatus = CompletionStatus;	
 		ExitState(Transition);
 
 		// No active states or global tasks anymore, reset frames.
@@ -1616,10 +1621,10 @@ EStateTreeRunStatus FStateTreeExecutionContext::EnterState(FStateTreeTransitionR
 		FCurrentlyProcessedFrameScope FrameScope(*this, CurrentParentFrame, CurrentFrame);
 		const UStateTree* CurrentStateTree = NextFrame.StateTree;
 
-		for (int32 Index = 0; Index < NextFrame.ActiveStates.Num() && Result != EStateTreeRunStatus::Failed; Index++)
+		for (int32 StateIndex = 0; StateIndex < NextFrame.ActiveStates.Num() && Result != EStateTreeRunStatus::Failed; StateIndex++)
 		{
-			const FStateTreeStateHandle CurrentHandle = NextFrame.ActiveStates[Index];
-			const FStateTreeStateHandle PreviousHandle = PreviousActiveStates.GetStateSafe(Index);
+			const FStateTreeStateHandle CurrentHandle = NextFrame.ActiveStates[StateIndex];
+			const FStateTreeStateHandle PreviousHandle = PreviousActiveStates.GetStateSafe(StateIndex);
 			const FCompactStateTreeState& State = CurrentStateTree->States[CurrentHandle.Index];
 
 			FCurrentlyProcessedStateScope StateScope(*this, CurrentHandle);
@@ -1656,8 +1661,8 @@ EStateTreeRunStatus FStateTreeExecutionContext::EnterState(FStateTreeTransitionR
 			if (bIsEnteringState)
 			{
 				STATETREE_TRACE_STATE_EVENT(CurrentHandle, EStateTreeTraceEventType::OnEntering);
-				STATETREE_LOG(Log, TEXT("%*sState '%s' %s"), Index*UE::StateTree::DebugIndentSize, TEXT(""),
-					*DebugGetStatePath(Transition.NextActiveFrames, &NextFrame, Index),
+				STATETREE_LOG(Log, TEXT("%*sState '%s' %s"), StateIndex*UE::StateTree::DebugIndentSize, TEXT(""),
+					*DebugGetStatePath(Transition.NextActiveFrames, &NextFrame, StateIndex),
 					*UEnum::GetDisplayValueAsText(CurrentTransition.ChangeType).ToString());
 			}
 
@@ -1682,7 +1687,8 @@ EStateTreeRunStatus FStateTreeExecutionContext::EnterState(FStateTreeTransitionR
 								// Use validated copy, since we test in situations where the sources are not always valid (e.g. enter conditions may try to access inactive parent state).
 								CopyBatchOnActiveInstances(CurrentParentFrame, CurrentFrame, ConditionInstanceView, Cond.BindingsBatch);
 							}
-								
+
+							UE_STATETREE_DEBUG_NODE_ENTERSTATE(this, CurrentFrame.StateTree, FStateTreeIndex16(ConditionIndex));
 							Cond.EnterState(*this, Transition);
 
 							// Reset copied properties that might contain object references.
@@ -1721,7 +1727,8 @@ EStateTreeRunStatus FStateTreeExecutionContext::EnterState(FStateTreeTransitionR
 
 				if (bIsEnteringState && bShouldCallStateChange)
 				{
-					STATETREE_LOG(Verbose, TEXT("%*s  Task '%s'"), Index*UE::StateTree::DebugIndentSize, TEXT(""), *Task.Name.ToString());
+					STATETREE_LOG(Verbose, TEXT("%*s  Task '%s'"), StateIndex*UE::StateTree::DebugIndentSize, TEXT(""), *Task.Name.ToString());
+					UE_STATETREE_DEBUG_NODE_ENTERSTATE(this, NextFrame.StateTree, FStateTreeIndex16(TaskIndex));
 
 					EStateTreeRunStatus Status = EStateTreeRunStatus::Unset;
 					{
@@ -1954,8 +1961,8 @@ void FStateTreeExecutionContext::ExitState(const FStateTreeTransitionResult& Tra
 
 						if (bShouldCallStateChange)
 						{
-
 							STATETREE_LOG(Verbose, TEXT("%*s  Task '%s'"), Index*UE::StateTree::DebugIndentSize, TEXT(""), *Task.Name.ToString());
+							UE_STATETREE_DEBUG_NODE_EXITSTATE(this, CurrentStateTree, FStateTreeIndex16(TaskIndex));
 							{
 								QUICK_SCOPE_CYCLE_COUNTER(StateTree_Task_ExitState);
 								CSV_SCOPED_TIMING_STAT_EXCLUSIVE(StateTree_Task_ExitState);
@@ -1987,7 +1994,8 @@ void FStateTreeExecutionContext::ExitState(const FStateTreeTransitionResult& Tra
 									// Use validated copy, since we test in situations where the sources are not always valid (e.g. enter conditions may try to access inactive parent state).
 									CopyBatchOnActiveInstances(CurrentParentFrame, CurrentFrame, ConditionInstanceView, Cond.BindingsBatch);
 								}
-								
+
+								UE_STATETREE_DEBUG_NODE_EXITSTATE(this, CurrentFrame.StateTree, FStateTreeIndex16(ConditionIndex));
 								Cond.ExitState(*this, Transition);
 
 								// Reset copied properties that might contain object references.
@@ -2283,9 +2291,10 @@ EStateTreeRunStatus FStateTreeExecutionContext::StartEvaluatorsAndGlobalTasks(FS
 				}
 
 				STATETREE_LOG(Verbose, TEXT("  Start: '%s'"), *Task.Name.ToString());
+				UE_STATETREE_DEBUG_NODE_ENTERSTATE(this, CurrentStateTree, FStateTreeIndex16(TaskIndex));
 				{
 					QUICK_SCOPE_CYCLE_COUNTER(StateTree_Task_TreeStart);
-					const EStateTreeRunStatus TaskStatus = Task.EnterState(*this, Transition); 
+					const EStateTreeRunStatus TaskStatus = Task.EnterState(*this, Transition);
 
 					STATETREE_TRACE_TASK_EVENT(TaskIndex, TaskInstanceView, EStateTreeTraceEventType::OnEntered, TaskStatus);
 
@@ -2395,6 +2404,7 @@ void FStateTreeExecutionContext::CallStopOnEvaluatorsAndGlobalTasks(const FState
 		if (TaskIndex <= LastInitializedTaskIndex.Get())
 		{
 			STATETREE_LOG(Verbose, TEXT("  Stop: '%s'"), *Task.Name.ToString());
+			UE_STATETREE_DEBUG_NODE_EXITSTATE(this, CurrentStateTree, FStateTreeIndex16(TaskIndex));
 			{
 				QUICK_SCOPE_CYCLE_COUNTER(StateTree_Task_TreeStop);
 				Task.ExitState(*this, Transition);
@@ -2500,9 +2510,10 @@ EStateTreeRunStatus FStateTreeExecutionContext::StartTemporaryEvaluatorsAndGloba
 			CopyBatchOnActiveInstances(CurrentParentFrame, CurrentFrame, TaskDataView, Task.BindingsBatch);
 		}
 
-		STATETREE_LOG(Verbose, TEXT("  Start: '%s'"), *Task.Name.ToString());
 		if (bWasCreated)
 		{
+			STATETREE_LOG(Verbose, TEXT("  Start: '%s'"), *Task.Name.ToString());
+			UE_STATETREE_DEBUG_NODE_ENTERSTATE(this, CurrentStateTree, FStateTreeIndex16(TaskIndex));
 			QUICK_SCOPE_CYCLE_COUNTER(StateTree_Task_TreeStart);
 			const EStateTreeRunStatus TaskStatus = Task.EnterState(*this, Transition);
 
@@ -2559,6 +2570,7 @@ void FStateTreeExecutionContext::StopTemporaryEvaluatorsAndGlobalTasks(const FSt
 			if (const FStateTreeTaskBase* Task = NodeView.GetPtr<const FStateTreeTaskBase>())
 			{
 				STATETREE_LOG(Verbose, TEXT("  Stop: '%s'"), *Task->Name.ToString());
+				UE_STATETREE_DEBUG_NODE_EXITSTATE(this, CurrentFrame.StateTree, FStateTreeIndex16(Index));
 				{
 					QUICK_SCOPE_CYCLE_COUNTER(StateTree_Task_TreeStop);
 					Task->ExitState(*this, Transition);
diff --git a/Engine/Plugins/Runtime/StateTree/Source/StateTreeModule/Private/StateTreeInstanceData.cpp b/Engine/Plugins/Runtime/StateTree/Source/StateTreeModule/Private/StateTreeInstanceData.cpp
--- a/Engine/Plugins/Runtime/StateTree/Source/StateTreeModule/Private/StateTreeInstanceData.cpp
+++ b/Engine/Plugins/Runtime/StateTree/Source/StateTreeModule/Private/StateTreeInstanceData.cpp
@@ -2,6 +2,7 @@
 
 #include "StateTreeInstanceData.h"
 #include "StateTreeExecutionTypes.h"
+#include "Debugger/StateTreeRuntimeValidationInstanceData.h"
 #include "Serialization/MemoryReader.h"
 #include "Serialization/MemoryWriter.h"
 #include "VisualLogger/VisualLogger.h"
@@ -83,7 +84,14 @@ namespace UE::StateTree
 // FStateTreeInstanceStorage
 //----------------------------------------------------------------//
 
+#if WITH_STATETREE_DEBUG
+FStateTreeInstanceStorage::FStateTreeInstanceStorage()
+	: RuntimeValidationData(MakePimpl<UE::StateTree::Debug::FRuntimeValidationInstanceData>())
+{
+}
+#else
 FStateTreeInstanceStorage::FStateTreeInstanceStorage() = default;
+#endif
 
 FStateTreeInstanceStorage::FStateTreeInstanceStorage(const FStateTreeInstanceStorage& Other)
 	: InstanceStructs(Other.InstanceStructs)
@@ -92,6 +100,9 @@ FStateTreeInstanceStorage::FStateTreeInstanceStorage(const FStateTreeInstanceSto
 	, EventQueue(MakeShared<FStateTreeEventQueue>(*Other.EventQueue))
 	, TransitionRequests(Other.TransitionRequests)
 	, GlobalParameters(Other.GlobalParameters)
+#if WITH_STATETREE_DEBUG
+	, RuntimeValidationData(MakePimpl<UE::StateTree::Debug::FRuntimeValidationInstanceData>(*Other.RuntimeValidationData))
+#endif
 {
 }
 
@@ -102,8 +113,14 @@ FStateTreeInstanceStorage::FStateTreeInstanceStorage(FStateTreeInstanceStorage&&
 	, EventQueue(Other.EventQueue)
 	, TransitionRequests(MoveTemp(Other.TransitionRequests))
 	, GlobalParameters(MoveTemp(Other.GlobalParameters))
+#if WITH_STATETREE_DEBUG
+	, RuntimeValidationData(MoveTemp(Other.RuntimeValidationData))
+#endif
 {
 	Other.EventQueue = MakeShared<FStateTreeEventQueue>();
+#if WITH_STATETREE_DEBUG
+	Other.RuntimeValidationData = MakePimpl<UE::StateTree::Debug::FRuntimeValidationInstanceData>();
+#endif
 }
 
 FStateTreeInstanceStorage& FStateTreeInstanceStorage::operator=(const FStateTreeInstanceStorage& Other)
@@ -267,6 +284,19 @@ void FStateTreeInstanceStorage::Reset()
 	ExecutionState.Reset();
 	TemporaryInstances.Reset();
 	GlobalParameters.Reset();
+
+#if WITH_STATETREE_DEBUG
+	RuntimeValidationData = MakePimpl<UE::StateTree::Debug::FRuntimeValidationInstanceData>();
+#endif
+}
+
+UE::StateTree::Debug::FRuntimeValidation FStateTreeInstanceStorage::GetRuntimeValidation() const
+{
+#if WITH_STATETREE_DEBUG
+	return UE::StateTree::Debug::FRuntimeValidation(RuntimeValidationData.Get());
+#else
+	return UE::StateTree::Debug::FRuntimeValidation();
+#endif
 }
 
 //----------------------------------------------------------------//
diff --git a/Engine/Plugins/Runtime/StateTree/Source/StateTreeModule/Public/Debugger/StateTreeDebug.h b/Engine/Plugins/Runtime/StateTree/Source/StateTreeModule/Public/Debugger/StateTreeDebug.h
new file mode 100644
--- /dev/null
+++ b/Engine/Plugins/Runtime/StateTree/Source/StateTreeModule/Public/Debugger/StateTreeDebug.h
@@ -0,0 +1,65 @@
+// Copyright Epic Games, Inc. All Rights Reserved.
+
+#pragma once
+
+#include "CoreMinimal.h"
+#include "StateTreeTypes.h"
+#include "Templates/NotNull.h"
+
+class UStateTree;
+struct FStateTreeExecutionContext;
+
+#if WITH_STATETREE_DEBUG
+
+#define UE_STATETREE_DEBUG_NODE_ENTERSTATE(InExecutionContextPtr, InStateTree, InNodeIndex) ::UE::StateTree::Debug::FStateTreeDebug::NodeEnterState(*(InExecutionContextPtr), ::UE::StateTree::Debug::FStateTreeDebug::FNodeReference((InStateTree), (InNodeIndex)));
+#define UE_STATETREE_DEBUG_NODE_EXITSTATE(InExecutionContextPtr, InStateTree, InNodeIndex) ::UE::StateTree::Debug::FStateTreeDebug::NodeExitState(*(InExecutionContextPtr), ::UE::StateTree::Debug::FStateTreeDebug::FNodeReference((InStateTree), (InNodeIndex)));
+
+namespace UE::StateTree::Debug
+{
+class STATETREEMODULE_API FStateTreeDebug final
+{
+public:
+	struct FNodeReference
+	{
+		FNodeReference(gsl::not_null<const UStateTree*> InStateTree, FStateTreeIndex16 InNodeIndex);
+		gsl::not_null<const UStateTree*> StateTree;
+		FStateTreeIndex16 Index;
+	};
+	
+	static void NodeEnterState(const FStateTreeExecutionContext& ExecutionContext, FNodeReference Node);
+	static void NodeExitState(const FStateTreeExecutionContext& ExecutionContext, FNodeReference Node);
+
+public:
+	struct FNodeDelegateArgs
+	{
+		FNodeReference Node;
+		FGuid NodeId;
+	};
+	
+	DECLARE_TS_MULTICAST_DELEGATE_TwoParams(FNodeDelegate, const FStateTreeExecutionContext& ExecutionContext, FNodeDelegateArgs Node);
+	/**
+	 * Debugging callback for when a node activates.
+	 * @note The callback executes inside the StateTree logic.
+	 * @note The StateTree can execute on any thread.
+	 */
+	static FNodeDelegate OnNodeEnterState_AnyThread;
+	
+	/**
+	 * Debugging callback for when a node deactivates.
+	 * @note The callback executes inside the StateTree logic.
+	 * @note The StateTree can execute on any thread.
+	 */
+	static FNodeDelegate OnNodeExitState_AnyThread;
+	
+private:
+	FStateTreeDebug() = default;
+};
+
+} //UE::StateTree::Debug
+
+#else
+
+#define UE_STATETREE_DEBUG_NODE_ENTERSTATE(...)
+#define UE_STATETREE_DEBUG_NODE_EXITSTATE(...)
+
+#endif // WITH_STATETREE_DEBUG
diff --git a/Engine/Plugins/Runtime/StateTree/Source/StateTreeModule/Public/Debugger/StateTreeRuntimeValidation.h b/Engine/Plugins/Runtime/StateTree/Source/StateTreeModule/Public/Debugger/StateTreeRuntimeValidation.h
new file mode 100644
--- /dev/null
+++ b/Engine/Plugins/Runtime/StateTree/Source/StateTreeModule/Public/Debugger/StateTreeRuntimeValidation.h
@@ -0,0 +1,42 @@
+// Copyright Epic Games, Inc. All Rights Reserved.
+
+#pragma once
+
+#include "Debugger/StateTreeDebug.h"
+
+class UStateTree;
+class UObject;
+
+#if WITH_STATETREE_DEBUG
+namespace UE::StateTree::Debug
+{
+	class FRuntimeValidationInstanceData;
+}
+#endif
+
+namespace UE::StateTree::Debug
+{
+
+/**
+ * For debugging purposes, test the actions that are made on the StateTree instance.
+ */
+struct FRuntimeValidation
+{
+public:
+	FRuntimeValidation() = default;
+
+#if WITH_STATETREE_DEBUG
+	FRuntimeValidation(gsl::not_null<FRuntimeValidationInstanceData*> Instance);
+	FRuntimeValidationInstanceData* GetInstanceData() const;
+#endif
+
+	/** Set the Owner of the data */
+	void SetContext(const UObject* Owner, const UStateTree* StateTree) const;
+
+private:
+#if WITH_STATETREE_DEBUG
+	FRuntimeValidationInstanceData* RuntimeValidationData = nullptr;
+#endif // WITH_STATETREE_DEBUG
+};
+
+} // UE::StateTree::Debug
diff --git a/Engine/Plugins/Runtime/StateTree/Source/StateTreeModule/Public/StateTreeInstanceData.h b/Engine/Plugins/Runtime/StateTree/Source/StateTreeModule/Public/StateTreeInstanceData.h
--- a/Engine/Plugins/Runtime/StateTree/Source/StateTreeModule/Public/StateTreeInstanceData.h
+++ b/Engine/Plugins/Runtime/StateTree/Source/StateTreeModule/Public/StateTreeInstanceData.h
@@ -3,15 +3,24 @@
 #pragma once
 
 #include "StructUtils/InstancedStructContainer.h"
+#include "Debugger/StateTreeRuntimeValidation.h"
 #include "StateTreeEvents.h"
 #include "StateTreeTypes.h"
 #include "StateTreeExecutionTypes.h"
+#include "Templates/PimplPtr.h"
 #include "Templates/SharedPointer.h"
 #include "StateTreeInstanceData.generated.h"
 
 struct FStateTreeTransitionRequest;
 struct FStateTreeExecutionState;
 
+#if WITH_STATETREE_DEBUG
+namespace UE::StateTree::Debug
+{
+	class FRuntimeValidationInstanceData;
+}
+#endif
+
 /** Wrapper class to store an object amongst the structs. */
 USTRUCT()
 struct STATETREEMODULE_API FStateTreeInstanceObjectWrapper
@@ -238,6 +247,9 @@ struct STATETREEMODULE_API FStateTreeInstanceStorage
 
 	/** Resets the storage to initial state. */
 	void Reset();
+
+	/** Get the data used at runtime to confirm the inner working of the StateTree. */
+	UE::StateTree::Debug::FRuntimeValidation GetRuntimeValidation() const;
 	
 	UE_DEPRECATED(5.4, "Use Num() instead.")
 	int32 NumStructs() const { return 0; }
@@ -266,9 +278,6 @@ protected:
 
 	/** Events (Transient) */
 	TSharedRef<FStateTreeEventQueue> EventQueue = MakeShared<FStateTreeEventQueue>();
-
-	/* True if the storage owns the event queue. */
-	bool bIsOwningEventQueue = true;
 	
 	/** Requested transitions */
 	UPROPERTY(Transient)
@@ -278,6 +287,13 @@ protected:
 	UPROPERTY(Transient)
 	FInstancedPropertyBag GlobalParameters;
 
+	/* True if the storage owns the event queue. */
+	bool bIsOwningEventQueue = true;
+
+#if WITH_STATETREE_DEBUG
+	TPimplPtr<UE::StateTree::Debug::FRuntimeValidationInstanceData> RuntimeValidationData;
+#endif
+
 	friend struct FStateTreeInstanceData;
 };
 
@@ -444,6 +460,12 @@ struct STATETREEMODULE_API FStateTreeInstanceData
 	{
 		return GetMutableStorage().ResetTemporaryInstances();
 	}
+
+	/** Get the data used at runtime to confirm the inner working of the StateTree. */
+	UE::StateTree::Debug::FRuntimeValidation GetRuntimeValidation() const
+	{
+		return GetStorage().GetRuntimeValidation();
+	}
 	
 	UE_DEPRECATED(5.4, "Use the structs only Init(), objects should be wrapped in FStateTreeInstanceObjectWrapper.")
 	void Init(UObject& InOwner, TConstArrayView<FInstancedStruct> InStructs, TConstArrayView<const UObject*> InObjects);
