StateTree: Add per project per user settings. With that setting, a user can configure its StateTree UX. Add settings to display condition and transition in the state view.

diff --git a/Engine/Plugins/Runtime/StateTree/Source/StateTreeEditorModule/Private/SStateTreeView.cpp b/Engine/Plugins/Runtime/StateTree/Source/StateTreeEditorModule/Private/SStateTreeView.cpp
--- a/Engine/Plugins/Runtime/StateTree/Source/StateTreeEditorModule/Private/SStateTreeView.cpp
+++ b/Engine/Plugins/Runtime/StateTree/Source/StateTreeEditorModule/Private/SStateTreeView.cpp
@@ -4,13 +4,23 @@
 #include "Debugger/StateTreeDebuggerTypes.h"
 #include "Framework/Commands/UICommandList.h"
 #include "Framework/MultiBox/MultiBoxBuilder.h"
+#include "Modules/ModuleManager.h"
+#include "SEnumCombo.h"
 #include "SPositiveActionButton.h"
 #include "SStateTreeViewRow.h"
 #include "StateTreeViewModel.h"
 #include "StateTreeState.h"
 #include "StateTreeEditorCommands.h"
-#include "StateTreeSettings.h"
+#include "StateTreeEditorUserSettings.h"
 #include "Widgets/Layout/SScrollBox.h"
+#include "Widgets/Layout/SSpacer.h"
+
+#include "DetailCategoryBuilder.h"
+#include "DetailLayoutBuilder.h"
+#include "DetailWidgetRow.h"
+#include "IDetailCustomization.h"
+#include "IDetailsView.h"
+#include "PropertyEditorModule.h"
 
 #define LOCTEXT_NAMESPACE "StateTreeEditor"
 
@@ -23,14 +33,19 @@ SStateTreeView::SStateTreeView()
 
 SStateTreeView::~SStateTreeView()
 {
-	if (StateTreeViewModel)
+	if (UObjectInitialized())
 	{
-		StateTreeViewModel->GetOnAssetChanged().RemoveAll(this);
-		StateTreeViewModel->GetOnStatesRemoved().RemoveAll(this);
-		StateTreeViewModel->GetOnStatesMoved().RemoveAll(this);
-		StateTreeViewModel->GetOnStateAdded().RemoveAll(this);
-		StateTreeViewModel->GetOnStatesChanged().RemoveAll(this);
-		StateTreeViewModel->GetOnSelectionChanged().RemoveAll(this);
+		GetMutableDefault<UStateTreeEditorUserSettings>()->OnSettingsChanged.Remove(SettingsChangedHandle);
+
+		if (StateTreeViewModel)
+		{
+			StateTreeViewModel->GetOnAssetChanged().RemoveAll(this);
+			StateTreeViewModel->GetOnStatesRemoved().RemoveAll(this);
+			StateTreeViewModel->GetOnStatesMoved().RemoveAll(this);
+			StateTreeViewModel->GetOnStateAdded().RemoveAll(this);
+			StateTreeViewModel->GetOnStatesChanged().RemoveAll(this);
+			StateTreeViewModel->GetOnSelectionChanged().RemoveAll(this);
+		}
 	}
 }
 
@@ -45,6 +60,8 @@ void SStateTreeView::Construct(const FArguments& InArgs, TSharedRef<FStateTreeVi
 	StateTreeViewModel->GetOnStatesChanged().AddSP(this, &SStateTreeView::HandleModelStatesChanged);
 	StateTreeViewModel->GetOnSelectionChanged().AddSP(this, &SStateTreeView::HandleModelSelectionChanged);
 
+	SettingsChangedHandle = GetMutableDefault<UStateTreeEditorUserSettings>()->OnSettingsChanged.AddSP(this, &SStateTreeView::HandleUserSettingsChanged);
+
 	bUpdatingSelection = false;
 
 	TSharedRef<SScrollBar> HorizontalScrollBar = SNew(SScrollBar)
@@ -93,6 +110,32 @@ void SStateTreeView::Construct(const FArguments& InArgs, TSharedRef<FStateTreeVi
 					.Text(LOCTEXT("AddState", "Add State"))
 					.OnClicked(this, &SStateTreeView::HandleAddStateButton)
 				]
+
+				+ SHorizontalBox::Slot()
+				.FillWidth(1.0f)
+				[
+					SNew(SSpacer)
+				]
+
+				+ SHorizontalBox::Slot()
+				.AutoWidth()
+				.HAlign(HAlign_Right)
+				[
+					SNew(SComboButton)
+					.HasDownArrow(false)
+					.ContentPadding(0.0f)
+					.ForegroundColor(FSlateColor::UseForeground())
+					.ButtonStyle(FAppStyle::Get(), "SimpleButton")
+					.MenuContent()
+					[
+						HandleGenerateSettingsMenu()
+					]
+					.ButtonContent()
+					[
+						SNew(SImage)
+						.Image(FAppStyle::GetBrush("DetailsView.ViewOptions"))
+					]
+				]
 			]
 		]
 
@@ -367,7 +410,10 @@ void SStateTreeView::UpdateTree(bool bExpandPersistent)
 	bItemsDirty = false;
 }
 
-
+void SStateTreeView::HandleUserSettingsChanged()
+{
+	TreeView->RebuildList();
+}
 
 void SStateTreeView::HandleModelAssetChanged()
 {
@@ -486,6 +532,66 @@ void SStateTreeView::HandleTreeExpansionChanged(TWeakObjectPtr<UStateTreeState>
 	}
 }
 
+TSharedRef<SWidget> SStateTreeView::HandleGenerateSettingsMenu()
+{
+	FDetailsViewArgs DetailsViewArgs;
+	DetailsViewArgs.bUpdatesFromSelection = false;
+	DetailsViewArgs.bLockable = false;
+	DetailsViewArgs.bShowPropertyMatrixButton = false;
+	DetailsViewArgs.bAllowSearch = false;
+	DetailsViewArgs.bShowOptions = false;
+	DetailsViewArgs.NameAreaSettings = FDetailsViewArgs::HideNameArea;
+	DetailsViewArgs.ViewIdentifier = NAME_None;
+
+	FPropertyEditorModule& PropertyEditorModule = FModuleManager::GetModuleChecked<FPropertyEditorModule>("PropertyEditor");
+	TSharedRef<IDetailsView> DetailsView = PropertyEditorModule.CreateDetailView(DetailsViewArgs);
+
+	DetailsView->RegisterInstancedCustomPropertyLayout(UStateTreeEditorUserSettings::StaticClass(), FOnGetDetailCustomizationInstance::CreateLambda([]()
+		{
+			class FStateTreeEditorUserSettingsDetailsCustomication : public IDetailCustomization
+			{
+			public:
+				virtual void CustomizeDetails(IDetailLayoutBuilder& DetailLayout) override
+				{
+					DetailLayout.HideCategory("OtherStuff");
+					{
+						IDetailCategoryBuilder& CategoryBuilder = DetailLayout.EditCategory("State View");
+						TArray<TSharedRef<IPropertyHandle>> AllProperties;
+						CategoryBuilder.GetDefaultProperties(AllProperties);
+
+						const FName PropertyToFind = "StatesViewDisplayNodeType";
+						TSharedRef<IPropertyHandle>* FoundProperty = AllProperties.FindByPredicate([PropertyToFind](TSharedRef<IPropertyHandle>& Other)
+							{
+								return Other->GetProperty()->GetFName() == PropertyToFind;
+							});
+						if (ensure(FoundProperty))
+						{
+							CategoryBuilder.AddProperty(*FoundProperty).CustomWidget()
+							.NameContent()
+							[
+								(*FoundProperty)->CreatePropertyNameWidget()
+							]
+							.ValueContent()
+							[
+								SNew(SEnumComboBox, StaticEnum<EStateTreeEditorUserSettingsNodeType>())
+									.OnEnumSelectionChanged(SEnumComboBox::FOnEnumSelectionChanged::CreateLambda([](int32 NewValue, ESelectInfo::Type)
+										{
+											GetMutableDefault<UStateTreeEditorUserSettings>()->SetStatesViewDisplayNodeType((EStateTreeEditorUserSettingsNodeType)NewValue);
+										}))
+									.CurrentValue(MakeAttributeLambda([]() { return (int32)GetDefault<UStateTreeEditorUserSettings>()->GetStatesViewDisplayNodeType(); }))
+											.Font(IDetailLayoutBuilder::GetDetailFont())
+							];
+						}
+					}
+				}
+			};
+			return MakeShared<FStateTreeEditorUserSettingsDetailsCustomication>();
+		}));
+
+	DetailsView->SetObject(GetMutableDefault<UStateTreeEditorUserSettings>());
+	return DetailsView;
+}
+
 TSharedPtr<SWidget> SStateTreeView::HandleContextMenuOpening()
 {
 	if (!StateTreeViewModel)
diff --git a/Engine/Plugins/Runtime/StateTree/Source/StateTreeEditorModule/Private/SStateTreeView.h b/Engine/Plugins/Runtime/StateTree/Source/StateTreeEditorModule/Private/SStateTreeView.h
--- a/Engine/Plugins/Runtime/StateTree/Source/StateTreeEditorModule/Private/SStateTreeView.h
+++ b/Engine/Plugins/Runtime/StateTree/Source/StateTreeEditorModule/Private/SStateTreeView.h
@@ -15,13 +15,6 @@ class UStateTreeState;
 class SScrollBox;
 class FUICommandList;
 
-namespace UE::StateTree::Editor
-{
-static constexpr float StateRowHeight = 32.0f;
-static constexpr float TaskRowHeight = 16.0f;
-
-} // UE::StateTree::Editor
-
 class SStateTreeView : public SCompoundWidget
 {
 public:
@@ -45,7 +38,10 @@ private:
 
 	void UpdateTree(bool bExpandPersistent = false);
 
-	// ViewModel handlers
+	//~ Editor/User settings handlers
+	void HandleUserSettingsChanged();
+
+	//~ ViewModel handlers
 	void HandleModelAssetChanged();
 	void HandleModelStatesRemoved(const TSet<UStateTreeState*>& AffectedParents);
 	void HandleModelStatesMoved(const TSet<UStateTreeState*>& AffectedParents, const TSet<UStateTreeState*>& MovedStates);
@@ -53,16 +49,17 @@ private:
 	void HandleModelStatesChanged(const TSet<UStateTreeState*>& AffectedStates, const FPropertyChangedEvent& PropertyChangedEvent);
 	void HandleModelSelectionChanged(const TArray<TWeakObjectPtr<UStateTreeState>>& SelectedStates);
 
-	// Treeview handlers
+	//~ Treeview handlers
 	TSharedRef<ITableRow> HandleGenerateRow(TWeakObjectPtr<UStateTreeState> InState, const TSharedRef<STableViewBase>& InOwnerTableView);
 	void HandleGetChildren(TWeakObjectPtr<UStateTreeState> InParent, TArray<TWeakObjectPtr<UStateTreeState>>& OutChildren);
 	void HandleTreeSelectionChanged(TWeakObjectPtr<UStateTreeState> InSelectedItem, ESelectInfo::Type SelectionType);
 	void HandleTreeExpansionChanged(TWeakObjectPtr<UStateTreeState> InSelectedItem, bool bExpanded);
 	
 	TSharedPtr<SWidget> HandleContextMenuOpening();
+	TSharedRef<SWidget> HandleGenerateSettingsMenu();
 
-	// Action handlers
-	// @todo: these are also defined in the outliner, figure out how to share code.
+	//~ Action handlers
+	//~ @todo: these are also defined in the outliner, figure out how to share code.
 	UStateTreeState* GetFirstSelectedState() const;
 	FReply HandleAddStateButton();
 	void HandleAddSiblingState();
@@ -94,6 +91,7 @@ private:
 	TSharedPtr<FUICommandList> CommandList;
 
 	UStateTreeState* RequestedRenameState;
+	FDelegateHandle SettingsChangedHandle;
 	bool bItemsDirty;
 	bool bUpdatingSelection;
 };
diff --git a/Engine/Plugins/Runtime/StateTree/Source/StateTreeEditorModule/Private/SStateTreeViewRow.cpp b/Engine/Plugins/Runtime/StateTree/Source/StateTreeEditorModule/Private/SStateTreeViewRow.cpp
--- a/Engine/Plugins/Runtime/StateTree/Source/StateTreeEditorModule/Private/SStateTreeViewRow.cpp
+++ b/Engine/Plugins/Runtime/StateTree/Source/StateTreeEditorModule/Private/SStateTreeViewRow.cpp
@@ -14,6 +14,7 @@
 #include "StateTreeDescriptionHelpers.h"
 #include "StateTreeDragDrop.h"
 #include "StateTreeEditorModule.h"
+#include "StateTreeEditorUserSettings.h"
 #include "StateTreeState.h"
 #include "StateTreeTaskBase.h"
 #include "StateTreeViewModel.h"
@@ -56,9 +57,6 @@ void SStateTreeViewRow::Construct(const FArguments& InArgs, const TSharedRef<STa
 		, InOwnerTableView);
 
 	static constexpr FLinearColor IconTint = FLinearColor(1, 1, 1, 0.5f);
-
-	const float ExpanderArrowTopPadding = FMath::FloorToFloat((UE::StateTree::Editor::StateRowHeight - 16.f) / 2.f);
-
 	
 	auto MakeTransitionWidget = [this](const EStateTreeTransitionTrigger Trigger, const FSlateBrush* Icon)
 	{
@@ -66,7 +64,7 @@ void SStateTreeViewRow::Construct(const FArguments& InArgs, const TSharedRef<STa
 		FilterOptions.bUseMask = EnumHasAnyFlags(Trigger, EStateTreeTransitionTrigger::OnTick | EStateTreeTransitionTrigger::OnEvent);
 		
 			return SNew(SBox)
-			.HeightOverride(UE::StateTree::Editor::StateRowHeight)
+			.HeightOverride(GetDefault<UStateTreeEditorUserSettings>()->GetStatesViewStateRowHeight())
 			.Visibility(this, &SStateTreeViewRow::GetTransitionsVisibility, Trigger)
 			[
 				SNew(SHorizontalBox)
@@ -136,6 +134,9 @@ void SStateTreeViewRow::Construct(const FArguments& InArgs, const TSharedRef<STa
 			];
 	};
 	
+	TSharedPtr<SVerticalBox> StateAndTasksVerticalBox;
+	TSharedPtr<SHorizontalBox> StateHorizontalBox;
+
 	this->ChildSlot
 	.HAlign(HAlign_Fill)
 	[
@@ -176,16 +177,16 @@ void SStateTreeViewRow::Construct(const FArguments& InArgs, const TSharedRef<STa
 			.Padding(FMargin(0, 6, 0, 6))
 			[
 				// State and tasks
-				SNew(SVerticalBox)
+				SAssignNew(StateAndTasksVerticalBox, SVerticalBox)
 				+ SVerticalBox::Slot()
 				.AutoHeight()
 				[
 					// State
 					SNew(SBox)
-					.HeightOverride(UE::StateTree::Editor::StateRowHeight)
+					.HeightOverride(GetDefault<UStateTreeEditorUserSettings>()->GetStatesViewStateRowHeight())
 					.HAlign(HAlign_Left)
 					[
-						SNew(SHorizontalBox)
+						SAssignNew(StateHorizontalBox, SHorizontalBox)
 
 						// State Box
 						+ SHorizontalBox::Slot()
@@ -194,7 +195,7 @@ void SStateTreeViewRow::Construct(const FArguments& InArgs, const TSharedRef<STa
 						.AutoWidth()
 						[
 							SNew(SBox)
-							.HeightOverride(UE::StateTree::Editor::StateRowHeight)
+							.HeightOverride(GetDefault<UStateTreeEditorUserSettings>()->GetStatesViewStateRowHeight())
 							.VAlign(VAlign_Fill)
 							[
 								SNew(SBorder)
@@ -224,7 +225,7 @@ void SStateTreeViewRow::Construct(const FArguments& InArgs, const TSharedRef<STa
 											[
 												SNew(SBox)
 												.WidthOverride(4.f)
-												.HeightOverride(UE::StateTree::Editor::StateRowHeight)
+												.HeightOverride(GetDefault<UStateTreeEditorUserSettings>()->GetStatesViewStateRowHeight())
 												.Visibility(this, &SStateTreeViewRow::GetSubTreeVisibility)
 												.VAlign(VAlign_Fill)
 												.HAlign(HAlign_Fill)
@@ -302,7 +303,7 @@ void SStateTreeViewRow::Construct(const FArguments& InArgs, const TSharedRef<STa
 											.AutoWidth()
 											[
 												SNew(SBox)
-												.HeightOverride(UE::StateTree::Editor::StateRowHeight)
+												.HeightOverride(GetDefault<UStateTreeEditorUserSettings>()->GetStatesViewStateRowHeight())
 												.VAlign(VAlign_Fill)
 												.Visibility(this, &SStateTreeViewRow::GetLinkedStateVisibility)
 												[
@@ -369,78 +370,89 @@ void SStateTreeViewRow::Construct(const FArguments& InArgs, const TSharedRef<STa
 								]
 							]
 						]
-						
-						+ SHorizontalBox::Slot()
-						.VAlign(VAlign_Fill)
-						.HAlign(HAlign_Left)
-						[
-							// Transitions
-							SNew(SHorizontalBox)
+					]
+				]
+			]
+		]
+	];
 
-							+ SHorizontalBox::Slot()
-							.VAlign(VAlign_Top)
-							.AutoWidth()
-							[
-								SNew(SBox)
-								.HeightOverride(UE::StateTree::Editor::StateRowHeight)
-								.Visibility(this, &SStateTreeViewRow::GetTransitionDashVisibility)
-								.VAlign(VAlign_Center)
-								[
-									SNew(SImage)
-									.Image(FStateTreeEditorStyle::Get().GetBrush("StateTreeEditor.Transition.Dash"))
-									.ColorAndOpacity(IconTint)
-								]
-							]
+	if (EnumHasAllFlags(GetDefault<UStateTreeEditorUserSettings>()->GetStatesViewDisplayNodeType(), EStateTreeEditorUserSettingsNodeType::Transition))
+	{
+		StateHorizontalBox->AddSlot()
+		.VAlign(VAlign_Fill)
+		.HAlign(HAlign_Left)
+		[
+			// Transitions
+			SNew(SHorizontalBox)
 
-							// Completed transitions
-							+ SHorizontalBox::Slot()
-							.VAlign(VAlign_Top)
-							.AutoWidth()
-							[
-								MakeTransitionWidget(EStateTreeTransitionTrigger::OnStateCompleted, nullptr)
-							]
+			+ SHorizontalBox::Slot()
+			.VAlign(VAlign_Top)
+			.AutoWidth()
+			[
+				SNew(SBox)
+				.HeightOverride(GetDefault<UStateTreeEditorUserSettings>()->GetStatesViewStateRowHeight())
+				.Visibility(this, &SStateTreeViewRow::GetTransitionDashVisibility)
+				.VAlign(VAlign_Center)
+				[
+					SNew(SImage)
+					.Image(FStateTreeEditorStyle::Get().GetBrush("StateTreeEditor.Transition.Dash"))
+					.ColorAndOpacity(IconTint)
+				]
+			]
 
-							// Succeeded transitions
-							+ SHorizontalBox::Slot()
-							.VAlign(VAlign_Top)
-							.AutoWidth()
-							[
-								MakeTransitionWidget(EStateTreeTransitionTrigger::OnStateSucceeded, FStateTreeEditorStyle::Get().GetBrush("StateTreeEditor.Transition.Succeeded"))
-							]
+			// Completed transitions
+			+ SHorizontalBox::Slot()
+			.VAlign(VAlign_Top)
+			.AutoWidth()
+			[
+				MakeTransitionWidget(EStateTreeTransitionTrigger::OnStateCompleted, nullptr)
+			]
 
-							// Failed transitions
-							+ SHorizontalBox::Slot()
-							.VAlign(VAlign_Top)
-							.AutoWidth()
-							[
-								MakeTransitionWidget(EStateTreeTransitionTrigger::OnStateFailed, FStateTreeEditorStyle::Get().GetBrush("StateTreeEditor.Transition.Failed"))
-							]
+			// Succeeded transitions
+			+ SHorizontalBox::Slot()
+			.VAlign(VAlign_Top)
+			.AutoWidth()
+			[
+				MakeTransitionWidget(EStateTreeTransitionTrigger::OnStateSucceeded, FStateTreeEditorStyle::Get().GetBrush("StateTreeEditor.Transition.Succeeded"))
+			]
 
-							// Transitions
-							+ SHorizontalBox::Slot()
-							.VAlign(VAlign_Top)
-							.AutoWidth()
-							[
-								MakeTransitionWidget(EStateTreeTransitionTrigger::OnTick | EStateTreeTransitionTrigger::OnEvent, FStateTreeEditorStyle::Get().GetBrush("StateTreeEditor.Transition.Condition"))
-							]
-						]
-					]
-				]
-				+ SVerticalBox::Slot()
-				.AutoHeight()
-				.Padding(FMargin(0, 2, 0, 0))
-				[
-					MakeConditionsWidget(ViewBox)
-				]
-				+ SVerticalBox::Slot()
-				.AutoHeight()
-				.Padding(FMargin(0, 2, 0, 0))
-				[
-					MakeTasksWidget(ViewBox)
-				]
+			// Failed transitions
+			+ SHorizontalBox::Slot()
+			.VAlign(VAlign_Top)
+			.AutoWidth()
+			[
+				MakeTransitionWidget(EStateTreeTransitionTrigger::OnStateFailed, FStateTreeEditorStyle::Get().GetBrush("StateTreeEditor.Transition.Failed"))
 			]
-		]
-	];
+
+			// Transitions
+			+ SHorizontalBox::Slot()
+			.VAlign(VAlign_Top)
+			.AutoWidth()
+			[
+				MakeTransitionWidget(EStateTreeTransitionTrigger::OnTick | EStateTreeTransitionTrigger::OnEvent, FStateTreeEditorStyle::Get().GetBrush("StateTreeEditor.Transition.Condition"))
+			]
+		];
+	}
+
+	if (EnumHasAllFlags(GetDefault<UStateTreeEditorUserSettings>()->GetStatesViewDisplayNodeType(), EStateTreeEditorUserSettingsNodeType::Condition))
+	{
+		StateAndTasksVerticalBox->AddSlot()
+		.AutoHeight()
+		.Padding(FMargin(0, 2, 0, 0))
+		[
+			MakeConditionsWidget(ViewBox)
+		];
+	}
+
+	if (EnumHasAllFlags(GetDefault<UStateTreeEditorUserSettings>()->GetStatesViewDisplayNodeType(), EStateTreeEditorUserSettingsNodeType::Task))
+	{
+		StateAndTasksVerticalBox->AddSlot()
+		.AutoHeight()
+		.Padding(FMargin(0, 2, 0, 0))
+		[
+			MakeTasksWidget(ViewBox)
+		];
+	}
 }
 
 TSharedRef<SWidget> SStateTreeViewRow::MakeTasksWidget(const TSharedPtr<SScrollBox>& ViewBox)
@@ -550,7 +562,7 @@ TSharedRef<SWidget> SStateTreeViewRow::MakeTasksWidget(const TSharedPtr<SScrollB
 							+ SOverlay::Slot()
 							[
 								SNew(SBox)
-								.HeightOverride(UE::StateTree::Editor::TaskRowHeight)
+								.HeightOverride(GetDefault<UStateTreeEditorUserSettings>()->GetStatesViewNodeRowHeight())
 								.Padding(FMargin(0.f, 0.f))
 								[
 									SNew(SHorizontalBox)
@@ -717,7 +729,7 @@ TSharedRef<SWidget> SStateTreeViewRow::MakeConditionsWidget(const TSharedPtr<SSc
 					.Padding(0)
 					[
 						SNew(SBox)
-						.HeightOverride(UE::StateTree::Editor::TaskRowHeight)
+						.HeightOverride(GetDefault<UStateTreeEditorUserSettings>()->GetStatesViewNodeRowHeight())
 						[
 							SNew(SHorizontalBox)
 
diff --git a/Engine/Plugins/Runtime/StateTree/Source/StateTreeEditorModule/Private/StateTreeEditorUserSettings.cpp b/Engine/Plugins/Runtime/StateTree/Source/StateTreeEditorModule/Private/StateTreeEditorUserSettings.cpp
new file mode 100644
--- /dev/null
+++ b/Engine/Plugins/Runtime/StateTree/Source/StateTreeEditorModule/Private/StateTreeEditorUserSettings.cpp
@@ -0,0 +1,26 @@
+// Copyright Epic Games, Inc. All Rights Reserved.
+
+#include "StateTreeEditorUserSettings.h"
+
+void UStateTreeEditorUserSettings::SetStatesViewDisplayNodeType(EStateTreeEditorUserSettingsNodeType Value)
+{
+	if (StatesViewDisplayNodeType != Value)
+	{
+		StatesViewDisplayNodeType = Value;
+		OnSettingsChanged.Broadcast();
+	}
+}
+
+
+void UStateTreeEditorUserSettings::PostEditChangeChainProperty(struct FPropertyChangedChainEvent& PropertyChangedEvent)
+{
+	FName PropertyName = PropertyChangedEvent.Property->GetFName();
+	if (PropertyName == GET_MEMBER_NAME_CHECKED(UStateTreeEditorUserSettings, StatesViewDisplayNodeType)
+		|| PropertyName == GET_MEMBER_NAME_CHECKED(UStateTreeEditorUserSettings, StatesViewStateRowHeight)
+		|| PropertyName == GET_MEMBER_NAME_CHECKED(UStateTreeEditorUserSettings, StatesViewNodeRowHeight))
+	{
+		OnSettingsChanged.Broadcast();
+	}
+
+	Super::PostEditChangeChainProperty(PropertyChangedEvent);
+}
\ No newline at end of file
diff --git a/Engine/Plugins/Runtime/StateTree/Source/StateTreeEditorModule/Private/StateTreeEditorUserSettings.h b/Engine/Plugins/Runtime/StateTree/Source/StateTreeEditorModule/Private/StateTreeEditorUserSettings.h
new file mode 100644
--- /dev/null
+++ b/Engine/Plugins/Runtime/StateTree/Source/StateTreeEditorModule/Private/StateTreeEditorUserSettings.h
@@ -0,0 +1,61 @@
+// Copyright Epic Games, Inc. All Rights Reserved.
+
+#pragma once
+
+#include "Engine/DeveloperSettings.h" //DeveloperSettings
+#include "StateTreeEditorUserSettings.generated.h"
+
+UENUM(meta=(Bitflags, UseEnumValuesAsMaskValuesInEditor = "true"))
+enum class EStateTreeEditorUserSettingsNodeType : uint8
+{
+	Condition = 1 << 0,
+	Task = 1 << 1,
+	Transition = 1 << 2,
+	All = (Condition | Task | Transition) UMETA(Hidden),
+};
+ENUM_CLASS_FLAGS(EStateTreeEditorUserSettingsNodeType);
+
+/** User settings for the StateTree editor */
+UCLASS(MinimalAPI, config=EditorPerProjectUserSettings, DisplayName="StateTree User Settings")
+class UStateTreeEditorUserSettings : public UDeveloperSettings
+{
+	GENERATED_BODY()
+
+public:
+	UStateTreeEditorUserSettings() = default;
+
+	EStateTreeEditorUserSettingsNodeType GetStatesViewDisplayNodeType() const
+	{
+		return StatesViewDisplayNodeType;
+	}
+	void SetStatesViewDisplayNodeType(EStateTreeEditorUserSettingsNodeType Value);
+
+	float GetStatesViewStateRowHeight() const
+	{
+		return StatesViewStateRowHeight;
+	}
+	
+	float GetStatesViewNodeRowHeight() const
+	{
+		return StatesViewNodeRowHeight;
+	}
+
+	/** Broadcast when a setting changes. */
+	FSimpleMulticastDelegate OnSettingsChanged;
+
+public:
+	virtual void PostEditChangeChainProperty(struct FPropertyChangedChainEvent& PropertyChangedEvent) override;
+
+private:
+	/** Which additional node type to display in the States View. */
+	UPROPERTY(config, EditAnywhere, Category = "State View", meta = (DisplayName = "Display Node"))
+	EStateTreeEditorUserSettingsNodeType StatesViewDisplayNodeType = EStateTreeEditorUserSettingsNodeType::All;
+	
+	/** Height of the a state in the Sates View. */
+	UPROPERTY(config, EditAnywhere, Category = "State View", meta = (DisplayName = "State Height", ClampMin = 8.0f))
+	float StatesViewStateRowHeight = 32.0f;
+	
+	/** Height of the a node in the Sates View. */
+	UPROPERTY(config, EditAnywhere, Category = "State View", meta = (DisplayName = "Node Height", ClampMin = 8.0f))
+	float StatesViewNodeRowHeight = 16.0f;
+};
diff --git a/Engine/Plugins/Runtime/StateTree/Source/StateTreeEditorModule/StateTreeEditorModule.Build.cs b/Engine/Plugins/Runtime/StateTree/Source/StateTreeEditorModule/StateTreeEditorModule.Build.cs
--- a/Engine/Plugins/Runtime/StateTree/Source/StateTreeEditorModule/StateTreeEditorModule.Build.cs
+++ b/Engine/Plugins/Runtime/StateTree/Source/StateTreeEditorModule/StateTreeEditorModule.Build.cs
@@ -55,6 +55,7 @@ namespace UnrealBuildTool.Rules
 				"Kismet",
 				"KismetCompiler",
 				"EditorInteractiveToolsFramework",
+				"EditorWidgets",
 				"InteractiveToolsFramework",
 			}
 			);
