StateTree: Add debug callback for evaluator. Split the delegates by type to help name them properly.

diff --git a/Engine/Plugins/Runtime/StateTree/Source/StateTreeModule/Private/Debugger/StateTreeDebug.cpp b/Engine/Plugins/Runtime/StateTree/Source/StateTreeModule/Private/Debugger/StateTreeDebug.cpp
--- a/Engine/Plugins/Runtime/StateTree/Source/StateTreeModule/Private/Debugger/StateTreeDebug.cpp
+++ b/Engine/Plugins/Runtime/StateTree/Source/StateTreeModule/Private/Debugger/StateTreeDebug.cpp
@@ -10,8 +10,15 @@
 namespace UE::StateTree::Debug
 {
 
-FStateTreeDebug::FNodeDelegate FStateTreeDebug::OnNodeEnterState_AnyThread;
-FStateTreeDebug::FNodeDelegate FStateTreeDebug::OnNodeExitState_AnyThread;
+FStateTreeDebug::FNodeDelegate FStateTreeDebug::OnConditionEnterState_AnyThread;
+FStateTreeDebug::FNodeDelegate FStateTreeDebug::OnTestCondition_AnyThread;
+FStateTreeDebug::FNodeDelegate FStateTreeDebug::OnConditionExitState_AnyThread;
+FStateTreeDebug::FNodeDelegate FStateTreeDebug::OnEvaluatorEnterTree_AnyThread;
+FStateTreeDebug::FNodeDelegate FStateTreeDebug::OnTickEvaluator_AnyThread;
+FStateTreeDebug::FNodeDelegate FStateTreeDebug::OnEvaluatorExitTree_AnyThread;
+FStateTreeDebug::FNodeDelegate FStateTreeDebug::OnTaskEnterState_AnyThread;
+FStateTreeDebug::FNodeDelegate FStateTreeDebug::OnTickTask_AnyThread;
+FStateTreeDebug::FNodeDelegate FStateTreeDebug::OnTaskExitState_AnyThread;
 
 FStateTreeDebug::FNodeReference::FNodeReference(gsl::not_null<const UStateTree*> InStateTree, FStateTreeIndex16 InNodeIndex)
 	: StateTree(InStateTree)
@@ -19,7 +26,9 @@ FStateTreeDebug::FNodeReference::FNodeReference(gsl::not_null<const UStateTree*>
 {
 }
 
-void FStateTreeDebug::NodeEnterState(const FStateTreeExecutionContext& ExecutionContext, FNodeReference Node)
+namespace Private
+{
+void NodeEnter(const FStateTreeExecutionContext& ExecutionContext, FStateTreeDebug::FNodeReference Node, FStateTreeDebug::FNodeDelegate& Delegate)
 {
 	FStateTreeInstanceData* InstanceData = ExecutionContext.GetMutableInstanceData();
 	FGuid NodeId = Node.StateTree->GetNodeIdFromIndex(Node.Index);
@@ -29,11 +38,11 @@ void FStateTreeDebug::NodeEnterState(const FStateTreeExecutionContext& Execution
 		{
 			RuntimeValidation->NodeEnterState(NodeId);
 		}
-		OnNodeEnterState_AnyThread.Broadcast(ExecutionContext, FNodeDelegateArgs{.Node = Node, .NodeId = NodeId});
+		Delegate.Broadcast(ExecutionContext, FStateTreeDebug::FNodeDelegateArgs{ .Node = Node, .NodeId = NodeId });
 	}
 }
 
-void FStateTreeDebug::NodeExitState(const FStateTreeExecutionContext& ExecutionContext, FNodeReference Node)
+void NodeExit(const FStateTreeExecutionContext& ExecutionContext, FStateTreeDebug::FNodeReference Node, FStateTreeDebug::FNodeDelegate& Delegate)
 {
 	FStateTreeInstanceData* InstanceData = ExecutionContext.GetMutableInstanceData();
 	FGuid NodeId = Node.StateTree->GetNodeIdFromIndex(Node.Index);
@@ -43,9 +52,65 @@ void FStateTreeDebug::NodeExitState(const FStateTreeExecutionContext& ExecutionC
 		{
 			RuntimeValidation->NodeExitState(NodeId);
 		}
-		OnNodeExitState_AnyThread.Broadcast(ExecutionContext, FNodeDelegateArgs{.Node = Node, .NodeId = NodeId});
+		Delegate.Broadcast(ExecutionContext, FStateTreeDebug::FNodeDelegateArgs{ .Node = Node, .NodeId = NodeId });
+	}
+}
+
+void NodeTick(const FStateTreeExecutionContext& ExecutionContext, FStateTreeDebug::FNodeReference Node, FStateTreeDebug::FNodeDelegate& Delegate)
+{
+	FStateTreeInstanceData* InstanceData = ExecutionContext.GetMutableInstanceData();
+	FGuid NodeId = Node.StateTree->GetNodeIdFromIndex(Node.Index);
+	if (ensure(NodeId.IsValid()))
+	{
+		Delegate.Broadcast(ExecutionContext, FStateTreeDebug::FNodeDelegateArgs{ .Node = Node, .NodeId = NodeId });
 	}
 }
+} //namespace Private
+
+void FStateTreeDebug::ConditionEnterState(const FStateTreeExecutionContext& ExecutionContext, FNodeReference Node)
+{
+	Private::NodeEnter(ExecutionContext, Node, OnConditionEnterState_AnyThread);
+}
+
+void FStateTreeDebug::ConditionTest(const FStateTreeExecutionContext& ExecutionContext, FNodeReference Node)
+{
+	Private::NodeTick(ExecutionContext, Node, OnTestCondition_AnyThread);
+}
+
+void FStateTreeDebug::ConditionExitState(const FStateTreeExecutionContext& ExecutionContext, FNodeReference Node)
+{
+	Private::NodeExit(ExecutionContext, Node, OnConditionExitState_AnyThread);
+}
+
+void FStateTreeDebug::EvaluatorEnterTree(const FStateTreeExecutionContext& ExecutionContext, FNodeReference Node)
+{
+	Private::NodeEnter(ExecutionContext, Node, OnEvaluatorEnterTree_AnyThread);
+}
+
+void FStateTreeDebug::EvaluatorTick(const FStateTreeExecutionContext& ExecutionContext, FNodeReference Node)
+{
+	Private::NodeTick(ExecutionContext, Node, OnTickEvaluator_AnyThread);
+}
+
+void FStateTreeDebug::EvaluatorExitTree(const FStateTreeExecutionContext& ExecutionContext, FNodeReference Node)
+{
+	Private::NodeExit(ExecutionContext, Node, OnEvaluatorExitTree_AnyThread);
+}
+
+void FStateTreeDebug::TaskEnterState(const FStateTreeExecutionContext& ExecutionContext, FNodeReference Node)
+{
+	Private::NodeEnter(ExecutionContext, Node, OnTaskEnterState_AnyThread);
+}
+
+void FStateTreeDebug::TaskTick(const FStateTreeExecutionContext& ExecutionContext, FNodeReference Node)
+{
+	Private::NodeTick(ExecutionContext, Node, OnTickTask_AnyThread);
+}
+
+void FStateTreeDebug::TaskExitState(const FStateTreeExecutionContext& ExecutionContext, FNodeReference Node)
+{
+	Private::NodeExit(ExecutionContext, Node, OnTaskExitState_AnyThread);
+}
 
 }//namespace UE::StateTree::Debug
 
diff --git a/Engine/Plugins/Runtime/StateTree/Source/StateTreeModule/Private/StateTreeExecutionContext.cpp b/Engine/Plugins/Runtime/StateTree/Source/StateTreeModule/Private/StateTreeExecutionContext.cpp
--- a/Engine/Plugins/Runtime/StateTree/Source/StateTreeModule/Private/StateTreeExecutionContext.cpp
+++ b/Engine/Plugins/Runtime/StateTree/Source/StateTreeModule/Private/StateTreeExecutionContext.cpp
@@ -276,6 +276,7 @@ EStateTreeRunStatus FStateTreeExecutionContext::Start(const FInstancedPropertyBa
 	}
 
 	// Initialize for the init frame.
+	ensure(Exec.ActiveFrames.Num() == 0);
 	FStateTreeExecutionFrame& InitFrame = Exec.ActiveFrames.AddDefaulted_GetRef();
 	InitFrame.StateTree = &RootStateTree;
 	InitFrame.RootState = FStateTreeStateHandle::Root;
@@ -1688,7 +1689,7 @@ EStateTreeRunStatus FStateTreeExecutionContext::EnterState(FStateTreeTransitionR
 								CopyBatchOnActiveInstances(CurrentParentFrame, CurrentFrame, ConditionInstanceView, Cond.BindingsBatch);
 							}
 
-							UE_STATETREE_DEBUG_NODE_ENTERSTATE(this, CurrentFrame.StateTree, FStateTreeIndex16(ConditionIndex));
+							UE_STATETREE_DEBUG_CONDITION_ENTER_STATE(this, CurrentFrame.StateTree, FStateTreeIndex16(ConditionIndex));
 							Cond.EnterState(*this, Transition);
 
 							// Reset copied properties that might contain object references.
@@ -1728,7 +1729,7 @@ EStateTreeRunStatus FStateTreeExecutionContext::EnterState(FStateTreeTransitionR
 				if (bIsEnteringState && bShouldCallStateChange)
 				{
 					STATETREE_LOG(Verbose, TEXT("%*s  Task '%s'"), StateIndex*UE::StateTree::DebugIndentSize, TEXT(""), *Task.Name.ToString());
-					UE_STATETREE_DEBUG_NODE_ENTERSTATE(this, NextFrame.StateTree, FStateTreeIndex16(TaskIndex));
+					UE_STATETREE_DEBUG_TASK_ENTER_STATE(this, NextFrame.StateTree, FStateTreeIndex16(TaskIndex));
 
 					EStateTreeRunStatus Status = EStateTreeRunStatus::Unset;
 					{
@@ -1962,7 +1963,7 @@ void FStateTreeExecutionContext::ExitState(const FStateTreeTransitionResult& Tra
 						if (bShouldCallStateChange)
 						{
 							STATETREE_LOG(Verbose, TEXT("%*s  Task '%s'"), Index*UE::StateTree::DebugIndentSize, TEXT(""), *Task.Name.ToString());
-							UE_STATETREE_DEBUG_NODE_EXITSTATE(this, CurrentStateTree, FStateTreeIndex16(TaskIndex));
+							UE_STATETREE_DEBUG_TASK_EXIT_STATE(this, CurrentStateTree, FStateTreeIndex16(TaskIndex));
 							{
 								QUICK_SCOPE_CYCLE_COUNTER(StateTree_Task_ExitState);
 								CSV_SCOPED_TIMING_STAT_EXCLUSIVE(StateTree_Task_ExitState);
@@ -1995,7 +1996,7 @@ void FStateTreeExecutionContext::ExitState(const FStateTreeTransitionResult& Tra
 									CopyBatchOnActiveInstances(CurrentParentFrame, CurrentFrame, ConditionInstanceView, Cond.BindingsBatch);
 								}
 
-								UE_STATETREE_DEBUG_NODE_EXITSTATE(this, CurrentFrame.StateTree, FStateTreeIndex16(ConditionIndex));
+								UE_STATETREE_DEBUG_CONDITION_EXIT_STATE(this, CurrentFrame.StateTree, FStateTreeIndex16(ConditionIndex));
 								Cond.ExitState(*this, Transition);
 
 								// Reset copied properties that might contain object references.
@@ -2152,6 +2153,7 @@ EStateTreeRunStatus FStateTreeExecutionContext::TickEvaluatorsAndGlobalTasks(con
 					CopyBatchOnActiveInstances(CurrentParentFrame, CurrentFrame, EvalInstanceView, Eval.BindingsBatch);
 				}
 				STATETREE_LOG(VeryVerbose, TEXT("  Tick: '%s'"), *Eval.Name.ToString());
+				UE_STATETREE_DEBUG_EVALUATOR_TICK(this, CurrentStateTree, EvalIndex);
 				{
 					QUICK_SCOPE_CYCLE_COUNTER(StateTree_Eval_Tick);
 					Eval.Tick(*this, DeltaTime);
@@ -2194,6 +2196,8 @@ EStateTreeRunStatus FStateTreeExecutionContext::TickEvaluatorsAndGlobalTasks(con
 					}
 
 					//STATETREE_TRACE_TASK_EVENT(TaskIndex, TaskDataView, EStateTreeTraceEventType::OnTickingTask, EStateTreeRunStatus::Running);
+					UE_STATETREE_DEBUG_TASK_TICK(this, CurrentStateTree, TaskIndex);
+
 					EStateTreeRunStatus TaskResult = EStateTreeRunStatus::Unset;
 					{
 						QUICK_SCOPE_CYCLE_COUNTER(StateTree_Task_Tick);
@@ -2259,6 +2263,7 @@ EStateTreeRunStatus FStateTreeExecutionContext::StartEvaluatorsAndGlobalTasks(FS
 					CopyBatchOnActiveInstances(CurrentParentFrame, CurrentFrame, EvalInstanceView, Eval.BindingsBatch);
 				}
 				STATETREE_LOG(Verbose, TEXT("  Start: '%s'"), *Eval.Name.ToString());
+				UE_STATETREE_DEBUG_EVALUATOR_ENTER_TREE(this, CurrentStateTree, FStateTreeIndex16(EvalIndex));
 				{
 					QUICK_SCOPE_CYCLE_COUNTER(StateTree_Eval_TreeStart);
 					Eval.TreeStart(*this);
@@ -2291,7 +2296,7 @@ EStateTreeRunStatus FStateTreeExecutionContext::StartEvaluatorsAndGlobalTasks(FS
 				}
 
 				STATETREE_LOG(Verbose, TEXT("  Start: '%s'"), *Task.Name.ToString());
-				UE_STATETREE_DEBUG_NODE_ENTERSTATE(this, CurrentStateTree, FStateTreeIndex16(TaskIndex));
+				UE_STATETREE_DEBUG_TASK_ENTER_STATE(this, CurrentStateTree, FStateTreeIndex16(TaskIndex));
 				{
 					QUICK_SCOPE_CYCLE_COUNTER(StateTree_Task_TreeStart);
 					const EStateTreeRunStatus TaskStatus = Task.EnterState(*this, Transition);
@@ -2404,7 +2409,7 @@ void FStateTreeExecutionContext::CallStopOnEvaluatorsAndGlobalTasks(const FState
 		if (TaskIndex <= LastInitializedTaskIndex.Get())
 		{
 			STATETREE_LOG(Verbose, TEXT("  Stop: '%s'"), *Task.Name.ToString());
-			UE_STATETREE_DEBUG_NODE_EXITSTATE(this, CurrentStateTree, FStateTreeIndex16(TaskIndex));
+			UE_STATETREE_DEBUG_TASK_EXIT_STATE(this, CurrentStateTree, FStateTreeIndex16(TaskIndex));
 			{
 				QUICK_SCOPE_CYCLE_COUNTER(StateTree_Task_TreeStop);
 				Task.ExitState(*this, Transition);
@@ -2420,6 +2425,7 @@ void FStateTreeExecutionContext::CallStopOnEvaluatorsAndGlobalTasks(const FState
 		FNodeInstanceDataScope DataScope(*this, Eval.InstanceDataHandle, EvalInstanceView);
 
 		STATETREE_LOG(Verbose, TEXT("  Stop: '%s'"), *Eval.Name.ToString());
+		UE_STATETREE_DEBUG_EVALUATOR_EXIT_TREE(this, CurrentStateTree, FStateTreeIndex16(EvalIndex));
 		{
 			QUICK_SCOPE_CYCLE_COUNTER(StateTree_Eval_TreeStop);
 			Eval.TreeStop(*this);
@@ -2513,7 +2519,7 @@ EStateTreeRunStatus FStateTreeExecutionContext::StartTemporaryEvaluatorsAndGloba
 		if (bWasCreated)
 		{
 			STATETREE_LOG(Verbose, TEXT("  Start: '%s'"), *Task.Name.ToString());
-			UE_STATETREE_DEBUG_NODE_ENTERSTATE(this, CurrentStateTree, FStateTreeIndex16(TaskIndex));
+			UE_STATETREE_DEBUG_TASK_ENTER_STATE(this, CurrentStateTree, FStateTreeIndex16(TaskIndex));
 			QUICK_SCOPE_CYCLE_COUNTER(StateTree_Task_TreeStart);
 			const EStateTreeRunStatus TaskStatus = Task.EnterState(*this, Transition);
 
@@ -2570,7 +2576,7 @@ void FStateTreeExecutionContext::StopTemporaryEvaluatorsAndGlobalTasks(const FSt
 			if (const FStateTreeTaskBase* Task = NodeView.GetPtr<const FStateTreeTaskBase>())
 			{
 				STATETREE_LOG(Verbose, TEXT("  Stop: '%s'"), *Task->Name.ToString());
-				UE_STATETREE_DEBUG_NODE_EXITSTATE(this, CurrentFrame.StateTree, FStateTreeIndex16(Index));
+				UE_STATETREE_DEBUG_TASK_EXIT_STATE(this, CurrentFrame.StateTree, FStateTreeIndex16(TempInstance.OwnerNodeIndex.Get()));
 				{
 					QUICK_SCOPE_CYCLE_COUNTER(StateTree_Task_TreeStop);
 					Task->ExitState(*this, Transition);
@@ -2580,6 +2586,7 @@ void FStateTreeExecutionContext::StopTemporaryEvaluatorsAndGlobalTasks(const FSt
 			else if (const FStateTreeEvaluatorBase* Eval = NodeView.GetPtr<const FStateTreeEvaluatorBase>())
 			{
 				STATETREE_LOG(Verbose, TEXT("  Stop: '%s'"), *Eval->Name.ToString());
+				UE_STATETREE_DEBUG_EVALUATOR_EXIT_TREE(this, CurrentFrame.StateTree, FStateTreeIndex16(TempInstance.OwnerNodeIndex.Get()));
 				{
 					QUICK_SCOPE_CYCLE_COUNTER(StateTree_Eval_TreeStop);
 					Eval->TreeStop(*this);
@@ -2674,6 +2681,8 @@ EStateTreeRunStatus FStateTreeExecutionContext::TickTasks(const float DeltaTime)
 				}
 
 				//STATETREE_TRACE_TASK_EVENT(TaskIndex, TaskDataView, EStateTreeTraceEventType::OnTickingTask, EStateTreeRunStatus::Running);
+				UE_STATETREE_DEBUG_TASK_TICK(this, CurrentStateTree, TaskIndex);
+
 				EStateTreeRunStatus TaskResult = EStateTreeRunStatus::Unset;
 				{
 					QUICK_SCOPE_CYCLE_COUNTER(StateTree_Task_Tick);
@@ -2756,6 +2765,8 @@ bool FStateTreeExecutionContext::TestAllConditions(const FStateTreeExecutionFram
 				}
 			}
 			
+			UE_STATETREE_DEBUG_CONDITION_TEST_CONDITION(this, CurrentFrame.StateTree, Index);
+
 			bValue = Cond.TestCondition(*this);
 			STATETREE_TRACE_CONDITION_EVENT(ConditionIndex, ConditionInstanceView, bValue ? EStateTreeTraceEventType::Passed : EStateTreeTraceEventType::Failed);
 			
diff --git a/Engine/Plugins/Runtime/StateTree/Source/StateTreeModule/Public/Debugger/StateTreeDebug.h b/Engine/Plugins/Runtime/StateTree/Source/StateTreeModule/Public/Debugger/StateTreeDebug.h
--- a/Engine/Plugins/Runtime/StateTree/Source/StateTreeModule/Public/Debugger/StateTreeDebug.h
+++ b/Engine/Plugins/Runtime/StateTree/Source/StateTreeModule/Public/Debugger/StateTreeDebug.h
@@ -11,8 +11,17 @@ struct FStateTreeExecutionContext;
 
 #if WITH_STATETREE_DEBUG
 
-#define UE_STATETREE_DEBUG_NODE_ENTERSTATE(InExecutionContextPtr, InStateTree, InNodeIndex) ::UE::StateTree::Debug::FStateTreeDebug::NodeEnterState(*(InExecutionContextPtr), ::UE::StateTree::Debug::FStateTreeDebug::FNodeReference((InStateTree), (InNodeIndex)));
-#define UE_STATETREE_DEBUG_NODE_EXITSTATE(InExecutionContextPtr, InStateTree, InNodeIndex) ::UE::StateTree::Debug::FStateTreeDebug::NodeExitState(*(InExecutionContextPtr), ::UE::StateTree::Debug::FStateTreeDebug::FNodeReference((InStateTree), (InNodeIndex)));
+#define UE_STATETREE_DEBUG_CONDITION_ENTER_STATE(InExecutionContextPtr, InStateTree, InNodeIndex) ::UE::StateTree::Debug::FStateTreeDebug::ConditionEnterState(*(InExecutionContextPtr), ::UE::StateTree::Debug::FStateTreeDebug::FNodeReference((InStateTree), FStateTreeIndex16(InNodeIndex)))
+#define UE_STATETREE_DEBUG_CONDITION_TEST_CONDITION(InExecutionContextPtr, InStateTree, InNodeIndex) ::UE::StateTree::Debug::FStateTreeDebug::ConditionTest(*(InExecutionContextPtr), ::UE::StateTree::Debug::FStateTreeDebug::FNodeReference((InStateTree), FStateTreeIndex16(InNodeIndex)))
+#define UE_STATETREE_DEBUG_CONDITION_EXIT_STATE(InExecutionContextPtr, InStateTree, InNodeIndex) ::UE::StateTree::Debug::FStateTreeDebug::ConditionExitState(*(InExecutionContextPtr), ::UE::StateTree::Debug::FStateTreeDebug::FNodeReference((InStateTree), FStateTreeIndex16(InNodeIndex)))
+
+#define UE_STATETREE_DEBUG_EVALUATOR_ENTER_TREE(InExecutionContextPtr, InStateTree, InNodeIndex) ::UE::StateTree::Debug::FStateTreeDebug::EvaluatorEnterTree(*(InExecutionContextPtr), ::UE::StateTree::Debug::FStateTreeDebug::FNodeReference((InStateTree), FStateTreeIndex16(InNodeIndex)))
+#define UE_STATETREE_DEBUG_EVALUATOR_TICK(InExecutionContextPtr, InStateTree, InNodeIndex) ::UE::StateTree::Debug::FStateTreeDebug::EvaluatorTick(*(InExecutionContextPtr), ::UE::StateTree::Debug::FStateTreeDebug::FNodeReference((InStateTree), FStateTreeIndex16(InNodeIndex)))
+#define UE_STATETREE_DEBUG_EVALUATOR_EXIT_TREE(InExecutionContextPtr, InStateTree, InNodeIndex) ::UE::StateTree::Debug::FStateTreeDebug::EvaluatorExitTree(*(InExecutionContextPtr), ::UE::StateTree::Debug::FStateTreeDebug::FNodeReference((InStateTree), FStateTreeIndex16(InNodeIndex)))
+
+#define UE_STATETREE_DEBUG_TASK_ENTER_STATE(InExecutionContextPtr, InStateTree, InNodeIndex) ::UE::StateTree::Debug::FStateTreeDebug::TaskEnterState(*(InExecutionContextPtr), ::UE::StateTree::Debug::FStateTreeDebug::FNodeReference((InStateTree), FStateTreeIndex16(InNodeIndex)))
+#define UE_STATETREE_DEBUG_TASK_TICK(InExecutionContextPtr, InStateTree, InNodeIndex) ::UE::StateTree::Debug::FStateTreeDebug::TaskTick(*(InExecutionContextPtr), ::UE::StateTree::Debug::FStateTreeDebug::FNodeReference((InStateTree), FStateTreeIndex16(InNodeIndex)))
+#define UE_STATETREE_DEBUG_TASK_EXIT_STATE(InExecutionContextPtr, InStateTree, InNodeIndex) ::UE::StateTree::Debug::FStateTreeDebug::TaskExitState(*(InExecutionContextPtr), ::UE::StateTree::Debug::FStateTreeDebug::FNodeReference((InStateTree), FStateTreeIndex16(InNodeIndex)))
 
 namespace UE::StateTree::Debug
 {
@@ -21,13 +30,25 @@ class STATETREEMODULE_API FStateTreeDebug final
 public:
 	struct FNodeReference
 	{
-		FNodeReference(gsl::not_null<const UStateTree*> InStateTree, FStateTreeIndex16 InNodeIndex);
+		explicit FNodeReference(gsl::not_null<const UStateTree*> InStateTree, FStateTreeIndex16 InNodeIndex);
 		gsl::not_null<const UStateTree*> StateTree;
 		FStateTreeIndex16 Index;
 	};
+
+	static void ConditionEnterState(const FStateTreeExecutionContext& ExecutionContext, FNodeReference Node);
+	static void ConditionTest(const FStateTreeExecutionContext& ExecutionContext, FNodeReference Node);
+	static void ConditionExitState(const FStateTreeExecutionContext& ExecutionContext, FNodeReference Node);
+
+	static void EvaluatorEnterTree(const FStateTreeExecutionContext& ExecutionContext, FNodeReference Node);
+	static void EvaluatorTick(const FStateTreeExecutionContext& ExecutionContext, FNodeReference Node);
+	static void EvaluatorExitTree(const FStateTreeExecutionContext& ExecutionContext, FNodeReference Node);
+
+	static void TaskEnterState(const FStateTreeExecutionContext& ExecutionContext, FNodeReference Node);
+	static void TaskTick(const FStateTreeExecutionContext& ExecutionContext, FNodeReference Node);
+	static void TaskExitState(const FStateTreeExecutionContext& ExecutionContext, FNodeReference Node);
+
+public:
 	
-	static void NodeEnterState(const FStateTreeExecutionContext& ExecutionContext, FNodeReference Node);
-	static void NodeExitState(const FStateTreeExecutionContext& ExecutionContext, FNodeReference Node);
 
 public:
 	struct FNodeDelegateArgs
@@ -35,21 +56,71 @@ public:
 		FNodeReference Node;
 		FGuid NodeId;
 	};
-	
+
 	DECLARE_TS_MULTICAST_DELEGATE_TwoParams(FNodeDelegate, const FStateTreeExecutionContext& ExecutionContext, FNodeDelegateArgs Node);
+
 	/**
-	 * Debugging callback for when a node activates.
+	 * Debugging callback for when a condition activates.
 	 * @note The callback executes inside the StateTree logic.
 	 * @note The StateTree can execute on any thread.
 	 */
-	static FNodeDelegate OnNodeEnterState_AnyThread;
-	
+	static FNodeDelegate OnConditionEnterState_AnyThread;
+
 	/**
-	 * Debugging callback for when a node deactivates.
+	 * Debugging callback for before a condition is tested.
 	 * @note The callback executes inside the StateTree logic.
 	 * @note The StateTree can execute on any thread.
 	 */
-	static FNodeDelegate OnNodeExitState_AnyThread;
+	static FNodeDelegate OnTestCondition_AnyThread;
+
+	/**
+	 * Debugging callback for when a condition deactivates.
+	 * @note The callback executes inside the StateTree logic.
+	 * @note The StateTree can execute on any thread.
+	 */
+	static FNodeDelegate OnConditionExitState_AnyThread;
+
+	/**
+	 * Debugging callback for when a evaluator activates.
+	 * @note The callback executes inside the StateTree logic.
+	 * @note The StateTree can execute on any thread.
+	 */
+	static FNodeDelegate OnEvaluatorEnterTree_AnyThread;
+
+	/**
+	 * Debugging callback for before an evaluator ticks.
+	 * @note The callback executes inside the StateTree logic.
+	 * @note The StateTree can execute on any thread.
+	 */
+	static FNodeDelegate OnTickEvaluator_AnyThread;
+
+	/**
+	 * Debugging callback for when a evaluator deactivates.
+	 * @note The callback executes inside the StateTree logic.
+	 * @note The StateTree can execute on any thread.
+	 */
+	static FNodeDelegate OnEvaluatorExitTree_AnyThread;
+
+	/**
+	 * Debugging callback for when a task activates.
+	 * @note The callback executes inside the StateTree logic.
+	 * @note The StateTree can execute on any thread.
+	 */
+	static FNodeDelegate OnTaskEnterState_AnyThread;
+
+	/**
+	 * Debugging callback executed before a task ticks.
+	 * @note The callback executes inside the StateTree logic.
+	 * @note The StateTree can execute on any thread.
+	 */
+	static FNodeDelegate OnTickTask_AnyThread;
+
+	/**
+	 * Debugging callback for when a task deactivates.
+	 * @note The callback executes inside the StateTree logic.
+	 * @note The StateTree can execute on any thread.
+	 */
+	static FNodeDelegate OnTaskExitState_AnyThread;
 	
 private:
 	FStateTreeDebug() = default;
@@ -59,7 +130,16 @@ private:
 
 #else
 
-#define UE_STATETREE_DEBUG_NODE_ENTERSTATE(...)
-#define UE_STATETREE_DEBUG_NODE_EXITSTATE(...)
+#define UE_STATETREE_DEBUG_CONDITION_ENTER_STATE(...)
+#define UE_STATETREE_DEBUG_CONDITION_TEST_CONDITION(...)
+#define UE_STATETREE_DEBUG_CONDITION_EXIT_STATE(...)
+
+#define UE_STATETREE_DEBUG_EVALUATOR_ENTER_TREE(...)
+#define UE_STATETREE_DEBUG_EVALUATOR_TICK(...)
+#define UE_STATETREE_DEBUG_EVALUATOR_EXIT_TREE(...)
+
+#define UE_STATETREE_DEBUG_TASK_ENTER_STATE(...)
+#define UE_STATETREE_DEBUG_TASK_TICK(...)
+#define UE_STATETREE_DEBUG_TASK_EXIT_STATE(...)
 
 #endif // WITH_STATETREE_DEBUG
