StateTree: Review feedback for CL 35822194. Change the name of the command.

diff --git a/Engine/Plugins/Runtime/StateTree/Source/StateTreeEditorModule/Private/StateTreeEditor.cpp b/Engine/Plugins/Runtime/StateTree/Source/StateTreeEditorModule/Private/StateTreeEditor.cpp
--- a/Engine/Plugins/Runtime/StateTree/Source/StateTreeEditorModule/Private/StateTreeEditor.cpp
+++ b/Engine/Plugins/Runtime/StateTree/Source/StateTreeEditorModule/Private/StateTreeEditor.cpp
@@ -356,7 +356,7 @@ void FillDeveloperMenu(UToolMenu* InMenu)
 	const FStateTreeEditorCommands& Commands = FStateTreeEditorCommands::Get();
 	{
 		FToolMenuSection& Section = InMenu->AddSection("FileDeveloperCompilerSettings", LOCTEXT("CompileOptionsHeading", "Compiler Settings"));
-		Section.AddMenuEntry(Commands.LogResultOnCompileSuccess);
+		Section.AddMenuEntry(Commands.LogCompilationResult);
 	}
 }
 void FillDynamicDeveloperMenu(FToolMenuSection& Section)
diff --git a/Engine/Plugins/Runtime/StateTree/Source/StateTreeEditorModule/Private/StateTreeEditorCommands.cpp b/Engine/Plugins/Runtime/StateTree/Source/StateTreeEditorModule/Private/StateTreeEditorCommands.cpp
--- a/Engine/Plugins/Runtime/StateTree/Source/StateTreeEditorModule/Private/StateTreeEditorCommands.cpp
+++ b/Engine/Plugins/Runtime/StateTree/Source/StateTreeEditorModule/Private/StateTreeEditorCommands.cpp
@@ -21,7 +21,7 @@ void FStateTreeEditorCommands::RegisterCommands()
 	UI_COMMAND(SaveOnCompile_Never, "Never", "Sets the save-on-compile option to 'Never', meaning that your StateTree will not be saved when they are compiled", EUserInterfaceActionType::RadioButton, FInputChord());
 	UI_COMMAND(SaveOnCompile_SuccessOnly, "On Success Only", "Sets the save-on-compile option to 'Success Only', meaning that your StateTree will be saved whenever they are successfully compiled", EUserInterfaceActionType::RadioButton, FInputChord());
 	UI_COMMAND(SaveOnCompile_Always, "Always", "Sets the save-on-compile option to 'Always', meaning that your StateTree will be saved whenever they are compiled (even if there were errors)", EUserInterfaceActionType::RadioButton, FInputChord());
-	UI_COMMAND(LogResultOnCompileSuccess, "Log compilation result", "After a StateTree compiles, log the internal content of the StateTree.", EUserInterfaceActionType::ToggleButton, FInputChord());
+	UI_COMMAND(LogCompilationResult, "Log compilation result", "After a StateTree compiles, log the internal content of the StateTree.", EUserInterfaceActionType::ToggleButton, FInputChord());
 
 	UI_COMMAND(AddSiblingState, "Add Sibling State", "Add a Sibling State", EUserInterfaceActionType::Button, FInputChord());
 	UI_COMMAND(AddChildState, "Add Child State", "Add a Child State", EUserInterfaceActionType::Button, FInputChord());
diff --git a/Engine/Plugins/Runtime/StateTree/Source/StateTreeEditorModule/Private/StateTreeEditorCommands.h b/Engine/Plugins/Runtime/StateTree/Source/StateTreeEditorModule/Private/StateTreeEditorCommands.h
--- a/Engine/Plugins/Runtime/StateTree/Source/StateTreeEditorModule/Private/StateTreeEditorCommands.h
+++ b/Engine/Plugins/Runtime/StateTree/Source/StateTreeEditorModule/Private/StateTreeEditorCommands.h
@@ -20,7 +20,7 @@ public:
 	TSharedPtr<FUICommandInfo> SaveOnCompile_Never;
 	TSharedPtr<FUICommandInfo> SaveOnCompile_SuccessOnly;
 	TSharedPtr<FUICommandInfo> SaveOnCompile_Always;
-	TSharedPtr<FUICommandInfo> LogResultOnCompileSuccess;
+	TSharedPtr<FUICommandInfo> LogCompilationResult;
 
 	TSharedPtr<FUICommandInfo> AddSiblingState;
 	TSharedPtr<FUICommandInfo> AddChildState;
diff --git a/Engine/Plugins/Runtime/StateTree/Source/StateTreeEditorModule/Private/StateTreeEditorMode.cpp b/Engine/Plugins/Runtime/StateTree/Source/StateTreeEditorModule/Private/StateTreeEditorMode.cpp
--- a/Engine/Plugins/Runtime/StateTree/Source/StateTreeEditorModule/Private/StateTreeEditorMode.cpp
+++ b/Engine/Plugins/Runtime/StateTree/Source/StateTreeEditorModule/Private/StateTreeEditorMode.cpp
@@ -304,24 +304,24 @@ static bool IsSaveOnCompileOptionSet(const EStateTreeSaveOnCompile Option)
 	return (Settings->SaveOnCompile == Option);
 }
 
-static IConsoleVariable* GetLogResultOnSuccessCVar()
+static IConsoleVariable* GetLogCompilationResultCVar()
 {
 	static IConsoleVariable* FoundVariable = IConsoleManager::Get().FindConsoleVariable(TEXT("StateTree.Compiler.LogResultOnCompilationCompleted"));
 	return FoundVariable;
 }
 
-static void ToggleLogResultOnCompileSuccess()
+static void ToggleLogCompilationResult()
 {
-	IConsoleVariable* LogResultCVar = GetLogResultOnSuccessCVar();
+	IConsoleVariable* LogResultCVar = GetLogCompilationResultCVar();
 	if (ensure(LogResultCVar))
 	{
 		LogResultCVar->Set(!LogResultCVar->GetBool(), ECVF_SetByConsole);
 	}
 }
 
-static bool IsLogResultOnCompileSuccess()
+static bool IsLogCompilationResult()
 {
-	IConsoleVariable* LogResultCVar = GetLogResultOnSuccessCVar();
+	IConsoleVariable* LogResultCVar = GetLogCompilationResultCVar();
 	return LogResultCVar ? LogResultCVar->GetBool() : false;
 }
 }
@@ -360,10 +360,10 @@ void UStateTreeEditorMode::BindToolkitCommands(const TSharedRef<FUICommandList>&
 		FIsActionButtonVisible::CreateUObject(this, &UStateTreeEditorMode::HasValidStateTree)
 	);
 	ToolkitCommands->MapAction(
-		FStateTreeEditorCommands::Get().LogResultOnCompileSuccess,
-		FExecuteAction::CreateStatic(&UE::StateTree::Editor::Internal::ToggleLogResultOnCompileSuccess),
+		FStateTreeEditorCommands::Get().LogCompilationResult,
+		FExecuteAction::CreateStatic(&UE::StateTree::Editor::Internal::ToggleLogCompilationResult),
 		FCanExecuteAction(),
-		FIsActionChecked::CreateStatic(&UE::StateTree::Editor::Internal::IsLogResultOnCompileSuccess)
+		FIsActionChecked::CreateStatic(&UE::StateTree::Editor::Internal::IsLogCompilationResult)
 	);
 }
 
