StateTree: Stop the tree evaluation in unit test. The stop is needed to call the "exit" function on tasks and conditions.

diff --git a/Engine/Plugins/Runtime/StateTree/Source/StateTreeTestSuite/Private/StateTreeTest.cpp b/Engine/Plugins/Runtime/StateTree/Source/StateTreeTestSuite/Private/StateTreeTest.cpp
--- a/Engine/Plugins/Runtime/StateTree/Source/StateTreeTestSuite/Private/StateTreeTest.cpp
+++ b/Engine/Plugins/Runtime/StateTree/Source/StateTreeTestSuite/Private/StateTreeTest.cpp
@@ -151,6 +151,8 @@ struct FStateTreeTest_EmptyStateTree : FAITestBase
 		AITEST_TRUE("StateTree should be completed", Status == EStateTreeRunStatus::Succeeded);
 		Exec.LogClear();
 
+		Exec.Stop();
+
 		return true;
 	}
 };
@@ -201,9 +203,9 @@ struct FStateTreeTest_Sequence : FAITestBase
 		Exec.LogClear();
 		
 		Status = Exec.Tick(0.1f);
-        AITEST_TRUE("StateTree Task2 should tick, and exit state", Exec.Expect(Task2.GetName(), TickStr).Then(Task2.GetName(), ExitStateStr));
+		AITEST_TRUE("StateTree Task2 should tick, and exit state", Exec.Expect(Task2.GetName(), TickStr).Then(Task2.GetName(), ExitStateStr));
 		AITEST_FALSE("StateTree Task1 should not tick", Exec.Expect(Task1.GetName(), TickStr));
-        AITEST_TRUE("StateTree should be completed", Status == EStateTreeRunStatus::Succeeded);
+		AITEST_TRUE("StateTree should be completed", Status == EStateTreeRunStatus::Succeeded);
 		Exec.LogClear();
 
 		Status = Exec.Tick(0.1f);
@@ -211,6 +213,8 @@ struct FStateTreeTest_Sequence : FAITestBase
 		AITEST_FALSE("StateTree Task2 should not tick", Exec.Expect(Task2.GetName(), TickStr));
 		Exec.LogClear();
 
+		Exec.Stop();
+
 		return true;
 	}
 };
@@ -281,7 +285,9 @@ struct FStateTreeTest_Select : FAITestBase
 		AITEST_TRUE("StateTree Task1 should tick, exit state, and enter state", Exec.Expect(Task1.GetName(), TickStr).Then(Task1.GetName(), ExitStateStr).Then(Task1.GetName(), EnterStateStr));
 		AITEST_TRUE("StateTree Task1A should tick, exit state, and enter state", Exec.Expect(Task1A.GetName(), TickStr).Then(Task1A.GetName(), ExitStateStr).Then(Task1A.GetName(), EnterStateStr));
 		AITEST_TRUE("StateTree should be running", Status == EStateTreeRunStatus::Running);
-        Exec.LogClear();
+		Exec.LogClear();
+
+		Exec.Stop();
 
 		return true;
 	}
@@ -399,6 +405,8 @@ struct FStateTreeTest_Restart : FAITestBase
 		AITEST_TRUE("StateTree exec status should be running", Exec.GetLastTickStatus() == EStateTreeRunStatus::Running);
 		Exec.LogClear();
 
+		Exec.Stop();
+
 		return true;
 	}
 };
@@ -474,6 +482,8 @@ struct FStateTreeTest_SubTree : FAITestBase
 		Status = Exec.Tick(0.1f);
 		AITEST_TRUE("StateTree should complete in succeeded", Status == EStateTreeRunStatus::Succeeded);
 		Exec.LogClear();
+
+		Exec.Stop();
 		
 		return true;
 	}
@@ -555,6 +565,8 @@ struct FStateTreeTest_SubTreeCondition : FAITestBase
 		AITEST_TRUE("StateTree should be running", Status == EStateTreeRunStatus::Running);
 		Exec.LogClear();
 
+		Exec.Stop();
+
 		return true;
 	}
 };
@@ -626,6 +638,8 @@ struct FStateTreeTest_SubTree_CascadedSucceeded : FAITestBase
 		Status = Exec.Tick(0.1f);
 		AITEST_TRUE("StateTree should be Failed", Status == EStateTreeRunStatus::Failed);
 		Exec.LogClear();
+
+		Exec.Stop();
 		
 		return true;
 	}
@@ -725,6 +739,16 @@ struct FStateTreeTest_SharedInstanceData : FAITestBase
 		AITEST_EQUAL("All StateTree contexts should be running after Tick", TickTotalRunning, NumConcurrent);
 		AITEST_EQUAL("Test condition global counter should equal context count after Tick", (int32)FStateTreeTestConditionInstanceData::GlobalCounter, InstanceDatas.Num());
 
+		ParallelForWithTaskContext(
+			RunContexts,
+			InstanceDatas.Num(),
+			[&InstanceDatas, &StateTree](FStateTreeTestRunContext& RunContext, int32 Index)
+			{
+				FTestStateTreeExecutionContext Exec(StateTree, StateTree, InstanceDatas[Index]);
+				Exec.Stop();
+			}
+		);
+
 		return true;
 	}
 };
@@ -803,6 +827,8 @@ struct FStateTreeTest_TransitionPriority : FAITestBase
 		AITEST_EQUAL("Tree execution should stop on success", Status, EStateTreeRunStatus::Succeeded);
 		Exec.LogClear();
 
+		Exec.Stop();
+
 		return true;
 	}
 };
@@ -867,6 +893,8 @@ struct FStateTreeTest_TransitionPriorityEnterState : FAITestBase
 		AITEST_FALSE("StateTree Task3 should not enter state", Exec.Expect(Task3.GetName(), EnterStateStr));
 		Exec.LogClear();
 
+		Exec.Stop();
+
 		return true;
 	}
 };
@@ -937,6 +965,8 @@ struct FStateTreeTest_TransitionNextSelectableState : FAITestBase
 		AITEST_TRUE("StateTree Task2 should complete", Exec.Expect(Task2.GetName(), StateCompletedStr));
 		Exec.LogClear();
 
+		Exec.Stop();
+
 		return true;
 	}
 };
@@ -994,6 +1024,8 @@ struct FStateTreeTest_TransitionNextWithParentData : FAITestBase
 		AITEST_TRUE("StateTree Task1A should enter state", Exec.Expect(Task1A.GetName(), EnterStateStr));
 		Exec.LogClear();
 
+		Exec.Stop();
+
 		return true;
 	}
 };
@@ -1047,6 +1079,8 @@ struct FStateTreeTest_LastConditionWithIndent : FAITestBase
 		AITEST_FALSE("StateTree Task1 should not tick", Exec.Expect(Task1.GetName(), TickStr));
 		Exec.LogClear();
 
+		Exec.Stop();
+
 		return true;
 	}
 };
@@ -1106,6 +1140,8 @@ struct FStateTreeTest_TransitionGlobalDataView : FAITestBase
 		AITEST_TRUE("StateTree Task1 should enter state with value 123", Exec.Expect(Task2.GetName(), EnterState123Str));
 		Exec.LogClear();
 
+		Exec.Stop();
+
 		return true;
 	}
 };
@@ -1187,6 +1223,8 @@ struct FStateTreeTest_TransitionDelay : FAITestBase
 		AITEST_TRUE("StateTree Task1 should enter state", Exec.Expect(Task1.GetName(), EnterStateStr));
 		Exec.LogClear();
 
+		Exec.Stop();
+
 		return true;
 	}
 };
@@ -1248,6 +1286,8 @@ struct FStateTreeTest_TransitionDelayZero : FAITestBase
 		AITEST_TRUE("StateTree Task1 should enter state", Exec.Expect(Task1.GetName(), EnterStateStr));
 		Exec.LogClear();
 
+		Exec.Stop();
+
 		return true;
 	}
 };
@@ -1333,6 +1373,8 @@ struct FStateTreeTest_StateRequiringEvent : FAITestBase
 		AITEST_FALSE("StateTree TaskD should not enter state", Exec.Expect(TaskD.GetName(), EnterStateStr));
 		AITEST_TRUE("StateTree TaskE should enter state", Exec.Expect(TaskE.GetName(), EnterStateStr));
 		Exec.LogClear();
+		
+		Exec.Stop();
 
 		return true;
 	}
@@ -1435,6 +1477,8 @@ struct FStateTreeTest_PassingTransitionEventToStateSelection : FAITestBase
 		AITEST_TRUE("StateTree TaskB should enter state", Exec.Expect(TaskB.GetName(), TEXT("EnterState1"))); // TaskB decorates "EnterState" with value from the payload.
 		Exec.LogClear();
 
+		Exec.Stop();
+
 		return true;
 	}
 };
@@ -1921,6 +1965,8 @@ struct FStateTreeTest_PropertyFunctions : FAITestBase
 		AITEST_TRUE(*FString::Printf(TEXT("StateTree TaskB should exit state with value %d"), ParameterPropertyFunctionsAmount), Exec.Expect(TaskB.GetName(), *FString::Printf(TEXT("ExitState%d"), ParameterPropertyFunctionsAmount)));
 		Exec.LogClear();
 
+		Exec.Stop();
+
 		return true;
 	}
 };
@@ -2290,6 +2336,8 @@ struct FStateTreeTest_FollowTransitions : FAITestBase
 		AITEST_TRUE("StateTree TaskC should enter state", Exec.Expect(TaskC.GetName(), EnterStateStr));
 		Exec.LogClear();
 
+		Exec.Stop();
+
 		return true;
 	}
 };
@@ -2365,6 +2413,8 @@ struct FStateTreeTest_InfiniteLoop : FAITestBase
 		AITEST_EQUAL("Start should fail", Status, EStateTreeRunStatus::Failed);
 		Exec.LogClear();
 
+		Exec.Stop();
+
 		return true;
 	}
 };
@@ -2433,6 +2483,8 @@ struct FStateTreeTest_Stop : FAITestBase
 		AITEST_TRUE("StateTree TaskA should get exit state expectedly", Exec.Expect(TaskAName, ExpectedExitStatusStr));
 		Exec.LogClear();
 		
+		Exec.Stop();
+
 		return true;
 	}
 
@@ -2582,7 +2634,7 @@ struct FStateTreeTest_StopEnterNormal : FStateTreeTest_Stop
 		AITEST_TRUE("StateTree TaskA should enter state", Exec.Expect(TaskAName, EnterStateStr));
 		AITEST_TRUE("StateTree TaskA should report exit status", Exec.Expect(TaskAName, ExpectedExitStatusStr));
 
-		Exec.LogClear();
+		Exec.Stop();
 		
 		return true;
 	}
@@ -2668,8 +2720,9 @@ struct FStateTreeTest_StopEnterGlobal : FStateTreeTest_Stop
 		// Normal tasks should not run
 		AITEST_FALSE("StateTree TaskA should not enter state", Exec.Expect(TaskAName, EnterStateStr));
 		AITEST_FALSE("StateTree TaskA should not report exit status", Exec.Expect(TaskAName, ExpectedExitStatusStr));
-
 		Exec.LogClear();
+
+		Exec.Stop();
 		
 		return true;
 	}
@@ -2873,6 +2926,8 @@ struct FStateTreeTest_DeferredStop_EnterGlobalTask : FStateTreeTest_DeferredStop
 		Status = Exec.Start();
 		AITEST_EQUAL("Tree should be stopped", Status, EStateTreeRunStatus::Stopped);
 
+		Exec.Stop();
+
 		return true;
 	}
 };
@@ -2891,6 +2946,8 @@ struct FStateTreeTest_DeferredStop_TickGlobalTask : FStateTreeTest_DeferredStop
 		Status = Exec.Tick(0.1f);
 		AITEST_EQUAL("Tree should be stopped", Status, EStateTreeRunStatus::Stopped);
 
+		Exec.Stop();
+
 		return true;
 	}
 };
@@ -2927,6 +2984,8 @@ struct FStateTreeTest_DeferredStop_EnterTask : FStateTreeTest_DeferredStop
 		Status = Exec.Start();
 		AITEST_EQUAL("Tree should be running", Status, EStateTreeRunStatus::Stopped);
 
+		Exec.Stop();
+
 		return true;
 	}
 };
@@ -2945,6 +3004,8 @@ struct FStateTreeTest_DeferredStop_TickTask : FStateTreeTest_DeferredStop
 		Status = Exec.Tick(0.1f);
 		AITEST_EQUAL("Tree should be stopped", Status, EStateTreeRunStatus::Stopped);
 
+		Exec.Stop();
+
 		return true;
 	}
 };
@@ -3026,7 +3087,7 @@ struct FStateTreeTest_FailEnterLinkedAsset : FAITestBase
 			AITEST_FALSE("StateTree should not enter Task2", Exec.Expect(Task2.GetName(), EnterStateStr));
 			AITEST_TRUE("StateTree should enter Task1", Exec.Expect(Task1.GetName(), EnterStateStr));
 
-			Exec.LogClear();
+			Exec.Stop();
 		}
 
 		return true;
@@ -3095,6 +3156,8 @@ struct FStateTreeTest_EnterAndExitLinkedAsset : FAITestBase
 			AITEST_TRUE("StateTree should exit Task2", Exec.Expect(Task2.GetName(), ExitStateStr));
 			AITEST_TRUE("StateTree should enter Task1", Exec.Expect(Task1.GetName(), EnterStateStr));
 			Exec.LogClear();
+
+			Exec.Stop();
 		}
 
 		return true;
@@ -3165,7 +3228,7 @@ struct FStateTreeTest_NestedOverride : FAITestBase
 			AITEST_EQUAL("Start should complete with Running", Status, EStateTreeRunStatus::Running);
 			AITEST_TRUE("StateTree should enter TaskRoot2", Exec.Expect(TaskRoot2.GetName(), EnterStateStr));
 
-			Exec.LogClear();
+			Exec.Stop();
 		}
 
 		// With overrides
@@ -3188,7 +3251,7 @@ struct FStateTreeTest_NestedOverride : FAITestBase
 			AITEST_EQUAL("Start should complete with Running", Status, EStateTreeRunStatus::Running);
 			AITEST_TRUE("StateTree should enter TaskRoot3", Exec.Expect(TaskRoot3.GetName(), EnterStateStr));
 
-			Exec.LogClear();
+			Exec.Stop();
 		}
 
 		return true;
@@ -3314,6 +3377,8 @@ struct FStateTreeTest_ParallelEventPriority : FAITestBase
 			AITEST_TRUE("StateTree should tick Task2, Task3", Exec.Expect(Task2.GetName(), TickStr).Then(Task3.GetName(), TickStr));
 			Exec.LogClear();
 		}
+
+		Exec.Stop();
 		
 		return true;
 	}
@@ -3415,6 +3480,8 @@ struct FStateTreeTest_SubTreeTransition : FAITestBase
 		AITEST_TRUE("StateTree Task4 should enter state", Exec.Expect(Task4.GetName(), EnterStateStr));
 		AITEST_TRUE("StateTree should be running", Status == EStateTreeRunStatus::Running);
 		Exec.LogClear();
+
+		Exec.Stop();
 		
 		return true;
 	}
