StateTree: Prevent crash when the property handle doesn't have a valid property after a undo operation.

diff --git a/Engine/Plugins/Runtime/StateTree/Source/StateTreeEditorModule/Private/Customizations/StateTreeBindingExtension.cpp b/Engine/Plugins/Runtime/StateTree/Source/StateTreeEditorModule/Private/Customizations/StateTreeBindingExtension.cpp
--- a/Engine/Plugins/Runtime/StateTree/Source/StateTreeEditorModule/Private/Customizations/StateTreeBindingExtension.cpp
+++ b/Engine/Plugins/Runtime/StateTree/Source/StateTreeEditorModule/Private/Customizations/StateTreeBindingExtension.cpp
@@ -1023,6 +1023,11 @@ struct FCachedBindingData : public TSharedFromThis<FCachedBindingData>
 
 		ConditionallyUpdateData();
 
+		if (!PropertyHandle.IsValid() || PropertyHandle->GetProperty() == nullptr)
+		{
+			return false;
+		}
+
 		const int32 SourceStructIndex = InBindingChain[0].ArrayIndex;
 		check(AccessibleStructs.IsValidIndex(SourceStructIndex));
 		const FStateTreeBindableStructDesc& StructDesc = AccessibleStructs[SourceStructIndex];
@@ -1470,11 +1475,10 @@ private:
 bool FStateTreeBindingExtension::IsPropertyExtendable(const UClass* InObjectClass, const IPropertyHandle& PropertyHandle) const
 {
 	const FProperty* Property = PropertyHandle.GetProperty();
-	if (Property->HasAnyPropertyFlags(CPF_PersistentInstance | CPF_EditorOnly | CPF_Config | CPF_Deprecated))
+	if (Property == nullptr || Property->HasAnyPropertyFlags(CPF_PersistentInstance | CPF_EditorOnly | CPF_Config | CPF_Deprecated))
 	{
 		return false;
 	}
-// private
 	
 	FStateTreePropertyPath TargetPath;
 	// Figure out the structs we're editing, and property path relative to current property.
