StateTree:Prevents binding to array, set and map element. Remove the ability to create bindings in the editor. It will be reenabled once all the bugs are fixed.

diff --git a/Engine/Plugins/Runtime/StateTree/Source/StateTreeEditorModule/Private/Customizations/StateTreeBindingExtension.cpp b/Engine/Plugins/Runtime/StateTree/Source/StateTreeEditorModule/Private/Customizations/StateTreeBindingExtension.cpp
--- a/Engine/Plugins/Runtime/StateTree/Source/StateTreeEditorModule/Private/Customizations/StateTreeBindingExtension.cpp
+++ b/Engine/Plugins/Runtime/StateTree/Source/StateTreeEditorModule/Private/Customizations/StateTreeBindingExtension.cpp
@@ -150,14 +150,27 @@ EStateTreePropertyUsage MakeStructPropertyPathFromPropertyHandle(TSharedPtr<cons
 	TSharedPtr<const IPropertyHandle> CurrentPropertyHandle = InPropertyHandle;
 	while (CurrentPropertyHandle.IsValid())
 	{
+		const int32 ArrayIndex = CurrentPropertyHandle->GetIndexInArray();
+		// Do not support TArray, TSet or TMap elements.
+		if (ArrayIndex != INDEX_NONE)
+		{
+			break;
+		}
+
 		const FProperty* Property = CurrentPropertyHandle->GetProperty();
 		if (Property)
 		{
+			// Do not support TSet or TMap. More work needs to be done to support these types.
+			if (Property->IsA<FSetProperty>() || Property->IsA<FMapProperty>())
+			{
+				break;
+			}
+
 			FStateTreePropertyPathSegment& Segment = PathSegments.InsertDefaulted_GetRef(0); // Traversing from leaf to root, insert in reverse.
 
 			// Store path up to the property which has ID.
 			Segment.SetName(Property->GetFName());
-			Segment.SetArrayIndex(CurrentPropertyHandle->GetIndexInArray());
+			Segment.SetArrayIndex(ArrayIndex);
 
 			// Store type of the object (e.g. for instanced objects or instanced structs).
 			if (const FObjectProperty* ObjectProperty = CastField<FObjectProperty>(Property))
