[Backout] - 36787551
Backing out the modification to prevent binding to array. Guarding beeing a cvar. Keeps the restriction for set and map

Original CL Desc
-----------------------------------------------------------------
StateTree:Prevents binding to array, set and map element. Remove the ability to create bindings in the editor. It will be reenabled once all the bugs are fixed.
#rnx

diff --git a/Engine/Plugins/Runtime/StateTree/Source/StateTreeEditorModule/Private/Customizations/StateTreeBindingExtension.cpp b/Engine/Plugins/Runtime/StateTree/Source/StateTreeEditorModule/Private/Customizations/StateTreeBindingExtension.cpp
--- a/Engine/Plugins/Runtime/StateTree/Source/StateTreeEditorModule/Private/Customizations/StateTreeBindingExtension.cpp
+++ b/Engine/Plugins/Runtime/StateTree/Source/StateTreeEditorModule/Private/Customizations/StateTreeBindingExtension.cpp
@@ -31,6 +31,12 @@
 
 namespace UE::StateTree::PropertyBinding
 {
+bool GbAllowArrayElementBindings = false;
+
+FAutoConsoleVariableRef CVarAllowArrayElementBindings(
+	TEXT("StateTree.Editor.AllowArrayElementBinding"),
+	GbAllowArrayElementBindings,
+	TEXT("Enable binding on array element in the StateTree editor."));
 
 /** Information for the types gathered from a FStateTreePropertyRef property meta-data */
 struct FRefTypeInfo
@@ -150,13 +156,6 @@ EStateTreePropertyUsage MakeStructPropertyPathFromPropertyHandle(TSharedPtr<cons
 	TSharedPtr<const IPropertyHandle> CurrentPropertyHandle = InPropertyHandle;
 	while (CurrentPropertyHandle.IsValid())
 	{
-		const int32 ArrayIndex = CurrentPropertyHandle->GetIndexInArray();
-		// Do not support TArray, TSet or TMap elements.
-		if (ArrayIndex != INDEX_NONE)
-		{
-			break;
-		}
-
 		const FProperty* Property = CurrentPropertyHandle->GetProperty();
 		if (Property)
 		{
@@ -170,7 +169,7 @@ EStateTreePropertyUsage MakeStructPropertyPathFromPropertyHandle(TSharedPtr<cons
 
 			// Store path up to the property which has ID.
 			Segment.SetName(Property->GetFName());
-			Segment.SetArrayIndex(ArrayIndex);
+			Segment.SetArrayIndex(CurrentPropertyHandle->GetIndexInArray());
 
 			// Store type of the object (e.g. for instanced objects or instanced structs).
 			if (const FObjectProperty* ObjectProperty = CastField<FObjectProperty>(Property))
@@ -210,12 +209,21 @@ EStateTreePropertyUsage MakeStructPropertyPathFromPropertyHandle(TSharedPtr<cons
 				TSharedPtr<const IPropertyHandle> ParentPropertyHandle = CurrentPropertyHandle->GetParentHandle();
 				if (ParentPropertyHandle.IsValid())
 				{
+					// Do not support TArray, TSet or TMap elements.
 					const FProperty* ParentProperty = ParentPropertyHandle->GetProperty();
-					if (ParentProperty
-						&& ParentProperty->IsA<FArrayProperty>()
-						&& Property->GetFName() == ParentProperty->GetFName())
+					if (ParentProperty)
 					{
-						CurrentPropertyHandle = ParentPropertyHandle;
+						if (UE::StateTree::PropertyBinding::GbAllowArrayElementBindings
+							&& ParentProperty->IsA<FArrayProperty>()
+							&& Property->GetFName() == ParentProperty->GetFName())
+						{
+							CurrentPropertyHandle = ParentPropertyHandle;
+						}
+						else if (ParentProperty->IsA<FMapProperty>() || ParentProperty->IsA<FSetProperty>())
+						{
+							// Prevents anything that uses TMap or TSet.
+							break;
+						}
 					}
 				}
 			}
