StateTree:  Log the current name of the node when validation fails.
#rnx

diff --git a/Engine/Plugins/Runtime/StateTree/Source/StateTreeModule/Private/Debugger/StateTreeRuntimeValidation.cpp b/Engine/Plugins/Runtime/StateTree/Source/StateTreeModule/Private/Debugger/StateTreeRuntimeValidation.cpp
--- a/Engine/Plugins/Runtime/StateTree/Source/StateTreeModule/Private/Debugger/StateTreeRuntimeValidation.cpp
+++ b/Engine/Plugins/Runtime/StateTree/Source/StateTreeModule/Private/Debugger/StateTreeRuntimeValidation.cpp
@@ -50,7 +50,7 @@ FRuntimeValidationInstanceData::~FRuntimeValidationInstanceData()
 		{
 			if (EnumHasAllFlags(Pair.State, EState::BetweenEnterExitState))
 			{
-				ensureAlwaysMsgf(false, TEXT("Tree exited. Missing ExitState on %s."), *Pair.Id.ToString());
+				ensureAlwaysMsgf(false, TEXT("Tree exited. Missing ExitState on %s."), *Private::NodeToString(StateTree.Get(), Pair.Id));
 				Private::RuntimeValidationEnterExitState = false;
 			}
 		}
