StateTree: Fix copied asset compilation. The copied asset had reference to the previous StateTree.

diff --git a/Engine/Plugins/Runtime/StateTree/Source/StateTreeEditorModule/Public/StateTreeEditorData.h b/Engine/Plugins/Runtime/StateTree/Source/StateTreeEditorModule/Public/StateTreeEditorData.h
--- a/Engine/Plugins/Runtime/StateTree/Source/StateTreeEditorModule/Public/StateTreeEditorData.h
+++ b/Engine/Plugins/Runtime/StateTree/Source/StateTreeEditorModule/Public/StateTreeEditorData.h
@@ -295,7 +295,7 @@ private:
 
 public:
 	/** Schema describing which inputs, evaluators, and tasks a StateTree can contain */	
-	UPROPERTY(EditDefaultsOnly, Category = Common, Instanced)
+	UPROPERTY(EditDefaultsOnly, Instanced, Category = Common, Instanced)
 	TObjectPtr<UStateTreeSchema> Schema = nullptr;
 
 	/** Public parameters that could be used for bindings within the Tree. */
@@ -316,7 +316,7 @@ public:
 	TSet<FStateTreeEditorColor> Colors;
 
 	/** Top level States. */
-	UPROPERTY()
+	UPROPERTY(Instanced)
 	TArray<TObjectPtr<UStateTreeState>> SubTrees;
 
 	/**
diff --git a/Engine/Plugins/Runtime/StateTree/Source/StateTreeEditorModule/Public/StateTreeState.h b/Engine/Plugins/Runtime/StateTree/Source/StateTreeEditorModule/Public/StateTreeState.h
--- a/Engine/Plugins/Runtime/StateTree/Source/StateTreeEditorModule/Public/StateTreeState.h
+++ b/Engine/Plugins/Runtime/StateTree/Source/StateTreeEditorModule/Public/StateTreeState.h
@@ -389,7 +389,7 @@ public:
 	UPROPERTY(EditDefaultsOnly, Category = "Transitions")
 	TArray<FStateTreeTransition> Transitions;
 
-	UPROPERTY()
+	UPROPERTY(Instanced)
 	TArray<TObjectPtr<UStateTreeState>> Children;
 
 	UPROPERTY(EditDefaultsOnly, Category = "State", meta = (IgnoreForMemberInitializationTest))
