StateTree: Fix the duplicated instanced on the Schema property. Requeted by review 36859934.

diff --git a/Engine/Plugins/Runtime/StateTree/Source/StateTreeEditorModule/Public/StateTreeEditorData.h b/Engine/Plugins/Runtime/StateTree/Source/StateTreeEditorModule/Public/StateTreeEditorData.h
--- a/Engine/Plugins/Runtime/StateTree/Source/StateTreeEditorModule/Public/StateTreeEditorData.h
+++ b/Engine/Plugins/Runtime/StateTree/Source/StateTreeEditorModule/Public/StateTreeEditorData.h
@@ -295,7 +295,7 @@ private:
 
 public:
 	/** Schema describing which inputs, evaluators, and tasks a StateTree can contain */	
-	UPROPERTY(EditDefaultsOnly, Instanced, Category = Common, Instanced)
+	UPROPERTY(EditDefaultsOnly, Instanced, Category = Common)
 	TObjectPtr<UStateTreeSchema> Schema = nullptr;
 
 	/** Public parameters that could be used for bindings within the Tree. */
