[Backout] - CL36787551
Original CL Desc
-----------------------------------------------------------------
StateTree:Prevents binding to array, set and map element. Remove the ability to create bindings in the editor. It will be reenabled once all the bugs are fixed.

diff --git a/Engine/Plugins/Runtime/StateTree/Source/StateTreeEditorModule/Private/Customizations/StateTreeBindingExtension.cpp b/Engine/Plugins/Runtime/StateTree/Source/StateTreeEditorModule/Private/Customizations/StateTreeBindingExtension.cpp
--- a/Engine/Plugins/Runtime/StateTree/Source/StateTreeEditorModule/Private/Customizations/StateTreeBindingExtension.cpp
+++ b/Engine/Plugins/Runtime/StateTree/Source/StateTreeEditorModule/Private/Customizations/StateTreeBindingExtension.cpp
@@ -31,12 +31,6 @@
 
 namespace UE::StateTree::PropertyBinding
 {
-bool GbAllowArrayElementBindings = false;
-
-FAutoConsoleVariableRef CVarAllowArrayElementBindings(
-	TEXT("StateTree.Editor.AllowArrayElementBinding"),
-	GbAllowArrayElementBindings,
-	TEXT("Enable binding on array element in the StateTree editor."));
 
 /** Information for the types gathered from a FStateTreePropertyRef property meta-data */
 struct FRefTypeInfo
@@ -159,12 +153,6 @@ EStateTreePropertyUsage MakeStructPropertyPathFromPropertyHandle(TSharedPtr<cons
 		const FProperty* Property = CurrentPropertyHandle->GetProperty();
 		if (Property)
 		{
-			// Do not support TSet or TMap. More work needs to be done to support these types.
-			if (Property->IsA<FSetProperty>() || Property->IsA<FMapProperty>())
-			{
-				break;
-			}
-
 			FStateTreePropertyPathSegment& Segment = PathSegments.InsertDefaulted_GetRef(0); // Traversing from leaf to root, insert in reverse.
 
 			// Store path up to the property which has ID.
@@ -209,21 +197,12 @@ EStateTreePropertyUsage MakeStructPropertyPathFromPropertyHandle(TSharedPtr<cons
 				TSharedPtr<const IPropertyHandle> ParentPropertyHandle = CurrentPropertyHandle->GetParentHandle();
 				if (ParentPropertyHandle.IsValid())
 				{
-					// Do not support TArray, TSet or TMap elements.
 					const FProperty* ParentProperty = ParentPropertyHandle->GetProperty();
-					if (ParentProperty)
+					if (ParentProperty
+						&& ParentProperty->IsA<FArrayProperty>()
+						&& Property->GetFName() == ParentProperty->GetFName())
 					{
-						if (UE::StateTree::PropertyBinding::GbAllowArrayElementBindings
-							&& ParentProperty->IsA<FArrayProperty>()
-							&& Property->GetFName() == ParentProperty->GetFName())
-						{
-							CurrentPropertyHandle = ParentPropertyHandle;
-						}
-						else if (ParentProperty->IsA<FMapProperty>() || ParentProperty->IsA<FSetProperty>())
-						{
-							// Prevents anything that uses TMap or TSet.
-							break;
-						}
+						CurrentPropertyHandle = ParentPropertyHandle;
 					}
 				}
 			}
