StateTree: Change the PropertyBindingUtils to use FStringView instead of FString. Removing the needs for dynamic allocation.

diff --git a/Engine/Plugins/Runtime/PropertyBindingUtils/Source/PropertyBindingUtils/Private/PropertyBindingPath.cpp b/Engine/Plugins/Runtime/PropertyBindingUtils/Source/PropertyBindingUtils/Private/PropertyBindingPath.cpp
--- a/Engine/Plugins/Runtime/PropertyBindingUtils/Source/PropertyBindingUtils/Private/PropertyBindingPath.cpp
+++ b/Engine/Plugins/Runtime/PropertyBindingUtils/Source/PropertyBindingUtils/Private/PropertyBindingPath.cpp
@@ -34,7 +34,7 @@ FString FBindableStructDesc::ToString() const
 //  FPropertyBindingPath
 //----------------------------------------------------------------//
 
-bool FPropertyBindingPath::FromString(const FString& InPath)
+bool FPropertyBindingPath::FromString(const FStringView InPath)
 {
 	Segments.Reset();
 	
@@ -43,48 +43,59 @@ bool FPropertyBindingPath::FromString(const FString& InPath)
 		return true;
 	}
 	
-	bool bResult = true;
-	TArray<FString> PathSegments;
-	InPath.ParseIntoArray(PathSegments, TEXT("."), /*InCullEmpty=*/false);
-	
-	for (const FString& Segment : PathSegments)
+	auto ParseSegment = [This = this](const FStringView Segment)
 	{
 		if (Segment.IsEmpty())
 		{
-			bResult = false;
-			break;
+			return false;
 		}
 
 		int32 FirstBracket = INDEX_NONE;
 		int32 LastBracket = INDEX_NONE;
-		if (Segment.FindChar(TEXT('['), FirstBracket)
-			&& Segment.FindLastChar(TEXT(']'), LastBracket))
+		if (Segment.FindChar(TEXT('['), FirstBracket) && Segment.FindLastChar(TEXT(']'), LastBracket))
 		{
 			const int32 NameStringLength = FirstBracket;
 			const int32 IndexStringLength = LastBracket - FirstBracket - 1;
 			if (NameStringLength < 1
 				|| IndexStringLength <= 0)
 			{
-				bResult = false;
-				break;
+				return false;
 			}
 
-			const FString NameString = Segment.Left(FirstBracket);
-			const FString IndexString = Segment.Mid(FirstBracket + 1, IndexStringLength);
+			const FStringView NameString = Segment.Left(FirstBracket);
+			const FStringView IndexString = Segment.Mid(FirstBracket + 1, IndexStringLength);
 			int32 ArrayIndex = INDEX_NONE;
-			LexFromString(ArrayIndex, *IndexString);
+			LexFromString(ArrayIndex, IndexString.GetData());
 			if (ArrayIndex < 0)
 			{
-				bResult = false;
-				break;
+				return false;
 			}
 			
-			AddPathSegment(FName(NameString), ArrayIndex);
+			This->AddPathSegment(FName(NameString), ArrayIndex);
 		}
 		else
 		{
-			AddPathSegment(FName(Segment));
+			This->AddPathSegment(FName(Segment));
 		}
+		return true;
+	};
+
+	bool bResult = true;
+	FStringView Path = InPath;
+	int32 FoundChar = INDEX_NONE;
+	while (Path.FindChar(TEXT('.'), FoundChar))
+	{
+		const FStringView Segment = Path.SubStr(0, FoundChar);
+		if (!ParseSegment(Segment))
+		{
+			bResult = false;
+			break;
+		}
+		Path = Path.Mid(FoundChar+1);
+	}
+	if (bResult)
+	{
+		bResult = ParseSegment(Path);
 	}
 
 	if (!bResult)
diff --git a/Engine/Plugins/Runtime/PropertyBindingUtils/Source/PropertyBindingUtils/Public/PropertyBindingPath.h b/Engine/Plugins/Runtime/PropertyBindingUtils/Source/PropertyBindingUtils/Public/PropertyBindingPath.h
--- a/Engine/Plugins/Runtime/PropertyBindingUtils/Source/PropertyBindingUtils/Public/PropertyBindingPath.h
+++ b/Engine/Plugins/Runtime/PropertyBindingUtils/Source/PropertyBindingUtils/Public/PropertyBindingPath.h
@@ -486,7 +486,7 @@ struct PROPERTYBINDINGUTILS_API FPropertyBindingPath
 	 * @param InPath Path string to parse
 	 * @return true if path was parsed successfully.
 	 */
-	bool FromString(const FString& InPath);
+	bool FromString(const FStringView InPath);
 
 	/**
 	 * Returns the property path as a one string. Highlight allows to decorate a specific segment.
diff --git a/Engine/Plugins/Runtime/StateTree/Source/StateTreeModule/Private/StateTreePropertyBindings.cpp b/Engine/Plugins/Runtime/StateTree/Source/StateTreeModule/Private/StateTreePropertyBindings.cpp
--- a/Engine/Plugins/Runtime/StateTree/Source/StateTreeModule/Private/StateTreePropertyBindings.cpp
+++ b/Engine/Plugins/Runtime/StateTree/Source/StateTreeModule/Private/StateTreePropertyBindings.cpp
@@ -1362,7 +1362,7 @@ FString FStateTreePropertyBindings::DebugInternalLayoutAsString() const
 //  FStateTreePropertyPath
 //----------------------------------------------------------------//
 
-bool FStateTreePropertyPath::FromString(const FString& InPath)
+bool FStateTreePropertyPath::FromString(const FStringView InPath)
 {
 	Segments.Reset();
 	
@@ -1371,48 +1371,59 @@ bool FStateTreePropertyPath::FromString(const FString& InPath)
 		return true;
 	}
 	
-	bool bResult = true;
-	TArray<FString> PathSegments;
-	InPath.ParseIntoArray(PathSegments, TEXT("."), /*InCullEmpty=*/false);
-	
-	for (const FString& Segment : PathSegments)
+	auto ParseSegment = [This = this](const FStringView Segment)
 	{
 		if (Segment.IsEmpty())
 		{
-			bResult = false;
-			break;
+			return false;
 		}
 
 		int32 FirstBracket = INDEX_NONE;
 		int32 LastBracket = INDEX_NONE;
-		if (Segment.FindChar(TEXT('['), FirstBracket)
-			&& Segment.FindLastChar(TEXT(']'), LastBracket))
+		if (Segment.FindChar(TEXT('['), FirstBracket) && Segment.FindLastChar(TEXT(']'), LastBracket))
 		{
 			const int32 NameStringLength = FirstBracket;
 			const int32 IndexStringLength = LastBracket - FirstBracket - 1;
 			if (NameStringLength < 1
 				|| IndexStringLength <= 0)
 			{
-				bResult = false;
-				break;
+				return false;
 			}
 
-			const FString NameString = Segment.Left(FirstBracket);
-			const FString IndexString = Segment.Mid(FirstBracket + 1, IndexStringLength);
+			const FStringView NameString = Segment.Left(FirstBracket);
+			const FStringView IndexString = Segment.Mid(FirstBracket + 1, IndexStringLength);
 			int32 ArrayIndex = INDEX_NONE;
-			LexFromString(ArrayIndex, *IndexString);
+			LexFromString(ArrayIndex, IndexString.GetData());
 			if (ArrayIndex < 0)
 			{
-				bResult = false;
-				break;
+				return false;
 			}
 			
-			AddPathSegment(FName(NameString), ArrayIndex);
+			This->AddPathSegment(FName(NameString), ArrayIndex);
 		}
 		else
 		{
-			AddPathSegment(FName(Segment));
+			This->AddPathSegment(FName(Segment));
+		}
+		return true;
+	};
+
+	bool bResult = true;
+	FStringView Path = InPath;
+	int32 FoundChar = INDEX_NONE;
+	while (Path.FindChar(TEXT('.'), FoundChar))
+	{
+		const FStringView Segment = Path.SubStr(0, FoundChar);
+		if (!ParseSegment(Segment))
+		{
+			bResult = false;
+			break;
 		}
+		Path = Path.Mid(FoundChar+1);
+	}
+	if (bResult)
+	{
+		bResult = ParseSegment(Path);
 	}
 
 	if (!bResult)
diff --git a/Engine/Plugins/Runtime/StateTree/Source/StateTreeModule/Public/StateTreePropertyBindings.h b/Engine/Plugins/Runtime/StateTree/Source/StateTreeModule/Public/StateTreePropertyBindings.h
--- a/Engine/Plugins/Runtime/StateTree/Source/StateTreeModule/Public/StateTreePropertyBindings.h
+++ b/Engine/Plugins/Runtime/StateTree/Source/StateTreeModule/Public/StateTreePropertyBindings.h
@@ -393,7 +393,7 @@ struct STATETREEMODULE_API FStateTreePropertyPath
 	 * @param InPath Path string to parse
 	 * @return true if path was parsed successfully.
 	 */
-	bool FromString(const FString& InPath);
+	bool FromString(const FStringView InPath);
 
 	/**
 	 * Returns the property path as a one string. Highlight allows to decorate a specific segment.
