StateTree: Fix compilation crash when the task fails and the state should not continue compiling. Regression introduced with CL 35821616.

diff --git a/Engine/Plugins/Runtime/StateTree/Source/StateTreeEditorModule/Private/StateTreeCompiler.cpp b/Engine/Plugins/Runtime/StateTree/Source/StateTreeEditorModule/Private/StateTreeCompiler.cpp
--- a/Engine/Plugins/Runtime/StateTree/Source/StateTreeEditorModule/Private/StateTreeCompiler.cpp
+++ b/Engine/Plugins/Runtime/StateTree/Source/StateTreeEditorModule/Private/StateTreeCompiler.cpp
@@ -757,6 +757,7 @@ bool FStateTreeCompiler::CreateStateTasksAndParameters()
 		}
 		
 		bool bStateHasTransitionTasks = false;
+		bool bCreateTaskSucceeded = true;
 		for (FStateTreeEditorNode& TaskNode : Tasks)
 		{
 			// Silently ignore empty nodes.
@@ -769,6 +770,7 @@ bool FStateTreeCompiler::CreateStateTasksAndParameters()
 			if (!CreateTask(State, TaskNode, TaskDataHandle))
 			{
 				bSucceeded = false;
+				bCreateTaskSucceeded = false;
 				continue;
 			}
 
@@ -778,6 +780,11 @@ bool FStateTreeCompiler::CreateStateTasksAndParameters()
 		}
 
 		CompactState.bHasTransitionTasks = bStateHasTransitionTasks;
+
+		if (!bCreateTaskSucceeded)
+		{
+			continue;
+		}
 		
 		const int32 TasksNum = Nodes.Num() - TasksBegin;
 		if (const auto Validation = UE::StateTree::Compiler::IsValidCount8(TasksNum); Validation.DidFail())
