StateTree: Split StateTreeTest.cpp into multiple files and categorize them. This is to help understand which cases are cover. Create a base class for StateTree test.

diff --git a/Engine/Plugins/Runtime/StateTree/Source/StateTreeTestSuite/Private/StateTreeBindingTest.cpp b/Engine/Plugins/Runtime/StateTree/Source/StateTreeTestSuite/Private/StateTreeBindingTest.cpp
new file mode 100644
--- /dev/null
+++ b/Engine/Plugins/Runtime/StateTree/Source/StateTreeTestSuite/Private/StateTreeBindingTest.cpp
@@ -0,0 +1,478 @@
+// Copyright Epic Games, Inc. All Rights Reserved.
+
+#include "StateTreeTest.h"
+#include "StateTreeTestBase.h"
+#include "StateTreeTestTypes.h"
+
+#include "StateTreeCompilerLog.h"
+#include "StateTreeEditorData.h"
+#include "StateTreeCompiler.h"
+#include "Conditions/StateTreeCommonConditions.h"
+
+#define LOCTEXT_NAMESPACE "AITestSuite_StateTreeTest"
+
+UE_DISABLE_OPTIMIZATION_SHIP
+
+namespace UE::StateTree::Tests
+{
+
+struct FStateTreeTest_BindingsCompiler : FStateTreeTestBase
+{
+	virtual bool InstantTest() override
+	{
+		FStateTreeCompilerLog Log;
+		FStateTreePropertyBindings Bindings;
+		FStateTreePropertyBindingCompiler BindingCompiler;
+
+		const bool bInitResult = BindingCompiler.Init(Bindings, Log);
+		AITEST_TRUE("Expect init to succeed", bInitResult);
+
+		FStateTreeBindableStructDesc SourceADesc;
+		SourceADesc.Name = FName(TEXT("SourceA"));
+		SourceADesc.Struct = TBaseStructure<FStateTreeTest_PropertyCopy>::Get();
+		SourceADesc.DataSource = EStateTreeBindableStructSource::Parameter;
+		SourceADesc.DataHandle = FStateTreeDataHandle(EStateTreeDataSourceType::ContextData, 0); // Used as index to SourceViews below.
+		SourceADesc.ID = FGuid::NewGuid();
+
+		FStateTreeBindableStructDesc SourceBDesc;
+		SourceBDesc.Name = FName(TEXT("SourceB"));
+		SourceBDesc.Struct = TBaseStructure<FStateTreeTest_PropertyCopy>::Get();
+		SourceBDesc.DataSource = EStateTreeBindableStructSource::Parameter;
+		SourceBDesc.DataHandle = FStateTreeDataHandle(EStateTreeDataSourceType::ContextData, 1); // Used as index to SourceViews below.
+		SourceBDesc.ID = FGuid::NewGuid();
+
+		FStateTreeBindableStructDesc TargetDesc;
+		TargetDesc.Name = FName(TEXT("Target"));
+		TargetDesc.Struct = TBaseStructure<FStateTreeTest_PropertyCopy>::Get();
+		TargetDesc.DataSource = EStateTreeBindableStructSource::Parameter;
+		TargetDesc.ID = FGuid::NewGuid();
+		
+		const int32 SourceAIndex = BindingCompiler.AddSourceStruct(SourceADesc);
+		const int32 SourceBIndex = BindingCompiler.AddSourceStruct(SourceBDesc);
+
+		TArray<FStateTreePropertyPathBinding> PropertyBindings;
+		PropertyBindings.Add(MakeBinding(SourceBDesc.ID, TEXT("Item"), TargetDesc.ID, TEXT("Array[1]")));
+		PropertyBindings.Add(MakeBinding(SourceADesc.ID, TEXT("Item.B"), TargetDesc.ID, TEXT("Array[1].B")));
+		PropertyBindings.Add(MakeBinding(SourceADesc.ID, TEXT("Array"), TargetDesc.ID, TEXT("Array")));
+
+		int32 CopyBatchIndex = INDEX_NONE;
+		const bool bCompileBatchResult = BindingCompiler.CompileBatch(TargetDesc, PropertyBindings, FStateTreeIndex16::Invalid, FStateTreeIndex16::Invalid, CopyBatchIndex);
+		AITEST_TRUE("CompileBatch should succeed", bCompileBatchResult);
+		AITEST_NOT_EQUAL("CopyBatchIndex should not be INDEX_NONE", CopyBatchIndex, (int32)INDEX_NONE);
+
+		BindingCompiler.Finalize();
+
+		const bool bResolveResult = Bindings.ResolvePaths();
+		AITEST_TRUE("ResolvePaths should succeed", bResolveResult);
+
+		FStateTreeTest_PropertyCopy SourceA;
+		SourceA.Item.B = 123;
+		SourceA.Array.AddDefaulted_GetRef().A = 1;
+		SourceA.Array.AddDefaulted_GetRef().B = 2;
+
+		FStateTreeTest_PropertyCopy SourceB;
+		SourceB.Item.A = 41;
+		SourceB.Item.B = 42;
+
+		FStateTreeTest_PropertyCopy Target;
+
+		AITEST_TRUE("SourceAIndex should be less than max number of source structs.", SourceAIndex < Bindings.GetSourceStructNum());
+		AITEST_TRUE("SourceBIndex should be less than max number of source structs.", SourceBIndex < Bindings.GetSourceStructNum());
+
+		TArray<FStateTreeDataView> SourceViews;
+		SourceViews.SetNum(Bindings.GetSourceStructNum());
+		SourceViews[SourceAIndex] = FStateTreeDataView(FStructView::Make(SourceA));
+		SourceViews[SourceBIndex] = FStateTreeDataView(FStructView::Make(SourceB));
+		FStateTreeDataView TargetView(FStructView::Make(Target));
+
+		bool bCopyResult = true;
+		for (const FStateTreePropertyCopy& Copy : Bindings.GetBatchCopies(FStateTreeIndex16(CopyBatchIndex)))
+		{
+			bCopyResult &= Bindings.CopyProperty(Copy, SourceViews[Copy.SourceDataHandle.GetIndex()], TargetView);
+		}
+		AITEST_TRUE("CopyTo should succeed", bCopyResult);
+
+		// Due to binding sorting, we expect them to executed in this order (sorted based on target access, earliest to latest)
+		// SourceA.Array -> Target.Array
+		// SourceB.Item -> Target.Array[1]
+		// SourceA.Item.B -> Target.Array[1].B
+
+		AITEST_EQUAL("Expect TargetArray to be copied from SourceA", Target.Array.Num(), SourceA.Array.Num());
+		AITEST_EQUAL("Expect Target.Array[0].A copied from SourceA.Array[0].A", Target.Array[0].A, SourceA.Array[0].A);
+		AITEST_EQUAL("Expect Target.Array[0].B copied from SourceA.Array[0].B", Target.Array[0].B, SourceA.Array[0].B);
+		AITEST_EQUAL("Expect Target.Array[1].A copied from SourceB.Item.A", Target.Array[1].A, SourceB.Item.A);
+		AITEST_EQUAL("Expect Target.Array[1].B copied from SourceA.Item.B", Target.Array[1].B, SourceA.Item.B);
+		
+		return true;
+	}
+};
+IMPLEMENT_AI_INSTANT_TEST(FStateTreeTest_BindingsCompiler, "System.StateTree.Binding.BindingsCompiler");
+
+struct FStateTreeTest_PropertyFunctions : FStateTreeTestBase
+{
+	virtual bool InstantTest() override
+	{
+		UStateTree& StateTree = NewStateTree();
+		UStateTreeEditorData& EditorData = *Cast<UStateTreeEditorData>(StateTree.EditorData);
+		UStateTreeState& Root = EditorData.AddSubTree(FName(TEXT("Root")));
+		FStateTreePropertyPathSegment PathSegmentToFuncResult = FStateTreePropertyPathSegment(TEXT("Result"));
+
+		// Condition with property function binding.
+		{
+			TStateTreeEditorNode<FStateTreeCompareIntCondition>& EnterCond = Root.AddEnterCondition<FStateTreeCompareIntCondition>(EGenericAICheck::Equal);
+			EnterCond.GetInstanceData().Right = 1;
+			EditorData.AddPropertyBinding(CastChecked<UScriptStruct>(FTestPropertyFunction::StaticStruct()), {PathSegmentToFuncResult}, FStateTreePropertyPath(EnterCond.ID, TEXT("Left")));
+		}
+
+		// Task with multiple nested property function bindings.
+		auto& TaskA = Root.AddTask<FTestTask_PrintAndResetValue>(FName(TEXT("TaskA")));
+		constexpr int32 TaskAPropertyFunctionsAmount = 10;
+		{
+			EditorData.AddPropertyBinding(CastChecked<UScriptStruct>(FTestPropertyFunction::StaticStruct()), {PathSegmentToFuncResult}, FStateTreePropertyPath(TaskA.ID, TEXT("Value")));
+		
+			for (int32 i = 0; i < TaskAPropertyFunctionsAmount - 1; ++i)
+			{
+				const FStateTreePropertyPathBinding& LastBinding = EditorData.GetPropertyEditorBindings()->GetBindings().Last();
+				const FGuid LastBindingPropertyFuncID = LastBinding.GetPropertyFunctionNode().Get<const FStateTreeEditorNode>().ID;
+				EditorData.AddPropertyBinding(CastChecked<UScriptStruct>(FTestPropertyFunction::StaticStruct()), {PathSegmentToFuncResult}, FStateTreePropertyPath(LastBindingPropertyFuncID, TEXT("Input")));
+			}
+		}
+
+		// Task bound to state parameter with multiple nested property function bindings.
+		auto& TaskB = Root.AddTask<FTestTask_PrintAndResetValue>(FName(TEXT("TaskB")));
+		constexpr int32 ParameterPropertyFunctionsAmount = 5;
+		{
+			Root.Parameters.Parameters.AddProperty(FName(TEXT("Int")), EPropertyBagPropertyType::Int32);
+			const FStateTreePropertyPath PathToProperty = FStateTreePropertyPath(Root.Parameters.ID, TEXT("Int"));
+			EditorData.AddPropertyBinding(PathToProperty, FStateTreePropertyPath(TaskB.ID, TEXT("Value")));
+			EditorData.AddPropertyBinding(CastChecked<UScriptStruct>(FTestPropertyFunction::StaticStruct()), {PathSegmentToFuncResult}, PathToProperty);
+		
+			for (int32 i = 0; i < ParameterPropertyFunctionsAmount - 1; ++i)
+			{
+				const FStateTreePropertyPathBinding& LastBinding = EditorData.GetPropertyEditorBindings()->GetBindings().Last();
+				const FGuid LastBindingPropertyFuncID = LastBinding.GetPropertyFunctionNode().Get<const FStateTreeEditorNode>().ID;
+				EditorData.AddPropertyBinding(CastChecked<UScriptStruct>(FTestPropertyFunction::StaticStruct()), {PathSegmentToFuncResult}, FStateTreePropertyPath(LastBindingPropertyFuncID, TEXT("Input")));
+			}
+		}
+
+		FStateTreeCompilerLog Log;
+		FStateTreeCompiler Compiler(Log);
+		const bool bResult = Compiler.Compile(StateTree);
+		AITEST_TRUE("StateTree should get compiled", bResult);
+
+		FStateTreeInstanceData InstanceData;
+		FTestStateTreeExecutionContext Exec(StateTree, StateTree, InstanceData);
+		const bool bInitSucceeded = Exec.IsValid();
+		AITEST_TRUE("StateTree should init", bInitSucceeded);
+
+		Exec.Start();
+		AITEST_TRUE(*FString::Printf(TEXT("StateTree TaskA should enter state with value %d"), TaskAPropertyFunctionsAmount), Exec.Expect(TaskA.GetName(), *FString::Printf(TEXT("EnterState%d"), TaskAPropertyFunctionsAmount)));
+		AITEST_TRUE(*FString::Printf(TEXT("StateTree TaskB should enter state with value %d"), ParameterPropertyFunctionsAmount), Exec.Expect(TaskB.GetName(), *FString::Printf(TEXT("EnterState%d"), ParameterPropertyFunctionsAmount)));
+		Exec.LogClear();
+
+		Exec.Tick(0.1f);
+		AITEST_TRUE(*FString::Printf(TEXT("StateTree TaskA should tick with value %d"), TaskAPropertyFunctionsAmount), Exec.Expect(TaskA.GetName(), *FString::Printf(TEXT("Tick%d"), TaskAPropertyFunctionsAmount)));
+		AITEST_TRUE(*FString::Printf(TEXT("StateTree TaskB should tick with value %d"), ParameterPropertyFunctionsAmount), Exec.Expect(TaskB.GetName(), *FString::Printf(TEXT("Tick%d"), ParameterPropertyFunctionsAmount)));
+		Exec.LogClear();
+
+		Exec.Stop(EStateTreeRunStatus::Stopped);
+		AITEST_TRUE(*FString::Printf(TEXT("StateTree TaskA should exit state with value %d"), TaskAPropertyFunctionsAmount), Exec.Expect(TaskA.GetName(), *FString::Printf(TEXT("ExitState%d"), TaskAPropertyFunctionsAmount)));
+		AITEST_TRUE(*FString::Printf(TEXT("StateTree TaskB should exit state with value %d"), ParameterPropertyFunctionsAmount), Exec.Expect(TaskB.GetName(), *FString::Printf(TEXT("ExitState%d"), ParameterPropertyFunctionsAmount)));
+		Exec.LogClear();
+
+		Exec.Stop();
+
+		return true;
+	}
+};
+IMPLEMENT_AI_INSTANT_TEST(FStateTreeTest_PropertyFunctions, "System.StateTree.Binding.PropertyFunctions");
+
+struct FStateTreeTest_CopyObjects : FStateTreeTestBase
+{
+	virtual bool InstantTest() override
+	{
+		FStateTreeCompilerLog Log;
+		FStateTreePropertyBindings Bindings;
+		FStateTreePropertyBindingCompiler BindingCompiler;
+
+		const bool bInitResult = BindingCompiler.Init(Bindings, Log);
+		AITEST_TRUE("Expect init to succeed", bInitResult);
+
+		FStateTreeBindableStructDesc SourceDesc;
+		SourceDesc.Name = FName(TEXT("Source"));
+		SourceDesc.Struct = TBaseStructure<FStateTreeTest_PropertyCopyObjects>::Get();
+		SourceDesc.DataSource = EStateTreeBindableStructSource::Parameter;
+		SourceDesc.DataHandle = FStateTreeDataHandle(EStateTreeDataSourceType::ContextData, 0); // Used as index to SourceViews below.
+		SourceDesc.ID = FGuid::NewGuid();
+
+		FStateTreeBindableStructDesc TargetADesc;
+		TargetADesc.Name = FName(TEXT("TargetA"));
+		TargetADesc.Struct = TBaseStructure<FStateTreeTest_PropertyCopyObjects>::Get();
+		TargetADesc.DataSource = EStateTreeBindableStructSource::Parameter;
+		TargetADesc.ID = FGuid::NewGuid();
+
+		FStateTreeBindableStructDesc TargetBDesc;
+		TargetBDesc.Name = FName(TEXT("TargetB"));
+		TargetBDesc.Struct = TBaseStructure<FStateTreeTest_PropertyCopyObjects>::Get();
+		TargetBDesc.DataSource = EStateTreeBindableStructSource::Parameter;
+		TargetBDesc.ID = FGuid::NewGuid();
+
+		const int32 SourceIndex = BindingCompiler.AddSourceStruct(SourceDesc);
+
+		TArray<FStateTreePropertyPathBinding> PropertyBindings;
+		// One-to-one copy from source to target A
+		PropertyBindings.Add(MakeBinding(SourceDesc.ID, TEXT("Object"), TargetADesc.ID, TEXT("Object")));
+		PropertyBindings.Add(MakeBinding(SourceDesc.ID, TEXT("SoftObject"), TargetADesc.ID, TEXT("SoftObject")));
+		PropertyBindings.Add(MakeBinding(SourceDesc.ID, TEXT("Class"), TargetADesc.ID, TEXT("Class")));
+		PropertyBindings.Add(MakeBinding(SourceDesc.ID, TEXT("SoftClass"), TargetADesc.ID, TEXT("SoftClass")));
+
+		// Cross copy from source to target B
+		PropertyBindings.Add(MakeBinding(SourceDesc.ID, TEXT("SoftObject"), TargetBDesc.ID, TEXT("Object")));
+		PropertyBindings.Add(MakeBinding(SourceDesc.ID, TEXT("Object"), TargetBDesc.ID, TEXT("SoftObject")));
+		PropertyBindings.Add(MakeBinding(SourceDesc.ID, TEXT("SoftClass"), TargetBDesc.ID, TEXT("Class")));
+		PropertyBindings.Add(MakeBinding(SourceDesc.ID, TEXT("Class"), TargetBDesc.ID, TEXT("SoftClass")));
+		
+		int32 TargetACopyBatchIndex = INDEX_NONE;
+		const bool bCompileBatchResultA = BindingCompiler.CompileBatch(TargetADesc, PropertyBindings, FStateTreeIndex16::Invalid, FStateTreeIndex16::Invalid, TargetACopyBatchIndex);
+		AITEST_TRUE("CompileBatchResultA should succeed", bCompileBatchResultA);
+		AITEST_NOT_EQUAL("TargetACopyBatchIndex should not be INDEX_NONE", TargetACopyBatchIndex, (int32)INDEX_NONE);
+
+		int32 TargetBCopyBatchIndex = INDEX_NONE;
+		const bool bCompileBatchResultB = BindingCompiler.CompileBatch(TargetBDesc, PropertyBindings, FStateTreeIndex16::Invalid, FStateTreeIndex16::Invalid, TargetBCopyBatchIndex);
+		AITEST_TRUE("CompileBatchResultB should succeed", bCompileBatchResultB);
+		AITEST_NOT_EQUAL("TargetBCopyBatchIndex should not be INDEX_NONE", TargetBCopyBatchIndex, (int32)INDEX_NONE);
+
+		BindingCompiler.Finalize();
+
+		const bool bResolveResult = Bindings.ResolvePaths();
+		AITEST_TRUE("ResolvePaths should succeed", bResolveResult);
+
+		UStateTreeTest_PropertyObject* ObjectA = NewObject<UStateTreeTest_PropertyObject>();
+		UStateTreeTest_PropertyObject2* ObjectB = NewObject<UStateTreeTest_PropertyObject2>();
+		
+		FStateTreeTest_PropertyCopyObjects Source;
+		Source.Object = ObjectA;
+		Source.SoftObject = ObjectB;
+		Source.Class = UStateTreeTest_PropertyObject::StaticClass();
+		Source.SoftClass = UStateTreeTest_PropertyObject::StaticClass();
+
+		AITEST_TRUE("SourceIndex should be less than max number of source structs.", SourceIndex < Bindings.GetSourceStructNum());
+
+		TArray<FStateTreeDataView> SourceViews;
+		SourceViews.SetNum(Bindings.GetSourceStructNum());
+		SourceViews[SourceIndex] = FStateTreeDataView(FStructView::Make(Source));
+
+		FStateTreeTest_PropertyCopyObjects TargetA;
+		bool bCopyResultA = true;
+		for (const FStateTreePropertyCopy& Copy : Bindings.GetBatchCopies(FStateTreeIndex16(TargetACopyBatchIndex)))
+		{
+			bCopyResultA &= Bindings.CopyProperty(Copy, SourceViews[Copy.SourceDataHandle.GetIndex()], FStructView::Make(TargetA));
+		}
+		AITEST_TRUE("CopyTo should succeed", bCopyResultA);
+
+		AITEST_TRUE("Expect TargetA.Object == Source.Object", TargetA.Object == Source.Object);
+		AITEST_TRUE("Expect TargetA.SoftObject == Source.SoftObject", TargetA.SoftObject == Source.SoftObject);
+		AITEST_TRUE("Expect TargetA.Class == Source.Class", TargetA.Class == Source.Class);
+		AITEST_TRUE("Expect TargetA.SoftClass == Source.SoftClass", TargetA.SoftClass == Source.SoftClass);
+
+		// Copying to TargetB should not affect TargetA
+		TargetA.Object = nullptr;
+		
+		FStateTreeTest_PropertyCopyObjects TargetB;
+		bool bCopyResultB = true;
+		for (const FStateTreePropertyCopy& Copy : Bindings.GetBatchCopies(FStateTreeIndex16(TargetBCopyBatchIndex)))
+		{
+			bCopyResultB &= Bindings.CopyProperty(Copy, SourceViews[Copy.SourceDataHandle.GetIndex()], FStructView::Make(TargetB));
+		}
+		AITEST_TRUE("CopyTo should succeed", bCopyResultB);
+
+		AITEST_TRUE("Expect TargetB.Object == Source.SoftObject", TSoftObjectPtr<UObject>(TargetB.Object) == Source.SoftObject);
+		AITEST_TRUE("Expect TargetB.SoftObject == Source.Object", TargetB.SoftObject == TSoftObjectPtr<UObject>(Source.Object));
+		AITEST_TRUE("Expect TargetB.Class == Source.SoftClass", TSoftClassPtr<UObject>(TargetB.Class) == Source.SoftClass);
+		AITEST_TRUE("Expect TargetB.SoftClass == Source.Class", TargetB.SoftClass == TSoftClassPtr<UObject>(Source.Class));
+
+		AITEST_TRUE("Expect TargetA.Object == nullptr after copy of TargetB", TargetA.Object == nullptr);
+
+		// Collect ObjectA and ObjectB, soft object paths should still copy ok.
+		ObjectA = nullptr;
+		ObjectB = nullptr;
+		Source.Object = nullptr;
+		CollectGarbage(GARBAGE_COLLECTION_KEEPFLAGS);
+
+		FStateTreeTest_PropertyCopyObjects TargetC;
+		bool bCopyResultC = true;
+		for (const FStateTreePropertyCopy& Copy : Bindings.GetBatchCopies(FStateTreeIndex16(TargetACopyBatchIndex)))
+		{
+			bCopyResultB &= Bindings.CopyProperty(Copy, SourceViews[Copy.SourceDataHandle.GetIndex()], FStructView::Make(TargetC));
+		}
+
+		
+		AITEST_TRUE("CopyTo should succeed", bCopyResultC);
+		AITEST_TRUE("Expect TargetC.SoftObject == Source.SoftObject after GC", TargetC.SoftObject == Source.SoftObject);
+
+		return true;
+	}
+};
+IMPLEMENT_AI_INSTANT_TEST(FStateTreeTest_CopyObjects, "System.StateTree.Binding.CopyObjects");
+
+
+struct FStateTreeTest_References : FStateTreeTestBase
+{
+	virtual bool InstantTest() override
+	{
+		FStateTreeCompilerLog Log;
+		FStateTreePropertyBindings Bindings;
+		FStateTreePropertyBindingCompiler BindingCompiler;
+
+		const bool bInitResult = BindingCompiler.Init(Bindings, Log);
+		AITEST_TRUE("Expect init to succeed", bInitResult);
+
+		FStateTreeBindableStructDesc SourceDesc;
+		SourceDesc.Name = FName(TEXT("Source"));
+		SourceDesc.Struct = TBaseStructure<FStateTreeTest_PropertyRefSourceStruct>::Get();
+		SourceDesc.DataSource = EStateTreeBindableStructSource::Parameter;
+		SourceDesc.DataHandle = FStateTreeDataHandle(EStateTreeDataSourceType::ContextData, 0);
+		SourceDesc.ID = FGuid::NewGuid();
+		BindingCompiler.AddSourceStruct(SourceDesc);
+
+		FStateTreeBindableStructDesc TargetDesc;
+		TargetDesc.Name = FName(TEXT("Target"));
+		TargetDesc.Struct = TBaseStructure<FStateTreeTest_PropertyRefTargetStruct>::Get();
+		TargetDesc.DataSource = EStateTreeBindableStructSource::Parameter;
+		TargetDesc.ID = FGuid::NewGuid();
+
+		TArray<FStateTreePropertyPathBinding> PropertyBindings;
+		PropertyBindings.Add(MakeBinding(SourceDesc.ID, TEXT("Item"), TargetDesc.ID, TEXT("RefToStruct")));
+		PropertyBindings.Add(MakeBinding(SourceDesc.ID, TEXT("Item.A"), TargetDesc.ID, TEXT("RefToInt")));
+		PropertyBindings.Add(MakeBinding(SourceDesc.ID, TEXT("Array"), TargetDesc.ID, TEXT("RefToStructArray")));
+
+		FStateTreeTest_PropertyRefSourceStruct Source;
+		FStateTreeDataView SourceView = FStateTreeDataView(FStructView::Make(Source));
+
+		FStateTreeTest_PropertyRefTargetStruct Target;
+		FStateTreeDataView TargetView(FStructView::Make(Target));
+
+		TMap<FGuid, const FStateTreeDataView> IDToStructValue;
+		IDToStructValue.Emplace(SourceDesc.ID, SourceView);
+		IDToStructValue.Emplace(TargetDesc.ID, TargetView);
+
+		const bool bCompileReferencesResult = BindingCompiler.CompileReferences(TargetDesc, PropertyBindings, TargetView, IDToStructValue);
+		AITEST_TRUE("CompileReferences should succeed", bCompileReferencesResult);
+
+		BindingCompiler.Finalize();
+
+		const bool bResolveResult = Bindings.ResolvePaths();
+		AITEST_TRUE("ResolvePaths should succeed", bResolveResult);
+
+		{
+			const FStateTreePropertyAccess* PropertyAccess = Bindings.GetPropertyAccess(Target.RefToStruct);
+			AITEST_NOT_NULL("GetPropertyAccess should succeed", PropertyAccess);
+
+			FStateTreeTest_PropertyStruct* Reference = Bindings.GetMutablePropertyPtr<FStateTreeTest_PropertyStruct>(SourceView, *PropertyAccess);
+			AITEST_EQUAL("Expect RefToStruct to point to SourceA.Item", Reference, &Source.Item);
+		}
+
+		{
+			const FStateTreePropertyAccess* PropertyAccess = Bindings.GetPropertyAccess(Target.RefToInt);
+			AITEST_NOT_NULL("GetPropertyAccess should succeed", PropertyAccess);
+
+			int32* Reference = Bindings.GetMutablePropertyPtr<int32>(SourceView, *PropertyAccess);
+			AITEST_EQUAL("Expect RefToInt to point to SourceA.Item.A", Reference, &Source.Item);
+		}
+
+		{
+			const FStateTreePropertyAccess* PropertyAccess = Bindings.GetPropertyAccess(Target.RefToStructArray);
+			AITEST_NOT_NULL("GetPropertyAccess should succeed", PropertyAccess);
+
+			TArray<FStateTreeTest_PropertyStruct>* Reference = Bindings.GetMutablePropertyPtr<TArray<FStateTreeTest_PropertyStruct>>(SourceView, *PropertyAccess);
+			AITEST_EQUAL("Expect RefToStructArray to point to SourceA.Array", Reference, &Source.Array);
+		}
+
+		return true;
+	}
+};
+IMPLEMENT_AI_INSTANT_TEST(FStateTreeTest_References, "System.StateTree.Binding.References");
+
+struct FStateTreeTest_ReferencesConstness : FStateTreeTestBase
+{
+	virtual bool InstantTest() override
+	{
+		FStateTreeCompilerLog Log;
+		FStateTreePropertyBindings Bindings;
+		FStateTreePropertyBindingCompiler BindingCompiler;
+
+		const bool bInitResult = BindingCompiler.Init(Bindings, Log);
+		AITEST_TRUE("Expect init to succeed", bInitResult);
+
+		FStateTreeBindableStructDesc SourceAsTaskDesc;
+		SourceAsTaskDesc.Name = FName(TEXT("SourceTask"));
+		SourceAsTaskDesc.Struct = TBaseStructure<FStateTreeTest_PropertyRefSourceStruct>::Get();
+		SourceAsTaskDesc.DataSource = EStateTreeBindableStructSource::Task;
+		SourceAsTaskDesc.DataHandle = FStateTreeDataHandle(EStateTreeDataSourceType::ContextData, 0);
+		SourceAsTaskDesc.ID = FGuid::NewGuid();
+		BindingCompiler.AddSourceStruct(SourceAsTaskDesc);
+
+		FStateTreeBindableStructDesc SourceAsContextDesc;
+		SourceAsContextDesc.Name = FName(TEXT("SourceContext"));
+		SourceAsContextDesc.Struct = TBaseStructure<FStateTreeTest_PropertyRefSourceStruct>::Get();
+		SourceAsContextDesc.DataSource = EStateTreeBindableStructSource::Context;
+		SourceAsContextDesc.DataHandle = FStateTreeDataHandle(EStateTreeDataSourceType::ContextData, 0);
+		SourceAsContextDesc.ID = FGuid::NewGuid();
+		BindingCompiler.AddSourceStruct(SourceAsContextDesc);
+
+		FStateTreeBindableStructDesc TargetDesc;
+		TargetDesc.Name = FName(TEXT("Target"));
+		TargetDesc.Struct = TBaseStructure<FStateTreeTest_PropertyRefTargetStruct>::Get();
+		TargetDesc.DataSource = EStateTreeBindableStructSource::Parameter;
+		TargetDesc.ID = FGuid::NewGuid();
+
+		FStateTreePropertyPathBinding TaskPropertyBinding = MakeBinding(SourceAsTaskDesc.ID, TEXT("Item"), TargetDesc.ID, TEXT("RefToStruct"));
+		FStateTreePropertyPathBinding TaskOutputPropertyBinding = MakeBinding(SourceAsTaskDesc.ID, TEXT("OutputItem"), TargetDesc.ID, TEXT("RefToStruct"));
+
+		FStateTreePropertyPathBinding ContextPropertyBinding = MakeBinding(SourceAsTaskDesc.ID, TEXT("Item"), TargetDesc.ID, TEXT("RefToStruct"));
+		FStateTreePropertyPathBinding ContextOutputPropertyBinding = MakeBinding(SourceAsTaskDesc.ID, TEXT("Item"), TargetDesc.ID, TEXT("RefToStruct"));
+
+		FStateTreeTest_PropertyRefSourceStruct SourceAsTask;
+		FStateTreeDataView SourceAsTaskView(FStructView::Make(SourceAsTask));
+
+		FStateTreeTest_PropertyRefSourceStruct SourceAsContext;
+		FStateTreeDataView SourceAsContextView(FStructView::Make(SourceAsContext));
+
+		FStateTreeTest_PropertyRefTargetStruct Target;
+		FStateTreeDataView TargetView(FStructView::Make(Target));
+
+		TMap<FGuid, const FStateTreeDataView> IDToStructValue;
+		IDToStructValue.Emplace(SourceAsTaskDesc.ID, SourceAsTaskView);
+		IDToStructValue.Emplace(SourceAsContextDesc.ID, SourceAsContextView);
+		IDToStructValue.Emplace(TargetDesc.ID, TargetView);
+
+		{
+			const bool bCompileReferenceResult = BindingCompiler.CompileReferences(TargetDesc, { TaskPropertyBinding }, TargetView, IDToStructValue);
+			AITEST_FALSE("CompileReferences should fail", bCompileReferenceResult);
+		}
+
+		{
+			const bool bCompileReferenceResult = BindingCompiler.CompileReferences(TargetDesc, { TaskOutputPropertyBinding }, TargetView, IDToStructValue);
+			AITEST_TRUE("CompileReferences should succeed", bCompileReferenceResult);
+		}
+
+		{
+			const bool bCompileReferenceResult = BindingCompiler.CompileReferences(TargetDesc, { ContextPropertyBinding }, TargetView, IDToStructValue);
+			AITEST_FALSE("CompileReferences should fail", bCompileReferenceResult);
+		}
+
+		{
+			const bool bCompileReferenceResult = BindingCompiler.CompileReferences(TargetDesc, { ContextOutputPropertyBinding }, TargetView, IDToStructValue);
+			AITEST_FALSE("CompileReferences should fail", bCompileReferenceResult);
+		}
+
+		return true;
+	}
+};
+IMPLEMENT_AI_INSTANT_TEST(FStateTreeTest_ReferencesConstness, "System.StateTree.Binding.ReferencesConstness");
+
+} // namespace UE::StateTree::Tests
+
+UE_ENABLE_OPTIMIZATION_SHIP
+
+#undef LOCTEXT_NAMESPACE
diff --git a/Engine/Plugins/Runtime/StateTree/Source/StateTreeTestSuite/Private/StateTreePropertyPathTest.cpp b/Engine/Plugins/Runtime/StateTree/Source/StateTreeTestSuite/Private/StateTreePropertyPathTest.cpp
new file mode 100644
--- /dev/null
+++ b/Engine/Plugins/Runtime/StateTree/Source/StateTreeTestSuite/Private/StateTreePropertyPathTest.cpp
@@ -0,0 +1,340 @@
+// Copyright Epic Games, Inc. All Rights Reserved.
+
+#include "StateTreeTest.h"
+#include "StateTreeTestBase.h"
+#include "StateTreeTestTypes.h"
+
+#include "StateTreeCompilerLog.h"
+#include "StateTreeEditorData.h"
+#include "StateTreeCompiler.h"
+#include "Conditions/StateTreeCommonConditions.h"
+
+#define LOCTEXT_NAMESPACE "AITestSuite_StateTreeTest"
+
+UE_DISABLE_OPTIMIZATION_SHIP
+
+namespace UE::StateTree::Tests
+{
+
+struct FStateTreeTest_PropertyPathOffset : FStateTreeTestBase
+{
+	virtual bool InstantTest() override
+	{
+		FStateTreePropertyPath Path;
+		const bool bParseResult = Path.FromString(TEXT("StructB.B"));
+
+		AITEST_TRUE("Parsing path should succeeed", bParseResult);
+		AITEST_EQUAL("Should have 2 path segments", Path.NumSegments(), 2);
+
+		FString ResolveErrors;
+		TArray<FStateTreePropertyPathIndirection> Indirections;
+		const bool bResolveResult = Path.ResolveIndirections(FStateTreeTest_PropertyStruct::StaticStruct(), Indirections, &ResolveErrors);
+
+		AITEST_TRUE("Resolve path should succeeed", bResolveResult);
+		AITEST_EQUAL("Should have no resolve errors", ResolveErrors.Len(), 0);
+		
+		AITEST_EQUAL("Should have 2 indirections", Indirections.Num(), 2);
+		AITEST_EQUAL("Indirection 0 should be Offset type", Indirections[0].GetAccessType(), EStateTreePropertyAccessType::Offset);
+		AITEST_EQUAL("Indirection 1 should be Offset type", Indirections[1].GetAccessType(), EStateTreePropertyAccessType::Offset);
+
+		return true;
+	}
+};
+IMPLEMENT_AI_INSTANT_TEST(FStateTreeTest_PropertyPathOffset, "System.StateTree.PropertyPath.Offset");
+
+struct FStateTreeTest_PropertyPathParseFail : FStateTreeTestBase
+{
+	virtual bool InstantTest() override
+	{
+		{
+			FStateTreePropertyPath Path;
+			const bool bParseResult = Path.FromString(TEXT("")); // empty is valid.
+			AITEST_TRUE("Parsing path should succeed", bParseResult);
+		}
+
+		{
+			FStateTreePropertyPath Path;
+			const bool bParseResult = Path.FromString(TEXT("StructB.[0]B"));
+			AITEST_FALSE("Parsing path should fail", bParseResult);
+		}
+
+		{
+			FStateTreePropertyPath Path;
+			const bool bParseResult = Path.FromString(TEXT("StructB..NoThere"));
+			AITEST_FALSE("Parsing path should fail", bParseResult);
+		}
+
+		{
+			FStateTreePropertyPath Path;
+			const bool bParseResult = Path.FromString(TEXT("."));
+			AITEST_FALSE("Parsing path should fail", bParseResult);
+		}
+
+		{
+			FStateTreePropertyPath Path;
+			const bool bParseResult = Path.FromString(TEXT("StructB..B"));
+			AITEST_FALSE("Parsing path should fail", bParseResult);
+		}
+
+		return true;
+	}
+};
+IMPLEMENT_AI_INSTANT_TEST(FStateTreeTest_PropertyPathParseFail, "System.StateTree.PropertyPath.ParseFail");
+
+struct FStateTreeTest_PropertyPathOffsetFail : FStateTreeTestBase
+{
+	virtual bool InstantTest() override
+	{
+		FStateTreePropertyPath Path;
+		const bool bParseResult = Path.FromString(TEXT("StructB.Q"));
+
+		AITEST_TRUE("Parsing path should succeeed", bParseResult);
+		AITEST_EQUAL("Should have 2 path segments", Path.NumSegments(), 2);
+
+		FString ResolveErrors;
+		TArray<FStateTreePropertyPathIndirection> Indirections;
+		const bool bResolveResult = Path.ResolveIndirections(FStateTreeTest_PropertyStruct::StaticStruct(), Indirections, &ResolveErrors);
+
+		AITEST_FALSE("Resolve path should not succeeed", bResolveResult);
+		AITEST_NOT_EQUAL("Should have errors", ResolveErrors.Len(), 0);
+		
+		AITEST_EQUAL("Should have 0 indirections", Indirections.Num(), 0);
+
+		return true;
+	}
+};
+IMPLEMENT_AI_INSTANT_TEST(FStateTreeTest_PropertyPathOffsetFail, "System.StateTree.PropertyPath.OffsetFail");
+
+struct FStateTreeTest_PropertyPathObject : FStateTreeTestBase
+{
+	virtual bool InstantTest() override
+	{
+		FStateTreePropertyPath Path;
+		const bool bParseResult = Path.FromString(TEXT("InstancedObject.A"));
+
+		AITEST_TRUE("Parsing path should succeeed", bParseResult);
+		AITEST_EQUAL("Should have 2 path segments", Path.NumSegments(), 2);
+
+		UStateTreeTest_PropertyObject* Object = NewObject<UStateTreeTest_PropertyObject>();
+		Object->InstancedObject = NewObject<UStateTreeTest_PropertyObjectInstanced>();
+		
+		const bool bUpdateResult = Path.UpdateSegmentsFromValue(FStateTreeDataView(Object));
+
+		AITEST_TRUE("Update instance types should succeeed", bUpdateResult);
+		AITEST_TRUE("Path segment 0 instance type should be UStateTreeTest_PropertyObjectInstanced", Path.GetSegment(0).GetInstanceStruct() == UStateTreeTest_PropertyObjectInstanced::StaticClass());
+		AITEST_TRUE("Path segment 1 instance type should be nullptr", Path.GetSegment(1).GetInstanceStruct() == nullptr);
+
+		return true;
+	}
+};
+IMPLEMENT_AI_INSTANT_TEST(FStateTreeTest_PropertyPathObject, "System.StateTree.PropertyPath.Object");
+
+struct FStateTreeTest_PropertyPathWrongObject : FStateTreeTestBase
+{
+	virtual bool InstantTest() override
+	{
+		FStateTreePropertyPath Path;
+		const bool bParseResult = Path.FromString(TEXT("InstancedObject.B"));
+
+		AITEST_TRUE("Parsing path should succeeed", bParseResult);
+		AITEST_EQUAL("Should have 2 path segments", Path.NumSegments(), 2);
+
+		UStateTreeTest_PropertyObject* Object = NewObject<UStateTreeTest_PropertyObject>();
+
+		Object->InstancedObject = NewObject<UStateTreeTest_PropertyObjectInstancedWithB>();
+		{
+			FString ResolveErrors;
+			TArray<FStateTreePropertyPathIndirection> Indirections;
+			const bool bResolveResult = Path.ResolveIndirectionsWithValue(FStateTreeDataView(Object), Indirections, &ResolveErrors);
+
+			AITEST_TRUE("Resolve path should succeeed", bResolveResult);
+			AITEST_EQUAL("Should have 2 indirections", Indirections.Num(), 2);
+			AITEST_TRUE("Object ", Indirections[0].GetAccessType() == EStateTreePropertyAccessType::ObjectInstance);
+			AITEST_TRUE("Object ", Indirections[0].GetContainerStruct() == Object->GetClass());
+			AITEST_TRUE("Object ", Indirections[0].GetInstanceStruct() == UStateTreeTest_PropertyObjectInstancedWithB::StaticClass());
+			AITEST_EQUAL("Should not have error", ResolveErrors.Len(), 0);
+		}
+
+		Object->InstancedObject = NewObject<UStateTreeTest_PropertyObjectInstanced>();
+		{
+			FString ResolveErrors;
+			TArray<FStateTreePropertyPathIndirection> Indirections;
+			const bool bResolveResult = Path.ResolveIndirectionsWithValue(FStateTreeDataView(Object), Indirections, &ResolveErrors);
+
+			AITEST_FALSE("Resolve path should fail", bResolveResult);
+			AITEST_EQUAL("Should have 0 indirections", Indirections.Num(), 0);
+			AITEST_NOT_EQUAL("Should have error", ResolveErrors.Len(), 0);
+		}
+		
+		return true;
+	}
+};
+IMPLEMENT_AI_INSTANT_TEST(FStateTreeTest_PropertyPathWrongObject, "System.StateTree.PropertyPath.WrongObject");
+
+struct FStateTreeTest_PropertyPathArray : FStateTreeTestBase
+{
+	virtual bool InstantTest() override
+	{
+		FStateTreePropertyPath Path;
+		const bool bParseResult = Path.FromString(TEXT("ArrayOfInts[1]"));
+
+		AITEST_TRUE("Parsing path should succeeed", bParseResult);
+		AITEST_EQUAL("Should have 1 path segments", Path.NumSegments(), 1);
+
+		UStateTreeTest_PropertyObject* Object = NewObject<UStateTreeTest_PropertyObject>();
+		Object->ArrayOfInts.Add(42);
+		Object->ArrayOfInts.Add(123);
+
+		FString ResolveErrors;
+		TArray<FStateTreePropertyPathIndirection> Indirections;
+		const bool bResolveResult = Path.ResolveIndirectionsWithValue(FStateTreeDataView(Object), Indirections, &ResolveErrors);
+
+		AITEST_TRUE("Resolve path should succeeed", bResolveResult);
+		AITEST_EQUAL("Should have no resolve errors", ResolveErrors.Len(), 0);
+		AITEST_EQUAL("Should have 2 indirections", Indirections.Num(), 2);
+		AITEST_EQUAL("Indirection 0 should be IndexArray type", Indirections[0].GetAccessType(), EStateTreePropertyAccessType::IndexArray);
+		AITEST_EQUAL("Indirection 1 should be Offset type", Indirections[1].GetAccessType(), EStateTreePropertyAccessType::Offset);
+
+		const int32 Value = *reinterpret_cast<const int32*>(Indirections[1].GetPropertyAddress());
+		AITEST_EQUAL("Value should be 123", Value, 123);
+		
+		return true;
+	}
+};
+IMPLEMENT_AI_INSTANT_TEST(FStateTreeTest_PropertyPathArray, "System.StateTree.PropertyPath.Array");
+
+struct FStateTreeTest_PropertyPathArrayInvalidIndex : FStateTreeTestBase
+{
+	virtual bool InstantTest() override
+	{
+		FStateTreePropertyPath Path;
+		const bool bParseResult = Path.FromString(TEXT("ArrayOfInts[123]"));
+
+		AITEST_TRUE("Parsing path should succeeed", bParseResult);
+		AITEST_EQUAL("Should have 1 path segments", Path.NumSegments(), 1);
+
+		UStateTreeTest_PropertyObject* Object = NewObject<UStateTreeTest_PropertyObject>();
+		Object->ArrayOfInts.Add(42);
+		Object->ArrayOfInts.Add(123);
+
+		FString ResolveErrors;
+		TArray<FStateTreePropertyPathIndirection> Indirections;
+		const bool bResolveResult = Path.ResolveIndirectionsWithValue(FStateTreeDataView(Object), Indirections, &ResolveErrors);
+
+		AITEST_FALSE("Resolve path should fail", bResolveResult);
+		
+		return true;
+	}
+};
+IMPLEMENT_AI_INSTANT_TEST(FStateTreeTest_PropertyPathArrayInvalidIndex, "System.StateTree.PropertyPath.ArrayInvalidIndex");
+
+struct FStateTreeTest_PropertyPathArrayOfStructs : FStateTreeTestBase
+{
+	virtual bool InstantTest() override
+	{
+		FStateTreePropertyPath Path1;
+		Path1.FromString(TEXT("ArrayOfStruct[0].B"));
+
+		FStateTreePropertyPath Path2;
+		Path2.FromString(TEXT("ArrayOfStruct[2].StructB.B"));
+
+		UStateTreeTest_PropertyObject* Object = NewObject<UStateTreeTest_PropertyObject>();
+		Object->ArrayOfStruct.AddDefaulted_GetRef().B = 3;
+		Object->ArrayOfStruct.AddDefaulted();
+		Object->ArrayOfStruct.AddDefaulted_GetRef().StructB.B = 42;
+
+		{
+			FString ResolveErrors;
+			TArray<FStateTreePropertyPathIndirection> Indirections;
+			const bool bResolveResult = Path1.ResolveIndirectionsWithValue(FStateTreeDataView(Object), Indirections, &ResolveErrors);
+
+			AITEST_TRUE("Resolve path1 should succeeed", bResolveResult);
+			AITEST_EQUAL("Should have no resolve errors", ResolveErrors.Len(), 0);
+			AITEST_EQUAL("Should have 3 indirections", Indirections.Num(), 3);
+			AITEST_EQUAL("Indirection 0 should be ArrayIndex type", Indirections[0].GetAccessType(), EStateTreePropertyAccessType::IndexArray);
+			AITEST_EQUAL("Indirection 1 should be Offset type", Indirections[1].GetAccessType(), EStateTreePropertyAccessType::Offset);
+			AITEST_EQUAL("Indirection 2 should be Offset type", Indirections[2].GetAccessType(), EStateTreePropertyAccessType::Offset);
+
+			const int32 Value = *reinterpret_cast<const int32*>(Indirections[2].GetPropertyAddress());
+			AITEST_EQUAL("Value should be 3", Value, 3);
+		}
+
+		{
+			FString ResolveErrors;
+			TArray<FStateTreePropertyPathIndirection> Indirections;
+			const bool bResolveResult = Path2.ResolveIndirectionsWithValue(FStateTreeDataView(Object), Indirections, &ResolveErrors);
+
+			AITEST_TRUE("Resolve path2 should succeeed", bResolveResult);
+			AITEST_EQUAL("Should have no resolve errors", ResolveErrors.Len(), 0);
+			AITEST_EQUAL("Should have 4 indirections", Indirections.Num(), 4);
+			AITEST_EQUAL("Indirection 0 should be ArrayIndex type", Indirections[0].GetAccessType(), EStateTreePropertyAccessType::IndexArray);
+			AITEST_EQUAL("Indirection 1 should be Offset type", Indirections[1].GetAccessType(), EStateTreePropertyAccessType::Offset);
+			AITEST_EQUAL("Indirection 2 should be Offset type", Indirections[2].GetAccessType(), EStateTreePropertyAccessType::Offset);
+			AITEST_EQUAL("Indirection 3 should be Offset type", Indirections[3].GetAccessType(), EStateTreePropertyAccessType::Offset);
+
+			const int32 Value = *reinterpret_cast<const int32*>(Indirections[3].GetPropertyAddress());
+			AITEST_EQUAL("Value should be 42", Value, 42);
+		}
+		return true;
+	}
+};
+IMPLEMENT_AI_INSTANT_TEST(FStateTreeTest_PropertyPathArrayOfStructs, "System.StateTree.PropertyPath.ArrayOfStructs");
+
+struct FStateTreeTest_PropertyPathArrayOfInstancedObjects : FStateTreeTestBase
+{
+	virtual bool InstantTest() override
+	{
+		FStateTreePropertyPath Path;
+		Path.FromString(TEXT("ArrayOfInstancedStructs[0].B"));
+
+		FStateTreeTest_PropertyStruct Struct;
+		Struct.B = 123;
+		
+		UStateTreeTest_PropertyObject* Object = NewObject<UStateTreeTest_PropertyObject>();
+		Object->ArrayOfInstancedStructs.Emplace(FConstStructView::Make(Struct));
+
+		const bool bUpdateResult = Path.UpdateSegmentsFromValue(FStateTreeDataView(Object));
+		AITEST_TRUE("Update instance types should succeeed", bUpdateResult);
+		AITEST_EQUAL("Should have 2 path segments", Path.NumSegments(), 2);
+		AITEST_TRUE("Path segment 0 instance type should be FStateTreeTest_PropertyStruct", Path.GetSegment(0).GetInstanceStruct() == FStateTreeTest_PropertyStruct::StaticStruct());
+		AITEST_TRUE("Path segment 1 instance type should be nullptr", Path.GetSegment(1).GetInstanceStruct() == nullptr);
+
+		{
+			FString ResolveErrors;
+			TArray<FStateTreePropertyPathIndirection> Indirections;
+			const bool bResolveResult = Path.ResolveIndirections(UStateTreeTest_PropertyObject::StaticClass(), Indirections, &ResolveErrors);
+
+			AITEST_TRUE("Resolve path should succeeed", bResolveResult);
+			AITEST_EQUAL("Should have no resolve errors", ResolveErrors.Len(), 0);
+			AITEST_EQUAL("Should have 3 indirections", Indirections.Num(), 3);
+			AITEST_EQUAL("Indirection 0 should be ArrayIndex type", Indirections[0].GetAccessType(), EStateTreePropertyAccessType::IndexArray);
+			AITEST_EQUAL("Indirection 1 should be StructInstance type", Indirections[1].GetAccessType(), EStateTreePropertyAccessType::StructInstance);
+			AITEST_EQUAL("Indirection 2 should be Offset type", Indirections[2].GetAccessType(), EStateTreePropertyAccessType::Offset);
+		}
+
+		{
+			FString ResolveErrors;
+			TArray<FStateTreePropertyPathIndirection> Indirections;
+			const bool bResolveResult = Path.ResolveIndirectionsWithValue(FStateTreeDataView(Object), Indirections, &ResolveErrors);
+
+			AITEST_TRUE("Resolve path should succeeed", bResolveResult);
+			AITEST_EQUAL("Should have no resolve errors", ResolveErrors.Len(), 0);
+			AITEST_EQUAL("Should have 3 indirections", Indirections.Num(), 3);
+			AITEST_EQUAL("Indirection 0 should be ArrayIndex type", Indirections[0].GetAccessType(), EStateTreePropertyAccessType::IndexArray);
+			AITEST_EQUAL("Indirection 1 should be StructInstance type", Indirections[1].GetAccessType(), EStateTreePropertyAccessType::StructInstance);
+			AITEST_EQUAL("Indirection 2 should be Offset type", Indirections[2].GetAccessType(), EStateTreePropertyAccessType::Offset);
+
+			const int32 Value = *reinterpret_cast<const int32*>(Indirections[2].GetPropertyAddress());
+			AITEST_EQUAL("Value should be 123", Value, 123);
+		}
+		
+		return true;
+	}
+};
+IMPLEMENT_AI_INSTANT_TEST(FStateTreeTest_PropertyPathArrayOfInstancedObjects, "System.StateTree.PropertyPath.ArrayOfInstancedObjects");
+
+} // namespace UE::StateTree::Tests
+
+UE_ENABLE_OPTIMIZATION_SHIP
+
+#undef LOCTEXT_NAMESPACE
diff --git a/Engine/Plugins/Runtime/StateTree/Source/StateTreeTestSuite/Private/StateTreeTest.cpp b/Engine/Plugins/Runtime/StateTree/Source/StateTreeTestSuite/Private/StateTreeTest.cpp
--- a/Engine/Plugins/Runtime/StateTree/Source/StateTreeTestSuite/Private/StateTreeTest.cpp
+++ b/Engine/Plugins/Runtime/StateTree/Source/StateTreeTestSuite/Private/StateTreeTest.cpp
@@ -1,6 +1,7 @@
 // Copyright Epic Games, Inc. All Rights Reserved.
 
 #include "StateTreeTest.h"
+#include "StateTreeTestBase.h"
 #include "AITestsCommon.h"
 #include "StateTreeCompilerLog.h"
 #include "StateTreeEditorData.h"
@@ -23,60 +24,13 @@ std::atomic<int32> FStateTreeTestConditionInstanceData::GlobalCounter = 0;
 
 namespace UE::StateTree::Tests
 {
-	UStateTree& NewStateTree(UObject* Outer = GetTransientPackage())
-	{
-		UStateTree* StateTree = NewObject<UStateTree>(Outer);
-		check(StateTree);
-		UStateTreeEditorData* EditorData = NewObject<UStateTreeEditorData>(StateTree);
-		check(EditorData);
-		StateTree->EditorData = EditorData;
-		EditorData->Schema = NewObject<UStateTreeTestSchema>();
-		return *StateTree;
-	}
-
-	FStateTreePropertyPathBinding MakeBinding(const FGuid& SourceID, const FString& Source, const FGuid& TargetID, const FString& Target)
-	{
-		FStateTreePropertyPath SourcePath;
-		SourcePath.FromString(Source);
-		SourcePath.SetStructID(SourceID);
-
-		FStateTreePropertyPath TargetPath;
-		TargetPath.FromString(Target);
-		TargetPath.SetStructID(TargetID);
-
-		return FStateTreePropertyPathBinding(SourcePath, TargetPath);
-	}
-
-	// Helper struct to define some test tags
-	struct FNativeGameplayTags : public FGameplayTagNativeAdder
-	{
-		virtual ~FNativeGameplayTags() {}
-		
-		FGameplayTag TestTag;
-		FGameplayTag TestTag2;
-
-		virtual void AddTags() override
-		{
-			UGameplayTagsManager& Manager = UGameplayTagsManager::Get();
-			TestTag = Manager.AddNativeGameplayTag(TEXT("Test.StateTree.Tag"));
-			TestTag2 = Manager.AddNativeGameplayTag(TEXT("Test.StateTree.Tag2"));
-		}
-
-		FORCEINLINE static const FNativeGameplayTags& Get()
-		{
-			return StaticInstance;
-		}
-		static FNativeGameplayTags StaticInstance;
-	};
-	FNativeGameplayTags FNativeGameplayTags::StaticInstance;
 
-}
 
-struct FStateTreeTest_MakeAndBakeStateTree : FAITestBase
+struct FStateTreeTest_MakeAndBakeStateTree : FStateTreeTestBase
 {
 	virtual bool InstantTest() override
 	{
-		UStateTree& StateTree = UE::StateTree::Tests::NewStateTree(&GetWorld());
+		UStateTree& StateTree = NewStateTree();
 		UStateTreeEditorData& EditorData = *Cast<UStateTreeEditorData>(StateTree.EditorData);
 		
 		UStateTreeState& Root = EditorData.AddSubTree(FName(TEXT("Root")));
@@ -121,11 +75,11 @@ struct FStateTreeTest_MakeAndBakeStateTree : FAITestBase
 IMPLEMENT_AI_INSTANT_TEST(FStateTreeTest_MakeAndBakeStateTree, "System.StateTree.MakeAndBakeStateTree");
 
 
-struct FStateTreeTest_EmptyStateTree : FAITestBase
+struct FStateTreeTest_EmptyStateTree : FStateTreeTestBase
 {
 	virtual bool InstantTest() override
 	{
-		UStateTree& StateTree = UE::StateTree::Tests::NewStateTree(&GetWorld());
+		UStateTree& StateTree = NewStateTree();
 		UStateTreeEditorData& EditorData = *Cast<UStateTreeEditorData>(StateTree.EditorData);
 		
 		UStateTreeState& Root = EditorData.AddSubTree(FName(TEXT("Root")));
@@ -158,11 +112,11 @@ struct FStateTreeTest_EmptyStateTree : FAITestBase
 };
 IMPLEMENT_AI_INSTANT_TEST(FStateTreeTest_EmptyStateTree, "System.StateTree.Empty");
 
-struct FStateTreeTest_Sequence : FAITestBase
+struct FStateTreeTest_Sequence : FStateTreeTestBase
 {
 	virtual bool InstantTest() override
 	{
-		UStateTree& StateTree = UE::StateTree::Tests::NewStateTree(&GetWorld());
+		UStateTree& StateTree = NewStateTree();
 		UStateTreeEditorData& EditorData = *Cast<UStateTreeEditorData>(StateTree.EditorData);
 		
 		UStateTreeState& Root = EditorData.AddSubTree(FName(TEXT("Root")));
@@ -220,11 +174,11 @@ struct FStateTreeTest_Sequence : FAITestBase
 };
 IMPLEMENT_AI_INSTANT_TEST(FStateTreeTest_Sequence, "System.StateTree.Sequence");
 
-struct FStateTreeTest_Select : FAITestBase
+struct FStateTreeTest_Select : FStateTreeTestBase
 {
 	virtual bool InstantTest() override
 	{
-		UStateTree& StateTree = UE::StateTree::Tests::NewStateTree(&GetWorld());
+		UStateTree& StateTree = NewStateTree();
 		UStateTreeEditorData& EditorData = *Cast<UStateTreeEditorData>(StateTree.EditorData);
 		
 		UStateTreeState& Root = EditorData.AddSubTree(FName(TEXT("Root")));
@@ -295,11 +249,11 @@ struct FStateTreeTest_Select : FAITestBase
 IMPLEMENT_AI_INSTANT_TEST(FStateTreeTest_Select, "System.StateTree.Select");
 
 
-struct FStateTreeTest_FailEnterState : FAITestBase
+struct FStateTreeTest_FailEnterState : FStateTreeTestBase
 {
 	virtual bool InstantTest() override
 	{
-		UStateTree& StateTree = UE::StateTree::Tests::NewStateTree(&GetWorld());
+		UStateTree& StateTree = NewStateTree();
 		UStateTreeEditorData& EditorData = *Cast<UStateTreeEditorData>(StateTree.EditorData);
 		
 		UStateTreeState& Root = EditorData.AddSubTree(FName(TEXT("Root")));
@@ -358,11 +312,11 @@ struct FStateTreeTest_FailEnterState : FAITestBase
 IMPLEMENT_AI_INSTANT_TEST(FStateTreeTest_FailEnterState, "System.StateTree.FailEnterState");
 
 
-struct FStateTreeTest_Restart : FAITestBase
+struct FStateTreeTest_Restart : FStateTreeTestBase
 {
 	virtual bool InstantTest() override
 	{
-		UStateTree& StateTree = UE::StateTree::Tests::NewStateTree(&GetWorld());
+		UStateTree& StateTree = NewStateTree();
 		UStateTreeEditorData& EditorData = *Cast<UStateTreeEditorData>(StateTree.EditorData);
 		
 		UStateTreeState& Root = EditorData.AddSubTree(FName(TEXT("Root")));
@@ -412,11 +366,11 @@ struct FStateTreeTest_Restart : FAITestBase
 };
 IMPLEMENT_AI_INSTANT_TEST(FStateTreeTest_Restart, "System.StateTree.Restart");
 
-struct FStateTreeTest_SubTree : FAITestBase
+struct FStateTreeTest_SubTree : FStateTreeTestBase
 {
 	virtual bool InstantTest() override
 	{
-		UStateTree& StateTree = UE::StateTree::Tests::NewStateTree(&GetWorld());
+		UStateTree& StateTree = NewStateTree();
 		UStateTreeEditorData& EditorData = *Cast<UStateTreeEditorData>(StateTree.EditorData);
 		
 		UStateTreeState& Root = EditorData.AddSubTree(FName(TEXT("Root")));
@@ -490,7 +444,7 @@ struct FStateTreeTest_SubTree : FAITestBase
 };
 IMPLEMENT_AI_INSTANT_TEST(FStateTreeTest_SubTree, "System.StateTree.SubTree");
 
-struct FStateTreeTest_SubTreeCondition : FAITestBase
+struct FStateTreeTest_SubTreeCondition : FStateTreeTestBase
 {
 	virtual bool InstantTest() override
 	{
@@ -502,7 +456,7 @@ struct FStateTreeTest_SubTreeCondition : FAITestBase
 			- State2 : Task3
 		*/
 		
-		UStateTree& StateTree = UE::StateTree::Tests::NewStateTree(&GetWorld());
+		UStateTree& StateTree = NewStateTree();
 		UStateTreeEditorData& EditorData = *Cast<UStateTreeEditorData>(StateTree.EditorData);
 		
 		UStateTreeState& Root = EditorData.AddSubTree(FName(TEXT("Root")));
@@ -572,11 +526,11 @@ struct FStateTreeTest_SubTreeCondition : FAITestBase
 };
 IMPLEMENT_AI_INSTANT_TEST(FStateTreeTest_SubTreeCondition, "System.StateTree.SubTreeCondition");
 
-struct FStateTreeTest_SubTree_CascadedSucceeded : FAITestBase
+struct FStateTreeTest_SubTree_CascadedSucceeded : FStateTreeTestBase
 {
 	virtual bool InstantTest() override
 	{
-		UStateTree& StateTree = UE::StateTree::Tests::NewStateTree(&GetWorld());
+		UStateTree& StateTree = NewStateTree();
 		UStateTreeEditorData& EditorData = *Cast<UStateTreeEditorData>(StateTree.EditorData);
 
 		//	- Root [TaskA]
@@ -602,1722 +556,180 @@ struct FStateTreeTest_SubTree_CascadedSucceeded : FAITestBase
 		SubLinkedState.AddTransition(EStateTreeTransitionTrigger::OnStateSucceeded, EStateTreeTransitionType::Failed);
 		SubSubLeaf.AddTransition(EStateTreeTransitionTrigger::OnStateSucceeded, EStateTreeTransitionType::Succeeded);
 
-		TStateTreeEditorNode<FTestTask_Stand>& TaskA = Root.AddTask<FTestTask_Stand>(FName(TEXT("TaskA")));
-		TStateTreeEditorNode<FTestTask_Stand>& TaskB = SubTreeState.AddTask<FTestTask_Stand>(FName(TEXT("TaskB")));
-		TStateTreeEditorNode<FTestTask_Stand>& TaskC = SubSubLeaf.AddTask<FTestTask_Stand>(FName(TEXT("TaskC")));
-
-		TaskA.GetNode().TicksToCompletion = 2;
-		TaskB.GetNode().TicksToCompletion = 2;
-		TaskC.GetNode().TicksToCompletion = 1; // The deepest task completes first.
-		
-		FStateTreeCompilerLog Log;
-		FStateTreeCompiler Compiler(Log);
-		const bool bResult = Compiler.Compile(StateTree);
-		AITEST_TRUE("StateTree should get compiled", bResult);
-
-		EStateTreeRunStatus Status = EStateTreeRunStatus::Unset;
-		FStateTreeInstanceData InstanceData;
-		FTestStateTreeExecutionContext Exec(StateTree, StateTree, InstanceData);
-		const bool bInitSucceeded = Exec.IsValid();
-		AITEST_TRUE("StateTree should init", bInitSucceeded);
-
-		const FString TickStr(TEXT("Tick"));
-		const FString EnterStateStr(TEXT("EnterState"));
-		const FString ExitStateStr(TEXT("ExitState"));
-		const FString StateCompletedStr(TEXT("StateCompleted"));
-
-		// Start and enter state
-		Status = Exec.Start();
-		AITEST_TRUE("StateTree Active States should be in Root/Linked/SubTreeState", Exec.ExpectInActiveStates(Root.Name, LinkedState.Name, SubTreeState.Name, SubLinkedState.Name, SubSubTreeState.Name, SubSubLeaf.Name));
-		AITEST_TRUE("TaskA,B,C should enter state", Exec.Expect(TaskA.GetName(), EnterStateStr).Then(TaskB.GetName(), EnterStateStr).Then(TaskC.GetName(), EnterStateStr));
-		AITEST_TRUE("StateTree should be running", Status == EStateTreeRunStatus::Running);
-		Exec.LogClear();
-
-		// Subtrees completes, and it completes the whole tree too.
-		// There's no good way to observe this externally. We switch the return along the way to make sure the transition does not happen directly from the leaf to failed.
-		Status = Exec.Tick(0.1f);
-		AITEST_TRUE("StateTree should be Failed", Status == EStateTreeRunStatus::Failed);
-		Exec.LogClear();
-
-		Exec.Stop();
-		
-		return true;
-	}
-};
-IMPLEMENT_AI_INSTANT_TEST(FStateTreeTest_SubTree_CascadedSucceeded, "System.StateTree.SubTree.CascadedSucceeded");
-
-
-struct FStateTreeTest_SharedInstanceData : FAITestBase
-{
-	virtual bool InstantTest() override
-	{
-		UStateTree& StateTree = UE::StateTree::Tests::NewStateTree(&GetWorld());
-		UStateTreeEditorData& EditorData = *Cast<UStateTreeEditorData>(StateTree.EditorData);
-		
-		UStateTreeState& Root = EditorData.AddSubTree(FName(TEXT("Root")));
-		auto& IntCond = Root.AddEnterCondition<FStateTreeTestCondition>();
-		IntCond.GetInstanceData().Count = 1;
-
-		auto& Task = Root.AddTask<FTestTask_Stand>(FName(TEXT("Task")));
-		Task.GetNode().TicksToCompletion = 2;
-
-		FStateTreeCompilerLog Log;
-		FStateTreeCompiler Compiler(Log);
-		const bool bResult = Compiler.Compile(StateTree);
-		AITEST_TRUE("StateTree should get compiled", bResult);
-
-		// Init, nothing should access the shared data.
-		constexpr int32 NumConcurrent = 100;
-		FStateTreeTestConditionInstanceData::GlobalCounter = 0;
-
-		bool bInitSucceeded = true;
-		TArray<FStateTreeInstanceData> InstanceDatas;
-
-		InstanceDatas.SetNum(NumConcurrent);
-		for (int32 Index = 0; Index < NumConcurrent; Index++)
-		{
-			FTestStateTreeExecutionContext Exec(StateTree, StateTree, InstanceDatas[Index]);
-			bInitSucceeded &= Exec.IsValid();
-		}
-		AITEST_TRUE("All StateTree contexts should init", bInitSucceeded);
-		AITEST_EQUAL("Test condition global counter should be 0", (int32)FStateTreeTestConditionInstanceData::GlobalCounter, 0);
-		
-		// Start in parallel
-		// This should create shared data per thread.
-		// We expect that ParallelForWithTaskContext() creates a context per thread.
-		TArray<FStateTreeTestRunContext> RunContexts;
-		
-		ParallelForWithTaskContext(
-			RunContexts,
-			InstanceDatas.Num(),
-			[&InstanceDatas, &StateTree](FStateTreeTestRunContext& RunContext, int32 Index)
-			{
-				FTestStateTreeExecutionContext Exec(StateTree, StateTree, InstanceDatas[Index]);
-				const EStateTreeRunStatus Status = Exec.Start();
-				if (Status == EStateTreeRunStatus::Running)
-				{
-					RunContext.Count++;
-				}
-			}
-		);
-
-		int32 StartTotalRunning = 0;
-		for (FStateTreeTestRunContext RunContext : RunContexts)
-		{
-			StartTotalRunning += RunContext.Count;
-		}
-		AITEST_EQUAL("All StateTree contexts should be running after Start", StartTotalRunning, NumConcurrent);
-		AITEST_EQUAL("Test condition global counter should equal context count after Start", (int32)FStateTreeTestConditionInstanceData::GlobalCounter, InstanceDatas.Num());
-
-		
-		// Tick in parallel
-		// This should not recreate the data, so FStateTreeTestConditionInstanceData::GlobalCounter should stay as is.
-		for (FStateTreeTestRunContext RunContext : RunContexts)
-		{
-			RunContext.Count = 0;
-		}
-
-		ParallelForWithTaskContext(
-			RunContexts,
-			InstanceDatas.Num(),
-			[&InstanceDatas, &StateTree](FStateTreeTestRunContext& RunContext, int32 Index)
-			{
-				FTestStateTreeExecutionContext Exec(StateTree, StateTree, InstanceDatas[Index]);
-				const EStateTreeRunStatus Status = Exec.Tick(0.1f);
-				if (Status == EStateTreeRunStatus::Running)
-				{
-					RunContext.Count++;
-				}
-			}
-		);
-
-		int32 TickTotalRunning = 0;
-		for (FStateTreeTestRunContext RunContext : RunContexts)
-		{
-			TickTotalRunning += RunContext.Count;
-		}
-		AITEST_EQUAL("All StateTree contexts should be running after Tick", TickTotalRunning, NumConcurrent);
-		AITEST_EQUAL("Test condition global counter should equal context count after Tick", (int32)FStateTreeTestConditionInstanceData::GlobalCounter, InstanceDatas.Num());
-
-		ParallelForWithTaskContext(
-			RunContexts,
-			InstanceDatas.Num(),
-			[&InstanceDatas, &StateTree](FStateTreeTestRunContext& RunContext, int32 Index)
-			{
-				FTestStateTreeExecutionContext Exec(StateTree, StateTree, InstanceDatas[Index]);
-				Exec.Stop();
-			}
-		);
-
-		return true;
-	}
-};
-IMPLEMENT_AI_INSTANT_TEST(FStateTreeTest_SharedInstanceData, "System.StateTree.SharedInstanceData");
-
-struct FStateTreeTest_TransitionPriority : FAITestBase
-{
-	virtual bool InstantTest() override
-	{
-		UStateTree& StateTree = UE::StateTree::Tests::NewStateTree(&GetWorld());
-		UStateTreeEditorData& EditorData = *Cast<UStateTreeEditorData>(StateTree.EditorData);
-
-		/*
-			- Root
-				- State1 : Task1 -> Succeeded
-					- State1A : Task1A -> Next
-					- State1B : Task1B -> Next
-					- State1C : Task1C
-		
-			Task1A completed first, transitioning to State1B.
-			Task1, Task1B, and Task1C complete at the same time, we should take the transition on the first completed state (State1).
-		*/
-		
-		UStateTreeState& Root = EditorData.AddSubTree(FName(TEXT("Root")));
-		UStateTreeState& State1 = Root.AddChildState(FName(TEXT("State1")));
-		UStateTreeState& State1A = State1.AddChildState(FName(TEXT("State1A")));
-		UStateTreeState& State1B = State1.AddChildState(FName(TEXT("State1B")));
-		UStateTreeState& State1C = State1.AddChildState(FName(TEXT("State1C")));
-
-		auto& Task1 = State1.AddTask<FTestTask_Stand>(FName(TEXT("Task1")));
-		Task1.GetNode().TicksToCompletion = 2;
-		State1.AddTransition(EStateTreeTransitionTrigger::OnStateCompleted, EStateTreeTransitionType::Succeeded);
-		
-		auto& Task1A = State1A.AddTask<FTestTask_Stand>(FName(TEXT("Task1A")));
-		Task1A.GetNode().TicksToCompletion = 1;
-		State1A.AddTransition(EStateTreeTransitionTrigger::OnStateCompleted, EStateTreeTransitionType::NextState);
-
-		auto& Task1B = State1B.AddTask<FTestTask_Stand>(FName(TEXT("Task1B")));
-		Task1B.GetNode().TicksToCompletion = 2;
-		State1B.AddTransition(EStateTreeTransitionTrigger::OnStateCompleted, EStateTreeTransitionType::NextState);
-
-		auto& Task1C = State1C.AddTask<FTestTask_Stand>(FName(TEXT("Task1C")));
-		Task1C.GetNode().TicksToCompletion = 2;
-		
-		FStateTreeCompilerLog Log;
-		FStateTreeCompiler Compiler(Log);
-		const bool bResult = Compiler.Compile(StateTree);
-		AITEST_TRUE("StateTree should get compiled", bResult);
-
-		EStateTreeRunStatus Status = EStateTreeRunStatus::Unset;
-		FStateTreeInstanceData InstanceData;
-		FTestStateTreeExecutionContext Exec(StateTree, StateTree, InstanceData);
-		const bool bInitSucceeded = Exec.IsValid();
-		AITEST_TRUE("StateTree should init", bInitSucceeded);
-
-		const FString TickStr(TEXT("Tick"));
-		const FString EnterStateStr(TEXT("EnterState"));
-		const FString ExitStateStr(TEXT("ExitState"));
-		const FString StateCompletedStr(TEXT("StateCompleted"));
-
-		// Start and enter state
-		Status = Exec.Start();
-		AITEST_TRUE("StateTree Task1 should enter state", Exec.Expect(Task1.GetName(), EnterStateStr));
-		AITEST_TRUE("StateTree Task1A should enter state", Exec.Expect(Task1A.GetName(), EnterStateStr));
-		Exec.LogClear();
-
-		// Transition from Task1A to Task1B
-		Status = Exec.Tick(0.1f);
-		AITEST_TRUE("StateTree Task1A should complete", Exec.Expect(Task1A.GetName(), StateCompletedStr));
-		AITEST_TRUE("StateTree Task1B should enter state", Exec.Expect(Task1B.GetName(), EnterStateStr));
-		Exec.LogClear();
-
-		// Task1 completes, and we should take State1 transition. 
-		Status = Exec.Tick(0.1f);
-		AITEST_TRUE("StateTree Task1 should complete", Exec.Expect(Task1.GetName(), StateCompletedStr));
-		AITEST_EQUAL("Tree execution should stop on success", Status, EStateTreeRunStatus::Succeeded);
-		Exec.LogClear();
-
-		Exec.Stop();
-
-		return true;
-	}
-};
-IMPLEMENT_AI_INSTANT_TEST(FStateTreeTest_TransitionPriority, "System.StateTree.Transition.Priority");
-
-struct FStateTreeTest_TransitionPriorityEnterState : FAITestBase
-{
-	virtual bool InstantTest() override
-	{
-		UStateTree& StateTree = UE::StateTree::Tests::NewStateTree(&GetWorld());
-		UStateTreeEditorData& EditorData = *Cast<UStateTreeEditorData>(StateTree.EditorData);
-		
-		UStateTreeState& Root =	EditorData.AddSubTree(FName(TEXT("Root")));
-		UStateTreeState& State0 = Root.AddChildState(FName(TEXT("State0")));
-		UStateTreeState& State1 = Root.AddChildState(FName(TEXT("State1")));
-		UStateTreeState& State1A = State1.AddChildState(FName(TEXT("State1A")));
-		UStateTreeState& State2 = Root.AddChildState(FName(TEXT("State2")));
-		UStateTreeState& State3 = Root.AddChildState(FName(TEXT("State3")));
-
-		auto& Task0 = State0.AddTask<FTestTask_Stand>(FName(TEXT("Task0")));
-		State0.AddTransition(EStateTreeTransitionTrigger::OnStateCompleted, EStateTreeTransitionType::GotoState, &State1);
-
-		auto& Task1 = State1.AddTask<FTestTask_Stand>(FName(TEXT("Task1")));
-		Task1.GetNode().EnterStateResult = EStateTreeRunStatus::Failed;
-		State1.AddTransition(EStateTreeTransitionTrigger::OnStateCompleted, EStateTreeTransitionType::GotoState, &State2);
-		
-		auto& Task1A = State1A.AddTask<FTestTask_Stand>(FName(TEXT("Task1A")));
-		State1A.AddTransition(EStateTreeTransitionTrigger::OnStateCompleted, EStateTreeTransitionType::GotoState, &State3);
-
-		auto& Task2 = State2.AddTask<FTestTask_Stand>(FName(TEXT("Task2")));
-		State2.AddTransition(EStateTreeTransitionTrigger::OnStateCompleted, EStateTreeTransitionType::Succeeded);
-
-		auto& Task3 = State3.AddTask<FTestTask_Stand>(FName(TEXT("Task3")));
-		State3.AddTransition(EStateTreeTransitionTrigger::OnStateCompleted, EStateTreeTransitionType::Succeeded);
-
-		
-		FStateTreeCompilerLog Log;
-		FStateTreeCompiler Compiler(Log);
-		const bool bResult = Compiler.Compile(StateTree);
-		AITEST_TRUE("StateTree should get compiled", bResult);
-
-		EStateTreeRunStatus Status = EStateTreeRunStatus::Unset;
-		FStateTreeInstanceData InstanceData;
-		FTestStateTreeExecutionContext Exec(StateTree, StateTree, InstanceData);
-		const bool bInitSucceeded = Exec.IsValid();
-		AITEST_TRUE("StateTree should init", bInitSucceeded);
-
-		const FString TickStr(TEXT("Tick"));
-		const FString EnterStateStr(TEXT("EnterState"));
-		const FString ExitStateStr(TEXT("ExitState"));
-		const FString StateCompletedStr(TEXT("StateCompleted"));
-
-		// Start and enter state
-		Status = Exec.Start();
-		AITEST_TRUE("StateTree Task0 should enter state", Exec.Expect(Task0.GetName(), EnterStateStr));
-		Exec.LogClear();
-
-		// Transition from State0 to State1, it should fail (Task1), and the transition on State1->State2 (and not State1A->State3)
-		Status = Exec.Tick(0.1f);
-		AITEST_TRUE("StateTree Task0 should complete", Exec.Expect(Task0.GetName(), StateCompletedStr));
-		AITEST_TRUE("StateTree Task2 should enter state", Exec.Expect(Task2.GetName(), EnterStateStr));
-		AITEST_FALSE("StateTree Task3 should not enter state", Exec.Expect(Task3.GetName(), EnterStateStr));
-		Exec.LogClear();
-
-		Exec.Stop();
-
-		return true;
-	}
-};
-IMPLEMENT_AI_INSTANT_TEST(FStateTreeTest_TransitionPriorityEnterState, "System.StateTree.Transition.PriorityEnterState");
-
-struct FStateTreeTest_TransitionNextSelectableState : FAITestBase
-{
-	virtual bool InstantTest() override
-	{
-		UStateTree& StateTree = UE::StateTree::Tests::NewStateTree(&GetWorld());
-		UStateTreeEditorData& EditorData = *Cast<UStateTreeEditorData>(StateTree.EditorData);
-
-		UStateTreeState& Root =	EditorData.AddSubTree(FName(TEXT("Root")));
-		UStateTreeState& State0 = Root.AddChildState(FName(TEXT("State0")));
-		UStateTreeState& State1 = Root.AddChildState(FName(TEXT("State1")));
-		UStateTreeState& State2 = Root.AddChildState(FName(TEXT("State2")));
-
-		auto& EvalA = EditorData.AddEvaluator<FTestEval_A>();
-		EvalA.GetInstanceData().bBoolA = true;
-
-		auto& Task0 = State0.AddTask<FTestTask_Stand>(FName(TEXT("Task0")));
-		State0.AddTransition(EStateTreeTransitionTrigger::OnStateCompleted, EStateTreeTransitionType::NextSelectableState);
-
-		// Add Task 1 with Condition that will always fail
-		auto& Task1 = State1.AddTask<FTestTask_Stand>(FName(TEXT("Task1")));
-		auto& BoolCond1 = State1.AddEnterCondition<FStateTreeCompareBoolCondition>();
-
-		EditorData.AddPropertyBinding(EvalA, TEXT("bBoolA"), BoolCond1, TEXT("bLeft"));
-		BoolCond1.GetInstanceData().bRight = !EvalA.GetInstanceData().bBoolA;
-
-		// Add Task 2 with Condition that will always succeed
-		auto& Task2 = State2.AddTask<FTestTask_Stand>(FName(TEXT("Task2")));
-		auto& BoolCond2 = State2.AddEnterCondition<FStateTreeCompareBoolCondition>();
-		State2.AddTransition(EStateTreeTransitionTrigger::OnStateCompleted, EStateTreeTransitionType::Succeeded);
-
-		EditorData.AddPropertyBinding(EvalA, TEXT("bBoolA"), BoolCond2, TEXT("bLeft"));
-		BoolCond2.GetInstanceData().bRight = EvalA.GetInstanceData().bBoolA;
-
-		FStateTreeCompilerLog Log;
-		FStateTreeCompiler Compiler(Log);
-		const bool bResult = Compiler.Compile(StateTree);
-		AITEST_TRUE("StateTree should get compiled", bResult);
-
-		FStateTreeInstanceData InstanceData;
-		FTestStateTreeExecutionContext Exec(StateTree, StateTree, InstanceData);
-		const bool bInitSucceeded = Exec.IsValid();
-		AITEST_TRUE("StateTree should init", bInitSucceeded);
-
-		const FString TickStr(TEXT("Tick"));
-		const FString EnterStateStr(TEXT("EnterState"));
-		const FString ExitStateStr(TEXT("ExitState"));
-		const FString StateCompletedStr(TEXT("StateCompleted"));
-
-		// Start and enter state
-		Exec.Start();
-		AITEST_TRUE("StateTree Task0 should enter state", Exec.Expect(Task0.GetName(), EnterStateStr));
-		Exec.LogClear();
-
-		// Transition from State0 and tries to select State1. It should fail (Task1) and because transition is set to "Next Selectable", it should now select Task 2 and Enter State
-		Exec.Tick(0.1f);
-		AITEST_TRUE("StateTree Task0 should complete", Exec.Expect(Task0.GetName(), StateCompletedStr));
-		AITEST_FALSE("StateTree Task1 should not enter state", Exec.Expect(Task1.GetName(), EnterStateStr));
-		AITEST_TRUE("StateTree Task2 should enter state", Exec.Expect(Task2.GetName(), EnterStateStr));
-		Exec.LogClear();
-
-		// Complete Task2
-		Exec.Tick(0.1f);
-		AITEST_TRUE("StateTree Task2 should complete", Exec.Expect(Task2.GetName(), StateCompletedStr));
-		Exec.LogClear();
-
-		Exec.Stop();
-
-		return true;
-	}
-};
-IMPLEMENT_AI_INSTANT_TEST(FStateTreeTest_TransitionNextSelectableState, "System.StateTree.Transition.NextSelectableState");
-
-
-struct FStateTreeTest_TransitionNextWithParentData : FAITestBase
-{
-	virtual bool InstantTest() override
-	{
-		UStateTree& StateTree = UE::StateTree::Tests::NewStateTree(&GetWorld());
-		UStateTreeEditorData& EditorData = *Cast<UStateTreeEditorData>(StateTree.EditorData);
-
-		UStateTreeState& Root =	EditorData.AddSubTree(FName(TEXT("Root")));
-		UStateTreeState& State0 = Root.AddChildState(FName(TEXT("State0")));
-		UStateTreeState& State1 = Root.AddChildState(FName(TEXT("State1")));
-		UStateTreeState& State1A = State1.AddChildState(FName(TEXT("State1A")));
-
-		auto& RootTask = Root.AddTask<FTestTask_B>(FName(TEXT("RootTask")));
-		RootTask.GetInstanceData().bBoolB = true;
-
-		auto& Task0 = State0.AddTask<FTestTask_Stand>(FName(TEXT("Task0")));
-		State0.AddTransition(EStateTreeTransitionTrigger::OnStateCompleted, EStateTreeTransitionType::NextState);
-
-		auto& Task1A = State1A.AddTask<FTestTask_Stand>(FName(TEXT("Task1A")));
-		auto& BoolCond1 = State1A.AddEnterCondition<FStateTreeCompareBoolCondition>();
-
-		EditorData.AddPropertyBinding(RootTask, TEXT("bBoolB"), BoolCond1, TEXT("bLeft"));
-		BoolCond1.GetInstanceData().bRight = true;
-
-		FStateTreeCompilerLog Log;
-		FStateTreeCompiler Compiler(Log);
-		const bool bResult = Compiler.Compile(StateTree);
-		AITEST_TRUE("StateTree should get compiled", bResult);
-
-		FStateTreeInstanceData InstanceData;
-		FTestStateTreeExecutionContext Exec(StateTree, StateTree, InstanceData);
-		const bool bInitSucceeded = Exec.IsValid();
-		AITEST_TRUE("StateTree should init", bInitSucceeded);
-
-		const FString TickStr(TEXT("Tick"));
-		const FString EnterStateStr(TEXT("EnterState"));
-		const FString ExitStateStr(TEXT("ExitState"));
-		const FString StateCompletedStr(TEXT("StateCompleted"));
-
-		// Start and enter state
-		Exec.Start();
-		AITEST_TRUE("StateTree Task0 should enter state", Exec.Expect(Task0.GetName(), EnterStateStr));
-		Exec.LogClear();
-
-		// Transition from State0 and tries to select State1.
-		// This tests that data from current shared active states (Root) is available during state selection.
-		Exec.Tick(0.1f);
-		AITEST_TRUE("StateTree Task0 should complete", Exec.Expect(Task0.GetName(), StateCompletedStr));
-		AITEST_TRUE("StateTree Task1A should enter state", Exec.Expect(Task1A.GetName(), EnterStateStr));
-		Exec.LogClear();
-
-		Exec.Stop();
-
-		return true;
-	}
-};
-IMPLEMENT_AI_INSTANT_TEST(FStateTreeTest_TransitionNextWithParentData, "System.StateTree.Transition.NextWithParentData");
-
-struct FStateTreeTest_LastConditionWithIndent : FAITestBase
-{
-	virtual bool InstantTest() override
-	{
-		UStateTree& StateTree = UE::StateTree::Tests::NewStateTree(&GetWorld());
-		UStateTreeEditorData& EditorData = *Cast<UStateTreeEditorData>(StateTree.EditorData);
-		
-		UStateTreeState& Root = EditorData.AddSubTree(FName(TEXT("Root")));
-		UStateTreeState& State1 = Root.AddChildState(FName(TEXT("State1")));
-
-		auto& Task1 = State1.AddTask<FTestTask_Stand>(FName(TEXT("Task1")));
-		State1.AddEnterCondition<FStateTreeTestCondition>();
-		auto& LastCondition = State1.AddEnterCondition<FStateTreeTestCondition>();
-
-		// Last condition has Indent
-		LastCondition.ExpressionIndent = 1;
-		
-		State1.AddTransition(EStateTreeTransitionTrigger::OnStateCompleted, EStateTreeTransitionType::Succeeded);
-
-		FStateTreeCompilerLog Log;
-		FStateTreeCompiler Compiler(Log);
-		const bool bResult = Compiler.Compile(StateTree);
-		AITEST_TRUE("StateTree should get compiled", bResult);
-
-		EStateTreeRunStatus Status = EStateTreeRunStatus::Unset;
-		FStateTreeInstanceData InstanceData;
-		FTestStateTreeExecutionContext Exec(StateTree, StateTree, InstanceData);
-		const bool bInitSucceeded = Exec.IsValid();
-		AITEST_TRUE("StateTree should init", bInitSucceeded);
-
-		const FString TickStr(TEXT("Tick"));
-		const FString EnterStateStr(TEXT("EnterState"));
-		const FString ExitStateStr(TEXT("ExitState"));
-
-		Status = Exec.Start();
-		AITEST_TRUE("StateTree Task1 should enter state", Exec.Expect(Task1.GetName(), EnterStateStr));
-		AITEST_FALSE("StateTree Task1 should not tick", Exec.Expect(Task1.GetName(), TickStr));
-		Exec.LogClear();
-
-		Status = Exec.Tick(0.1f);
-		AITEST_TRUE("StateTree Task1 should tick, and exit state", Exec.Expect(Task1.GetName(), TickStr).Then(Task1.GetName(), ExitStateStr));
-		AITEST_TRUE("StateTree should be completed", Status == EStateTreeRunStatus::Succeeded);
-		Exec.LogClear();
-
-		Status = Exec.Tick(0.1f);
-		AITEST_FALSE("StateTree Task1 should not tick", Exec.Expect(Task1.GetName(), TickStr));
-		Exec.LogClear();
-
-		Exec.Stop();
-
-		return true;
-	}
-};
-IMPLEMENT_AI_INSTANT_TEST(FStateTreeTest_LastConditionWithIndent, "System.StateTree.LastConditionWithIndent");
-
-struct FStateTreeTest_TransitionGlobalDataView : FAITestBase
-{
-	// Tests that the global eval and task dataviews are kept up to date when transitioning from  
-	virtual bool InstantTest() override
-	{
-		UStateTree& StateTree = UE::StateTree::Tests::NewStateTree(&GetWorld());
-		UStateTreeEditorData& EditorData = *Cast<UStateTreeEditorData>(StateTree.EditorData);
-
-		UStateTreeState& Root = EditorData.AddSubTree(FName(TEXT("Root")));
-		UStateTreeState& StateA = Root.AddChildState(FName(TEXT("A")));
-		UStateTreeState& StateB = Root.AddChildState(FName(TEXT("B")));
-
-		auto& EvalA = EditorData.AddEvaluator<FTestEval_A>(FName(TEXT("Eval")));
-		EvalA.GetInstanceData().IntA = 42;
-		auto& GlobalTask = EditorData.AddGlobalTask<FTestTask_PrintValue>(FName(TEXT("Global")));
-		GlobalTask.GetInstanceData().Value = 123;
-		
-		// State A
-		auto& Task0 = StateA.AddTask<FTestTask_Stand>(FName(TEXT("Task0")));
-		StateA.AddTransition(EStateTreeTransitionTrigger::OnStateCompleted, EStateTreeTransitionType::GotoState, &StateB);
-
-		// State B
-		auto& Task1 = StateB.AddTask<FTestTask_PrintValue>(FName(TEXT("Task1")));
-		EditorData.AddPropertyBinding(EvalA, TEXT("IntA"), Task1, TEXT("Value"));
-		auto& Task2 = StateB.AddTask<FTestTask_PrintValue>(FName(TEXT("Task2")));
-		EditorData.AddPropertyBinding(GlobalTask, TEXT("Value"), Task2, TEXT("Value"));
-
-		FStateTreeCompilerLog Log;
-		FStateTreeCompiler Compiler(Log);
-		const bool bResult = Compiler.Compile(StateTree);
-
-		AITEST_TRUE("StateTree should get compiled", bResult);
-
-		EStateTreeRunStatus Status = EStateTreeRunStatus::Unset;
-		FStateTreeInstanceData InstanceData;
-		FTestStateTreeExecutionContext Exec(StateTree, StateTree, InstanceData);
-		const bool bInitSucceeded = Exec.IsValid();
-		AITEST_TRUE("StateTree should init", bInitSucceeded);
-
-		const FString EnterStateStr(TEXT("EnterState"));
-		const FString EnterState42Str(TEXT("EnterState42"));
-		const FString EnterState123Str(TEXT("EnterState123"));
-
-		// Start and enter state
-		Status = Exec.Start();
-		AITEST_TRUE("StateTree Task0 should enter state", Exec.Expect(Task0.GetName(), EnterStateStr));
-		Exec.LogClear();
-
-		// Transition from StateA to StateB, Task0 should enter state with evaluator value copied.
-		Status = Exec.Tick(0.1f);
-		AITEST_TRUE("StateTree Task0 should enter state with value 42", Exec.Expect(Task1.GetName(), EnterState42Str));
-		AITEST_TRUE("StateTree Task1 should enter state with value 123", Exec.Expect(Task2.GetName(), EnterState123Str));
-		Exec.LogClear();
-
-		Exec.Stop();
-
-		return true;
-	}
-};
-IMPLEMENT_AI_INSTANT_TEST(FStateTreeTest_TransitionGlobalDataView, "System.StateTree.Transition.GlobalDataView");
-
-struct FStateTreeTest_TransitionDelay : FAITestBase
-{
-	virtual bool InstantTest() override
-	{
-		UStateTree& StateTree = UE::StateTree::Tests::NewStateTree(&GetWorld());
-		UStateTreeEditorData& EditorData = *Cast<UStateTreeEditorData>(StateTree.EditorData);
-		const FGameplayTag Tag = UE::StateTree::Tests::FNativeGameplayTags::Get().TestTag;
-
-		UStateTreeState& Root = EditorData.AddSubTree(FName(TEXT("Root")));
-		UStateTreeState& StateA = Root.AddChildState(FName(TEXT("A")));
-		UStateTreeState& StateB = Root.AddChildState(FName(TEXT("B")));
-
-		// State A
-		auto& Task0 = StateA.AddTask<FTestTask_Stand>(FName(TEXT("Task0")));
-		Task0.GetNode().TicksToCompletion = 100;
-		
-		FStateTreeTransition& Transition = StateA.AddTransition(EStateTreeTransitionTrigger::OnEvent, EStateTreeTransitionType::GotoState, &StateB);
-		Transition.bDelayTransition = true;
-		Transition.DelayDuration = 0.15f;
-		Transition.DelayRandomVariance = 0.0f;
-		Transition.RequiredEvent.Tag = Tag;
-
-		// State B
-		auto& Task1 = StateB.AddTask<FTestTask_Stand>(FName(TEXT("Task1")));
-		Task1.GetNode().TicksToCompletion = 100;
-
-		FStateTreeCompilerLog Log;
-		FStateTreeCompiler Compiler(Log);
-		const bool bResult = Compiler.Compile(StateTree);
-
-		AITEST_TRUE("StateTree should get compiled", bResult);
-
-		EStateTreeRunStatus Status = EStateTreeRunStatus::Unset;
-		FStateTreeInstanceData InstanceData;
-		FTestStateTreeExecutionContext Exec(StateTree, StateTree, InstanceData);
-		const bool bInitSucceeded = Exec.IsValid();
-		AITEST_TRUE("StateTree should init", bInitSucceeded);
-
-		const FString TickStr(TEXT("Tick"));
-		const FString EnterStateStr(TEXT("EnterState"));
-		const FString ExitStateStr(TEXT("ExitState"));
-		const FString StateCompletedStr(TEXT("StateCompleted"));
-
-		// Start and enter state
-		Status = Exec.Start();
-		AITEST_TRUE("StateTree Task0 should enter state", Exec.Expect(Task0.GetName(), EnterStateStr));
-		Exec.LogClear();
-
-		// This should cause delayed transition.
-		Exec.SendEvent(Tag);
-		
-		Status = Exec.Tick(0.1f);
-		AITEST_TRUE("StateTree Task0 should tick", Exec.Expect(Task0.GetName(), TickStr));
-		Exec.LogClear();
-
-		// Should have execution frames
-		AITEST_TRUE("Should have active frames", InstanceData.GetExecutionState()->ActiveFrames.Num() > 0);
-
-		// Should have delayed transitions
-		const int32 NumDelayedTransitions0 = InstanceData.GetExecutionState()->DelayedTransitions.Num();
-		AITEST_EQUAL("Should have a delayed transition", NumDelayedTransitions0, 1);
-
-		// Tick and expect a delayed transition. 
-		Status = Exec.Tick(0.1f);
-		AITEST_TRUE("StateTree Task0 should tick", Exec.Expect(Task0.GetName(), TickStr));
-		Exec.LogClear();
-
-		const int32 NumDelayedTransitions1 = InstanceData.GetExecutionState()->DelayedTransitions.Num();
-		AITEST_EQUAL("Should have a delayed transition", NumDelayedTransitions1, 1);
-
-		// Should complete delayed transition.
-		Status = Exec.Tick(0.1f);
-		AITEST_TRUE("StateTree Task0 should exit state", Exec.Expect(Task0.GetName(), ExitStateStr));
-		AITEST_TRUE("StateTree Task1 should enter state", Exec.Expect(Task1.GetName(), EnterStateStr));
-		Exec.LogClear();
-
-		Exec.Stop();
-
-		return true;
-	}
-};
-IMPLEMENT_AI_INSTANT_TEST(FStateTreeTest_TransitionDelay, "System.StateTree.TransitionDelay");
-
-struct FStateTreeTest_TransitionDelayZero : FAITestBase
-{
-	virtual bool InstantTest() override
-	{
-		UStateTree& StateTree = UE::StateTree::Tests::NewStateTree(&GetWorld());
-		UStateTreeEditorData& EditorData = *Cast<UStateTreeEditorData>(StateTree.EditorData);
-		const FGameplayTag Tag = UE::StateTree::Tests::FNativeGameplayTags::Get().TestTag;
-
-		UStateTreeState& Root = EditorData.AddSubTree(FName(TEXT("Root")));
-		UStateTreeState& StateA = Root.AddChildState(FName(TEXT("A")));
-		UStateTreeState& StateB = Root.AddChildState(FName(TEXT("B")));
-
-		// State A
-		auto& Task0 = StateA.AddTask<FTestTask_Stand>(FName(TEXT("Task0")));
-		Task0.GetNode().TicksToCompletion = 100;
-		
-		FStateTreeTransition& Transition = StateA.AddTransition(EStateTreeTransitionTrigger::OnEvent, EStateTreeTransitionType::GotoState, &StateB);
-		Transition.bDelayTransition = true;
-		Transition.DelayDuration = 0.0f;
-		Transition.DelayRandomVariance = 0.0f;
-		Transition.RequiredEvent.Tag = Tag;
-
-		// State B
-		auto& Task1 = StateB.AddTask<FTestTask_Stand>(FName(TEXT("Task1")));
-		Task1.GetNode().TicksToCompletion = 100;
-
-		FStateTreeCompilerLog Log;
-		FStateTreeCompiler Compiler(Log);
-		const bool bResult = Compiler.Compile(StateTree);
-
-		AITEST_TRUE("StateTree should get compiled", bResult);
-
-		EStateTreeRunStatus Status = EStateTreeRunStatus::Unset;
-		FStateTreeInstanceData InstanceData;
-		FTestStateTreeExecutionContext Exec(StateTree, StateTree, InstanceData);
-		const bool bInitSucceeded = Exec.IsValid();
-		AITEST_TRUE("StateTree should init", bInitSucceeded);
-
-		const FString TickStr(TEXT("Tick"));
-		const FString EnterStateStr(TEXT("EnterState"));
-		const FString ExitStateStr(TEXT("ExitState"));
-		const FString StateCompletedStr(TEXT("StateCompleted"));
-
-		// Start and enter state
-		Status = Exec.Start();
-		AITEST_TRUE("StateTree Task0 should enter state", Exec.Expect(Task0.GetName(), EnterStateStr));
-		Exec.LogClear();
-
-		// This should cause delayed transition. Because the time is 0, it should happen immediately.
-		Exec.SendEvent(Tag);
-		
-		Status = Exec.Tick(0.1f);
-		AITEST_TRUE("StateTree Task0 should exit state", Exec.Expect(Task0.GetName(), ExitStateStr));
-		AITEST_TRUE("StateTree Task1 should enter state", Exec.Expect(Task1.GetName(), EnterStateStr));
-		Exec.LogClear();
-
-		Exec.Stop();
-
-		return true;
-	}
-};
-IMPLEMENT_AI_INSTANT_TEST(FStateTreeTest_TransitionDelayZero, "System.StateTree.TransitionDelayZero");
-
-struct FStateTreeTest_StateRequiringEvent : FAITestBase
-{
-	virtual bool InstantTest() override
-	{
-		UStateTree& StateTree = UE::StateTree::Tests::NewStateTree(&GetWorld());
-		UStateTreeEditorData& EditorData = *Cast<UStateTreeEditorData>(StateTree.EditorData);
-
-		UStateTreeState& Root = EditorData.AddSubTree(FName(TEXT("Root")));
-		
-		FGameplayTag ValidTag = UE::StateTree::Tests::FNativeGameplayTags::Get().TestTag;
-		FGameplayTag InvalidTag = UE::StateTree::Tests::FNativeGameplayTags::Get().TestTag2;
-
-		using FValidPayload = FStateTreeTest_PropertyStructA;
-		using FInvalidPayload = FStateTreeTest_PropertyStructB;
-
-		// This state shouldn't be selected as it requires different tag.
-		UStateTreeState& StateA = Root.AddChildState(FName(TEXT("A")));
-		StateA.bHasRequiredEventToEnter  = true;
-		StateA.RequiredEventToEnter.Tag = InvalidTag;
-		auto& TaskA = StateA.AddTask<FTestTask_Stand>(FName(TEXT("TaskA")));
-
-		// This state shouldn't be selected as it requires different payload.
-		UStateTreeState& StateB = Root.AddChildState(FName(TEXT("B")));
-		StateB.bHasRequiredEventToEnter  = true;
-		StateB.RequiredEventToEnter.PayloadStruct = FInvalidPayload::StaticStruct();
-		auto& TaskB = StateB.AddTask<FTestTask_Stand>(FName(TEXT("TaskB")));
-
-		// This state shouldn't be selected as it requires the same tag, but different payload.
-		UStateTreeState& StateC = Root.AddChildState(FName(TEXT("C")));
-		StateC.bHasRequiredEventToEnter  = true;
-		StateC.RequiredEventToEnter.Tag = ValidTag;
-		StateC.RequiredEventToEnter.PayloadStruct = FInvalidPayload::StaticStruct();
-		auto& TaskC = StateC.AddTask<FTestTask_Stand>(FName(TEXT("TaskC")));
-
-		// This state shouldn't be selected as it requires the same payload, but different tag.
-		UStateTreeState& StateD = Root.AddChildState(FName(TEXT("D")));
-		StateD.bHasRequiredEventToEnter  = true;
-		StateD.RequiredEventToEnter.Tag = InvalidTag;
-		StateD.RequiredEventToEnter.PayloadStruct = FValidPayload::StaticStruct();
-		auto& TaskD = StateD.AddTask<FTestTask_Stand>(FName(TEXT("TaskD")));
-
-		// This state should be selected as the arrived event matches the requirement.
-		UStateTreeState& StateE = Root.AddChildState(FName(TEXT("E")));
-		StateE.bHasRequiredEventToEnter  = true;
-		StateE.RequiredEventToEnter.Tag = ValidTag;
-		StateE.RequiredEventToEnter.PayloadStruct = FValidPayload::StaticStruct();
-		auto& TaskE = StateE.AddTask<FTestTask_Stand>(FName(TEXT("TaskE")));
-
-		// This state should be selected only initially when there's not event in the queue.
-		UStateTreeState& StateInitial = Root.AddChildState(FName(TEXT("Initial")));
-		auto& TaskInitial = StateInitial.AddTask<FTestTask_Stand>(FName(TEXT("TaskInitial")));
-		StateInitial.AddTransition(EStateTreeTransitionTrigger::OnEvent, ValidTag, EStateTreeTransitionType::GotoState, &Root);
-
-		FStateTreeCompilerLog Log;
-		FStateTreeCompiler Compiler(Log);
-		const bool bResult = Compiler.Compile(StateTree);
-
-		AITEST_TRUE("StateTree should get compiled", bResult);
-
-		EStateTreeRunStatus Status = EStateTreeRunStatus::Unset;
-		FStateTreeInstanceData InstanceData;
-		FTestStateTreeExecutionContext Exec(StateTree, StateTree, InstanceData);
-		const bool bInitSucceeded = Exec.IsValid();
-		AITEST_TRUE("StateTree should init", bInitSucceeded);
-
-		const FString EnterStateStr(TEXT("EnterState"));
-
-		Status = Exec.Start();
-		AITEST_TRUE("StateTree TaskInitial should enter state", Exec.Expect(TaskInitial.GetName(), EnterStateStr));
-		Exec.LogClear();
-
-		Exec.SendEvent(ValidTag, FConstStructView::Make(FValidPayload()));
-		Status = Exec.Tick(0.1f);
-
-		AITEST_FALSE("StateTree TaskA should not enter state", Exec.Expect(TaskA.GetName(), EnterStateStr));
-		AITEST_FALSE("StateTree TaskB should not enter state", Exec.Expect(TaskB.GetName(), EnterStateStr));
-		AITEST_FALSE("StateTree TaskC should not enter state", Exec.Expect(TaskC.GetName(), EnterStateStr));
-		AITEST_FALSE("StateTree TaskD should not enter state", Exec.Expect(TaskD.GetName(), EnterStateStr));
-		AITEST_TRUE("StateTree TaskE should enter state", Exec.Expect(TaskE.GetName(), EnterStateStr));
-		Exec.LogClear();
-		
-		Exec.Stop();
-
-		return true;
-	}
-};
-IMPLEMENT_AI_INSTANT_TEST(FStateTreeTest_StateRequiringEvent, "System.StateTree.StateRequiringEvent");
-
-struct FStateTreeTest_PassingTransitionEventToStateSelection : FAITestBase
-{
-	virtual bool InstantTest() override
-	{
-		UStateTree& StateTree = UE::StateTree::Tests::NewStateTree(&GetWorld());
-		UStateTreeEditorData& EditorData = *Cast<UStateTreeEditorData>(StateTree.EditorData);
-
-		UStateTreeState& Root = EditorData.AddSubTree(FName(TEXT("Root")));
-
-		FStateTreePropertyPath PathToPayloadMember;
-		{
-			const bool bParseResult = PathToPayloadMember.FromString(TEXT("Payload.A"));
-
-			AITEST_TRUE("Parsing path should succeeed", bParseResult);
-
-			FStateTreeEvent EventWithPayload;
-			EventWithPayload.Payload = FInstancedStruct::Make<FStateTreeTest_PropertyStructA>();
-			const bool bUpdateSegments = PathToPayloadMember.UpdateSegmentsFromValue(FStateTreeDataView(FStructView::Make(EventWithPayload)));
-			AITEST_TRUE("Updating segments should succeeed", bUpdateSegments);
-		}
-
-		// This state shouldn't be selected, because transition's condition and state's enter condition exlude each other.
-		UStateTreeState& StateA = Root.AddChildState(FName(TEXT("A")));
-		StateA.bHasRequiredEventToEnter  = true;
-		StateA.RequiredEventToEnter.PayloadStruct = FStateTreeTest_PropertyStructA::StaticStruct();
-		auto& TaskA = StateA.AddTask<FTestTask_Stand>(FName(TEXT("TaskA")));
-		TStateTreeEditorNode<FStateTreeCompareIntCondition>& AIntCond = StateA.AddEnterCondition<FStateTreeCompareIntCondition>(EGenericAICheck::Equal);
-		AIntCond.GetInstanceData().Right = 0;
-		EditorData.AddPropertyBinding(
-			FStateTreePropertyPath(StateA.GetEventID(), PathToPayloadMember.GetSegments()),
-			FStateTreePropertyPath(AIntCond.ID, TEXT("Left")));
-
-		// This state should be selected as the sent event fullfils both transition's condition and state's enter condition.
-		UStateTreeState& StateB = Root.AddChildState(FName(TEXT("B")));
-		StateB.bHasRequiredEventToEnter  = true;
-		StateB.RequiredEventToEnter.PayloadStruct = FStateTreeTest_PropertyStructA::StaticStruct();
-		auto& TaskB = StateB.AddTask<FTestTask_PrintValue>(FName(TEXT("TaskB")));
-		// Test copying data from the state event. The condition properties are copied from temp instance data during selection, this gets copied from active instance data.
-		TaskB.GetInstanceData().Value = -1; // Initially -1, expected to be overridden by property binding below. 
-		EditorData.AddPropertyBinding(
-			FStateTreePropertyPath(StateB.GetEventID(), PathToPayloadMember.GetSegments()),
-			FStateTreePropertyPath(TaskB.ID, TEXT("Value")));
-		
-		TStateTreeEditorNode<FStateTreeCompareIntCondition>& BIntCond = StateB.AddEnterCondition<FStateTreeCompareIntCondition>(EGenericAICheck::Equal);
-		BIntCond.GetInstanceData().Right = 1;
-		EditorData.AddPropertyBinding(
-			FStateTreePropertyPath(StateB.GetEventID(), PathToPayloadMember.GetSegments()),
-			FStateTreePropertyPath(BIntCond.ID, TEXT("Left")));
-
-		// This state should be selected only initially when there's not event in the queue.
-		UStateTreeState& StateInitial = Root.AddChildState(FName(TEXT("Initial")));
-		auto& TaskInitial = StateInitial.AddTask<FTestTask_Stand>(FName(TEXT("TaskInitial")));
-		// Transition from Initial -> StateA
-		FStateTreeTransition& TransA = StateInitial.AddTransition(EStateTreeTransitionTrigger::OnEvent, FGameplayTag(), EStateTreeTransitionType::GotoState, &StateA);
-		TransA.RequiredEvent.PayloadStruct = FStateTreeTest_PropertyStructA::StaticStruct();
-		TStateTreeEditorNode<FStateTreeCompareIntCondition>& TransAIntCond = TransA.AddCondition<FStateTreeCompareIntCondition>(EGenericAICheck::Equal);
-		TransAIntCond.GetInstanceData().Right = 1;
-		EditorData.AddPropertyBinding(
-			FStateTreePropertyPath(TransA.GetEventID(), PathToPayloadMember.GetSegments()),
-			FStateTreePropertyPath(TransAIntCond.ID, TEXT("Left")));
-		// Transition from Initial -> StateB
-		FStateTreeTransition& TransB = StateInitial.AddTransition(EStateTreeTransitionTrigger::OnEvent, FGameplayTag(), EStateTreeTransitionType::GotoState, &StateB);
-		TransB.RequiredEvent.PayloadStruct = FStateTreeTest_PropertyStructA::StaticStruct();
-		TStateTreeEditorNode<FStateTreeCompareIntCondition>& TransBIntCond = TransB.AddCondition<FStateTreeCompareIntCondition>(EGenericAICheck::Equal);
-		TransBIntCond.GetInstanceData().Right = 1;
-		EditorData.AddPropertyBinding(
-			FStateTreePropertyPath(TransB.GetEventID(), PathToPayloadMember.GetSegments()),
-			FStateTreePropertyPath(TransBIntCond.ID, TEXT("Left")));
-
-		FStateTreeCompilerLog Log;
-		FStateTreeCompiler Compiler(Log);
-		const bool bResult = Compiler.Compile(StateTree);
-
-		AITEST_TRUE("StateTree should get compiled", bResult);
-
-		EStateTreeRunStatus Status = EStateTreeRunStatus::Unset;
-		FStateTreeInstanceData InstanceData;
-		FTestStateTreeExecutionContext Exec(StateTree, StateTree, InstanceData);
-		const bool bInitSucceeded = Exec.IsValid();
-		AITEST_TRUE("StateTree should init", bInitSucceeded);
-
-		const FString EnterStateStr(TEXT("EnterState"));
-
-		Status = Exec.Start();
-		AITEST_TRUE("StateTree TaskInitial should enter state", Exec.Expect(TaskInitial.GetName(), EnterStateStr));
-		Exec.LogClear();
-
-		// The conditions test for payload Value=1, the first event should not trigger transition. 
-		Exec.SendEvent(UE::StateTree::Tests::FNativeGameplayTags::Get().TestTag, FConstStructView::Make(FStateTreeTest_PropertyStructA{0}));
-		Exec.SendEvent(UE::StateTree::Tests::FNativeGameplayTags::Get().TestTag, FConstStructView::Make(FStateTreeTest_PropertyStructA{1}));
-		Status = Exec.Tick(0.1f);
-
-		AITEST_FALSE("StateTree TaskA should not enter state", Exec.Expect(TaskA.GetName(), EnterStateStr));
-		AITEST_TRUE("StateTree TaskB should enter state", Exec.Expect(TaskB.GetName(), TEXT("EnterState1"))); // TaskB decorates "EnterState" with value from the payload.
-		Exec.LogClear();
-
-		Exec.Stop();
-
-		return true;
-	}
-};
-IMPLEMENT_AI_INSTANT_TEST(FStateTreeTest_PassingTransitionEventToStateSelection, "System.StateTree.PassingTransitionEventToStateSelection");
-
-struct FStateTreeTest_PropertyPathOffset : FAITestBase
-{
-	virtual bool InstantTest() override
-	{
-		FStateTreePropertyPath Path;
-		const bool bParseResult = Path.FromString(TEXT("StructB.B"));
-
-		AITEST_TRUE("Parsing path should succeeed", bParseResult);
-		AITEST_EQUAL("Should have 2 path segments", Path.NumSegments(), 2);
-
-		FString ResolveErrors;
-		TArray<FStateTreePropertyPathIndirection> Indirections;
-		const bool bResolveResult = Path.ResolveIndirections(FStateTreeTest_PropertyStruct::StaticStruct(), Indirections, &ResolveErrors);
-
-		AITEST_TRUE("Resolve path should succeeed", bResolveResult);
-		AITEST_EQUAL("Should have no resolve errors", ResolveErrors.Len(), 0);
-		
-		AITEST_EQUAL("Should have 2 indirections", Indirections.Num(), 2);
-		AITEST_EQUAL("Indirection 0 should be Offset type", Indirections[0].GetAccessType(), EStateTreePropertyAccessType::Offset);
-		AITEST_EQUAL("Indirection 1 should be Offset type", Indirections[1].GetAccessType(), EStateTreePropertyAccessType::Offset);
-
-		return true;
-	}
-};
-IMPLEMENT_AI_INSTANT_TEST(FStateTreeTest_PropertyPathOffset, "System.StateTree.PropertyPath.Offset");
-
-struct FStateTreeTest_PropertyPathParseFail : FAITestBase
-{
-	virtual bool InstantTest() override
-	{
-		{
-			FStateTreePropertyPath Path;
-			const bool bParseResult = Path.FromString(TEXT("")); // empty is valid.
-			AITEST_TRUE("Parsing path should succeed", bParseResult);
-		}
-
-		{
-			FStateTreePropertyPath Path;
-			const bool bParseResult = Path.FromString(TEXT("StructB.[0]B"));
-			AITEST_FALSE("Parsing path should fail", bParseResult);
-		}
-
-		{
-			FStateTreePropertyPath Path;
-			const bool bParseResult = Path.FromString(TEXT("StructB..NoThere"));
-			AITEST_FALSE("Parsing path should fail", bParseResult);
-		}
-
-		{
-			FStateTreePropertyPath Path;
-			const bool bParseResult = Path.FromString(TEXT("."));
-			AITEST_FALSE("Parsing path should fail", bParseResult);
-		}
-
-		{
-			FStateTreePropertyPath Path;
-			const bool bParseResult = Path.FromString(TEXT("StructB..B"));
-			AITEST_FALSE("Parsing path should fail", bParseResult);
-		}
-
-		return true;
-	}
-};
-IMPLEMENT_AI_INSTANT_TEST(FStateTreeTest_PropertyPathParseFail, "System.StateTree.PropertyPath.ParseFail");
-
-struct FStateTreeTest_PropertyPathOffsetFail : FAITestBase
-{
-	virtual bool InstantTest() override
-	{
-		FStateTreePropertyPath Path;
-		const bool bParseResult = Path.FromString(TEXT("StructB.Q"));
-
-		AITEST_TRUE("Parsing path should succeeed", bParseResult);
-		AITEST_EQUAL("Should have 2 path segments", Path.NumSegments(), 2);
-
-		FString ResolveErrors;
-		TArray<FStateTreePropertyPathIndirection> Indirections;
-		const bool bResolveResult = Path.ResolveIndirections(FStateTreeTest_PropertyStruct::StaticStruct(), Indirections, &ResolveErrors);
-
-		AITEST_FALSE("Resolve path should not succeeed", bResolveResult);
-		AITEST_NOT_EQUAL("Should have errors", ResolveErrors.Len(), 0);
-		
-		AITEST_EQUAL("Should have 0 indirections", Indirections.Num(), 0);
-
-		return true;
-	}
-};
-IMPLEMENT_AI_INSTANT_TEST(FStateTreeTest_PropertyPathOffsetFail, "System.StateTree.PropertyPath.OffsetFail");
-
-struct FStateTreeTest_PropertyPathObject : FAITestBase
-{
-	virtual bool InstantTest() override
-	{
-		FStateTreePropertyPath Path;
-		const bool bParseResult = Path.FromString(TEXT("InstancedObject.A"));
-
-		AITEST_TRUE("Parsing path should succeeed", bParseResult);
-		AITEST_EQUAL("Should have 2 path segments", Path.NumSegments(), 2);
-
-		UStateTreeTest_PropertyObject* Object = NewObject<UStateTreeTest_PropertyObject>();
-		Object->InstancedObject = NewObject<UStateTreeTest_PropertyObjectInstanced>();
-		
-		const bool bUpdateResult = Path.UpdateSegmentsFromValue(FStateTreeDataView(Object));
-
-		AITEST_TRUE("Update instance types should succeeed", bUpdateResult);
-		AITEST_TRUE("Path segment 0 instance type should be UStateTreeTest_PropertyObjectInstanced", Path.GetSegment(0).GetInstanceStruct() == UStateTreeTest_PropertyObjectInstanced::StaticClass());
-		AITEST_TRUE("Path segment 1 instance type should be nullptr", Path.GetSegment(1).GetInstanceStruct() == nullptr);
-
-		return true;
-	}
-};
-IMPLEMENT_AI_INSTANT_TEST(FStateTreeTest_PropertyPathObject, "System.StateTree.PropertyPath.Object");
-
-struct FStateTreeTest_PropertyPathWrongObject : FAITestBase
-{
-	virtual bool InstantTest() override
-	{
-		FStateTreePropertyPath Path;
-		const bool bParseResult = Path.FromString(TEXT("InstancedObject.B"));
-
-		AITEST_TRUE("Parsing path should succeeed", bParseResult);
-		AITEST_EQUAL("Should have 2 path segments", Path.NumSegments(), 2);
-
-		UStateTreeTest_PropertyObject* Object = NewObject<UStateTreeTest_PropertyObject>();
-
-		Object->InstancedObject = NewObject<UStateTreeTest_PropertyObjectInstancedWithB>();
-		{
-			FString ResolveErrors;
-			TArray<FStateTreePropertyPathIndirection> Indirections;
-			const bool bResolveResult = Path.ResolveIndirectionsWithValue(FStateTreeDataView(Object), Indirections, &ResolveErrors);
-
-			AITEST_TRUE("Resolve path should succeeed", bResolveResult);
-			AITEST_EQUAL("Should have 2 indirections", Indirections.Num(), 2);
-			AITEST_TRUE("Object ", Indirections[0].GetAccessType() == EStateTreePropertyAccessType::ObjectInstance);
-			AITEST_TRUE("Object ", Indirections[0].GetContainerStruct() == Object->GetClass());
-			AITEST_TRUE("Object ", Indirections[0].GetInstanceStruct() == UStateTreeTest_PropertyObjectInstancedWithB::StaticClass());
-			AITEST_EQUAL("Should not have error", ResolveErrors.Len(), 0);
-		}
-
-		Object->InstancedObject = NewObject<UStateTreeTest_PropertyObjectInstanced>();
-		{
-			FString ResolveErrors;
-			TArray<FStateTreePropertyPathIndirection> Indirections;
-			const bool bResolveResult = Path.ResolveIndirectionsWithValue(FStateTreeDataView(Object), Indirections, &ResolveErrors);
-
-			AITEST_FALSE("Resolve path should fail", bResolveResult);
-			AITEST_EQUAL("Should have 0 indirections", Indirections.Num(), 0);
-			AITEST_NOT_EQUAL("Should have error", ResolveErrors.Len(), 0);
-		}
-		
-		return true;
-	}
-};
-IMPLEMENT_AI_INSTANT_TEST(FStateTreeTest_PropertyPathWrongObject, "System.StateTree.PropertyPath.WrongObject");
-
-struct FStateTreeTest_PropertyPathArray : FAITestBase
-{
-	virtual bool InstantTest() override
-	{
-		FStateTreePropertyPath Path;
-		const bool bParseResult = Path.FromString(TEXT("ArrayOfInts[1]"));
-
-		AITEST_TRUE("Parsing path should succeeed", bParseResult);
-		AITEST_EQUAL("Should have 1 path segments", Path.NumSegments(), 1);
-
-		UStateTreeTest_PropertyObject* Object = NewObject<UStateTreeTest_PropertyObject>();
-		Object->ArrayOfInts.Add(42);
-		Object->ArrayOfInts.Add(123);
-
-		FString ResolveErrors;
-		TArray<FStateTreePropertyPathIndirection> Indirections;
-		const bool bResolveResult = Path.ResolveIndirectionsWithValue(FStateTreeDataView(Object), Indirections, &ResolveErrors);
-
-		AITEST_TRUE("Resolve path should succeeed", bResolveResult);
-		AITEST_EQUAL("Should have no resolve errors", ResolveErrors.Len(), 0);
-		AITEST_EQUAL("Should have 2 indirections", Indirections.Num(), 2);
-		AITEST_EQUAL("Indirection 0 should be IndexArray type", Indirections[0].GetAccessType(), EStateTreePropertyAccessType::IndexArray);
-		AITEST_EQUAL("Indirection 1 should be Offset type", Indirections[1].GetAccessType(), EStateTreePropertyAccessType::Offset);
-
-		const int32 Value = *reinterpret_cast<const int32*>(Indirections[1].GetPropertyAddress());
-		AITEST_EQUAL("Value should be 123", Value, 123);
-		
-		return true;
-	}
-};
-IMPLEMENT_AI_INSTANT_TEST(FStateTreeTest_PropertyPathArray, "System.StateTree.PropertyPath.Array");
-
-struct FStateTreeTest_PropertyPathArrayInvalidIndex : FAITestBase
-{
-	virtual bool InstantTest() override
-	{
-		FStateTreePropertyPath Path;
-		const bool bParseResult = Path.FromString(TEXT("ArrayOfInts[123]"));
-
-		AITEST_TRUE("Parsing path should succeeed", bParseResult);
-		AITEST_EQUAL("Should have 1 path segments", Path.NumSegments(), 1);
-
-		UStateTreeTest_PropertyObject* Object = NewObject<UStateTreeTest_PropertyObject>();
-		Object->ArrayOfInts.Add(42);
-		Object->ArrayOfInts.Add(123);
-
-		FString ResolveErrors;
-		TArray<FStateTreePropertyPathIndirection> Indirections;
-		const bool bResolveResult = Path.ResolveIndirectionsWithValue(FStateTreeDataView(Object), Indirections, &ResolveErrors);
-
-		AITEST_FALSE("Resolve path should fail", bResolveResult);
-		
-		return true;
-	}
-};
-IMPLEMENT_AI_INSTANT_TEST(FStateTreeTest_PropertyPathArrayInvalidIndex, "System.StateTree.PropertyPath.ArrayInvalidIndex");
-
-struct FStateTreeTest_PropertyPathArrayOfStructs : FAITestBase
-{
-	virtual bool InstantTest() override
-	{
-		FStateTreePropertyPath Path1;
-		Path1.FromString(TEXT("ArrayOfStruct[0].B"));
-
-		FStateTreePropertyPath Path2;
-		Path2.FromString(TEXT("ArrayOfStruct[2].StructB.B"));
-
-		UStateTreeTest_PropertyObject* Object = NewObject<UStateTreeTest_PropertyObject>();
-		Object->ArrayOfStruct.AddDefaulted_GetRef().B = 3;
-		Object->ArrayOfStruct.AddDefaulted();
-		Object->ArrayOfStruct.AddDefaulted_GetRef().StructB.B = 42;
-
-		{
-			FString ResolveErrors;
-			TArray<FStateTreePropertyPathIndirection> Indirections;
-			const bool bResolveResult = Path1.ResolveIndirectionsWithValue(FStateTreeDataView(Object), Indirections, &ResolveErrors);
-
-			AITEST_TRUE("Resolve path1 should succeeed", bResolveResult);
-			AITEST_EQUAL("Should have no resolve errors", ResolveErrors.Len(), 0);
-			AITEST_EQUAL("Should have 3 indirections", Indirections.Num(), 3);
-			AITEST_EQUAL("Indirection 0 should be ArrayIndex type", Indirections[0].GetAccessType(), EStateTreePropertyAccessType::IndexArray);
-			AITEST_EQUAL("Indirection 1 should be Offset type", Indirections[1].GetAccessType(), EStateTreePropertyAccessType::Offset);
-			AITEST_EQUAL("Indirection 2 should be Offset type", Indirections[2].GetAccessType(), EStateTreePropertyAccessType::Offset);
-
-			const int32 Value = *reinterpret_cast<const int32*>(Indirections[2].GetPropertyAddress());
-			AITEST_EQUAL("Value should be 3", Value, 3);
-		}
-
-		{
-			FString ResolveErrors;
-			TArray<FStateTreePropertyPathIndirection> Indirections;
-			const bool bResolveResult = Path2.ResolveIndirectionsWithValue(FStateTreeDataView(Object), Indirections, &ResolveErrors);
-
-			AITEST_TRUE("Resolve path2 should succeeed", bResolveResult);
-			AITEST_EQUAL("Should have no resolve errors", ResolveErrors.Len(), 0);
-			AITEST_EQUAL("Should have 4 indirections", Indirections.Num(), 4);
-			AITEST_EQUAL("Indirection 0 should be ArrayIndex type", Indirections[0].GetAccessType(), EStateTreePropertyAccessType::IndexArray);
-			AITEST_EQUAL("Indirection 1 should be Offset type", Indirections[1].GetAccessType(), EStateTreePropertyAccessType::Offset);
-			AITEST_EQUAL("Indirection 2 should be Offset type", Indirections[2].GetAccessType(), EStateTreePropertyAccessType::Offset);
-			AITEST_EQUAL("Indirection 3 should be Offset type", Indirections[3].GetAccessType(), EStateTreePropertyAccessType::Offset);
-
-			const int32 Value = *reinterpret_cast<const int32*>(Indirections[3].GetPropertyAddress());
-			AITEST_EQUAL("Value should be 42", Value, 42);
-		}
-		return true;
-	}
-};
-IMPLEMENT_AI_INSTANT_TEST(FStateTreeTest_PropertyPathArrayOfStructs, "System.StateTree.PropertyPath.ArrayOfStructs");
-
-struct FStateTreeTest_PropertyPathArrayOfInstancedObjects : FAITestBase
-{
-	virtual bool InstantTest() override
-	{
-		FStateTreePropertyPath Path;
-		Path.FromString(TEXT("ArrayOfInstancedStructs[0].B"));
-
-		FStateTreeTest_PropertyStruct Struct;
-		Struct.B = 123;
-		
-		UStateTreeTest_PropertyObject* Object = NewObject<UStateTreeTest_PropertyObject>();
-		Object->ArrayOfInstancedStructs.Emplace(FConstStructView::Make(Struct));
-
-		const bool bUpdateResult = Path.UpdateSegmentsFromValue(FStateTreeDataView(Object));
-		AITEST_TRUE("Update instance types should succeeed", bUpdateResult);
-		AITEST_EQUAL("Should have 2 path segments", Path.NumSegments(), 2);
-		AITEST_TRUE("Path segment 0 instance type should be FStateTreeTest_PropertyStruct", Path.GetSegment(0).GetInstanceStruct() == FStateTreeTest_PropertyStruct::StaticStruct());
-		AITEST_TRUE("Path segment 1 instance type should be nullptr", Path.GetSegment(1).GetInstanceStruct() == nullptr);
-
-		{
-			FString ResolveErrors;
-			TArray<FStateTreePropertyPathIndirection> Indirections;
-			const bool bResolveResult = Path.ResolveIndirections(UStateTreeTest_PropertyObject::StaticClass(), Indirections, &ResolveErrors);
-
-			AITEST_TRUE("Resolve path should succeeed", bResolveResult);
-			AITEST_EQUAL("Should have no resolve errors", ResolveErrors.Len(), 0);
-			AITEST_EQUAL("Should have 3 indirections", Indirections.Num(), 3);
-			AITEST_EQUAL("Indirection 0 should be ArrayIndex type", Indirections[0].GetAccessType(), EStateTreePropertyAccessType::IndexArray);
-			AITEST_EQUAL("Indirection 1 should be StructInstance type", Indirections[1].GetAccessType(), EStateTreePropertyAccessType::StructInstance);
-			AITEST_EQUAL("Indirection 2 should be Offset type", Indirections[2].GetAccessType(), EStateTreePropertyAccessType::Offset);
-		}
-
-		{
-			FString ResolveErrors;
-			TArray<FStateTreePropertyPathIndirection> Indirections;
-			const bool bResolveResult = Path.ResolveIndirectionsWithValue(FStateTreeDataView(Object), Indirections, &ResolveErrors);
-
-			AITEST_TRUE("Resolve path should succeeed", bResolveResult);
-			AITEST_EQUAL("Should have no resolve errors", ResolveErrors.Len(), 0);
-			AITEST_EQUAL("Should have 3 indirections", Indirections.Num(), 3);
-			AITEST_EQUAL("Indirection 0 should be ArrayIndex type", Indirections[0].GetAccessType(), EStateTreePropertyAccessType::IndexArray);
-			AITEST_EQUAL("Indirection 1 should be StructInstance type", Indirections[1].GetAccessType(), EStateTreePropertyAccessType::StructInstance);
-			AITEST_EQUAL("Indirection 2 should be Offset type", Indirections[2].GetAccessType(), EStateTreePropertyAccessType::Offset);
-
-			const int32 Value = *reinterpret_cast<const int32*>(Indirections[2].GetPropertyAddress());
-			AITEST_EQUAL("Value should be 123", Value, 123);
-		}
-		
-		return true;
-	}
-};
-IMPLEMENT_AI_INSTANT_TEST(FStateTreeTest_PropertyPathArrayOfInstancedObjects, "System.StateTree.PropertyPath.ArrayOfInstancedObjects");
-
-struct FStateTreeTest_BindingsCompiler : FAITestBase
-{
-	virtual bool InstantTest() override
-	{
-		FStateTreeCompilerLog Log;
-		FStateTreePropertyBindings Bindings;
-		FStateTreePropertyBindingCompiler BindingCompiler;
-
-		const bool bInitResult = BindingCompiler.Init(Bindings, Log);
-		AITEST_TRUE("Expect init to succeed", bInitResult);
-
-		FStateTreeBindableStructDesc SourceADesc;
-		SourceADesc.Name = FName(TEXT("SourceA"));
-		SourceADesc.Struct = TBaseStructure<FStateTreeTest_PropertyCopy>::Get();
-		SourceADesc.DataSource = EStateTreeBindableStructSource::Parameter;
-		SourceADesc.DataHandle = FStateTreeDataHandle(EStateTreeDataSourceType::ContextData, 0); // Used as index to SourceViews below.
-		SourceADesc.ID = FGuid::NewGuid();
-
-		FStateTreeBindableStructDesc SourceBDesc;
-		SourceBDesc.Name = FName(TEXT("SourceB"));
-		SourceBDesc.Struct = TBaseStructure<FStateTreeTest_PropertyCopy>::Get();
-		SourceBDesc.DataSource = EStateTreeBindableStructSource::Parameter;
-		SourceBDesc.DataHandle = FStateTreeDataHandle(EStateTreeDataSourceType::ContextData, 1); // Used as index to SourceViews below.
-		SourceBDesc.ID = FGuid::NewGuid();
-
-		FStateTreeBindableStructDesc TargetDesc;
-		TargetDesc.Name = FName(TEXT("Target"));
-		TargetDesc.Struct = TBaseStructure<FStateTreeTest_PropertyCopy>::Get();
-		TargetDesc.DataSource = EStateTreeBindableStructSource::Parameter;
-		TargetDesc.ID = FGuid::NewGuid();
-		
-		const int32 SourceAIndex = BindingCompiler.AddSourceStruct(SourceADesc);
-		const int32 SourceBIndex = BindingCompiler.AddSourceStruct(SourceBDesc);
-
-		TArray<FStateTreePropertyPathBinding> PropertyBindings;
-		PropertyBindings.Add(UE::StateTree::Tests::MakeBinding(SourceBDesc.ID, TEXT("Item"), TargetDesc.ID, TEXT("Array[1]")));
-		PropertyBindings.Add(UE::StateTree::Tests::MakeBinding(SourceADesc.ID, TEXT("Item.B"), TargetDesc.ID, TEXT("Array[1].B")));
-		PropertyBindings.Add(UE::StateTree::Tests::MakeBinding(SourceADesc.ID, TEXT("Array"), TargetDesc.ID, TEXT("Array")));
-
-		int32 CopyBatchIndex = INDEX_NONE;
-		const bool bCompileBatchResult = BindingCompiler.CompileBatch(TargetDesc, PropertyBindings, FStateTreeIndex16::Invalid, FStateTreeIndex16::Invalid, CopyBatchIndex);
-		AITEST_TRUE("CompileBatch should succeed", bCompileBatchResult);
-		AITEST_NOT_EQUAL("CopyBatchIndex should not be INDEX_NONE", CopyBatchIndex, (int32)INDEX_NONE);
-
-		BindingCompiler.Finalize();
-
-		const bool bResolveResult = Bindings.ResolvePaths();
-		AITEST_TRUE("ResolvePaths should succeed", bResolveResult);
-
-		FStateTreeTest_PropertyCopy SourceA;
-		SourceA.Item.B = 123;
-		SourceA.Array.AddDefaulted_GetRef().A = 1;
-		SourceA.Array.AddDefaulted_GetRef().B = 2;
-
-		FStateTreeTest_PropertyCopy SourceB;
-		SourceB.Item.A = 41;
-		SourceB.Item.B = 42;
-
-		FStateTreeTest_PropertyCopy Target;
-
-		AITEST_TRUE("SourceAIndex should be less than max number of source structs.", SourceAIndex < Bindings.GetSourceStructNum());
-		AITEST_TRUE("SourceBIndex should be less than max number of source structs.", SourceBIndex < Bindings.GetSourceStructNum());
-
-		TArray<FStateTreeDataView> SourceViews;
-		SourceViews.SetNum(Bindings.GetSourceStructNum());
-		SourceViews[SourceAIndex] = FStateTreeDataView(FStructView::Make(SourceA));
-		SourceViews[SourceBIndex] = FStateTreeDataView(FStructView::Make(SourceB));
-		FStateTreeDataView TargetView(FStructView::Make(Target));
-
-		bool bCopyResult = true;
-		for (const FStateTreePropertyCopy& Copy : Bindings.GetBatchCopies(FStateTreeIndex16(CopyBatchIndex)))
-		{
-			bCopyResult &= Bindings.CopyProperty(Copy, SourceViews[Copy.SourceDataHandle.GetIndex()], TargetView);
-		}
-		AITEST_TRUE("CopyTo should succeed", bCopyResult);
-
-		// Due to binding sorting, we expect them to executed in this order (sorted based on target access, earliest to latest)
-		// SourceA.Array -> Target.Array
-		// SourceB.Item -> Target.Array[1]
-		// SourceA.Item.B -> Target.Array[1].B
-
-		AITEST_EQUAL("Expect TargetArray to be copied from SourceA", Target.Array.Num(), SourceA.Array.Num());
-		AITEST_EQUAL("Expect Target.Array[0].A copied from SourceA.Array[0].A", Target.Array[0].A, SourceA.Array[0].A);
-		AITEST_EQUAL("Expect Target.Array[0].B copied from SourceA.Array[0].B", Target.Array[0].B, SourceA.Array[0].B);
-		AITEST_EQUAL("Expect Target.Array[1].A copied from SourceB.Item.A", Target.Array[1].A, SourceB.Item.A);
-		AITEST_EQUAL("Expect Target.Array[1].B copied from SourceA.Item.B", Target.Array[1].B, SourceA.Item.B);
-		
-		return true;
-	}
-};
-IMPLEMENT_AI_INSTANT_TEST(FStateTreeTest_BindingsCompiler, "System.StateTree.BindingsCompiler");
-
-struct FStateTreeTest_PropertyFunctions : FAITestBase
-{
-	virtual bool InstantTest() override
-	{
-		UStateTree& StateTree = UE::StateTree::Tests::NewStateTree(&GetWorld());
-		UStateTreeEditorData& EditorData = *Cast<UStateTreeEditorData>(StateTree.EditorData);
-		UStateTreeState& Root = EditorData.AddSubTree(FName(TEXT("Root")));
-		FStateTreePropertyPathSegment PathSegmentToFuncResult = FStateTreePropertyPathSegment(TEXT("Result"));
-
-		// Condition with property function binding.
-		{
-			TStateTreeEditorNode<FStateTreeCompareIntCondition>& EnterCond = Root.AddEnterCondition<FStateTreeCompareIntCondition>(EGenericAICheck::Equal);
-			EnterCond.GetInstanceData().Right = 1;
-			EditorData.AddPropertyBinding(CastChecked<UScriptStruct>(FTestPropertyFunction::StaticStruct()), {PathSegmentToFuncResult}, FStateTreePropertyPath(EnterCond.ID, TEXT("Left")));
-		}
-
-		// Task with multiple nested property function bindings.
-		auto& TaskA = Root.AddTask<FTestTask_PrintAndResetValue>(FName(TEXT("TaskA")));
-		constexpr int32 TaskAPropertyFunctionsAmount = 10;
-		{
-			EditorData.AddPropertyBinding(CastChecked<UScriptStruct>(FTestPropertyFunction::StaticStruct()), {PathSegmentToFuncResult}, FStateTreePropertyPath(TaskA.ID, TEXT("Value")));
-		
-			for (int32 i = 0; i < TaskAPropertyFunctionsAmount - 1; ++i)
-			{
-				const FStateTreePropertyPathBinding& LastBinding = EditorData.GetPropertyEditorBindings()->GetBindings().Last();
-				const FGuid LastBindingPropertyFuncID = LastBinding.GetPropertyFunctionNode().Get<const FStateTreeEditorNode>().ID;
-				EditorData.AddPropertyBinding(CastChecked<UScriptStruct>(FTestPropertyFunction::StaticStruct()), {PathSegmentToFuncResult}, FStateTreePropertyPath(LastBindingPropertyFuncID, TEXT("Input")));
-			}
-		}
-
-		// Task bound to state parameter with multiple nested property function bindings.
-		auto& TaskB = Root.AddTask<FTestTask_PrintAndResetValue>(FName(TEXT("TaskB")));
-		constexpr int32 ParameterPropertyFunctionsAmount = 5;
-		{
-			Root.Parameters.Parameters.AddProperty(FName(TEXT("Int")), EPropertyBagPropertyType::Int32);
-			const FStateTreePropertyPath PathToProperty = FStateTreePropertyPath(Root.Parameters.ID, TEXT("Int"));
-			EditorData.AddPropertyBinding(PathToProperty, FStateTreePropertyPath(TaskB.ID, TEXT("Value")));
-			EditorData.AddPropertyBinding(CastChecked<UScriptStruct>(FTestPropertyFunction::StaticStruct()), {PathSegmentToFuncResult}, PathToProperty);
-		
-			for (int32 i = 0; i < ParameterPropertyFunctionsAmount - 1; ++i)
-			{
-				const FStateTreePropertyPathBinding& LastBinding = EditorData.GetPropertyEditorBindings()->GetBindings().Last();
-				const FGuid LastBindingPropertyFuncID = LastBinding.GetPropertyFunctionNode().Get<const FStateTreeEditorNode>().ID;
-				EditorData.AddPropertyBinding(CastChecked<UScriptStruct>(FTestPropertyFunction::StaticStruct()), {PathSegmentToFuncResult}, FStateTreePropertyPath(LastBindingPropertyFuncID, TEXT("Input")));
-			}
-		}
-
-		FStateTreeCompilerLog Log;
-		FStateTreeCompiler Compiler(Log);
-		const bool bResult = Compiler.Compile(StateTree);
-		AITEST_TRUE("StateTree should get compiled", bResult);
-
-		FStateTreeInstanceData InstanceData;
-		FTestStateTreeExecutionContext Exec(StateTree, StateTree, InstanceData);
-		const bool bInitSucceeded = Exec.IsValid();
-		AITEST_TRUE("StateTree should init", bInitSucceeded);
-
-		Exec.Start();
-		AITEST_TRUE(*FString::Printf(TEXT("StateTree TaskA should enter state with value %d"), TaskAPropertyFunctionsAmount), Exec.Expect(TaskA.GetName(), *FString::Printf(TEXT("EnterState%d"), TaskAPropertyFunctionsAmount)));
-		AITEST_TRUE(*FString::Printf(TEXT("StateTree TaskB should enter state with value %d"), ParameterPropertyFunctionsAmount), Exec.Expect(TaskB.GetName(), *FString::Printf(TEXT("EnterState%d"), ParameterPropertyFunctionsAmount)));
-		Exec.LogClear();
-
-		Exec.Tick(0.1f);
-		AITEST_TRUE(*FString::Printf(TEXT("StateTree TaskA should tick with value %d"), TaskAPropertyFunctionsAmount), Exec.Expect(TaskA.GetName(), *FString::Printf(TEXT("Tick%d"), TaskAPropertyFunctionsAmount)));
-		AITEST_TRUE(*FString::Printf(TEXT("StateTree TaskB should tick with value %d"), ParameterPropertyFunctionsAmount), Exec.Expect(TaskB.GetName(), *FString::Printf(TEXT("Tick%d"), ParameterPropertyFunctionsAmount)));
-		Exec.LogClear();
-
-		Exec.Stop(EStateTreeRunStatus::Stopped);
-		AITEST_TRUE(*FString::Printf(TEXT("StateTree TaskA should exit state with value %d"), TaskAPropertyFunctionsAmount), Exec.Expect(TaskA.GetName(), *FString::Printf(TEXT("ExitState%d"), TaskAPropertyFunctionsAmount)));
-		AITEST_TRUE(*FString::Printf(TEXT("StateTree TaskB should exit state with value %d"), ParameterPropertyFunctionsAmount), Exec.Expect(TaskB.GetName(), *FString::Printf(TEXT("ExitState%d"), ParameterPropertyFunctionsAmount)));
-		Exec.LogClear();
-
-		Exec.Stop();
-
-		return true;
-	}
-};
-IMPLEMENT_AI_INSTANT_TEST(FStateTreeTest_PropertyFunctions, "System.StateTree.PropertyFunctions");
-
-struct FStateTreeTest_CopyObjects : FAITestBase
-{
-	virtual bool InstantTest() override
-	{
-		FStateTreeCompilerLog Log;
-		FStateTreePropertyBindings Bindings;
-		FStateTreePropertyBindingCompiler BindingCompiler;
-
-		const bool bInitResult = BindingCompiler.Init(Bindings, Log);
-		AITEST_TRUE("Expect init to succeed", bInitResult);
-
-		FStateTreeBindableStructDesc SourceDesc;
-		SourceDesc.Name = FName(TEXT("Source"));
-		SourceDesc.Struct = TBaseStructure<FStateTreeTest_PropertyCopyObjects>::Get();
-		SourceDesc.DataSource = EStateTreeBindableStructSource::Parameter;
-		SourceDesc.DataHandle = FStateTreeDataHandle(EStateTreeDataSourceType::ContextData, 0); // Used as index to SourceViews below.
-		SourceDesc.ID = FGuid::NewGuid();
-
-		FStateTreeBindableStructDesc TargetADesc;
-		TargetADesc.Name = FName(TEXT("TargetA"));
-		TargetADesc.Struct = TBaseStructure<FStateTreeTest_PropertyCopyObjects>::Get();
-		TargetADesc.DataSource = EStateTreeBindableStructSource::Parameter;
-		TargetADesc.ID = FGuid::NewGuid();
-
-		FStateTreeBindableStructDesc TargetBDesc;
-		TargetBDesc.Name = FName(TEXT("TargetB"));
-		TargetBDesc.Struct = TBaseStructure<FStateTreeTest_PropertyCopyObjects>::Get();
-		TargetBDesc.DataSource = EStateTreeBindableStructSource::Parameter;
-		TargetBDesc.ID = FGuid::NewGuid();
-
-		const int32 SourceIndex = BindingCompiler.AddSourceStruct(SourceDesc);
-
-		TArray<FStateTreePropertyPathBinding> PropertyBindings;
-		// One-to-one copy from source to target A
-		PropertyBindings.Add(UE::StateTree::Tests::MakeBinding(SourceDesc.ID, TEXT("Object"), TargetADesc.ID, TEXT("Object")));
-		PropertyBindings.Add(UE::StateTree::Tests::MakeBinding(SourceDesc.ID, TEXT("SoftObject"), TargetADesc.ID, TEXT("SoftObject")));
-		PropertyBindings.Add(UE::StateTree::Tests::MakeBinding(SourceDesc.ID, TEXT("Class"), TargetADesc.ID, TEXT("Class")));
-		PropertyBindings.Add(UE::StateTree::Tests::MakeBinding(SourceDesc.ID, TEXT("SoftClass"), TargetADesc.ID, TEXT("SoftClass")));
-
-		// Cross copy from source to target B
-		PropertyBindings.Add(UE::StateTree::Tests::MakeBinding(SourceDesc.ID, TEXT("SoftObject"), TargetBDesc.ID, TEXT("Object")));
-		PropertyBindings.Add(UE::StateTree::Tests::MakeBinding(SourceDesc.ID, TEXT("Object"), TargetBDesc.ID, TEXT("SoftObject")));
-		PropertyBindings.Add(UE::StateTree::Tests::MakeBinding(SourceDesc.ID, TEXT("SoftClass"), TargetBDesc.ID, TEXT("Class")));
-		PropertyBindings.Add(UE::StateTree::Tests::MakeBinding(SourceDesc.ID, TEXT("Class"), TargetBDesc.ID, TEXT("SoftClass")));
-		
-		int32 TargetACopyBatchIndex = INDEX_NONE;
-		const bool bCompileBatchResultA = BindingCompiler.CompileBatch(TargetADesc, PropertyBindings, FStateTreeIndex16::Invalid, FStateTreeIndex16::Invalid, TargetACopyBatchIndex);
-		AITEST_TRUE("CompileBatchResultA should succeed", bCompileBatchResultA);
-		AITEST_NOT_EQUAL("TargetACopyBatchIndex should not be INDEX_NONE", TargetACopyBatchIndex, (int32)INDEX_NONE);
-
-		int32 TargetBCopyBatchIndex = INDEX_NONE;
-		const bool bCompileBatchResultB = BindingCompiler.CompileBatch(TargetBDesc, PropertyBindings, FStateTreeIndex16::Invalid, FStateTreeIndex16::Invalid, TargetBCopyBatchIndex);
-		AITEST_TRUE("CompileBatchResultB should succeed", bCompileBatchResultB);
-		AITEST_NOT_EQUAL("TargetBCopyBatchIndex should not be INDEX_NONE", TargetBCopyBatchIndex, (int32)INDEX_NONE);
-
-		BindingCompiler.Finalize();
-
-		const bool bResolveResult = Bindings.ResolvePaths();
-		AITEST_TRUE("ResolvePaths should succeed", bResolveResult);
-
-		UStateTreeTest_PropertyObject* ObjectA = NewObject<UStateTreeTest_PropertyObject>();
-		UStateTreeTest_PropertyObject2* ObjectB = NewObject<UStateTreeTest_PropertyObject2>();
-		
-		FStateTreeTest_PropertyCopyObjects Source;
-		Source.Object = ObjectA;
-		Source.SoftObject = ObjectB;
-		Source.Class = UStateTreeTest_PropertyObject::StaticClass();
-		Source.SoftClass = UStateTreeTest_PropertyObject::StaticClass();
-
-		AITEST_TRUE("SourceIndex should be less than max number of source structs.", SourceIndex < Bindings.GetSourceStructNum());
-
-		TArray<FStateTreeDataView> SourceViews;
-		SourceViews.SetNum(Bindings.GetSourceStructNum());
-		SourceViews[SourceIndex] = FStateTreeDataView(FStructView::Make(Source));
-
-		FStateTreeTest_PropertyCopyObjects TargetA;
-		bool bCopyResultA = true;
-		for (const FStateTreePropertyCopy& Copy : Bindings.GetBatchCopies(FStateTreeIndex16(TargetACopyBatchIndex)))
-		{
-			bCopyResultA &= Bindings.CopyProperty(Copy, SourceViews[Copy.SourceDataHandle.GetIndex()], FStructView::Make(TargetA));
-		}
-		AITEST_TRUE("CopyTo should succeed", bCopyResultA);
-
-		AITEST_TRUE("Expect TargetA.Object == Source.Object", TargetA.Object == Source.Object);
-		AITEST_TRUE("Expect TargetA.SoftObject == Source.SoftObject", TargetA.SoftObject == Source.SoftObject);
-		AITEST_TRUE("Expect TargetA.Class == Source.Class", TargetA.Class == Source.Class);
-		AITEST_TRUE("Expect TargetA.SoftClass == Source.SoftClass", TargetA.SoftClass == Source.SoftClass);
+		TStateTreeEditorNode<FTestTask_Stand>& TaskA = Root.AddTask<FTestTask_Stand>(FName(TEXT("TaskA")));
+		TStateTreeEditorNode<FTestTask_Stand>& TaskB = SubTreeState.AddTask<FTestTask_Stand>(FName(TEXT("TaskB")));
+		TStateTreeEditorNode<FTestTask_Stand>& TaskC = SubSubLeaf.AddTask<FTestTask_Stand>(FName(TEXT("TaskC")));
 
-		// Copying to TargetB should not affect TargetA
-		TargetA.Object = nullptr;
+		TaskA.GetNode().TicksToCompletion = 2;
+		TaskB.GetNode().TicksToCompletion = 2;
+		TaskC.GetNode().TicksToCompletion = 1; // The deepest task completes first.
 		
-		FStateTreeTest_PropertyCopyObjects TargetB;
-		bool bCopyResultB = true;
-		for (const FStateTreePropertyCopy& Copy : Bindings.GetBatchCopies(FStateTreeIndex16(TargetBCopyBatchIndex)))
-		{
-			bCopyResultB &= Bindings.CopyProperty(Copy, SourceViews[Copy.SourceDataHandle.GetIndex()], FStructView::Make(TargetB));
-		}
-		AITEST_TRUE("CopyTo should succeed", bCopyResultB);
+		FStateTreeCompilerLog Log;
+		FStateTreeCompiler Compiler(Log);
+		const bool bResult = Compiler.Compile(StateTree);
+		AITEST_TRUE("StateTree should get compiled", bResult);
 
-		AITEST_TRUE("Expect TargetB.Object == Source.SoftObject", TSoftObjectPtr<UObject>(TargetB.Object) == Source.SoftObject);
-		AITEST_TRUE("Expect TargetB.SoftObject == Source.Object", TargetB.SoftObject == TSoftObjectPtr<UObject>(Source.Object));
-		AITEST_TRUE("Expect TargetB.Class == Source.SoftClass", TSoftClassPtr<UObject>(TargetB.Class) == Source.SoftClass);
-		AITEST_TRUE("Expect TargetB.SoftClass == Source.Class", TargetB.SoftClass == TSoftClassPtr<UObject>(Source.Class));
+		EStateTreeRunStatus Status = EStateTreeRunStatus::Unset;
+		FStateTreeInstanceData InstanceData;
+		FTestStateTreeExecutionContext Exec(StateTree, StateTree, InstanceData);
+		const bool bInitSucceeded = Exec.IsValid();
+		AITEST_TRUE("StateTree should init", bInitSucceeded);
 
-		AITEST_TRUE("Expect TargetA.Object == nullptr after copy of TargetB", TargetA.Object == nullptr);
+		const FString TickStr(TEXT("Tick"));
+		const FString EnterStateStr(TEXT("EnterState"));
+		const FString ExitStateStr(TEXT("ExitState"));
+		const FString StateCompletedStr(TEXT("StateCompleted"));
 
-		// Collect ObjectA and ObjectB, soft object paths should still copy ok.
-		ObjectA = nullptr;
-		ObjectB = nullptr;
-		Source.Object = nullptr;
-		CollectGarbage(GARBAGE_COLLECTION_KEEPFLAGS);
+		// Start and enter state
+		Status = Exec.Start();
+		AITEST_TRUE("StateTree Active States should be in Root/Linked/SubTreeState", Exec.ExpectInActiveStates(Root.Name, LinkedState.Name, SubTreeState.Name, SubLinkedState.Name, SubSubTreeState.Name, SubSubLeaf.Name));
+		AITEST_TRUE("TaskA,B,C should enter state", Exec.Expect(TaskA.GetName(), EnterStateStr).Then(TaskB.GetName(), EnterStateStr).Then(TaskC.GetName(), EnterStateStr));
+		AITEST_TRUE("StateTree should be running", Status == EStateTreeRunStatus::Running);
+		Exec.LogClear();
 
-		FStateTreeTest_PropertyCopyObjects TargetC;
-		bool bCopyResultC = true;
-		for (const FStateTreePropertyCopy& Copy : Bindings.GetBatchCopies(FStateTreeIndex16(TargetACopyBatchIndex)))
-		{
-			bCopyResultB &= Bindings.CopyProperty(Copy, SourceViews[Copy.SourceDataHandle.GetIndex()], FStructView::Make(TargetC));
-		}
+		// Subtrees completes, and it completes the whole tree too.
+		// There's no good way to observe this externally. We switch the return along the way to make sure the transition does not happen directly from the leaf to failed.
+		Status = Exec.Tick(0.1f);
+		AITEST_TRUE("StateTree should be Failed", Status == EStateTreeRunStatus::Failed);
+		Exec.LogClear();
 
+		Exec.Stop();
 		
-		AITEST_TRUE("CopyTo should succeed", bCopyResultC);
-		AITEST_TRUE("Expect TargetC.SoftObject == Source.SoftObject after GC", TargetC.SoftObject == Source.SoftObject);
-
 		return true;
 	}
 };
-IMPLEMENT_AI_INSTANT_TEST(FStateTreeTest_CopyObjects, "System.StateTree.CopyObjects");
+IMPLEMENT_AI_INSTANT_TEST(FStateTreeTest_SubTree_CascadedSucceeded, "System.StateTree.SubTree.CascadedSucceeded");
+
 
-struct FStateTreeTest_References : FAITestBase
+struct FStateTreeTest_SharedInstanceData : FStateTreeTestBase
 {
 	virtual bool InstantTest() override
 	{
-		FStateTreeCompilerLog Log;
-		FStateTreePropertyBindings Bindings;
-		FStateTreePropertyBindingCompiler BindingCompiler;
-
-		const bool bInitResult = BindingCompiler.Init(Bindings, Log);
-		AITEST_TRUE("Expect init to succeed", bInitResult);
-
-		FStateTreeBindableStructDesc SourceDesc;
-		SourceDesc.Name = FName(TEXT("Source"));
-		SourceDesc.Struct = TBaseStructure<FStateTreeTest_PropertyRefSourceStruct>::Get();
-		SourceDesc.DataSource = EStateTreeBindableStructSource::Parameter;
-		SourceDesc.DataHandle = FStateTreeDataHandle(EStateTreeDataSourceType::ContextData, 0);
-		SourceDesc.ID = FGuid::NewGuid();
-		BindingCompiler.AddSourceStruct(SourceDesc);
-
-		FStateTreeBindableStructDesc TargetDesc;
-		TargetDesc.Name = FName(TEXT("Target"));
-		TargetDesc.Struct = TBaseStructure<FStateTreeTest_PropertyRefTargetStruct>::Get();
-		TargetDesc.DataSource = EStateTreeBindableStructSource::Parameter;
-		TargetDesc.ID = FGuid::NewGuid();
+		UStateTree& StateTree = NewStateTree();
+		UStateTreeEditorData& EditorData = *Cast<UStateTreeEditorData>(StateTree.EditorData);
 		
-		TArray<FStateTreePropertyPathBinding> PropertyBindings;
-		PropertyBindings.Add(UE::StateTree::Tests::MakeBinding(SourceDesc.ID, TEXT("Item"), TargetDesc.ID, TEXT("RefToStruct")));
-		PropertyBindings.Add(UE::StateTree::Tests::MakeBinding(SourceDesc.ID, TEXT("Item.A"), TargetDesc.ID, TEXT("RefToInt")));
-		PropertyBindings.Add(UE::StateTree::Tests::MakeBinding(SourceDesc.ID, TEXT("Array"), TargetDesc.ID, TEXT("RefToStructArray")));
-
-		FStateTreeTest_PropertyRefSourceStruct Source;
-		FStateTreeDataView SourceView = FStateTreeDataView(FStructView::Make(Source));
-
-		FStateTreeTest_PropertyRefTargetStruct Target;
-		FStateTreeDataView TargetView(FStructView::Make(Target));
-
-		TMap<FGuid, const FStateTreeDataView> IDToStructValue;
-		IDToStructValue.Emplace(SourceDesc.ID, SourceView);
-		IDToStructValue.Emplace(TargetDesc.ID, TargetView);
-
-		const bool bCompileReferencesResult = BindingCompiler.CompileReferences(TargetDesc, PropertyBindings, TargetView, IDToStructValue);
-		AITEST_TRUE("CompileReferences should succeed", bCompileReferencesResult);	
-
-		BindingCompiler.Finalize();
+		UStateTreeState& Root = EditorData.AddSubTree(FName(TEXT("Root")));
+		auto& IntCond = Root.AddEnterCondition<FStateTreeTestCondition>();
+		IntCond.GetInstanceData().Count = 1;
 
-		const bool bResolveResult = Bindings.ResolvePaths();
-		AITEST_TRUE("ResolvePaths should succeed", bResolveResult);
+		auto& Task = Root.AddTask<FTestTask_Stand>(FName(TEXT("Task")));
+		Task.GetNode().TicksToCompletion = 2;
 
-		{
-			const FStateTreePropertyAccess* PropertyAccess = Bindings.GetPropertyAccess(Target.RefToStruct);
-			AITEST_NOT_NULL("GetPropertyAccess should succeed", PropertyAccess);
-			
-			FStateTreeTest_PropertyStruct* Reference = Bindings.GetMutablePropertyPtr<FStateTreeTest_PropertyStruct>(SourceView, *PropertyAccess);
-			AITEST_EQUAL("Expect RefToStruct to point to SourceA.Item", Reference, &Source.Item);
-		}
+		FStateTreeCompilerLog Log;
+		FStateTreeCompiler Compiler(Log);
+		const bool bResult = Compiler.Compile(StateTree);
+		AITEST_TRUE("StateTree should get compiled", bResult);
 
-		{
-			const FStateTreePropertyAccess* PropertyAccess = Bindings.GetPropertyAccess(Target.RefToInt);
-			AITEST_NOT_NULL("GetPropertyAccess should succeed", PropertyAccess);
+		// Init, nothing should access the shared data.
+		constexpr int32 NumConcurrent = 100;
+		FStateTreeTestConditionInstanceData::GlobalCounter = 0;
 
-			int32* Reference = Bindings.GetMutablePropertyPtr<int32>(SourceView, *PropertyAccess);
-			AITEST_EQUAL("Expect RefToInt to point to SourceA.Item.A", Reference, &Source.Item);
-		}
+		bool bInitSucceeded = true;
+		TArray<FStateTreeInstanceData> InstanceDatas;
 
+		InstanceDatas.SetNum(NumConcurrent);
+		for (int32 Index = 0; Index < NumConcurrent; Index++)
 		{
-			const FStateTreePropertyAccess* PropertyAccess = Bindings.GetPropertyAccess(Target.RefToStructArray);
-			AITEST_NOT_NULL("GetPropertyAccess should succeed", PropertyAccess);
-
-			TArray<FStateTreeTest_PropertyStruct>* Reference = Bindings.GetMutablePropertyPtr<TArray<FStateTreeTest_PropertyStruct>>(SourceView, *PropertyAccess);
-			AITEST_EQUAL("Expect RefToStructArray to point to SourceA.Array", Reference, &Source.Array);
+			FTestStateTreeExecutionContext Exec(StateTree, StateTree, InstanceDatas[Index]);
+			bInitSucceeded &= Exec.IsValid();
 		}
+		AITEST_TRUE("All StateTree contexts should init", bInitSucceeded);
+		AITEST_EQUAL("Test condition global counter should be 0", (int32)FStateTreeTestConditionInstanceData::GlobalCounter, 0);
 		
-		return true;
-	}
-};
-IMPLEMENT_AI_INSTANT_TEST(FStateTreeTest_References, "System.StateTree.References");
-
-struct FStateTreeTest_ReferencesConstness : FAITestBase
-{
-	virtual bool InstantTest() override
-	{
-		FStateTreeCompilerLog Log;
-		FStateTreePropertyBindings Bindings;
-		FStateTreePropertyBindingCompiler BindingCompiler;
-
-		const bool bInitResult = BindingCompiler.Init(Bindings, Log);
-		AITEST_TRUE("Expect init to succeed", bInitResult);
-
-		FStateTreeBindableStructDesc SourceAsTaskDesc;
-		SourceAsTaskDesc.Name = FName(TEXT("SourceTask"));
-		SourceAsTaskDesc.Struct = TBaseStructure<FStateTreeTest_PropertyRefSourceStruct>::Get();
-		SourceAsTaskDesc.DataSource = EStateTreeBindableStructSource::Task;
-		SourceAsTaskDesc.DataHandle = FStateTreeDataHandle(EStateTreeDataSourceType::ContextData, 0);
-		SourceAsTaskDesc.ID = FGuid::NewGuid();
-		BindingCompiler.AddSourceStruct(SourceAsTaskDesc);
-
-		FStateTreeBindableStructDesc SourceAsContextDesc;
-		SourceAsContextDesc.Name = FName(TEXT("SourceContext"));
-		SourceAsContextDesc.Struct = TBaseStructure<FStateTreeTest_PropertyRefSourceStruct>::Get();
-		SourceAsContextDesc.DataSource = EStateTreeBindableStructSource::Context;
-		SourceAsContextDesc.DataHandle = FStateTreeDataHandle(EStateTreeDataSourceType::ContextData, 0);
-		SourceAsContextDesc.ID = FGuid::NewGuid();
-		BindingCompiler.AddSourceStruct(SourceAsContextDesc);
-
-		FStateTreeBindableStructDesc TargetDesc;
-		TargetDesc.Name = FName(TEXT("Target"));
-		TargetDesc.Struct = TBaseStructure<FStateTreeTest_PropertyRefTargetStruct>::Get();
-		TargetDesc.DataSource = EStateTreeBindableStructSource::Parameter;
-		TargetDesc.ID = FGuid::NewGuid();
+		// Start in parallel
+		// This should create shared data per thread.
+		// We expect that ParallelForWithTaskContext() creates a context per thread.
+		TArray<FStateTreeTestRunContext> RunContexts;
 		
-		FStateTreePropertyPathBinding TaskPropertyBinding = UE::StateTree::Tests::MakeBinding(SourceAsTaskDesc.ID, TEXT("Item"), TargetDesc.ID, TEXT("RefToStruct"));
-		FStateTreePropertyPathBinding TaskOutputPropertyBinding = UE::StateTree::Tests::MakeBinding(SourceAsTaskDesc.ID, TEXT("OutputItem"), TargetDesc.ID, TEXT("RefToStruct"));
-
-		FStateTreePropertyPathBinding ContextPropertyBinding = UE::StateTree::Tests::MakeBinding(SourceAsTaskDesc.ID, TEXT("Item"), TargetDesc.ID, TEXT("RefToStruct"));
-		FStateTreePropertyPathBinding ContextOutputPropertyBinding = UE::StateTree::Tests::MakeBinding(SourceAsTaskDesc.ID, TEXT("Item"), TargetDesc.ID, TEXT("RefToStruct"));
-
-		FStateTreeTest_PropertyRefSourceStruct SourceAsTask;
-		FStateTreeDataView SourceAsTaskView(FStructView::Make(SourceAsTask));
-
-		FStateTreeTest_PropertyRefSourceStruct SourceAsContext;
-		FStateTreeDataView SourceAsContextView(FStructView::Make(SourceAsContext));
-
-		FStateTreeTest_PropertyRefTargetStruct Target;
-		FStateTreeDataView TargetView(FStructView::Make(Target));
-
-		TMap<FGuid, const FStateTreeDataView> IDToStructValue;
-		IDToStructValue.Emplace(SourceAsTaskDesc.ID, SourceAsTaskView);
-		IDToStructValue.Emplace(SourceAsContextDesc.ID, SourceAsContextView);
-		IDToStructValue.Emplace(TargetDesc.ID, TargetView);
+		ParallelForWithTaskContext(
+			RunContexts,
+			InstanceDatas.Num(),
+			[&InstanceDatas, &StateTree](FStateTreeTestRunContext& RunContext, int32 Index)
+			{
+				FTestStateTreeExecutionContext Exec(StateTree, StateTree, InstanceDatas[Index]);
+				const EStateTreeRunStatus Status = Exec.Start();
+				if (Status == EStateTreeRunStatus::Running)
+				{
+					RunContext.Count++;
+				}
+			}
+		);
 
+		int32 StartTotalRunning = 0;
+		for (FStateTreeTestRunContext RunContext : RunContexts)
 		{
-			const bool bCompileReferenceResult = BindingCompiler.CompileReferences(TargetDesc, {TaskPropertyBinding}, TargetView, IDToStructValue);
-			AITEST_FALSE("CompileReferences should fail", bCompileReferenceResult);
+			StartTotalRunning += RunContext.Count;
 		}
+		AITEST_EQUAL("All StateTree contexts should be running after Start", StartTotalRunning, NumConcurrent);
+		AITEST_EQUAL("Test condition global counter should equal context count after Start", (int32)FStateTreeTestConditionInstanceData::GlobalCounter, InstanceDatas.Num());
 
+		
+		// Tick in parallel
+		// This should not recreate the data, so FStateTreeTestConditionInstanceData::GlobalCounter should stay as is.
+		for (FStateTreeTestRunContext RunContext : RunContexts)
 		{
-			const bool bCompileReferenceResult = BindingCompiler.CompileReferences(TargetDesc, {TaskOutputPropertyBinding}, TargetView, IDToStructValue);
-			AITEST_TRUE("CompileReferences should succeed", bCompileReferenceResult);
+			RunContext.Count = 0;
 		}
 
+		ParallelForWithTaskContext(
+			RunContexts,
+			InstanceDatas.Num(),
+			[&InstanceDatas, &StateTree](FStateTreeTestRunContext& RunContext, int32 Index)
+			{
+				FTestStateTreeExecutionContext Exec(StateTree, StateTree, InstanceDatas[Index]);
+				const EStateTreeRunStatus Status = Exec.Tick(0.1f);
+				if (Status == EStateTreeRunStatus::Running)
+				{
+					RunContext.Count++;
+				}
+			}
+		);
+
+		int32 TickTotalRunning = 0;
+		for (FStateTreeTestRunContext RunContext : RunContexts)
 		{
-			const bool bCompileReferenceResult = BindingCompiler.CompileReferences(TargetDesc, {ContextPropertyBinding}, TargetView, IDToStructValue);
-			AITEST_FALSE("CompileReferences should fail", bCompileReferenceResult);
+			TickTotalRunning += RunContext.Count;
 		}
+		AITEST_EQUAL("All StateTree contexts should be running after Tick", TickTotalRunning, NumConcurrent);
+		AITEST_EQUAL("Test condition global counter should equal context count after Tick", (int32)FStateTreeTestConditionInstanceData::GlobalCounter, InstanceDatas.Num());
 
-		{	
-			const bool bCompileReferenceResult = BindingCompiler.CompileReferences(TargetDesc, {ContextOutputPropertyBinding}, TargetView, IDToStructValue);
-			AITEST_FALSE("CompileReferences should fail", bCompileReferenceResult);
-		}
+		ParallelForWithTaskContext(
+			RunContexts,
+			InstanceDatas.Num(),
+			[&InstanceDatas, &StateTree](FStateTreeTestRunContext& RunContext, int32 Index)
+			{
+				FTestStateTreeExecutionContext Exec(StateTree, StateTree, InstanceDatas[Index]);
+				Exec.Stop();
+			}
+		);
 
 		return true;
 	}
 };
-IMPLEMENT_AI_INSTANT_TEST(FStateTreeTest_ReferencesConstness, "System.StateTree.ReferencesConstness");
+IMPLEMENT_AI_INSTANT_TEST(FStateTreeTest_SharedInstanceData, "System.StateTree.SharedInstanceData");
 
-struct FStateTreeTest_FollowTransitions : FAITestBase
+struct FStateTreeTest_LastConditionWithIndent : FStateTreeTestBase
 {
 	virtual bool InstantTest() override
 	{
-		UStateTree& StateTree = UE::StateTree::Tests::NewStateTree(&GetWorld());
+		UStateTree& StateTree = NewStateTree();
 		UStateTreeEditorData& EditorData = *Cast<UStateTreeEditorData>(StateTree.EditorData);
-
-		EditorData.RootParameters.Parameters.AddProperty(FName(TEXT("Int")), EPropertyBagPropertyType::Int32);
-		EditorData.RootParameters.Parameters.SetValueInt32(FName(TEXT("Int")), 1);
 		
 		UStateTreeState& Root = EditorData.AddSubTree(FName(TEXT("Root")));
-		UStateTreeState& StateTrans = Root.AddChildState(FName(TEXT("Trans")));
-		UStateTreeState& StateA = Root.AddChildState(FName(TEXT("A")));
-		UStateTreeState& StateB = Root.AddChildState(FName(TEXT("B")));
-		UStateTreeState& StateC = Root.AddChildState(FName(TEXT("C")));
-
-		// Root
-
-		// Trans
-		{
-			StateTrans.SelectionBehavior = EStateTreeStateSelectionBehavior::TryFollowTransitions;
-
-			{
-				// This transition should be skipped due to the condition
-				FStateTreeTransition& TransA = StateTrans.AddTransition(EStateTreeTransitionTrigger::OnTick, EStateTreeTransitionType::GotoState, &StateA);
-				TStateTreeEditorNode<FStateTreeCompareIntCondition>& TransIntCond = TransA.AddCondition<FStateTreeCompareIntCondition>(EGenericAICheck::Equal);
-				TransIntCond.GetInstanceData().Right = 0;
-				EditorData.AddPropertyBinding(
-					FStateTreePropertyPath(EditorData.RootParameters.ID, TEXT("Int")),
-					FStateTreePropertyPath(TransIntCond.ID, TEXT("Left")));
-			}
-
-			{
-				// This transition leads to selection, but will be overridden.
-				FStateTreeTransition& TransB = StateTrans.AddTransition(EStateTreeTransitionTrigger::OnTick, EStateTreeTransitionType::GotoState, &StateB);
-				TransB.Priority = EStateTreeTransitionPriority::Normal;
-				TStateTreeEditorNode<FStateTreeCompareIntCondition>& TransIntCond = TransB.AddCondition<FStateTreeCompareIntCondition>(EGenericAICheck::Equal);
-				TransIntCond.GetInstanceData().Right = 1;
-				EditorData.AddPropertyBinding(
-					FStateTreePropertyPath(EditorData.RootParameters.ID, TEXT("Int")),
-					FStateTreePropertyPath(TransIntCond.ID, TEXT("Left")));
-			}
+		UStateTreeState& State1 = Root.AddChildState(FName(TEXT("State1")));
 
-			{
-				// This transition is selected, should override previous one due to priority.
-				FStateTreeTransition& TransC = StateTrans.AddTransition(EStateTreeTransitionTrigger::OnTick, EStateTreeTransitionType::GotoState, &StateC);
-				TransC.Priority = EStateTreeTransitionPriority::High;
-				TStateTreeEditorNode<FStateTreeCompareIntCondition>& TransIntCond = TransC.AddCondition<FStateTreeCompareIntCondition>(EGenericAICheck::Equal);
-				TransIntCond.GetInstanceData().Right = 1;
-				EditorData.AddPropertyBinding(
-					FStateTreePropertyPath(EditorData.RootParameters.ID, TEXT("Int")),
-					FStateTreePropertyPath(TransIntCond.ID, TEXT("Left")));
-			}
-		}
+		auto& Task1 = State1.AddTask<FTestTask_Stand>(FName(TEXT("Task1")));
+		State1.AddEnterCondition<FStateTreeTestCondition>();
+		auto& LastCondition = State1.AddEnterCondition<FStateTreeTestCondition>();
 
-		auto& TaskA = StateA.AddTask<FTestTask_Stand>(FName(TEXT("TaskA")));
-		auto& TaskB = StateB.AddTask<FTestTask_Stand>(FName(TEXT("TaskB")));
-		auto& TaskC = StateC.AddTask<FTestTask_Stand>(FName(TEXT("TaskC")));
+		// Last condition has Indent
+		LastCondition.ExpressionIndent = 1;
+		
+		State1.AddTransition(EStateTreeTransitionTrigger::OnStateCompleted, EStateTreeTransitionType::Succeeded);
 
 		FStateTreeCompilerLog Log;
 		FStateTreeCompiler Compiler(Log);
 		const bool bResult = Compiler.Compile(StateTree);
-
 		AITEST_TRUE("StateTree should get compiled", bResult);
 
 		EStateTreeRunStatus Status = EStateTreeRunStatus::Unset;
@@ -2331,9 +743,17 @@ struct FStateTreeTest_FollowTransitions : FAITestBase
 		const FString ExitStateStr(TEXT("ExitState"));
 
 		Status = Exec.Start();
-		AITEST_FALSE("StateTree TaskA should not enter state", Exec.Expect(TaskA.GetName(), EnterStateStr));
-		AITEST_FALSE("StateTree TaskB should not enter state", Exec.Expect(TaskB.GetName(), EnterStateStr));
-		AITEST_TRUE("StateTree TaskC should enter state", Exec.Expect(TaskC.GetName(), EnterStateStr));
+		AITEST_TRUE("StateTree Task1 should enter state", Exec.Expect(Task1.GetName(), EnterStateStr));
+		AITEST_FALSE("StateTree Task1 should not tick", Exec.Expect(Task1.GetName(), TickStr));
+		Exec.LogClear();
+
+		Status = Exec.Tick(0.1f);
+		AITEST_TRUE("StateTree Task1 should tick, and exit state", Exec.Expect(Task1.GetName(), TickStr).Then(Task1.GetName(), ExitStateStr));
+		AITEST_TRUE("StateTree should be completed", Status == EStateTreeRunStatus::Succeeded);
+		Exec.LogClear();
+
+		Status = Exec.Tick(0.1f);
+		AITEST_FALSE("StateTree Task1 should not tick", Exec.Expect(Task1.GetName(), TickStr));
 		Exec.LogClear();
 
 		Exec.Stop();
@@ -2341,55 +761,61 @@ struct FStateTreeTest_FollowTransitions : FAITestBase
 		return true;
 	}
 };
-IMPLEMENT_AI_INSTANT_TEST(FStateTreeTest_FollowTransitions, "System.StateTree.FollowTransitions");
+IMPLEMENT_AI_INSTANT_TEST(FStateTreeTest_LastConditionWithIndent, "System.StateTree.LastConditionWithIndent");
 
-struct FStateTreeTest_InfiniteLoop : FAITestBase
+struct FStateTreeTest_StateRequiringEvent : FStateTreeTestBase
 {
 	virtual bool InstantTest() override
 	{
-		UStateTree& StateTree = UE::StateTree::Tests::NewStateTree(&GetWorld());
+		UStateTree& StateTree = NewStateTree();
 		UStateTreeEditorData& EditorData = *Cast<UStateTreeEditorData>(StateTree.EditorData);
 
-		EditorData.RootParameters.Parameters.AddProperty(FName(TEXT("Int")), EPropertyBagPropertyType::Int32);
-		EditorData.RootParameters.Parameters.SetValueInt32(FName(TEXT("Int")), 1);
-		
 		UStateTreeState& Root = EditorData.AddSubTree(FName(TEXT("Root")));
-		UStateTreeState& StateA = Root.AddChildState(FName(TEXT("A")));
-		UStateTreeState& StateB = StateA.AddChildState(FName(TEXT("B")));
-
-		// Root
+		
+		FGameplayTag ValidTag = GetTestTag1();
+		FGameplayTag InvalidTag = GetTestTag2();
 
-		// State A
-		{
-			StateA.SelectionBehavior = EStateTreeStateSelectionBehavior::TryFollowTransitions;
-			{
-				// A -> B
-				FStateTreeTransition& Trans = StateA.AddTransition(EStateTreeTransitionTrigger::OnTick, EStateTreeTransitionType::GotoState, &StateB);
-				TStateTreeEditorNode<FStateTreeCompareIntCondition>& TransIntCond = Trans.AddCondition<FStateTreeCompareIntCondition>(EGenericAICheck::Equal);
-				TransIntCond.GetInstanceData().Right = 1;
-				EditorData.AddPropertyBinding(
-					FStateTreePropertyPath(EditorData.RootParameters.ID, TEXT("Int")),
-					FStateTreePropertyPath(TransIntCond.ID, TEXT("Left")));
-			}
-		}
+		using FValidPayload = FStateTreeTest_PropertyStructA;
+		using FInvalidPayload = FStateTreeTest_PropertyStructB;
 
-		// State B
-		{
-			StateB.SelectionBehavior = EStateTreeStateSelectionBehavior::TryFollowTransitions;
-			{
-				// B -> A
-				FStateTreeTransition& Trans = StateB.AddTransition(EStateTreeTransitionTrigger::OnTick, EStateTreeTransitionType::GotoState, &StateA);
-				TStateTreeEditorNode<FStateTreeCompareIntCondition>& TransIntCond = Trans.AddCondition<FStateTreeCompareIntCondition>(EGenericAICheck::Equal);
-				TransIntCond.GetInstanceData().Right = 1;
-				EditorData.AddPropertyBinding(
-					FStateTreePropertyPath(EditorData.RootParameters.ID, TEXT("Int")),
-					FStateTreePropertyPath(TransIntCond.ID, TEXT("Left")));
-			}
-		}
-		
+		// This state shouldn't be selected as it requires different tag.
+		UStateTreeState& StateA = Root.AddChildState(FName(TEXT("A")));
+		StateA.bHasRequiredEventToEnter  = true;
+		StateA.RequiredEventToEnter.Tag = InvalidTag;
 		auto& TaskA = StateA.AddTask<FTestTask_Stand>(FName(TEXT("TaskA")));
+
+		// This state shouldn't be selected as it requires different payload.
+		UStateTreeState& StateB = Root.AddChildState(FName(TEXT("B")));
+		StateB.bHasRequiredEventToEnter  = true;
+		StateB.RequiredEventToEnter.PayloadStruct = FInvalidPayload::StaticStruct();
 		auto& TaskB = StateB.AddTask<FTestTask_Stand>(FName(TEXT("TaskB")));
 
+		// This state shouldn't be selected as it requires the same tag, but different payload.
+		UStateTreeState& StateC = Root.AddChildState(FName(TEXT("C")));
+		StateC.bHasRequiredEventToEnter  = true;
+		StateC.RequiredEventToEnter.Tag = ValidTag;
+		StateC.RequiredEventToEnter.PayloadStruct = FInvalidPayload::StaticStruct();
+		auto& TaskC = StateC.AddTask<FTestTask_Stand>(FName(TEXT("TaskC")));
+
+		// This state shouldn't be selected as it requires the same payload, but different tag.
+		UStateTreeState& StateD = Root.AddChildState(FName(TEXT("D")));
+		StateD.bHasRequiredEventToEnter  = true;
+		StateD.RequiredEventToEnter.Tag = InvalidTag;
+		StateD.RequiredEventToEnter.PayloadStruct = FValidPayload::StaticStruct();
+		auto& TaskD = StateD.AddTask<FTestTask_Stand>(FName(TEXT("TaskD")));
+
+		// This state should be selected as the arrived event matches the requirement.
+		UStateTreeState& StateE = Root.AddChildState(FName(TEXT("E")));
+		StateE.bHasRequiredEventToEnter  = true;
+		StateE.RequiredEventToEnter.Tag = ValidTag;
+		StateE.RequiredEventToEnter.PayloadStruct = FValidPayload::StaticStruct();
+		auto& TaskE = StateE.AddTask<FTestTask_Stand>(FName(TEXT("TaskE")));
+
+		// This state should be selected only initially when there's not event in the queue.
+		UStateTreeState& StateInitial = Root.AddChildState(FName(TEXT("Initial")));
+		auto& TaskInitial = StateInitial.AddTask<FTestTask_Stand>(FName(TEXT("TaskInitial")));
+		StateInitial.AddTransition(EStateTreeTransitionTrigger::OnEvent, ValidTag, EStateTreeTransitionType::GotoState, &Root);
+
 		FStateTreeCompilerLog Log;
 		FStateTreeCompiler Compiler(Log);
 		const bool bResult = Compiler.Compile(StateTree);
@@ -2402,33 +828,37 @@ struct FStateTreeTest_InfiniteLoop : FAITestBase
 		const bool bInitSucceeded = Exec.IsValid();
 		AITEST_TRUE("StateTree should init", bInitSucceeded);
 
-		const FString TickStr(TEXT("Tick"));
 		const FString EnterStateStr(TEXT("EnterState"));
-		const FString ExitStateStr(TEXT("ExitState"));
 
-		GetTestRunner().AddExpectedError(TEXT("Loop detected when trying to select state"), EAutomationExpectedErrorFlags::Contains, 1);
-		GetTestRunner().AddExpectedError(TEXT("Failed to select initial state"), EAutomationExpectedErrorFlags::Contains, 1);
-		
 		Status = Exec.Start();
-		AITEST_EQUAL("Start should fail", Status, EStateTreeRunStatus::Failed);
+		AITEST_TRUE("StateTree TaskInitial should enter state", Exec.Expect(TaskInitial.GetName(), EnterStateStr));
 		Exec.LogClear();
 
+		Exec.SendEvent(ValidTag, FConstStructView::Make(FValidPayload()));
+		Status = Exec.Tick(0.1f);
+
+		AITEST_FALSE("StateTree TaskA should not enter state", Exec.Expect(TaskA.GetName(), EnterStateStr));
+		AITEST_FALSE("StateTree TaskB should not enter state", Exec.Expect(TaskB.GetName(), EnterStateStr));
+		AITEST_FALSE("StateTree TaskC should not enter state", Exec.Expect(TaskC.GetName(), EnterStateStr));
+		AITEST_FALSE("StateTree TaskD should not enter state", Exec.Expect(TaskD.GetName(), EnterStateStr));
+		AITEST_TRUE("StateTree TaskE should enter state", Exec.Expect(TaskE.GetName(), EnterStateStr));
+		Exec.LogClear();
+		
 		Exec.Stop();
 
 		return true;
 	}
 };
-IMPLEMENT_AI_INSTANT_TEST(FStateTreeTest_InfiniteLoop, "System.StateTree.InfiniteLoop");
-
+IMPLEMENT_AI_INSTANT_TEST(FStateTreeTest_StateRequiringEvent, "System.StateTree.StateRequiringEvent");
 
 //
 // The stop tests test how the combinations of execution path to stop the tree are reported on ExitState() transition.  
 //
-struct FStateTreeTest_Stop : FAITestBase
+struct FStateTreeTest_Stop : FStateTreeTestBase
 {
 	UStateTree& SetupTree()
 	{
-		UStateTree& StateTree = UE::StateTree::Tests::NewStateTree(&GetWorld());
+		UStateTree& StateTree = NewStateTree();
 		UStateTreeEditorData& EditorData = *Cast<UStateTreeEditorData>(StateTree.EditorData);
 
 		UStateTreeState& Root = EditorData.AddSubTree(FName(TEXT("Root")));
@@ -2869,11 +1299,11 @@ IMPLEMENT_AI_INSTANT_TEST(FStateTreeTest_Stop_AlreadyStopped, "System.StateTree.
 //
 // The deferred stop tests validates that the tree can be properly stopped if requested in the main entry points (Start, Tick, Stop).  
 //
-struct FStateTreeTest_DeferredStop : FAITestBase
+struct FStateTreeTest_DeferredStop : FStateTreeTestBase
 {
 	UStateTree& SetupTree() const
 	{
-		UStateTree& StateTree = UE::StateTree::Tests::NewStateTree(&GetWorld());
+		UStateTree& StateTree = NewStateTree();
 		UStateTreeEditorData& EditorData = *Cast<UStateTreeEditorData>(StateTree.EditorData);
 
 		UStateTreeState& Root = EditorData.AddSubTree(FName(TEXT("Root")));
@@ -3032,14 +1462,14 @@ struct FStateTreeTest_DeferredStop_ExitTask : FStateTreeTest_DeferredStop
 };
 IMPLEMENT_AI_INSTANT_TEST(FStateTreeTest_DeferredStop_ExitTask, "System.StateTree.DeferredStop.ExitTask");
 
-struct FStateTreeTest_FailEnterLinkedAsset : FAITestBase
+struct FStateTreeTest_FailEnterLinkedAsset : FStateTreeTestBase
 {
 	virtual bool InstantTest() override
 	{
 		FStateTreeCompilerLog Log;
 
 		// Asset 2
-		UStateTree& StateTree2 = UE::StateTree::Tests::NewStateTree(&GetWorld());
+		UStateTree& StateTree2 = NewStateTree();
 		UStateTreeEditorData& EditorData2 = *Cast<UStateTreeEditorData>(StateTree2.EditorData);
 		UStateTreeState& Root2 = EditorData2.AddSubTree(FName(TEXT("Root2")));
 		TStateTreeEditorNode<FTestTask_Stand>& Task2 = Root2.AddTask<FTestTask_Stand>(FName(TEXT("Task2")));
@@ -3056,7 +1486,7 @@ struct FStateTreeTest_FailEnterLinkedAsset : FAITestBase
 		AITEST_TRUE("StateTree2 should get compiled", bResult2);
 
 		// Main asset
-		UStateTree& StateTree = UE::StateTree::Tests::NewStateTree(&GetWorld());
+		UStateTree& StateTree = NewStateTree();
 		UStateTreeEditorData& EditorData = *Cast<UStateTreeEditorData>(StateTree.EditorData);
 		
 		UStateTreeState& Root = EditorData.AddSubTree(FName(TEXT("Root1")));
@@ -3095,14 +1525,14 @@ struct FStateTreeTest_FailEnterLinkedAsset : FAITestBase
 };
 IMPLEMENT_AI_INSTANT_TEST(FStateTreeTest_FailEnterLinkedAsset, "System.StateTree.FailEnterLinkedAsset");
 
-struct FStateTreeTest_EnterAndExitLinkedAsset : FAITestBase
+struct FStateTreeTest_EnterAndExitLinkedAsset : FStateTreeTestBase
 {
 	virtual bool InstantTest() override
 	{
 		FStateTreeCompilerLog Log;
 
 		// Asset 2
-		UStateTree& StateTree2 = UE::StateTree::Tests::NewStateTree(&GetWorld());
+		UStateTree& StateTree2 = NewStateTree();
 		UStateTreeEditorData& EditorData2 = *Cast<UStateTreeEditorData>(StateTree2.EditorData);
 		UStateTreeState& Root2 = EditorData2.AddSubTree(FName(TEXT("Root2")));
 		TStateTreeEditorNode<FTestTask_Stand>& Task2 = Root2.AddTask<FTestTask_Stand>(FName(TEXT("Task2")));
@@ -3114,7 +1544,7 @@ struct FStateTreeTest_EnterAndExitLinkedAsset : FAITestBase
 		AITEST_TRUE("StateTree2 should get compiled", bResult2);
 
 		// Main asset
-		UStateTree& StateTree = UE::StateTree::Tests::NewStateTree(&GetWorld());
+		UStateTree& StateTree = NewStateTree();
 		UStateTreeEditorData& EditorData = *Cast<UStateTreeEditorData>(StateTree.EditorData);
 		
 		UStateTreeState& Root = EditorData.AddSubTree(FName(TEXT("Root1")));
@@ -3166,16 +1596,16 @@ struct FStateTreeTest_EnterAndExitLinkedAsset : FAITestBase
 IMPLEMENT_AI_INSTANT_TEST(FStateTreeTest_EnterAndExitLinkedAsset, "System.StateTree.EnterAndExitLinkedAsset");
 
 // Test nested tree overrides
-struct FStateTreeTest_NestedOverride : FAITestBase
+struct FStateTreeTest_NestedOverride : FStateTreeTestBase
 {
 	virtual bool InstantTest() override
 	{
 		FStateTreeCompilerLog Log;
 		
-		const FGameplayTag Tag = UE::StateTree::Tests::FNativeGameplayTags::Get().TestTag;
+		const FGameplayTag Tag = GetTestTag1();
 
 		// Asset 2
-		UStateTree& StateTree2 = UE::StateTree::Tests::NewStateTree(&GetWorld());
+		UStateTree& StateTree2 = NewStateTree();
 		UStateTreeEditorData& EditorData2 = *Cast<UStateTreeEditorData>(StateTree2.EditorData);
 		EditorData2.RootParameters.Parameters.AddProperty(FName(TEXT("Int")), EPropertyBagPropertyType::Int32);
 		UStateTreeState& Root2 = EditorData2.AddSubTree(FName(TEXT("Root2")));
@@ -3187,7 +1617,7 @@ struct FStateTreeTest_NestedOverride : FAITestBase
 
 		
 		// Asset 3
-		UStateTree& StateTree3 = UE::StateTree::Tests::NewStateTree(&GetWorld());
+		UStateTree& StateTree3 = NewStateTree();
 		UStateTreeEditorData& EditorData3 = *Cast<UStateTreeEditorData>(StateTree3.EditorData);
 		EditorData3.RootParameters.Parameters.AddProperty(FName(TEXT("Float")), EPropertyBagPropertyType::Float); // Different parameters
 		UStateTreeState& Root3 = EditorData3.AddSubTree(FName(TEXT("Root3")));
@@ -3198,7 +1628,7 @@ struct FStateTreeTest_NestedOverride : FAITestBase
 		AITEST_TRUE("StateTree3 should get compiled", bResult3);
 
 		// Main asset
-		UStateTree& StateTree = UE::StateTree::Tests::NewStateTree(&GetWorld());
+		UStateTree& StateTree = NewStateTree();
 		UStateTreeEditorData& EditorData = *Cast<UStateTreeEditorData>(StateTree.EditorData);
 
 		EditorData.RootParameters.Parameters.AddProperty(FName(TEXT("Int")), EPropertyBagPropertyType::Int32);
@@ -3260,7 +1690,7 @@ struct FStateTreeTest_NestedOverride : FAITestBase
 IMPLEMENT_AI_INSTANT_TEST(FStateTreeTest_NestedOverride, "System.StateTree.NestedOverride");
 
 // Test parallel tree event priority handling.
-struct FStateTreeTest_ParallelEventPriority : FAITestBase
+struct FStateTreeTest_ParallelEventPriority : FStateTreeTestBase
 {
 	EStateTreeTransitionPriority ParallelTreePriority = EStateTreeTransitionPriority::Normal;
 	
@@ -3268,13 +1698,13 @@ struct FStateTreeTest_ParallelEventPriority : FAITestBase
 	{
 		FStateTreeCompilerLog Log;
 		
-		const FGameplayTag EventTag = UE::StateTree::Tests::FNativeGameplayTags::Get().TestTag;
+		const FGameplayTag EventTag = GetTestTag1();
 
 		// Parallel tree
 		// - Root
 		//   - State1 ?-> State2
 		//   - State2
-		UStateTree& StateTreePar = UE::StateTree::Tests::NewStateTree(&GetWorld());
+		UStateTree& StateTreePar = NewStateTree();
 		UStateTreeEditorData& EditorDataPar = *Cast<UStateTreeEditorData>(StateTreePar.EditorData);
 
 		UStateTreeState& RootPar = EditorDataPar.AddSubTree(FName(TEXT("Root")));
@@ -3298,7 +1728,7 @@ struct FStateTreeTest_ParallelEventPriority : FAITestBase
 		// - Root [StateTreePar]
 		//   - State3 ?-> State4
 		//   - State4
-		UStateTree& StateTree = UE::StateTree::Tests::NewStateTree(&GetWorld());
+		UStateTree& StateTree = NewStateTree();
 		UStateTreeEditorData& EditorData = *Cast<UStateTreeEditorData>(StateTree.EditorData);
 
 		UStateTreeState& Root = EditorData.AddSubTree(FName(TEXT("Root")));
@@ -3404,11 +1834,11 @@ struct FStateTreeTest_ParallelEventPriority_High : FStateTreeTest_ParallelEventP
 };
 IMPLEMENT_AI_INSTANT_TEST(FStateTreeTest_ParallelEventPriority_High, "System.StateTree.ParallelEventPriority.High");
 
-struct FStateTreeTest_SubTreeTransition : FAITestBase
+struct FStateTreeTest_SubTreeTransition : FStateTreeTestBase
 {
 	virtual bool InstantTest() override
 	{
-		UStateTree& StateTree = UE::StateTree::Tests::NewStateTree(&GetWorld());
+		UStateTree& StateTree = NewStateTree();
 		UStateTreeEditorData& EditorData = *Cast<UStateTreeEditorData>(StateTree.EditorData);
 
 		/*
@@ -3488,6 +1918,8 @@ struct FStateTreeTest_SubTreeTransition : FAITestBase
 };
 IMPLEMENT_AI_INSTANT_TEST(FStateTreeTest_SubTreeTransition, "System.StateTree.SubTreeTransition");
 
+} // namespace UE::StateTree::Tests
+
 UE_ENABLE_OPTIMIZATION_SHIP
 
 #undef LOCTEXT_NAMESPACE
diff --git a/Engine/Plugins/Runtime/StateTree/Source/StateTreeTestSuite/Private/StateTreeTestBase.cpp b/Engine/Plugins/Runtime/StateTree/Source/StateTreeTestSuite/Private/StateTreeTestBase.cpp
new file mode 100644
--- /dev/null
+++ b/Engine/Plugins/Runtime/StateTree/Source/StateTreeTestSuite/Private/StateTreeTestBase.cpp
@@ -0,0 +1,77 @@
+// Copyright Epic Games, Inc. All Rights Reserved.
+
+#include "StateTreeTestBase.h"
+#include "StateTreeTest.h"
+#include "StateTreeEditorData.h"
+#include "Engine/World.h"
+#include "GameplayTagsManager.h"
+
+#define LOCTEXT_NAMESPACE "AITestSuite_StateTreeTest"
+
+namespace UE::StateTree::Tests
+{
+
+UStateTree& FStateTreeTestBase::NewStateTree() const
+{
+	UStateTree* StateTree = NewObject<UStateTree>(&GetWorld());
+	check(StateTree);
+	UStateTreeEditorData* EditorData = NewObject<UStateTreeEditorData>(StateTree);
+	check(EditorData);
+	StateTree->EditorData = EditorData;
+	EditorData->Schema = NewObject<UStateTreeTestSchema>();
+	return *StateTree;
+}
+
+FStateTreePropertyPathBinding FStateTreeTestBase::MakeBinding(const FGuid& SourceID, const FStringView Source, const FGuid& TargetID, const FStringView Target)
+{
+	FStateTreePropertyPath SourcePath;
+	SourcePath.FromString(Source);
+	SourcePath.SetStructID(SourceID);
+
+	FStateTreePropertyPath TargetPath;
+	TargetPath.FromString(Target);
+	TargetPath.SetStructID(TargetID);
+
+	return FStateTreePropertyPathBinding(SourcePath, TargetPath);
+}
+
+// Helper struct to define some test tags
+struct FNativeGameplayTags : public FGameplayTagNativeAdder
+{
+	virtual ~FNativeGameplayTags() {}
+
+	FGameplayTag TestTag;
+	FGameplayTag TestTag2;
+	FGameplayTag TestTag3;
+
+	virtual void AddTags() override
+	{
+		UGameplayTagsManager& Manager = UGameplayTagsManager::Get();
+		TestTag = Manager.AddNativeGameplayTag(TEXT("Test.StateTree.Tag"));
+		TestTag2 = Manager.AddNativeGameplayTag(TEXT("Test.StateTree.Tag2"));
+		TestTag2 = Manager.AddNativeGameplayTag(TEXT("Test.StateTree.Tag3"));
+	}
+};
+static FNativeGameplayTags GameplayTagStaticInstance;
+
+FGameplayTag FStateTreeTestBase::GetTestTag1()
+{
+	return GameplayTagStaticInstance.TestTag;
+}
+
+FGameplayTag FStateTreeTestBase::GetTestTag2()
+{
+	return GameplayTagStaticInstance.TestTag2;
+}
+
+FGameplayTag FStateTreeTestBase::GetTestTag3()
+{
+	return GameplayTagStaticInstance.TestTag3;
+}
+
+}//namespace UE::StateTree::Tests
+
+UE_ENABLE_OPTIMIZATION_SHIP
+
+#undef LOCTEXT_NAMESPACE
+
diff --git a/Engine/Plugins/Runtime/StateTree/Source/StateTreeTestSuite/Private/StateTreeTransitionTest.cpp b/Engine/Plugins/Runtime/StateTree/Source/StateTreeTestSuite/Private/StateTreeTransitionTest.cpp
new file mode 100644
--- /dev/null
+++ b/Engine/Plugins/Runtime/StateTree/Source/StateTreeTestSuite/Private/StateTreeTransitionTest.cpp
@@ -0,0 +1,775 @@
+// Copyright Epic Games, Inc. All Rights Reserved.
+
+#include "StateTreeTest.h"
+#include "StateTreeTestBase.h"
+#include "StateTreeTestTypes.h"
+
+#include "StateTreeCompilerLog.h"
+#include "StateTreeEditorData.h"
+#include "StateTreeCompiler.h"
+#include "Conditions/StateTreeCommonConditions.h"
+
+#define LOCTEXT_NAMESPACE "AITestSuite_StateTreeTest"
+
+UE_DISABLE_OPTIMIZATION_SHIP
+
+namespace UE::StateTree::Tests
+{
+
+struct FStateTreeTest_TransitionPriority : FStateTreeTestBase
+{
+	virtual bool InstantTest() override
+	{
+		UStateTree& StateTree = NewStateTree();
+		UStateTreeEditorData& EditorData = *Cast<UStateTreeEditorData>(StateTree.EditorData);
+
+		/*
+			- Root
+				- State1 : Task1 -> Succeeded
+					- State1A : Task1A -> Next
+					- State1B : Task1B -> Next
+					- State1C : Task1C
+		
+			Task1A completed first, transitioning to State1B.
+			Task1, Task1B, and Task1C complete at the same time, we should take the transition on the first completed state (State1).
+		*/
+		
+		UStateTreeState& Root = EditorData.AddSubTree(FName(TEXT("Root")));
+		UStateTreeState& State1 = Root.AddChildState(FName(TEXT("State1")));
+		UStateTreeState& State1A = State1.AddChildState(FName(TEXT("State1A")));
+		UStateTreeState& State1B = State1.AddChildState(FName(TEXT("State1B")));
+		UStateTreeState& State1C = State1.AddChildState(FName(TEXT("State1C")));
+
+		auto& Task1 = State1.AddTask<FTestTask_Stand>(FName(TEXT("Task1")));
+		Task1.GetNode().TicksToCompletion = 2;
+		State1.AddTransition(EStateTreeTransitionTrigger::OnStateCompleted, EStateTreeTransitionType::Succeeded);
+		
+		auto& Task1A = State1A.AddTask<FTestTask_Stand>(FName(TEXT("Task1A")));
+		Task1A.GetNode().TicksToCompletion = 1;
+		State1A.AddTransition(EStateTreeTransitionTrigger::OnStateCompleted, EStateTreeTransitionType::NextState);
+
+		auto& Task1B = State1B.AddTask<FTestTask_Stand>(FName(TEXT("Task1B")));
+		Task1B.GetNode().TicksToCompletion = 2;
+		State1B.AddTransition(EStateTreeTransitionTrigger::OnStateCompleted, EStateTreeTransitionType::NextState);
+
+		auto& Task1C = State1C.AddTask<FTestTask_Stand>(FName(TEXT("Task1C")));
+		Task1C.GetNode().TicksToCompletion = 2;
+		
+		FStateTreeCompilerLog Log;
+		FStateTreeCompiler Compiler(Log);
+		const bool bResult = Compiler.Compile(StateTree);
+		AITEST_TRUE("StateTree should get compiled", bResult);
+
+		EStateTreeRunStatus Status = EStateTreeRunStatus::Unset;
+		FStateTreeInstanceData InstanceData;
+		FTestStateTreeExecutionContext Exec(StateTree, StateTree, InstanceData);
+		const bool bInitSucceeded = Exec.IsValid();
+		AITEST_TRUE("StateTree should init", bInitSucceeded);
+
+		const FString TickStr(TEXT("Tick"));
+		const FString EnterStateStr(TEXT("EnterState"));
+		const FString ExitStateStr(TEXT("ExitState"));
+		const FString StateCompletedStr(TEXT("StateCompleted"));
+
+		// Start and enter state
+		Status = Exec.Start();
+		AITEST_TRUE("StateTree Task1 should enter state", Exec.Expect(Task1.GetName(), EnterStateStr));
+		AITEST_TRUE("StateTree Task1A should enter state", Exec.Expect(Task1A.GetName(), EnterStateStr));
+		Exec.LogClear();
+
+		// Transition from Task1A to Task1B
+		Status = Exec.Tick(0.1f);
+		AITEST_TRUE("StateTree Task1A should complete", Exec.Expect(Task1A.GetName(), StateCompletedStr));
+		AITEST_TRUE("StateTree Task1B should enter state", Exec.Expect(Task1B.GetName(), EnterStateStr));
+		Exec.LogClear();
+
+		// Task1 completes, and we should take State1 transition. 
+		Status = Exec.Tick(0.1f);
+		AITEST_TRUE("StateTree Task1 should complete", Exec.Expect(Task1.GetName(), StateCompletedStr));
+		AITEST_EQUAL("Tree execution should stop on success", Status, EStateTreeRunStatus::Succeeded);
+		Exec.LogClear();
+
+		Exec.Stop();
+
+		return true;
+	}
+};
+IMPLEMENT_AI_INSTANT_TEST(FStateTreeTest_TransitionPriority, "System.StateTree.Transition.Priority");
+
+struct FStateTreeTest_TransitionPriorityEnterState : FStateTreeTestBase
+{
+	virtual bool InstantTest() override
+	{
+		UStateTree& StateTree = NewStateTree();
+		UStateTreeEditorData& EditorData = *Cast<UStateTreeEditorData>(StateTree.EditorData);
+		
+		UStateTreeState& Root =	EditorData.AddSubTree(FName(TEXT("Root")));
+		UStateTreeState& State0 = Root.AddChildState(FName(TEXT("State0")));
+		UStateTreeState& State1 = Root.AddChildState(FName(TEXT("State1")));
+		UStateTreeState& State1A = State1.AddChildState(FName(TEXT("State1A")));
+		UStateTreeState& State2 = Root.AddChildState(FName(TEXT("State2")));
+		UStateTreeState& State3 = Root.AddChildState(FName(TEXT("State3")));
+
+		auto& Task0 = State0.AddTask<FTestTask_Stand>(FName(TEXT("Task0")));
+		State0.AddTransition(EStateTreeTransitionTrigger::OnStateCompleted, EStateTreeTransitionType::GotoState, &State1);
+
+		auto& Task1 = State1.AddTask<FTestTask_Stand>(FName(TEXT("Task1")));
+		Task1.GetNode().EnterStateResult = EStateTreeRunStatus::Failed;
+		State1.AddTransition(EStateTreeTransitionTrigger::OnStateCompleted, EStateTreeTransitionType::GotoState, &State2);
+		
+		auto& Task1A = State1A.AddTask<FTestTask_Stand>(FName(TEXT("Task1A")));
+		State1A.AddTransition(EStateTreeTransitionTrigger::OnStateCompleted, EStateTreeTransitionType::GotoState, &State3);
+
+		auto& Task2 = State2.AddTask<FTestTask_Stand>(FName(TEXT("Task2")));
+		State2.AddTransition(EStateTreeTransitionTrigger::OnStateCompleted, EStateTreeTransitionType::Succeeded);
+
+		auto& Task3 = State3.AddTask<FTestTask_Stand>(FName(TEXT("Task3")));
+		State3.AddTransition(EStateTreeTransitionTrigger::OnStateCompleted, EStateTreeTransitionType::Succeeded);
+
+		
+		FStateTreeCompilerLog Log;
+		FStateTreeCompiler Compiler(Log);
+		const bool bResult = Compiler.Compile(StateTree);
+		AITEST_TRUE("StateTree should get compiled", bResult);
+
+		EStateTreeRunStatus Status = EStateTreeRunStatus::Unset;
+		FStateTreeInstanceData InstanceData;
+		FTestStateTreeExecutionContext Exec(StateTree, StateTree, InstanceData);
+		const bool bInitSucceeded = Exec.IsValid();
+		AITEST_TRUE("StateTree should init", bInitSucceeded);
+
+		const FString TickStr(TEXT("Tick"));
+		const FString EnterStateStr(TEXT("EnterState"));
+		const FString ExitStateStr(TEXT("ExitState"));
+		const FString StateCompletedStr(TEXT("StateCompleted"));
+
+		// Start and enter state
+		Status = Exec.Start();
+		AITEST_TRUE("StateTree Task0 should enter state", Exec.Expect(Task0.GetName(), EnterStateStr));
+		Exec.LogClear();
+
+		// Transition from State0 to State1, it should fail (Task1), and the transition on State1->State2 (and not State1A->State3)
+		Status = Exec.Tick(0.1f);
+		AITEST_TRUE("StateTree Task0 should complete", Exec.Expect(Task0.GetName(), StateCompletedStr));
+		AITEST_TRUE("StateTree Task2 should enter state", Exec.Expect(Task2.GetName(), EnterStateStr));
+		AITEST_FALSE("StateTree Task3 should not enter state", Exec.Expect(Task3.GetName(), EnterStateStr));
+		Exec.LogClear();
+
+		Exec.Stop();
+
+		return true;
+	}
+};
+IMPLEMENT_AI_INSTANT_TEST(FStateTreeTest_TransitionPriorityEnterState, "System.StateTree.Transition.PriorityEnterState");
+
+struct FStateTreeTest_TransitionNextSelectableState : FStateTreeTestBase
+{
+	virtual bool InstantTest() override
+	{
+		UStateTree& StateTree = NewStateTree();
+		UStateTreeEditorData& EditorData = *Cast<UStateTreeEditorData>(StateTree.EditorData);
+
+		UStateTreeState& Root =	EditorData.AddSubTree(FName(TEXT("Root")));
+		UStateTreeState& State0 = Root.AddChildState(FName(TEXT("State0")));
+		UStateTreeState& State1 = Root.AddChildState(FName(TEXT("State1")));
+		UStateTreeState& State2 = Root.AddChildState(FName(TEXT("State2")));
+
+		auto& EvalA = EditorData.AddEvaluator<FTestEval_A>();
+		EvalA.GetInstanceData().bBoolA = true;
+
+		auto& Task0 = State0.AddTask<FTestTask_Stand>(FName(TEXT("Task0")));
+		State0.AddTransition(EStateTreeTransitionTrigger::OnStateCompleted, EStateTreeTransitionType::NextSelectableState);
+
+		// Add Task 1 with Condition that will always fail
+		auto& Task1 = State1.AddTask<FTestTask_Stand>(FName(TEXT("Task1")));
+		auto& BoolCond1 = State1.AddEnterCondition<FStateTreeCompareBoolCondition>();
+
+		EditorData.AddPropertyBinding(EvalA, TEXT("bBoolA"), BoolCond1, TEXT("bLeft"));
+		BoolCond1.GetInstanceData().bRight = !EvalA.GetInstanceData().bBoolA;
+
+		// Add Task 2 with Condition that will always succeed
+		auto& Task2 = State2.AddTask<FTestTask_Stand>(FName(TEXT("Task2")));
+		auto& BoolCond2 = State2.AddEnterCondition<FStateTreeCompareBoolCondition>();
+		State2.AddTransition(EStateTreeTransitionTrigger::OnStateCompleted, EStateTreeTransitionType::Succeeded);
+
+		EditorData.AddPropertyBinding(EvalA, TEXT("bBoolA"), BoolCond2, TEXT("bLeft"));
+		BoolCond2.GetInstanceData().bRight = EvalA.GetInstanceData().bBoolA;
+
+		FStateTreeCompilerLog Log;
+		FStateTreeCompiler Compiler(Log);
+		const bool bResult = Compiler.Compile(StateTree);
+		AITEST_TRUE("StateTree should get compiled", bResult);
+
+		FStateTreeInstanceData InstanceData;
+		FTestStateTreeExecutionContext Exec(StateTree, StateTree, InstanceData);
+		const bool bInitSucceeded = Exec.IsValid();
+		AITEST_TRUE("StateTree should init", bInitSucceeded);
+
+		const FString TickStr(TEXT("Tick"));
+		const FString EnterStateStr(TEXT("EnterState"));
+		const FString ExitStateStr(TEXT("ExitState"));
+		const FString StateCompletedStr(TEXT("StateCompleted"));
+
+		// Start and enter state
+		Exec.Start();
+		AITEST_TRUE("StateTree Task0 should enter state", Exec.Expect(Task0.GetName(), EnterStateStr));
+		Exec.LogClear();
+
+		// Transition from State0 and tries to select State1. It should fail (Task1) and because transition is set to "Next Selectable", it should now select Task 2 and Enter State
+		Exec.Tick(0.1f);
+		AITEST_TRUE("StateTree Task0 should complete", Exec.Expect(Task0.GetName(), StateCompletedStr));
+		AITEST_FALSE("StateTree Task1 should not enter state", Exec.Expect(Task1.GetName(), EnterStateStr));
+		AITEST_TRUE("StateTree Task2 should enter state", Exec.Expect(Task2.GetName(), EnterStateStr));
+		Exec.LogClear();
+
+		// Complete Task2
+		Exec.Tick(0.1f);
+		AITEST_TRUE("StateTree Task2 should complete", Exec.Expect(Task2.GetName(), StateCompletedStr));
+		Exec.LogClear();
+
+		Exec.Stop();
+
+		return true;
+	}
+};
+IMPLEMENT_AI_INSTANT_TEST(FStateTreeTest_TransitionNextSelectableState, "System.StateTree.Transition.NextSelectableState");
+
+struct FStateTreeTest_TransitionNextWithParentData : FStateTreeTestBase
+{
+	virtual bool InstantTest() override
+	{
+		UStateTree& StateTree = NewStateTree();
+		UStateTreeEditorData& EditorData = *Cast<UStateTreeEditorData>(StateTree.EditorData);
+
+		UStateTreeState& Root =	EditorData.AddSubTree(FName(TEXT("Root")));
+		UStateTreeState& State0 = Root.AddChildState(FName(TEXT("State0")));
+		UStateTreeState& State1 = Root.AddChildState(FName(TEXT("State1")));
+		UStateTreeState& State1A = State1.AddChildState(FName(TEXT("State1A")));
+
+		auto& RootTask = Root.AddTask<FTestTask_B>(FName(TEXT("RootTask")));
+		RootTask.GetInstanceData().bBoolB = true;
+
+		auto& Task0 = State0.AddTask<FTestTask_Stand>(FName(TEXT("Task0")));
+		State0.AddTransition(EStateTreeTransitionTrigger::OnStateCompleted, EStateTreeTransitionType::NextState);
+
+		auto& Task1A = State1A.AddTask<FTestTask_Stand>(FName(TEXT("Task1A")));
+		auto& BoolCond1 = State1A.AddEnterCondition<FStateTreeCompareBoolCondition>();
+
+		EditorData.AddPropertyBinding(RootTask, TEXT("bBoolB"), BoolCond1, TEXT("bLeft"));
+		BoolCond1.GetInstanceData().bRight = true;
+
+		FStateTreeCompilerLog Log;
+		FStateTreeCompiler Compiler(Log);
+		const bool bResult = Compiler.Compile(StateTree);
+		AITEST_TRUE("StateTree should get compiled", bResult);
+
+		FStateTreeInstanceData InstanceData;
+		FTestStateTreeExecutionContext Exec(StateTree, StateTree, InstanceData);
+		const bool bInitSucceeded = Exec.IsValid();
+		AITEST_TRUE("StateTree should init", bInitSucceeded);
+
+		const FString TickStr(TEXT("Tick"));
+		const FString EnterStateStr(TEXT("EnterState"));
+		const FString ExitStateStr(TEXT("ExitState"));
+		const FString StateCompletedStr(TEXT("StateCompleted"));
+
+		// Start and enter state
+		Exec.Start();
+		AITEST_TRUE("StateTree Task0 should enter state", Exec.Expect(Task0.GetName(), EnterStateStr));
+		Exec.LogClear();
+
+		// Transition from State0 and tries to select State1.
+		// This tests that data from current shared active states (Root) is available during state selection.
+		Exec.Tick(0.1f);
+		AITEST_TRUE("StateTree Task0 should complete", Exec.Expect(Task0.GetName(), StateCompletedStr));
+		AITEST_TRUE("StateTree Task1A should enter state", Exec.Expect(Task1A.GetName(), EnterStateStr));
+		Exec.LogClear();
+
+		Exec.Stop();
+
+		return true;
+	}
+};
+IMPLEMENT_AI_INSTANT_TEST(FStateTreeTest_TransitionNextWithParentData, "System.StateTree.Transition.NextWithParentData");
+
+struct FStateTreeTest_TransitionGlobalDataView : FStateTreeTestBase
+{
+	// Tests that the global eval and task dataviews are kept up to date when transitioning from  
+	virtual bool InstantTest() override
+	{
+		UStateTree& StateTree = NewStateTree();
+		UStateTreeEditorData& EditorData = *Cast<UStateTreeEditorData>(StateTree.EditorData);
+
+		UStateTreeState& Root = EditorData.AddSubTree(FName(TEXT("Root")));
+		UStateTreeState& StateA = Root.AddChildState(FName(TEXT("A")));
+		UStateTreeState& StateB = Root.AddChildState(FName(TEXT("B")));
+
+		auto& EvalA = EditorData.AddEvaluator<FTestEval_A>(FName(TEXT("Eval")));
+		EvalA.GetInstanceData().IntA = 42;
+		auto& GlobalTask = EditorData.AddGlobalTask<FTestTask_PrintValue>(FName(TEXT("Global")));
+		GlobalTask.GetInstanceData().Value = 123;
+		
+		// State A
+		auto& Task0 = StateA.AddTask<FTestTask_Stand>(FName(TEXT("Task0")));
+		StateA.AddTransition(EStateTreeTransitionTrigger::OnStateCompleted, EStateTreeTransitionType::GotoState, &StateB);
+
+		// State B
+		auto& Task1 = StateB.AddTask<FTestTask_PrintValue>(FName(TEXT("Task1")));
+		EditorData.AddPropertyBinding(EvalA, TEXT("IntA"), Task1, TEXT("Value"));
+		auto& Task2 = StateB.AddTask<FTestTask_PrintValue>(FName(TEXT("Task2")));
+		EditorData.AddPropertyBinding(GlobalTask, TEXT("Value"), Task2, TEXT("Value"));
+
+		FStateTreeCompilerLog Log;
+		FStateTreeCompiler Compiler(Log);
+		const bool bResult = Compiler.Compile(StateTree);
+
+		AITEST_TRUE("StateTree should get compiled", bResult);
+
+		EStateTreeRunStatus Status = EStateTreeRunStatus::Unset;
+		FStateTreeInstanceData InstanceData;
+		FTestStateTreeExecutionContext Exec(StateTree, StateTree, InstanceData);
+		const bool bInitSucceeded = Exec.IsValid();
+		AITEST_TRUE("StateTree should init", bInitSucceeded);
+
+		const FString EnterStateStr(TEXT("EnterState"));
+		const FString EnterState42Str(TEXT("EnterState42"));
+		const FString EnterState123Str(TEXT("EnterState123"));
+
+		// Start and enter state
+		Status = Exec.Start();
+		AITEST_TRUE("StateTree Task0 should enter state", Exec.Expect(Task0.GetName(), EnterStateStr));
+		Exec.LogClear();
+
+		// Transition from StateA to StateB, Task0 should enter state with evaluator value copied.
+		Status = Exec.Tick(0.1f);
+		AITEST_TRUE("StateTree Task0 should enter state with value 42", Exec.Expect(Task1.GetName(), EnterState42Str));
+		AITEST_TRUE("StateTree Task1 should enter state with value 123", Exec.Expect(Task2.GetName(), EnterState123Str));
+		Exec.LogClear();
+
+		Exec.Stop();
+
+		return true;
+	}
+};
+IMPLEMENT_AI_INSTANT_TEST(FStateTreeTest_TransitionGlobalDataView, "System.StateTree.Transition.GlobalDataView");
+
+struct FStateTreeTest_TransitionDelay : FStateTreeTestBase
+{
+	virtual bool InstantTest() override
+	{
+		UStateTree& StateTree = NewStateTree();
+		UStateTreeEditorData& EditorData = *Cast<UStateTreeEditorData>(StateTree.EditorData);
+		const FGameplayTag Tag = GetTestTag1();
+
+		UStateTreeState& Root = EditorData.AddSubTree(FName(TEXT("Root")));
+		UStateTreeState& StateA = Root.AddChildState(FName(TEXT("A")));
+		UStateTreeState& StateB = Root.AddChildState(FName(TEXT("B")));
+
+		// State A
+		auto& Task0 = StateA.AddTask<FTestTask_Stand>(FName(TEXT("Task0")));
+		Task0.GetNode().TicksToCompletion = 100;
+		
+		FStateTreeTransition& Transition = StateA.AddTransition(EStateTreeTransitionTrigger::OnEvent, EStateTreeTransitionType::GotoState, &StateB);
+		Transition.bDelayTransition = true;
+		Transition.DelayDuration = 0.15f;
+		Transition.DelayRandomVariance = 0.0f;
+		Transition.RequiredEvent.Tag = Tag;
+
+		// State B
+		auto& Task1 = StateB.AddTask<FTestTask_Stand>(FName(TEXT("Task1")));
+		Task1.GetNode().TicksToCompletion = 100;
+
+		FStateTreeCompilerLog Log;
+		FStateTreeCompiler Compiler(Log);
+		const bool bResult = Compiler.Compile(StateTree);
+
+		AITEST_TRUE("StateTree should get compiled", bResult);
+
+		EStateTreeRunStatus Status = EStateTreeRunStatus::Unset;
+		FStateTreeInstanceData InstanceData;
+		FTestStateTreeExecutionContext Exec(StateTree, StateTree, InstanceData);
+		const bool bInitSucceeded = Exec.IsValid();
+		AITEST_TRUE("StateTree should init", bInitSucceeded);
+
+		const FString TickStr(TEXT("Tick"));
+		const FString EnterStateStr(TEXT("EnterState"));
+		const FString ExitStateStr(TEXT("ExitState"));
+		const FString StateCompletedStr(TEXT("StateCompleted"));
+
+		// Start and enter state
+		Status = Exec.Start();
+		AITEST_TRUE("StateTree Task0 should enter state", Exec.Expect(Task0.GetName(), EnterStateStr));
+		Exec.LogClear();
+
+		// This should cause delayed transition.
+		Exec.SendEvent(Tag);
+		
+		Status = Exec.Tick(0.1f);
+		AITEST_TRUE("StateTree Task0 should tick", Exec.Expect(Task0.GetName(), TickStr));
+		Exec.LogClear();
+
+		// Should have execution frames
+		AITEST_TRUE("Should have active frames", InstanceData.GetExecutionState()->ActiveFrames.Num() > 0);
+
+		// Should have delayed transitions
+		const int32 NumDelayedTransitions0 = InstanceData.GetExecutionState()->DelayedTransitions.Num();
+		AITEST_EQUAL("Should have a delayed transition", NumDelayedTransitions0, 1);
+
+		// Tick and expect a delayed transition. 
+		Status = Exec.Tick(0.1f);
+		AITEST_TRUE("StateTree Task0 should tick", Exec.Expect(Task0.GetName(), TickStr));
+		Exec.LogClear();
+
+		const int32 NumDelayedTransitions1 = InstanceData.GetExecutionState()->DelayedTransitions.Num();
+		AITEST_EQUAL("Should have a delayed transition", NumDelayedTransitions1, 1);
+
+		// Should complete delayed transition.
+		Status = Exec.Tick(0.1f);
+		AITEST_TRUE("StateTree Task0 should exit state", Exec.Expect(Task0.GetName(), ExitStateStr));
+		AITEST_TRUE("StateTree Task1 should enter state", Exec.Expect(Task1.GetName(), EnterStateStr));
+		Exec.LogClear();
+
+		Exec.Stop();
+
+		return true;
+	}
+};
+IMPLEMENT_AI_INSTANT_TEST(FStateTreeTest_TransitionDelay, "System.StateTree.Transition.Delay");
+
+struct FStateTreeTest_TransitionDelayZero : FStateTreeTestBase
+{
+	virtual bool InstantTest() override
+	{
+		UStateTree& StateTree = NewStateTree();
+		UStateTreeEditorData& EditorData = *Cast<UStateTreeEditorData>(StateTree.EditorData);
+		const FGameplayTag Tag = GetTestTag1();
+
+		UStateTreeState& Root = EditorData.AddSubTree(FName(TEXT("Root")));
+		UStateTreeState& StateA = Root.AddChildState(FName(TEXT("A")));
+		UStateTreeState& StateB = Root.AddChildState(FName(TEXT("B")));
+
+		// State A
+		auto& Task0 = StateA.AddTask<FTestTask_Stand>(FName(TEXT("Task0")));
+		Task0.GetNode().TicksToCompletion = 100;
+		
+		FStateTreeTransition& Transition = StateA.AddTransition(EStateTreeTransitionTrigger::OnEvent, EStateTreeTransitionType::GotoState, &StateB);
+		Transition.bDelayTransition = true;
+		Transition.DelayDuration = 0.0f;
+		Transition.DelayRandomVariance = 0.0f;
+		Transition.RequiredEvent.Tag = Tag;
+
+		// State B
+		auto& Task1 = StateB.AddTask<FTestTask_Stand>(FName(TEXT("Task1")));
+		Task1.GetNode().TicksToCompletion = 100;
+
+		FStateTreeCompilerLog Log;
+		FStateTreeCompiler Compiler(Log);
+		const bool bResult = Compiler.Compile(StateTree);
+
+		AITEST_TRUE("StateTree should get compiled", bResult);
+
+		EStateTreeRunStatus Status = EStateTreeRunStatus::Unset;
+		FStateTreeInstanceData InstanceData;
+		FTestStateTreeExecutionContext Exec(StateTree, StateTree, InstanceData);
+		const bool bInitSucceeded = Exec.IsValid();
+		AITEST_TRUE("StateTree should init", bInitSucceeded);
+
+		const FString TickStr(TEXT("Tick"));
+		const FString EnterStateStr(TEXT("EnterState"));
+		const FString ExitStateStr(TEXT("ExitState"));
+		const FString StateCompletedStr(TEXT("StateCompleted"));
+
+		// Start and enter state
+		Status = Exec.Start();
+		AITEST_TRUE("StateTree Task0 should enter state", Exec.Expect(Task0.GetName(), EnterStateStr));
+		Exec.LogClear();
+
+		// This should cause delayed transition. Because the time is 0, it should happen immediately.
+		Exec.SendEvent(Tag);
+		
+		Status = Exec.Tick(0.1f);
+		AITEST_TRUE("StateTree Task0 should exit state", Exec.Expect(Task0.GetName(), ExitStateStr));
+		AITEST_TRUE("StateTree Task1 should enter state", Exec.Expect(Task1.GetName(), EnterStateStr));
+		Exec.LogClear();
+
+		Exec.Stop();
+
+		return true;
+	}
+};
+IMPLEMENT_AI_INSTANT_TEST(FStateTreeTest_TransitionDelayZero, "System.StateTree.Transition.DelayZero");
+
+struct FStateTreeTest_PassingTransitionEventToStateSelection : FStateTreeTestBase
+{
+	virtual bool InstantTest() override
+	{
+		UStateTree& StateTree = NewStateTree();
+		UStateTreeEditorData& EditorData = *Cast<UStateTreeEditorData>(StateTree.EditorData);
+
+		UStateTreeState& Root = EditorData.AddSubTree(FName(TEXT("Root")));
+
+		FStateTreePropertyPath PathToPayloadMember;
+		{
+			const bool bParseResult = PathToPayloadMember.FromString(TEXT("Payload.A"));
+
+			AITEST_TRUE("Parsing path should succeeed", bParseResult);
+
+			FStateTreeEvent EventWithPayload;
+			EventWithPayload.Payload = FInstancedStruct::Make<FStateTreeTest_PropertyStructA>();
+			const bool bUpdateSegments = PathToPayloadMember.UpdateSegmentsFromValue(FStateTreeDataView(FStructView::Make(EventWithPayload)));
+			AITEST_TRUE("Updating segments should succeeed", bUpdateSegments);
+		}
+
+		// This state shouldn't be selected, because transition's condition and state's enter condition exlude each other.
+		UStateTreeState& StateA = Root.AddChildState(FName(TEXT("A")));
+		StateA.bHasRequiredEventToEnter  = true;
+		StateA.RequiredEventToEnter.PayloadStruct = FStateTreeTest_PropertyStructA::StaticStruct();
+		auto& TaskA = StateA.AddTask<FTestTask_Stand>(FName(TEXT("TaskA")));
+		TStateTreeEditorNode<FStateTreeCompareIntCondition>& AIntCond = StateA.AddEnterCondition<FStateTreeCompareIntCondition>(EGenericAICheck::Equal);
+		AIntCond.GetInstanceData().Right = 0;
+		EditorData.AddPropertyBinding(
+			FStateTreePropertyPath(StateA.GetEventID(), PathToPayloadMember.GetSegments()),
+			FStateTreePropertyPath(AIntCond.ID, TEXT("Left")));
+
+		// This state should be selected as the sent event fullfils both transition's condition and state's enter condition.
+		UStateTreeState& StateB = Root.AddChildState(FName(TEXT("B")));
+		StateB.bHasRequiredEventToEnter  = true;
+		StateB.RequiredEventToEnter.PayloadStruct = FStateTreeTest_PropertyStructA::StaticStruct();
+		auto& TaskB = StateB.AddTask<FTestTask_PrintValue>(FName(TEXT("TaskB")));
+		// Test copying data from the state event. The condition properties are copied from temp instance data during selection, this gets copied from active instance data.
+		TaskB.GetInstanceData().Value = -1; // Initially -1, expected to be overridden by property binding below. 
+		EditorData.AddPropertyBinding(
+			FStateTreePropertyPath(StateB.GetEventID(), PathToPayloadMember.GetSegments()),
+			FStateTreePropertyPath(TaskB.ID, TEXT("Value")));
+		
+		TStateTreeEditorNode<FStateTreeCompareIntCondition>& BIntCond = StateB.AddEnterCondition<FStateTreeCompareIntCondition>(EGenericAICheck::Equal);
+		BIntCond.GetInstanceData().Right = 1;
+		EditorData.AddPropertyBinding(
+			FStateTreePropertyPath(StateB.GetEventID(), PathToPayloadMember.GetSegments()),
+			FStateTreePropertyPath(BIntCond.ID, TEXT("Left")));
+
+		// This state should be selected only initially when there's not event in the queue.
+		UStateTreeState& StateInitial = Root.AddChildState(FName(TEXT("Initial")));
+		auto& TaskInitial = StateInitial.AddTask<FTestTask_Stand>(FName(TEXT("TaskInitial")));
+		// Transition from Initial -> StateA
+		FStateTreeTransition& TransA = StateInitial.AddTransition(EStateTreeTransitionTrigger::OnEvent, FGameplayTag(), EStateTreeTransitionType::GotoState, &StateA);
+		TransA.RequiredEvent.PayloadStruct = FStateTreeTest_PropertyStructA::StaticStruct();
+		TStateTreeEditorNode<FStateTreeCompareIntCondition>& TransAIntCond = TransA.AddCondition<FStateTreeCompareIntCondition>(EGenericAICheck::Equal);
+		TransAIntCond.GetInstanceData().Right = 1;
+		EditorData.AddPropertyBinding(
+			FStateTreePropertyPath(TransA.GetEventID(), PathToPayloadMember.GetSegments()),
+			FStateTreePropertyPath(TransAIntCond.ID, TEXT("Left")));
+		// Transition from Initial -> StateB
+		FStateTreeTransition& TransB = StateInitial.AddTransition(EStateTreeTransitionTrigger::OnEvent, FGameplayTag(), EStateTreeTransitionType::GotoState, &StateB);
+		TransB.RequiredEvent.PayloadStruct = FStateTreeTest_PropertyStructA::StaticStruct();
+		TStateTreeEditorNode<FStateTreeCompareIntCondition>& TransBIntCond = TransB.AddCondition<FStateTreeCompareIntCondition>(EGenericAICheck::Equal);
+		TransBIntCond.GetInstanceData().Right = 1;
+		EditorData.AddPropertyBinding(
+			FStateTreePropertyPath(TransB.GetEventID(), PathToPayloadMember.GetSegments()),
+			FStateTreePropertyPath(TransBIntCond.ID, TEXT("Left")));
+
+		FStateTreeCompilerLog Log;
+		FStateTreeCompiler Compiler(Log);
+		const bool bResult = Compiler.Compile(StateTree);
+
+		AITEST_TRUE("StateTree should get compiled", bResult);
+
+		EStateTreeRunStatus Status = EStateTreeRunStatus::Unset;
+		FStateTreeInstanceData InstanceData;
+		FTestStateTreeExecutionContext Exec(StateTree, StateTree, InstanceData);
+		const bool bInitSucceeded = Exec.IsValid();
+		AITEST_TRUE("StateTree should init", bInitSucceeded);
+
+		const FString EnterStateStr(TEXT("EnterState"));
+
+		Status = Exec.Start();
+		AITEST_TRUE("StateTree TaskInitial should enter state", Exec.Expect(TaskInitial.GetName(), EnterStateStr));
+		Exec.LogClear();
+
+		// The conditions test for payload Value=1, the first event should not trigger transition. 
+		Exec.SendEvent(GetTestTag1(), FConstStructView::Make(FStateTreeTest_PropertyStructA{0}));
+		Exec.SendEvent(GetTestTag1(), FConstStructView::Make(FStateTreeTest_PropertyStructA{1}));
+		Status = Exec.Tick(0.1f);
+
+		AITEST_FALSE("StateTree TaskA should not enter state", Exec.Expect(TaskA.GetName(), EnterStateStr));
+		AITEST_TRUE("StateTree TaskB should enter state", Exec.Expect(TaskB.GetName(), TEXT("EnterState1"))); // TaskB decorates "EnterState" with value from the payload.
+		Exec.LogClear();
+
+		Exec.Stop();
+
+		return true;
+	}
+};
+IMPLEMENT_AI_INSTANT_TEST(FStateTreeTest_PassingTransitionEventToStateSelection, "System.StateTree.Transition.PassingTransitionEventToStateSelection");
+
+struct FStateTreeTest_FollowTransitions : FStateTreeTestBase
+{
+	virtual bool InstantTest() override
+	{
+		UStateTree& StateTree = NewStateTree();
+		UStateTreeEditorData& EditorData = *Cast<UStateTreeEditorData>(StateTree.EditorData);
+
+		EditorData.RootParameters.Parameters.AddProperty(FName(TEXT("Int")), EPropertyBagPropertyType::Int32);
+		EditorData.RootParameters.Parameters.SetValueInt32(FName(TEXT("Int")), 1);
+
+		UStateTreeState& Root = EditorData.AddSubTree(FName(TEXT("Root")));
+		UStateTreeState& StateTrans = Root.AddChildState(FName(TEXT("Trans")));
+		UStateTreeState& StateA = Root.AddChildState(FName(TEXT("A")));
+		UStateTreeState& StateB = Root.AddChildState(FName(TEXT("B")));
+		UStateTreeState& StateC = Root.AddChildState(FName(TEXT("C")));
+
+		// Root
+
+		// Trans
+		{
+			StateTrans.SelectionBehavior = EStateTreeStateSelectionBehavior::TryFollowTransitions;
+
+			{
+				// This transition should be skipped due to the condition
+				FStateTreeTransition& TransA = StateTrans.AddTransition(EStateTreeTransitionTrigger::OnTick, EStateTreeTransitionType::GotoState, &StateA);
+				TStateTreeEditorNode<FStateTreeCompareIntCondition>& TransIntCond = TransA.AddCondition<FStateTreeCompareIntCondition>(EGenericAICheck::Equal);
+				TransIntCond.GetInstanceData().Right = 0;
+				EditorData.AddPropertyBinding(
+					FStateTreePropertyPath(EditorData.RootParameters.ID, TEXT("Int")),
+					FStateTreePropertyPath(TransIntCond.ID, TEXT("Left")));
+			}
+
+			{
+				// This transition leads to selection, but will be overridden.
+				FStateTreeTransition& TransB = StateTrans.AddTransition(EStateTreeTransitionTrigger::OnTick, EStateTreeTransitionType::GotoState, &StateB);
+				TransB.Priority = EStateTreeTransitionPriority::Normal;
+				TStateTreeEditorNode<FStateTreeCompareIntCondition>& TransIntCond = TransB.AddCondition<FStateTreeCompareIntCondition>(EGenericAICheck::Equal);
+				TransIntCond.GetInstanceData().Right = 1;
+				EditorData.AddPropertyBinding(
+					FStateTreePropertyPath(EditorData.RootParameters.ID, TEXT("Int")),
+					FStateTreePropertyPath(TransIntCond.ID, TEXT("Left")));
+			}
+
+			{
+				// This transition is selected, should override previous one due to priority.
+				FStateTreeTransition& TransC = StateTrans.AddTransition(EStateTreeTransitionTrigger::OnTick, EStateTreeTransitionType::GotoState, &StateC);
+				TransC.Priority = EStateTreeTransitionPriority::High;
+				TStateTreeEditorNode<FStateTreeCompareIntCondition>& TransIntCond = TransC.AddCondition<FStateTreeCompareIntCondition>(EGenericAICheck::Equal);
+				TransIntCond.GetInstanceData().Right = 1;
+				EditorData.AddPropertyBinding(
+					FStateTreePropertyPath(EditorData.RootParameters.ID, TEXT("Int")),
+					FStateTreePropertyPath(TransIntCond.ID, TEXT("Left")));
+			}
+		}
+
+		auto& TaskA = StateA.AddTask<FTestTask_Stand>(FName(TEXT("TaskA")));
+		auto& TaskB = StateB.AddTask<FTestTask_Stand>(FName(TEXT("TaskB")));
+		auto& TaskC = StateC.AddTask<FTestTask_Stand>(FName(TEXT("TaskC")));
+
+		FStateTreeCompilerLog Log;
+		FStateTreeCompiler Compiler(Log);
+		const bool bResult = Compiler.Compile(StateTree);
+
+		AITEST_TRUE("StateTree should get compiled", bResult);
+
+		EStateTreeRunStatus Status = EStateTreeRunStatus::Unset;
+		FStateTreeInstanceData InstanceData;
+		FTestStateTreeExecutionContext Exec(StateTree, StateTree, InstanceData);
+		const bool bInitSucceeded = Exec.IsValid();
+		AITEST_TRUE("StateTree should init", bInitSucceeded);
+
+		const FString TickStr(TEXT("Tick"));
+		const FString EnterStateStr(TEXT("EnterState"));
+		const FString ExitStateStr(TEXT("ExitState"));
+
+		Status = Exec.Start();
+		AITEST_FALSE("StateTree TaskA should not enter state", Exec.Expect(TaskA.GetName(), EnterStateStr));
+		AITEST_FALSE("StateTree TaskB should not enter state", Exec.Expect(TaskB.GetName(), EnterStateStr));
+		AITEST_TRUE("StateTree TaskC should enter state", Exec.Expect(TaskC.GetName(), EnterStateStr));
+		Exec.LogClear();
+
+		Exec.Stop();
+
+		return true;
+	}
+};
+IMPLEMENT_AI_INSTANT_TEST(FStateTreeTest_FollowTransitions, "System.StateTree.Transition.FollowTransitions");
+
+
+struct FStateTreeTest_InfiniteLoop : FStateTreeTestBase
+{
+	virtual bool InstantTest() override
+	{
+		UStateTree& StateTree = NewStateTree();
+		UStateTreeEditorData& EditorData = *Cast<UStateTreeEditorData>(StateTree.EditorData);
+
+		EditorData.RootParameters.Parameters.AddProperty(FName(TEXT("Int")), EPropertyBagPropertyType::Int32);
+		EditorData.RootParameters.Parameters.SetValueInt32(FName(TEXT("Int")), 1);
+		
+		UStateTreeState& Root = EditorData.AddSubTree(FName(TEXT("Root")));
+		UStateTreeState& StateA = Root.AddChildState(FName(TEXT("A")));
+		UStateTreeState& StateB = StateA.AddChildState(FName(TEXT("B")));
+
+		// Root
+
+		// State A
+		{
+			StateA.SelectionBehavior = EStateTreeStateSelectionBehavior::TryFollowTransitions;
+			{
+				// A -> B
+				FStateTreeTransition& Trans = StateA.AddTransition(EStateTreeTransitionTrigger::OnTick, EStateTreeTransitionType::GotoState, &StateB);
+				TStateTreeEditorNode<FStateTreeCompareIntCondition>& TransIntCond = Trans.AddCondition<FStateTreeCompareIntCondition>(EGenericAICheck::Equal);
+				TransIntCond.GetInstanceData().Right = 1;
+				EditorData.AddPropertyBinding(
+					FStateTreePropertyPath(EditorData.RootParameters.ID, TEXT("Int")),
+					FStateTreePropertyPath(TransIntCond.ID, TEXT("Left")));
+			}
+		}
+
+		// State B
+		{
+			StateB.SelectionBehavior = EStateTreeStateSelectionBehavior::TryFollowTransitions;
+			{
+				// B -> A
+				FStateTreeTransition& Trans = StateB.AddTransition(EStateTreeTransitionTrigger::OnTick, EStateTreeTransitionType::GotoState, &StateA);
+				TStateTreeEditorNode<FStateTreeCompareIntCondition>& TransIntCond = Trans.AddCondition<FStateTreeCompareIntCondition>(EGenericAICheck::Equal);
+				TransIntCond.GetInstanceData().Right = 1;
+				EditorData.AddPropertyBinding(
+					FStateTreePropertyPath(EditorData.RootParameters.ID, TEXT("Int")),
+					FStateTreePropertyPath(TransIntCond.ID, TEXT("Left")));
+			}
+		}
+		
+		auto& TaskA = StateA.AddTask<FTestTask_Stand>(FName(TEXT("TaskA")));
+		auto& TaskB = StateB.AddTask<FTestTask_Stand>(FName(TEXT("TaskB")));
+
+		FStateTreeCompilerLog Log;
+		FStateTreeCompiler Compiler(Log);
+		const bool bResult = Compiler.Compile(StateTree);
+
+		AITEST_TRUE("StateTree should get compiled", bResult);
+
+		EStateTreeRunStatus Status = EStateTreeRunStatus::Unset;
+		FStateTreeInstanceData InstanceData;
+		FTestStateTreeExecutionContext Exec(StateTree, StateTree, InstanceData);
+		const bool bInitSucceeded = Exec.IsValid();
+		AITEST_TRUE("StateTree should init", bInitSucceeded);
+
+		const FString TickStr(TEXT("Tick"));
+		const FString EnterStateStr(TEXT("EnterState"));
+		const FString ExitStateStr(TEXT("ExitState"));
+
+		GetTestRunner().AddExpectedError(TEXT("Loop detected when trying to select state"), EAutomationExpectedErrorFlags::Contains, 1);
+		GetTestRunner().AddExpectedError(TEXT("Failed to select initial state"), EAutomationExpectedErrorFlags::Contains, 1);
+		
+		Status = Exec.Start();
+		AITEST_EQUAL("Start should fail", Status, EStateTreeRunStatus::Failed);
+		Exec.LogClear();
+
+		Exec.Stop();
+
+		return true;
+	}
+};
+IMPLEMENT_AI_INSTANT_TEST(FStateTreeTest_InfiniteLoop, "System.StateTree.Transition.InfiniteLoop");
+
+} // namespace UE::StateTree::Tests
+
+UE_ENABLE_OPTIMIZATION_SHIP
+
+#undef LOCTEXT_NAMESPACE
diff --git a/Engine/Plugins/Runtime/StateTree/Source/StateTreeTestSuite/Public/StateTreeTestBase.h b/Engine/Plugins/Runtime/StateTree/Source/StateTreeTestSuite/Public/StateTreeTestBase.h
new file mode 100644
--- /dev/null
+++ b/Engine/Plugins/Runtime/StateTree/Source/StateTreeTestSuite/Public/StateTreeTestBase.h
@@ -0,0 +1,26 @@
+// Copyright Epic Games, Inc. All Rights Reserved.
+
+#pragma  once
+
+#include "AITestsCommon.h"
+
+#include "StateTree.h"
+#include "StateTreePropertyBindings.h"
+
+namespace UE::StateTree::Tests
+{
+	
+/**
+ * Base class for StateTree test
+ */
+struct FStateTreeTestBase : public FAITestBase
+{
+protected:
+	UStateTree& NewStateTree() const;
+	static FStateTreePropertyPathBinding MakeBinding(const FGuid& SourceID, const FStringView Source, const FGuid& TargetID, const FStringView Target);
+	static FGameplayTag GetTestTag1();
+	static FGameplayTag GetTestTag2();
+	static FGameplayTag GetTestTag3();
+};
+
+} // namespace UE::StateTree::Tests
\ No newline at end of file
