StateTree: Small optimization to skip tasks ticking if no task requests a tick. Moved bHasGlobalTransitionTasks and  bHasTransitionTasks from the compiler to the runtime flag initialization. Reorder the data in FCompactStateTreeState for better memory padding.

diff --git a/Engine/Plugins/Runtime/StateTree/Source/StateTreeEditorModule/Private/StateTreeCompiler.cpp b/Engine/Plugins/Runtime/StateTree/Source/StateTreeEditorModule/Private/StateTreeCompiler.cpp
--- a/Engine/Plugins/Runtime/StateTree/Source/StateTreeEditorModule/Private/StateTreeCompiler.cpp
+++ b/Engine/Plugins/Runtime/StateTree/Source/StateTreeEditorModule/Private/StateTreeCompiler.cpp
@@ -565,7 +565,6 @@ bool FStateTreeCompiler::CreateGlobalTasks()
 	}
 	StateTree->GlobalTasksBegin = uint16(GlobalTasksBegin);
 
-	StateTree->bHasGlobalTransitionTasks = false;
 	for (FStateTreeEditorNode& TaskNode : EditorData->GlobalTasks)
 	{
 		// Silently ignore empty nodes.
@@ -581,10 +580,6 @@ bool FStateTreeCompiler::CreateGlobalTasks()
 			bSucceeded = false;
 			continue;
 		}
-		
-		const FStateTreeTaskBase& LastAddedTask = Nodes.Last().Get<FStateTreeTaskBase>();
-		
-		StateTree->bHasGlobalTransitionTasks |= LastAddedTask.bShouldAffectTransitions;
 	}
 	
 	const int32 GlobalTasksNum = Nodes.Num() - GlobalTasksBegin;
@@ -756,7 +751,6 @@ bool FStateTreeCompiler::CreateStateTasksAndParameters()
 			Tasks = TArrayView<FStateTreeEditorNode>(&State->SingleTask, 1);
 		}
 		
-		bool bStateHasTransitionTasks = false;
 		bool bCreateTaskSucceeded = true;
 		for (FStateTreeEditorNode& TaskNode : Tasks)
 		{
@@ -773,14 +767,8 @@ bool FStateTreeCompiler::CreateStateTasksAndParameters()
 				bCreateTaskSucceeded = false;
 				continue;
 			}
-
-			const FStateTreeTaskBase& LastAddedTask = Nodes.Last().Get<FStateTreeTaskBase>();
-			
-			bStateHasTransitionTasks |= LastAddedTask.bShouldAffectTransitions;
 		}
 
-		CompactState.bHasTransitionTasks = bStateHasTransitionTasks;
-
 		if (!bCreateTaskSucceeded)
 		{
 			continue;
diff --git a/Engine/Plugins/Runtime/StateTree/Source/StateTreeModule/Private/StateTree.cpp b/Engine/Plugins/Runtime/StateTree/Source/StateTreeModule/Private/StateTree.cpp
--- a/Engine/Plugins/Runtime/StateTree/Source/StateTreeModule/Private/StateTree.cpp
+++ b/Engine/Plugins/Runtime/StateTree/Source/StateTreeModule/Private/StateTree.cpp
@@ -637,6 +637,8 @@ bool UStateTree::Link()
 		}
 	}
 
+	UpdateRuntimeFlags();
+
 	if (!DefaultInstanceData.AreAllInstancesValid())
 	{
 		UE_LOG(LogStateTree, Error, TEXT("%s: State Tree asset was not properly loaded (missing instance data). See log for loading failures, or recompile the StateTree asset."), *GetFullName());
@@ -666,6 +668,51 @@ bool UStateTree::Link()
 	return true;
 }
 
+void UStateTree::UpdateRuntimeFlags()
+{
+	// Set the tick flags at runtime instead of compilation.
+	//This is to support hotfix (when we only modify cpp code).
+
+	for (FCompactStateTreeState& State : States)
+	{
+		// Update the state task flags.
+		
+		// Tick Linked and LinkedAsset because they CopyBatchOnActiveInstances when they tick. That info might be important when doing transaction.
+		const bool bShouldTickByDefault = (State.Type == EStateTreeStateType::Linked || State.Type == EStateTreeStateType::LinkedAsset);
+		State.bHasTickTasks = bShouldTickByDefault;
+		State.bHasTickTasksOnlyOnEvents = bShouldTickByDefault;
+		State.bHasTransitionTasks = false;
+
+		for (int32 TaskIndex = State.TasksBegin; TaskIndex < (State.TasksBegin + State.TasksNum); ++TaskIndex)
+		{
+			const FStateTreeTaskBase& Task = Nodes[TaskIndex].Get<const FStateTreeTaskBase>();
+			if (Task.bTaskEnabled)
+			{
+				State.bHasTickTasks |= Task.bShouldCallTick;
+				State.bHasTickTasksOnlyOnEvents |= Task.bShouldCallTickOnlyOnEvents;
+				State.bHasTransitionTasks |= Task.bShouldAffectTransitions;
+			}
+		}
+	}
+
+	// Update the global task flags.
+	{
+		bHasGlobalTickTasks = false;
+		bHasGlobalTickTasksOnlyOnEvents = false;
+		bHasGlobalTransitionTasks = false;
+		for (int32 TaskIndex = GlobalTasksBegin; TaskIndex < (GlobalTasksBegin + GlobalTasksNum); ++TaskIndex)
+		{
+			const FStateTreeTaskBase& Task = Nodes[TaskIndex].Get<const FStateTreeTaskBase>();
+			if (Task.bTaskEnabled)
+			{
+				bHasGlobalTickTasks |= Task.bShouldCallTick;
+				bHasGlobalTickTasksOnlyOnEvents |= Task.bShouldCallTickOnlyOnEvents;
+				bHasGlobalTransitionTasks |= Task.bShouldAffectTransitions;
+			}
+		}
+	}
+}
+
 bool UStateTree::PatchBindings()
 {
 	const TArrayView<FStateTreeBindableStructDesc> SourceStructs = PropertyBindings.SourceStructs;
diff --git a/Engine/Plugins/Runtime/StateTree/Source/StateTreeModule/Private/StateTreeExecutionContext.cpp b/Engine/Plugins/Runtime/StateTree/Source/StateTreeModule/Private/StateTreeExecutionContext.cpp
--- a/Engine/Plugins/Runtime/StateTree/Source/StateTreeModule/Private/StateTreeExecutionContext.cpp
+++ b/Engine/Plugins/Runtime/StateTree/Source/StateTreeModule/Private/StateTreeExecutionContext.cpp
@@ -2168,58 +2168,65 @@ EStateTreeRunStatus FStateTreeExecutionContext::TickEvaluatorsAndGlobalTasks(con
 				bool bShouldTickTasks = true;
 				const bool bHasEvents = EventQueue && EventQueue->HasEvents();
 
-				for (int32 TaskIndex = CurrentStateTree->GlobalTasksBegin; TaskIndex < (CurrentStateTree->GlobalTasksBegin + CurrentStateTree->GlobalTasksNum); TaskIndex++)
+				if (CurrentStateTree->ShouldTickGlobalTasks(bHasEvents))
 				{
-					const FStateTreeTaskBase& Task = CurrentStateTree->Nodes[TaskIndex].Get<const FStateTreeTaskBase>();
-					const FStateTreeDataView TaskInstanceView = GetDataView(CurrentParentFrame, CurrentFrame, Task.InstanceDataHandle);
-					FNodeInstanceDataScope DataScope(*this, Task.InstanceDataHandle, TaskInstanceView);
-
-					// Ignore disabled task
-					if (Task.bTaskEnabled == false)
+					bool bTaskTicked = false;
+					for (int32 TaskIndex = CurrentStateTree->GlobalTasksBegin; TaskIndex < (CurrentStateTree->GlobalTasksBegin + CurrentStateTree->GlobalTasksNum); TaskIndex++)
 					{
-						STATETREE_LOG(VeryVerbose, TEXT("%*sSkipped 'Tick' for disabled Task: '%s'"), UE::StateTree::DebugIndentSize, TEXT(""), *Task.Name.ToString());
-						continue;
-					}
+						const FStateTreeTaskBase& Task = CurrentStateTree->Nodes[TaskIndex].Get<const FStateTreeTaskBase>();
+						const FStateTreeDataView TaskInstanceView = GetDataView(CurrentParentFrame, CurrentFrame, Task.InstanceDataHandle);
+						FNodeInstanceDataScope DataScope(*this, Task.InstanceDataHandle, TaskInstanceView);
 
-					const bool bNeedsTick = bShouldTickTasks && (Task.bShouldCallTick || (bHasEvents && Task.bShouldCallTickOnlyOnEvents));
-					STATETREE_LOG(VeryVerbose, TEXT("  Tick: '%s' %s"), *Task.Name.ToString(), !bNeedsTick ? TEXT("[not ticked]") : TEXT(""));
-					if (!bNeedsTick)
-					{
-						continue;
-					}
+						// Ignore disabled task
+						if (Task.bTaskEnabled == false)
+						{
+							STATETREE_LOG(VeryVerbose, TEXT("%*sSkipped 'Tick' for disabled Task: '%s'"), UE::StateTree::DebugIndentSize, TEXT(""), *Task.Name.ToString());
+							continue;
+						}
 
-					// Copy bound properties.
-					// Only copy properties when the task is actually ticked, and copy properties at tick is requested.
-					if (Task.BindingsBatch.IsValid() && Task.bShouldCopyBoundPropertiesOnTick)
-					{
-						CopyBatchOnActiveInstances(CurrentParentFrame, CurrentFrame, TaskInstanceView, Task.BindingsBatch);
-					}
+						const bool bNeedsTick = bShouldTickTasks && (Task.bShouldCallTick || (bHasEvents && Task.bShouldCallTickOnlyOnEvents));
+						STATETREE_LOG(VeryVerbose, TEXT("  Tick: '%s' %s"), *Task.Name.ToString(), !bNeedsTick ? TEXT("[not ticked]") : TEXT(""));
+						if (!bNeedsTick)
+						{
+							continue;
+						}
 
-					//STATETREE_TRACE_TASK_EVENT(TaskIndex, TaskDataView, EStateTreeTraceEventType::OnTickingTask, EStateTreeRunStatus::Running);
-					UE_STATETREE_DEBUG_TASK_TICK(this, CurrentStateTree, TaskIndex);
+						// Copy bound properties.
+						// Only copy properties when the task is actually ticked, and copy properties at tick is requested.
+						if (Task.BindingsBatch.IsValid() && Task.bShouldCopyBoundPropertiesOnTick)
+						{
+							CopyBatchOnActiveInstances(CurrentParentFrame, CurrentFrame, TaskInstanceView, Task.BindingsBatch);
+						}
 
-					EStateTreeRunStatus TaskResult = EStateTreeRunStatus::Unset;
-					{
-						QUICK_SCOPE_CYCLE_COUNTER(StateTree_Task_Tick);
-						CSV_SCOPED_TIMING_STAT_EXCLUSIVE(StateTree_Task_Tick);
+						//STATETREE_TRACE_TASK_EVENT(TaskIndex, TaskDataView, EStateTreeTraceEventType::OnTickingTask, EStateTreeRunStatus::Running);
+						UE_STATETREE_DEBUG_TASK_TICK(this, CurrentStateTree, TaskIndex);
 
-						TaskResult = Task.Tick(*this, DeltaTime);
-					}
+						EStateTreeRunStatus TaskResult = EStateTreeRunStatus::Unset;
+						{
+							QUICK_SCOPE_CYCLE_COUNTER(StateTree_Task_Tick);
+							CSV_SCOPED_TIMING_STAT_EXCLUSIVE(StateTree_Task_Tick);
 
-					STATETREE_TRACE_TASK_EVENT(TaskIndex, TaskInstanceView,
-						TaskResult != EStateTreeRunStatus::Running ? EStateTreeTraceEventType::OnTaskCompleted : EStateTreeTraceEventType::OnTicked,
-						TaskResult);
+							TaskResult = Task.Tick(*this, DeltaTime);
+							bTaskTicked = true;
+						}
 
-					// If a global task succeeds or fails, it will stop the whole tree.
-					if (TaskResult != EStateTreeRunStatus::Running)
-					{
-						Result = TaskResult;
-					}
+						STATETREE_TRACE_TASK_EVENT(TaskIndex, TaskInstanceView,
+							TaskResult != EStateTreeRunStatus::Running ? EStateTreeTraceEventType::OnTaskCompleted : EStateTreeTraceEventType::OnTicked,
+							TaskResult);
+
+						// If a global task succeeds or fails, it will stop the whole tree.
+						if (TaskResult != EStateTreeRunStatus::Running)
+						{
+							Result = TaskResult;
+						}
 						
-					if (TaskResult == EStateTreeRunStatus::Failed)
-					{
-						bShouldTickTasks = false;
+						if (TaskResult == EStateTreeRunStatus::Failed)
+						{
+							bShouldTickTasks = false;
+						}
 					}
+
+					ensureMsgf(bTaskTicked, TEXT("ShouldTickGlobalTasks is true but no tasks were tick for the StateTree `%s`."), *CurrentStateTree->GetPathName());
 				}
 			}
 		}
@@ -2641,76 +2648,84 @@ EStateTreeRunStatus FStateTreeExecutionContext::TickTasks(const float DeltaTime)
 
 			STATETREE_CLOG(State.TasksNum > 0, VeryVerbose, TEXT("%*sState '%s'"), Index*UE::StateTree::DebugIndentSize, TEXT(""), *DebugGetStatePath(Exec.ActiveFrames, &CurrentFrame, Index));
 
-			if (State.Type == EStateTreeStateType::Linked
-				|| State.Type == EStateTreeStateType::LinkedAsset)
-			{
-				if (State.ParameterDataHandle.IsValid()
-					&& State.ParameterBindingsBatch.IsValid())
-				{
-					const FStateTreeDataView StateParamsDataView = GetDataView(CurrentParentFrame, CurrentFrame, State.ParameterDataHandle);
-					CopyBatchOnActiveInstances(CurrentParentFrame, CurrentFrame, StateParamsDataView, State.ParameterBindingsBatch);
-				}
-			}
-
-			// Update Tasks data and tick if possible (ie. if no task has yet failed and so bShouldTickTasks is true)
-			for (int32 TaskIndex = State.TasksBegin; TaskIndex < (State.TasksBegin + State.TasksNum); TaskIndex++)
+			if (State.ShouldTickTasks(bHasEvents))
 			{
-				const FStateTreeTaskBase& Task = CurrentStateTree->Nodes[TaskIndex].Get<const FStateTreeTaskBase>();
-				const FStateTreeDataView TaskInstanceView = GetDataView(CurrentParentFrame, CurrentFrame, Task.InstanceDataHandle);
-				FNodeInstanceDataScope DataScope(*this, Task.InstanceDataHandle, TaskInstanceView);
+				bool bTaskTicked = false;
 
-				// Ignore disabled task
-				if (Task.bTaskEnabled == false)
+				if (State.Type == EStateTreeStateType::Linked || State.Type == EStateTreeStateType::LinkedAsset)
 				{
-					STATETREE_LOG(VeryVerbose, TEXT("%*sSkipped 'Tick' for disabled Task: '%s'"), UE::StateTree::DebugIndentSize, TEXT(""), *Task.Name.ToString());
-					continue;
+					// Linked and LinkedAsset are considered to always tick because the set data that can be used in transitions.
+					bTaskTicked = true;
+
+					if (State.ParameterDataHandle.IsValid() && State.ParameterBindingsBatch.IsValid())
+					{
+						const FStateTreeDataView StateParamsDataView = GetDataView(CurrentParentFrame, CurrentFrame, State.ParameterDataHandle);
+						CopyBatchOnActiveInstances(CurrentParentFrame, CurrentFrame, StateParamsDataView, State.ParameterBindingsBatch);
+					}
 				}
 
-				const bool bNeedsTick = bShouldTickTasks && (Task.bShouldCallTick || (bHasEvents && Task.bShouldCallTickOnlyOnEvents));
-				STATETREE_LOG(VeryVerbose, TEXT("%*s  Tick: '%s' %s"), Index*UE::StateTree::DebugIndentSize, TEXT(""), *Task.Name.ToString(), !bNeedsTick ? TEXT("[not ticked]") : TEXT(""));
-				if (!bNeedsTick)
+				// Update Tasks data and tick if possible (ie. if no task has yet failed and so bShouldTickTasks is true)
+				for (int32 TaskIndex = State.TasksBegin; TaskIndex < (State.TasksBegin + State.TasksNum); TaskIndex++)
 				{
-					continue;
-				}
+					const FStateTreeTaskBase& Task = CurrentStateTree->Nodes[TaskIndex].Get<const FStateTreeTaskBase>();
+					const FStateTreeDataView TaskInstanceView = GetDataView(CurrentParentFrame, CurrentFrame, Task.InstanceDataHandle);
+					FNodeInstanceDataScope DataScope(*this, Task.InstanceDataHandle, TaskInstanceView);
+
+					// Ignore disabled task
+					if (Task.bTaskEnabled == false)
+					{
+						STATETREE_LOG(VeryVerbose, TEXT("%*sSkipped 'Tick' for disabled Task: '%s'"), UE::StateTree::DebugIndentSize, TEXT(""), *Task.Name.ToString());
+						continue;
+					}
+
+					const bool bNeedsTick = bShouldTickTasks && (Task.bShouldCallTick || (bHasEvents && Task.bShouldCallTickOnlyOnEvents));
+					STATETREE_LOG(VeryVerbose, TEXT("%*s  Tick: '%s' %s"), Index*UE::StateTree::DebugIndentSize, TEXT(""), *Task.Name.ToString(), !bNeedsTick ? TEXT("[not ticked]") : TEXT(""));
+					if (!bNeedsTick)
+					{
+						continue;
+					}
 				
-				// Copy bound properties.
-				// Only copy properties when the task is actually ticked, and copy properties at tick is requested.
-				if (Task.BindingsBatch.IsValid() && Task.bShouldCopyBoundPropertiesOnTick)
-				{
-					CopyBatchOnActiveInstances(CurrentParentFrame, CurrentFrame, TaskInstanceView, Task.BindingsBatch);
-				}
+					// Copy bound properties.
+					// Only copy properties when the task is actually ticked, and copy properties at tick is requested.
+					if (Task.BindingsBatch.IsValid() && Task.bShouldCopyBoundPropertiesOnTick)
+					{
+						CopyBatchOnActiveInstances(CurrentParentFrame, CurrentFrame, TaskInstanceView, Task.BindingsBatch);
+					}
 
-				//STATETREE_TRACE_TASK_EVENT(TaskIndex, TaskDataView, EStateTreeTraceEventType::OnTickingTask, EStateTreeRunStatus::Running);
-				UE_STATETREE_DEBUG_TASK_TICK(this, CurrentStateTree, TaskIndex);
+					//STATETREE_TRACE_TASK_EVENT(TaskIndex, TaskDataView, EStateTreeTraceEventType::OnTickingTask, EStateTreeRunStatus::Running);
+					UE_STATETREE_DEBUG_TASK_TICK(this, CurrentStateTree, TaskIndex);
 
-				EStateTreeRunStatus TaskResult = EStateTreeRunStatus::Unset;
-				{
-					QUICK_SCOPE_CYCLE_COUNTER(StateTree_Task_Tick);
-					CSV_SCOPED_TIMING_STAT_EXCLUSIVE(StateTree_Task_Tick);
+					EStateTreeRunStatus TaskResult = EStateTreeRunStatus::Unset;
+					{
+						QUICK_SCOPE_CYCLE_COUNTER(StateTree_Task_Tick);
+						CSV_SCOPED_TIMING_STAT_EXCLUSIVE(StateTree_Task_Tick);
 
-					TaskResult = Task.Tick(*this, DeltaTime);
-				}
+						TaskResult = Task.Tick(*this, DeltaTime);
+						bTaskTicked = true;
+					}
 
-				STATETREE_TRACE_TASK_EVENT(TaskIndex, TaskInstanceView,
-					TaskResult != EStateTreeRunStatus::Running ? EStateTreeTraceEventType::OnTaskCompleted : EStateTreeTraceEventType::OnTicked,
-					TaskResult);
+					STATETREE_TRACE_TASK_EVENT(TaskIndex, TaskInstanceView,
+						TaskResult != EStateTreeRunStatus::Running ? EStateTreeTraceEventType::OnTaskCompleted : EStateTreeTraceEventType::OnTicked,
+						TaskResult);
 				
-				// TODO: Add more control over which states can control the failed/succeeded result.
-				if (TaskResult != EStateTreeRunStatus::Running)
-				{
-					// Store the first state that completed, will be used to decide where to trigger transitions.
-					if (!Exec.CompletedStateHandle.IsValid())
+					// TODO: Add more control over which states can control the failed/succeeded result.
+					if (TaskResult != EStateTreeRunStatus::Running)
 					{
-						Exec.CompletedFrameIndex = FStateTreeIndex16(FrameIndex);
-						Exec.CompletedStateHandle = CurrentHandle;
+						// Store the first state that completed, will be used to decide where to trigger transitions.
+						if (!Exec.CompletedStateHandle.IsValid())
+						{
+							Exec.CompletedFrameIndex = FStateTreeIndex16(FrameIndex);
+							Exec.CompletedStateHandle = CurrentHandle;
+						}
+						Result = TaskResult;
 					}
-					Result = TaskResult;
-				}
 				
-				if (TaskResult == EStateTreeRunStatus::Failed)
-				{
-					bShouldTickTasks = false;
+					if (TaskResult == EStateTreeRunStatus::Failed)
+					{
+						bShouldTickTasks = false;
+					}
 				}
+				ensureMsgf(bTaskTicked, TEXT("ShouldTickTasks is true but no tasks were tick for State: '%s' inside theStateTree %s."), *State.Name.ToString(), *CurrentStateTree->GetPathName());
 			}
 			NumTotalTasks += State.TasksNum;
 		}
@@ -3179,14 +3194,17 @@ bool FStateTreeExecutionContext::TriggerTransitions()
 				// Transition tasks.
 				if (State.bHasTransitionTasks)
 				{
+					bool bAdded = false;
 					for (int32 TaskIndex = (State.TasksBegin + State.TasksNum) - 1; TaskIndex >= State.TasksBegin; TaskIndex--)
 					{
 						const FStateTreeTaskBase& Task = CurrentStateTree->Nodes[TaskIndex].Get<const FStateTreeTaskBase>();
-						if (Task.bShouldAffectTransitions)
+						if (Task.bShouldAffectTransitions && Task.bTaskEnabled)
 						{
 							TransitionHandlers.Emplace((uint8)FrameIndex, StateHandle, FStateTreeIndex16(TaskIndex), Task.TransitionHandlingPriority);
+							bAdded = true;
 						}
 					}
+					ensureMsgf(bAdded, TEXT("bHasTransitionTasks is set but not task were added for the State: '%s' inside theStateTree %s"), *State.Name.ToString(), *CurrentStateTree->GetPathName());
 				}
 				
 				// Regular transitions on state
@@ -3201,15 +3219,18 @@ bool FStateTreeExecutionContext::TriggerTransitions()
 				// Global transition tasks.
 				if (CurrentFrame.StateTree->bHasGlobalTransitionTasks)
 				{
+					bool bAdded = false;
 					for (int32 TaskIndex = (CurrentStateTree->GlobalTasksBegin + CurrentStateTree->GlobalTasksNum) - 1; TaskIndex >= CurrentFrame.StateTree->GlobalTasksBegin; TaskIndex--)
 					{
 						const FStateTreeTaskBase& Task = CurrentStateTree->Nodes[TaskIndex].Get<const FStateTreeTaskBase>();
-						if (Task.bShouldAffectTransitions)
+						if (Task.bShouldAffectTransitions && Task.bTaskEnabled)
 						{
 							TransitionHandlers.Emplace((uint8)FrameIndex, FStateTreeStateHandle(), FStateTreeIndex16(TaskIndex), Task.TransitionHandlingPriority);
+							bAdded = true;
 						}
 					}
-				}				
+					ensureMsgf(bAdded, TEXT("bHasGlobalTransitionTasks is set but not task were added for the StateTree `%s`"), *CurrentStateTree->GetPathName());
+				}
 			}
 		}
 
diff --git a/Engine/Plugins/Runtime/StateTree/Source/StateTreeModule/Public/StateTree.h b/Engine/Plugins/Runtime/StateTree/Source/StateTreeModule/Public/StateTree.h
--- a/Engine/Plugins/Runtime/StateTree/Source/StateTreeModule/Public/StateTree.h
+++ b/Engine/Plugins/Runtime/StateTree/Source/StateTreeModule/Public/StateTree.h
@@ -159,11 +159,17 @@ public:
 	FStateTreeIndex16 GetTransitionIndexFromId(const FGuid Id) const;
 
 	/** @return Id of the transition matching a given runtime transition index; invalid Id if transition not found. */
-	FGuid GetTransitionIdFromIndex(const FStateTreeIndex16 Index) const;	
+	FGuid GetTransitionIdFromIndex(const FStateTreeIndex16 Index) const;
 
 	/** @return Property bindings */
 	const FStateTreePropertyBindings& GetPropertyBindings() const { return PropertyBindings; }
 
+	/** @return True if there is any global tasks need ticking. */
+	bool ShouldTickGlobalTasks(bool bHasEvent) const
+	{
+		return bHasGlobalTickTasks || (bHasEvent && bHasGlobalTickTasksOnlyOnEvents);
+	}
+
 	UE_DEPRECATED(5.4, "Replaced with GetNumContextDataViews() which contains context data and external data only.")
 	int32 GetNumDataViews() const { return 0; }
 
@@ -245,6 +251,9 @@ private:
 	/** @return true if all the source instance data types match with the node's instance data types */
 	bool ValidateInstanceData();
 
+	/** Set StateTree's flag that can't be set when compiled. */
+	void UpdateRuntimeFlags();
+
 	bool PatchBindings();
 
 	// Data created during compilation, source data in EditorData.
@@ -303,6 +312,11 @@ private:
 	UPROPERTY()
 	FInstancedPropertyBag Parameters;
 
+	//~ Data created during linking.
+	/** List of external data required by the state tree, created during linking. */
+	UPROPERTY(Transient)
+	TArray<FStateTreeExternalDataDesc> ExternalDataDescs;
+
 	/** Number of context data, include parameters and all context data. */
 	UPROPERTY()
 	uint16 NumContextData = 0;
@@ -329,16 +343,23 @@ private:
 
 	/** True if any global task is a transition task. */
 	UPROPERTY()
-	bool bHasGlobalTransitionTasks = false;
-	
-	// Data created during linking.
-	
-	/** List of external data required by the state tree, created during linking. */
-	UPROPERTY(Transient)
-	TArray<FStateTreeExternalDataDesc> ExternalDataDescs;
+	uint8 bHasGlobalTransitionTasks : 1 = false;
+
+	/**
+	 * True if any global task has bShouldCallTick.
+	 * Not ticking implies no property copy.
+	 */
+	uint8 bHasGlobalTickTasks : 1 = false;
+
+	/**
+	 * True if any global task has bShouldCallTickOnlyOnEvents.
+	 * No effect if bHasGlobalTickTasks is true.
+	 * Not ticking implies no property copy.
+	 */
+	uint8 bHasGlobalTickTasksOnlyOnEvents : 1 = false;
 
 	/** True if the StateTree was linked successfully. */
-	bool bIsLinked = false;
+	uint8 bIsLinked : 1 = false;
 
 #if WITH_EDITORONLY_DATA
 	/** List of Struct that are out of date and waiting to be replaced with the new instance. */
diff --git a/Engine/Plugins/Runtime/StateTree/Source/StateTreeModule/Public/StateTreeTaskBase.h b/Engine/Plugins/Runtime/StateTree/Source/StateTreeModule/Public/StateTreeTaskBase.h
--- a/Engine/Plugins/Runtime/StateTree/Source/StateTreeModule/Public/StateTreeTaskBase.h
+++ b/Engine/Plugins/Runtime/StateTree/Source/StateTreeModule/Public/StateTreeTaskBase.h
@@ -106,7 +106,7 @@ struct STATETREEMODULE_API FStateTreeTaskBase : public FStateTreeNodeBase
 
 	/** If set to true, Tick() is called. Not ticking implies no property copy. Default true. */
 	uint8 bShouldCallTick : 1;
-	/** If set to true, Tick() is called only when there are events. No effect if bShouldCallTickState is true. Not ticking implies no property copy. Default false. */
+	/** If set to true, Tick() is called only when there are events. No effect if bShouldCallTick is true. Not ticking implies no property copy. Default false. */
 	uint8 bShouldCallTickOnlyOnEvents : 1;
 
 	/** If set to true, copy the values of bound properties before calling Tick(). Default true. */
diff --git a/Engine/Plugins/Runtime/StateTree/Source/StateTreeModule/Public/StateTreeTypes.h b/Engine/Plugins/Runtime/StateTree/Source/StateTreeModule/Public/StateTreeTypes.h
--- a/Engine/Plugins/Runtime/StateTree/Source/StateTreeModule/Public/StateTreeTypes.h
+++ b/Engine/Plugins/Runtime/StateTree/Source/StateTreeModule/Public/StateTreeTypes.h
@@ -660,10 +660,22 @@ struct STATETREEMODULE_API FCompactStateTreeState
 	GENERATED_BODY()
 	
 	/** @return Index to the next sibling state. */
-	uint16 GetNextSibling() const { return ChildrenEnd; }
+	uint16 GetNextSibling() const
+	{
+		return ChildrenEnd;
+	}
 
 	/** @return True if the state has any child states */
-	bool HasChildren() const { return ChildrenEnd > ChildrenBegin; }
+	bool HasChildren() const
+	{
+		return ChildrenEnd > ChildrenBegin;
+	}
+
+	/** @return True if the state has any tasks that need ticking. */
+	bool ShouldTickTasks(bool bHasEvent) const
+	{
+		return bHasTickTasks || (bHasEvent && bHasTickTasksOnlyOnEvents);
+	}
 
 	/** Description of an event required to enter the state. */
 	UPROPERTY()
@@ -725,6 +737,10 @@ struct STATETREEMODULE_API FCompactStateTreeState
 	UPROPERTY()
 	FStateTreeIndex16 EventDataIndex = FStateTreeIndex16::Invalid;
 
+	/** Weight used to scale the normalized final utility score for this state */
+	UPROPERTY()
+	float Weight = 1.f;
+
 	/** Number of enter conditions */
 	UPROPERTY()
 	uint8 EnterConditionsNum = 0;
@@ -765,13 +781,22 @@ struct STATETREEMODULE_API FCompactStateTreeState
 	UPROPERTY()
 	uint8 bHasStateChangeConditions : 1 = false;
 
+	/**
+	 * True if any task has bShouldCallTick.
+	 * Not ticking implies no property copy.
+	 */
+	uint8 bHasTickTasks : 1 = false;
+
+	/**
+	 * True if any task has bShouldCallTickOnlyOnEvents.
+	 * No effect if bHasGlobalTickTasks is true.
+	 * Not ticking implies no property copy.
+	 */
+	uint8 bHasTickTasksOnlyOnEvents : 1 = false;
+
 	/** Should state's required event and enter conditions be evaluated when transition leads directly to it's child. */
 	UPROPERTY()
 	uint8 bCheckPrerequisitesWhenActivatingChildDirectly : 1 = false;
-	
-	/** Weight used to scale the normalized final utility score for this state */
-	UPROPERTY()
-	float Weight = 1.f;
 
 	/** True if the state is Enabled (i.e. not explicitly marked as disabled). */
 	UPROPERTY()
