StateTree: Global tasks and evaluators are not started if temporary data for them already exists. It's a leftover from times when we were not stopping them in such case.

diff --git a/Engine/Plugins/Runtime/StateTree/Source/StateTreeModule/Private/StateTreeExecutionContext.cpp b/Engine/Plugins/Runtime/StateTree/Source/StateTreeModule/Private/StateTreeExecutionContext.cpp
--- a/Engine/Plugins/Runtime/StateTree/Source/StateTreeModule/Private/StateTreeExecutionContext.cpp
+++ b/Engine/Plugins/Runtime/StateTree/Source/StateTreeModule/Private/StateTreeExecutionContext.cpp
@@ -2465,12 +2465,10 @@ EStateTreeRunStatus FStateTreeExecutionContext::StartTemporaryEvaluatorsAndGloba
 	{
 		const FStateTreeEvaluatorBase& Eval = CurrentStateTree->Nodes[EvalIndex].Get<const FStateTreeEvaluatorBase>();
 		FStateTreeDataView EvalInstanceView = GetDataViewOrTemporary(CurrentParentFrame, CurrentFrame, Eval.InstanceDataHandle);
-		bool bWasCreated = false;
 		if (!EvalInstanceView.IsValid())
 		{
 			EvalInstanceView = AddTemporaryInstance(CurrentFrame, FStateTreeIndex16(EvalIndex), Eval.InstanceDataHandle, CurrentFrame.StateTree->DefaultInstanceData.GetStruct(Eval.InstanceTemplateIndex.Get()));
 			check(EvalInstanceView.IsValid());
-			bWasCreated = true;
 		}
 		
 		FNodeInstanceDataScope DataScope(*this, Eval.InstanceDataHandle, EvalInstanceView);
@@ -2480,15 +2478,12 @@ EStateTreeRunStatus FStateTreeExecutionContext::StartTemporaryEvaluatorsAndGloba
 			CopyBatchOnActiveInstances(CurrentParentFrame, CurrentFrame, EvalInstanceView, Eval.BindingsBatch);
 		}
 
-		if (bWasCreated)
+		STATETREE_LOG(Verbose, TEXT("  Start: '%s'"), *Eval.Name.ToString());
 		{
-			STATETREE_LOG(Verbose, TEXT("  Start: '%s'"), *Eval.Name.ToString());
-			{
-				QUICK_SCOPE_CYCLE_COUNTER(StateTree_Eval_TreeStart);
-				Eval.TreeStart(*this);
+			QUICK_SCOPE_CYCLE_COUNTER(StateTree_Eval_TreeStart);
+			Eval.TreeStart(*this);
 
-				STATETREE_TRACE_EVALUATOR_EVENT(EvalIndex, EvalInstanceView, EStateTreeTraceEventType::OnTreeStarted);
-			}
+			STATETREE_TRACE_EVALUATOR_EVENT(EvalIndex, EvalInstanceView, EStateTreeTraceEventType::OnTreeStarted);
 		}
 	}
 
@@ -2507,12 +2502,10 @@ EStateTreeRunStatus FStateTreeExecutionContext::StartTemporaryEvaluatorsAndGloba
 		}
 
 		FStateTreeDataView TaskDataView = GetDataViewOrTemporary(CurrentParentFrame, CurrentFrame, Task.InstanceDataHandle);
-		bool bWasCreated = false;
 		if (!TaskDataView.IsValid())
 		{
 			TaskDataView = AddTemporaryInstance(CurrentFrame, FStateTreeIndex16(TaskIndex), Task.InstanceDataHandle, CurrentFrame.StateTree->DefaultInstanceData.GetStruct(Task.InstanceTemplateIndex.Get()));
 			check(TaskDataView.IsValid())
-			bWasCreated = true;
 		}
 
 		FNodeInstanceDataScope DataScope(*this, Task.InstanceDataHandle, TaskDataView);
@@ -2523,20 +2516,17 @@ EStateTreeRunStatus FStateTreeExecutionContext::StartTemporaryEvaluatorsAndGloba
 			CopyBatchOnActiveInstances(CurrentParentFrame, CurrentFrame, TaskDataView, Task.BindingsBatch);
 		}
 
-		if (bWasCreated)
-		{
-			STATETREE_LOG(Verbose, TEXT("  Start: '%s'"), *Task.Name.ToString());
-			UE_STATETREE_DEBUG_TASK_ENTER_STATE(this, CurrentStateTree, FStateTreeIndex16(TaskIndex));
-			QUICK_SCOPE_CYCLE_COUNTER(StateTree_Task_TreeStart);
-			const EStateTreeRunStatus TaskStatus = Task.EnterState(*this, Transition);
+		STATETREE_LOG(Verbose, TEXT("  Start: '%s'"), *Task.Name.ToString());
+		UE_STATETREE_DEBUG_TASK_ENTER_STATE(this, CurrentStateTree, FStateTreeIndex16(TaskIndex));
+		QUICK_SCOPE_CYCLE_COUNTER(StateTree_Task_TreeStart);
+		const EStateTreeRunStatus TaskStatus = Task.EnterState(*this, Transition);
 
-			STATETREE_TRACE_TASK_EVENT(TaskIndex, TaskDataView, EStateTreeTraceEventType::OnEntered, TaskStatus);
+		STATETREE_TRACE_TASK_EVENT(TaskIndex, TaskDataView, EStateTreeTraceEventType::OnEntered, TaskStatus);
 
-			if (TaskStatus != EStateTreeRunStatus::Running)
-			{
-				Result = TaskStatus;
-				break;
-			}
+		if (TaskStatus != EStateTreeRunStatus::Running)
+		{
+			Result = TaskStatus;
+			break;
 		}
 	}
 
