StateTree: Add icons in the editor to show which task ticks.

diff --git a/Engine/Plugins/Runtime/StateTree/Source/StateTreeEditorModule/Private/Customizations/StateTreeEditorNodeDetails.cpp b/Engine/Plugins/Runtime/StateTree/Source/StateTreeEditorModule/Private/Customizations/StateTreeEditorNodeDetails.cpp
--- a/Engine/Plugins/Runtime/StateTree/Source/StateTreeEditorModule/Private/Customizations/StateTreeEditorNodeDetails.cpp
+++ b/Engine/Plugins/Runtime/StateTree/Source/StateTreeEditorModule/Private/Customizations/StateTreeEditorNodeDetails.cpp
@@ -8,6 +8,8 @@
 #include "StateTree.h"
 #include "StateTreeEditor.h"
 #include "StateTreeEditorData.h"
+#include "StateTreeEditorUserSettings.h"
+#include "StateTreePropertyFunctionBase.h"
 #include "StateTreePropertyRef.h"
 #include "Widgets/Input/SButton.h"
 #include "Widgets/Input/SComboButton.h"
@@ -39,7 +41,6 @@
 #include "Styling/SlateTypes.h"
 #include "TextStyleDecorator.h"
 #include "Framework/Application/SlateApplication.h"
-#include "StateTreePropertyFunctionBase.h"
 
 #define LOCTEXT_NAMESPACE "StateTreeEditor"
 
@@ -393,6 +394,8 @@ void FStateTreeEditorNodeDetails::CustomizeHeader(TSharedRef<class IPropertyHand
 		return (RowBorder && RowBorder->IsHovered()) ? FSlateColor::UseForeground() : FSlateColor(FLinearColor::Transparent);
 	};
 
+	TSharedPtr<SBorder> FlagBorder;
+
 	HeaderRow
 		.WholeRowContent()
 		.VAlign(VAlign_Center)
@@ -501,6 +504,7 @@ void FStateTreeEditorNodeDetails::CustomizeHeader(TSharedRef<class IPropertyHand
 				.Padding(FMargin(0, 4))
 				[
 					SNew(SHorizontalBox)
+
 					// Icon
 					+ SHorizontalBox::Slot()
 					.AutoWidth()
@@ -512,6 +516,7 @@ void FStateTreeEditorNodeDetails::CustomizeHeader(TSharedRef<class IPropertyHand
 						.ColorAndOpacity(this, &FStateTreeEditorNodeDetails::GetIconColor)
 						.Visibility(this, &FStateTreeEditorNodeDetails::IsIconVisible)
 					]
+
 					// Rich text description and name edit 
 					+ SHorizontalBox::Slot()
 					.VAlign(VAlign_Center)
@@ -546,7 +551,19 @@ void FStateTreeEditorNodeDetails::CustomizeHeader(TSharedRef<class IPropertyHand
 							.Visibility(this, &FStateTreeEditorNodeDetails::IsNodeDescriptionVisible)
 						]
 					]
+
+					// Flags icons
+					+ SHorizontalBox::Slot()
+					.AutoWidth()
+					.VAlign(VAlign_Center)
+					.Padding(FMargin(4.0f, 0.0f))
+					[
+						SAssignNew(FlagBorder, SBorder)
+							.BorderImage(FStyleDefaults::GetNoBrush())
+							.Visibility(this, &FStateTreeEditorNodeDetails::AreFlagsVisible)
+					]
 				]
+
 				// Close parens
 				+ SHorizontalBox::Slot()
 				.AutoWidth()
@@ -563,10 +580,10 @@ void FStateTreeEditorNodeDetails::CustomizeHeader(TSharedRef<class IPropertyHand
 				.FillContentWidth(1.0f, 0.0f) // grow, no shrinking
 				.VAlign(VAlign_Center)
 				.HAlign(HAlign_Right)
-				.Padding(FMargin(8, 0, 2, 0))
+				.Padding(FMargin(8.0f, 0.0f, 2.0f, 0.0f))
 				[
 					SNew(SHorizontalBox)
-					
+
 					// Debugger labels
 					+ SHorizontalBox::Slot()
 					.AutoWidth()
@@ -643,6 +660,58 @@ void FStateTreeEditorNodeDetails::CustomizeHeader(TSharedRef<class IPropertyHand
 		.OverrideResetToDefault(ResetOverride)
 		.CopyAction(FUIAction(FExecuteAction::CreateSP(this, &FStateTreeEditorNodeDetails::OnCopyNode)))
 		.PasteAction(FUIAction(FExecuteAction::CreateSP(this, &FStateTreeEditorNodeDetails::OnPasteNode)));
+
+	if (FlagBorder)
+	{
+		TArray<void*> RawNodeDatas;
+		StructProperty->AccessRawData(RawNodeDatas);
+		bool bShowCallTick = false;
+		bool bShouldCallTickOnlyOnEvents = false;
+		for (void* RawNodeData : RawNodeDatas)
+		{
+			if (const FStateTreeEditorNode* EditorNode = reinterpret_cast<const FStateTreeEditorNode*>(RawNodeData))
+			{
+				if (const FStateTreeTaskBase* TreeTaskNodePtr = EditorNode->Node.GetPtr<FStateTreeTaskBase>())
+				{
+					bShowCallTick = bShowCallTick || TreeTaskNodePtr->bShouldCallTick;
+					bShouldCallTickOnlyOnEvents = bShouldCallTickOnlyOnEvents || TreeTaskNodePtr->bShouldCallTickOnlyOnEvents;
+				}
+			}
+		}
+
+		if (bShowCallTick || bShouldCallTickOnlyOnEvents)
+		{
+			TSharedRef<SHorizontalBox> Box = SNew(SHorizontalBox);
+			if (bShowCallTick)
+			{
+				Box->AddSlot()
+				.AutoWidth()
+				.VAlign(VAlign_Center)
+				[
+					SNew(SImage)
+					.Image(FStateTreeEditorStyle::Get().GetBrush("StateTreeEditor.Flags.Tick"))
+					.ColorAndOpacity(FSlateColor::UseSubduedForeground())
+					.ToolTipText(LOCTEXT("NodeTick", "The task ticks at runtime."))
+				];
+			}
+
+			if (bShouldCallTickOnlyOnEvents)
+			{
+				Box->AddSlot()
+				.AutoWidth()
+				.VAlign(VAlign_Center)
+				[
+					SNew(SImage)
+					.Image(FStateTreeEditorStyle::Get().GetBrush("StateTreeEditor.Flags.TickOnEvent"))
+					.ColorAndOpacity(FSlateColor::UseSubduedForeground())
+					.ToolTipText(LOCTEXT("NodeTickEvent", "The task ticks on event at runtime."))
+				];
+			}
+
+			FlagBorder->SetPadding(FMargin(4.f, 4.f, 4.f, 4.f));
+			FlagBorder->SetContent(Box);
+		}
+	}
 }
 
 
@@ -1309,6 +1378,12 @@ EVisibility FStateTreeEditorNodeDetails::AreParensVisible() const
 	return EVisibility::Collapsed;
 }
 
+EVisibility FStateTreeEditorNodeDetails::AreFlagsVisible() const
+{
+	bool bVisible = EnumHasAllFlags(GetDefault<UStateTreeEditorUserSettings>()->GetStatesViewDisplayNodeType(), EStateTreeEditorUserSettingsNodeType::Flag);
+	return bVisible ? EVisibility::Visible : EVisibility::Collapsed;
+}
+
 EVisibility FStateTreeEditorNodeDetails::IsIconVisible() const
 {
 	return UE::StateTreeEditor::EditorNodeUtils::IsIconVisible(StructProperty);
diff --git a/Engine/Plugins/Runtime/StateTree/Source/StateTreeEditorModule/Private/Customizations/StateTreeEditorNodeDetails.h b/Engine/Plugins/Runtime/StateTree/Source/StateTreeEditorModule/Private/Customizations/StateTreeEditorNodeDetails.h
--- a/Engine/Plugins/Runtime/StateTree/Source/StateTreeEditorModule/Private/Customizations/StateTreeEditorNodeDetails.h
+++ b/Engine/Plugins/Runtime/StateTree/Source/StateTreeEditorModule/Private/Customizations/StateTreeEditorNodeDetails.h
@@ -44,7 +44,6 @@ public:
 	virtual void CustomizeChildren(TSharedRef<IPropertyHandle> StructPropertyHandle, IDetailChildrenBuilder& StructBuilder, IPropertyTypeCustomizationUtils& StructCustomizationUtils) override;
 
 private:
-
 	bool ShouldResetToDefault(TSharedPtr<IPropertyHandle> PropertyHandle) const;
 	void ResetToDefault(TSharedPtr<IPropertyHandle> PropertyHandle);
 	void OnCopyNode();
@@ -79,6 +78,7 @@ private:
 	EVisibility IsOperandVisible() const;
 	EVisibility AreIndentButtonsVisible() const;
 	EVisibility AreParensVisible() const;
+	EVisibility AreFlagsVisible() const;
 	
 	EVisibility IsIconVisible() const;
 	const FSlateBrush* GetIcon() const;
diff --git a/Engine/Plugins/Runtime/StateTree/Source/StateTreeEditorModule/Private/SStateTreeViewRow.cpp b/Engine/Plugins/Runtime/StateTree/Source/StateTreeEditorModule/Private/SStateTreeViewRow.cpp
--- a/Engine/Plugins/Runtime/StateTree/Source/StateTreeEditorModule/Private/SStateTreeViewRow.cpp
+++ b/Engine/Plugins/Runtime/StateTree/Source/StateTreeEditorModule/Private/SStateTreeViewRow.cpp
@@ -17,6 +17,7 @@
 #include "StateTreeEditorUserSettings.h"
 #include "StateTreeState.h"
 #include "StateTreeTaskBase.h"
+#include "StateTreeTypes.h"
 #include "StateTreeViewModel.h"
 #include "Widgets/Views/SListView.h"
 #include "TextStyleDecorator.h"
@@ -136,6 +137,7 @@ void SStateTreeViewRow::Construct(const FArguments& InArgs, const TSharedRef<STa
 	
 	TSharedPtr<SVerticalBox> StateAndTasksVerticalBox;
 	TSharedPtr<SHorizontalBox> StateHorizontalBox;
+	TSharedPtr<SBorder> FlagBorder;
 
 	this->ChildSlot
 	.HAlign(HAlign_Fill)
@@ -297,6 +299,16 @@ void SStateTreeViewRow::Construct(const FArguments& InArgs, const TSharedRef<STa
 												.IsSelected(this, &SStateTreeViewRow::IsStateSelected)
 											]
 
+											// Flags icons
+											+ SHorizontalBox::Slot()
+											.VAlign(VAlign_Center)
+											.AutoWidth()
+											.Padding(FMargin(0.0f))
+											[
+												SAssignNew(FlagBorder, SBorder)
+												.BorderImage(FStyleDefaults::GetNoBrush())
+											]
+
 											// Linked State
 											+ SHorizontalBox::Slot()
 											.VAlign(VAlign_Center)
@@ -453,6 +465,60 @@ void SStateTreeViewRow::Construct(const FArguments& InArgs, const TSharedRef<STa
 			MakeTasksWidget(ViewBox)
 		];
 	}
+
+	if (EnumHasAllFlags(GetDefault<UStateTreeEditorUserSettings>()->GetStatesViewDisplayNodeType(), EStateTreeEditorUserSettingsNodeType::Flag))
+	{
+		const UStateTree* StateTree = StateTreeViewModel->GetStateTree();
+		if (State && StateTree)
+		{
+			const FCompactStateTreeState* RuntimeState = StateTree->GetStateFromHandle(StateTree->GetStateHandleFromId(State->ID));
+			if (RuntimeState->bHasTickTasks || RuntimeState->bHasTickTasksOnlyOnEvents || RuntimeState->bHasTransitionTasks)
+			{
+				TSharedRef<SHorizontalBox> Box = SNew(SHorizontalBox);
+				if (RuntimeState->bHasTickTasks)
+				{
+					Box->AddSlot()
+					.AutoWidth()
+					.VAlign(VAlign_Center)
+					[
+						SNew(SImage)
+						.Image(FStateTreeEditorStyle::Get().GetBrush("StateTreeEditor.Flags.Tick"))
+						.ColorAndOpacity(IconTint)
+						.ToolTipText(LOCTEXT("NodeTick", "The node ticks at runtime."))
+					];
+				}
+
+				if (RuntimeState->bHasTickTasksOnlyOnEvents)
+				{
+					Box->AddSlot()
+					.AutoWidth()
+					.VAlign(VAlign_Center)
+					[
+						SNew(SImage)
+						.Image(FStateTreeEditorStyle::Get().GetBrush("StateTreeEditor.Flags.TickOnEvent"))
+						.ColorAndOpacity(IconTint)
+						.ToolTipText(LOCTEXT("NodeTickEvent", "The node ticks on event at runtime."))
+					];
+				}
+
+				if (RuntimeState->bHasTransitionTasks)
+				{
+					Box->AddSlot()
+					.AutoWidth()
+					.VAlign(VAlign_Center)
+					[
+						SNew(SImage)
+						.Image(FStateTreeEditorStyle::Get().GetBrush("StateTreeEditor.Transitions"))
+						.ColorAndOpacity(IconTint)
+						.ToolTipText(LOCTEXT("NodeTickTransition", "The node ticks on to trigger transition."))
+					];
+				}
+
+				FlagBorder->SetPadding(FMargin(4.f, 4.f, 4.f, 4.f));
+				FlagBorder->SetContent(Box);
+			}
+		}
+	}
 }
 
 TSharedRef<SWidget> SStateTreeViewRow::MakeTasksWidget(const TSharedPtr<SScrollBox>& ViewBox)
diff --git a/Engine/Plugins/Runtime/StateTree/Source/StateTreeEditorModule/Private/StateTreeEditorStyle.cpp b/Engine/Plugins/Runtime/StateTree/Source/StateTreeEditorModule/Private/StateTreeEditorStyle.cpp
--- a/Engine/Plugins/Runtime/StateTree/Source/StateTreeEditorModule/Private/StateTreeEditorStyle.cpp
+++ b/Engine/Plugins/Runtime/StateTree/Source/StateTreeEditorModule/Private/StateTreeEditorStyle.cpp
@@ -370,7 +370,7 @@ FStateTreeEditorStyle::FStateTreeEditorStyle()
 	{
 		// From plugin
 		Set("StateTreeEditor.AddSiblingState", new IMAGE_BRUSH_SVG("Icons/Sibling_State", CoreStyleConstants::Icon16x16));
-		Set("StateTreeEditor.AddChildState", new IMAGE_BRUSH_SVG("Icons/Child_State", CoreStyleConstants::Icon16x16));
+		Set("StateTreeditor.AddChildState", new IMAGE_BRUSH_SVG("Icons/Child_State", CoreStyleConstants::Icon16x16));
 
 		Set("StateTreeEditor.PasteStatesAsSiblings", new IMAGE_BRUSH_SVG("Icons/Sibling_State", CoreStyleConstants::Icon16x16));
 		Set("StateTreeEditor.PasteStatesAsChildren", new IMAGE_BRUSH_SVG("Icons/Child_State", CoreStyleConstants::Icon16x16));
@@ -413,6 +413,11 @@ FStateTreeEditorStyle::FStateTreeEditorStyle()
 		Set("Node.Task", new IMAGE_BRUSH_SVG("Icons/Task", CoreStyleConstants::Icon16x16));
 		Set("Node.Text", new IMAGE_BRUSH_SVG("Icons/Text", CoreStyleConstants::Icon16x16));
 		Set("Node.Function", new IMAGE_BRUSH_SVG("Icons/Function", CoreStyleConstants::Icon16x16));
+
+
+		// Runtime flag
+		Set("StateTreeEditor.Flags.Tick", new IMAGE_BRUSH_SVG("Icons/Tick", CoreStyleConstants::Icon16x16));
+		Set("StateTreeEditor.Flags.TickOnEvent", new IMAGE_BRUSH_SVG("Icons/TickEvent", CoreStyleConstants::Icon16x16));
 	}
 }
 
diff --git a/Engine/Plugins/Runtime/StateTree/Source/StateTreeEditorModule/Private/StateTreeEditorUserSettings.h b/Engine/Plugins/Runtime/StateTree/Source/StateTreeEditorModule/Private/StateTreeEditorUserSettings.h
--- a/Engine/Plugins/Runtime/StateTree/Source/StateTreeEditorModule/Private/StateTreeEditorUserSettings.h
+++ b/Engine/Plugins/Runtime/StateTree/Source/StateTreeEditorModule/Private/StateTreeEditorUserSettings.h
@@ -11,7 +11,8 @@ enum class EStateTreeEditorUserSettingsNodeType : uint8
 	Condition = 1 << 0,
 	Task = 1 << 1,
 	Transition = 1 << 2,
-	All = (Condition | Task | Transition) UMETA(Hidden),
+	Flag = 1 << 3,
+	All = (Condition | Task | Transition | Flag) UMETA(Hidden),
 };
 ENUM_CLASS_FLAGS(EStateTreeEditorUserSettingsNodeType);
 
