StateTree: Fix UE_ENABLE_OPTIMIZATION_SHIP macro appears without matching UE_DISABLE_OPTIMIZATION_SHIP macro.
Fix issue 620452

#rnx

diff --git a/Engine/Plugins/Runtime/StateTree/Source/StateTreeTestSuite/Private/StateTreeTestBase.cpp b/Engine/Plugins/Runtime/StateTree/Source/StateTreeTestSuite/Private/StateTreeTestBase.cpp
--- a/Engine/Plugins/Runtime/StateTree/Source/StateTreeTestSuite/Private/StateTreeTestBase.cpp
+++ b/Engine/Plugins/Runtime/StateTree/Source/StateTreeTestSuite/Private/StateTreeTestBase.cpp
@@ -71,7 +71,5 @@ FGameplayTag FStateTreeTestBase::GetTestTag3()
 
 }//namespace UE::StateTree::Tests
 
-UE_ENABLE_OPTIMIZATION_SHIP
-
 #undef LOCTEXT_NAMESPACE
 
