StateTree: Fix localization conflict.
#rnx

diff --git a/Engine/Plugins/Runtime/StateTree/Source/StateTreeEditorModule/Private/Customizations/StateTreeEditorNodeDetails.cpp b/Engine/Plugins/Runtime/StateTree/Source/StateTreeEditorModule/Private/Customizations/StateTreeEditorNodeDetails.cpp
--- a/Engine/Plugins/Runtime/StateTree/Source/StateTreeEditorModule/Private/Customizations/StateTreeEditorNodeDetails.cpp
+++ b/Engine/Plugins/Runtime/StateTree/Source/StateTreeEditorModule/Private/Customizations/StateTreeEditorNodeDetails.cpp
@@ -691,7 +691,7 @@ void FStateTreeEditorNodeDetails::CustomizeHeader(TSharedRef<class IPropertyHand
 					SNew(SImage)
 					.Image(FStateTreeEditorStyle::Get().GetBrush("StateTreeEditor.Flags.Tick"))
 					.ColorAndOpacity(FSlateColor::UseSubduedForeground())
-					.ToolTipText(LOCTEXT("NodeTick", "The task ticks at runtime."))
+					.ToolTipText(LOCTEXT("TaskTick", "The task ticks at runtime."))
 				];
 			}
 
@@ -704,7 +704,7 @@ void FStateTreeEditorNodeDetails::CustomizeHeader(TSharedRef<class IPropertyHand
 					SNew(SImage)
 					.Image(FStateTreeEditorStyle::Get().GetBrush("StateTreeEditor.Flags.TickOnEvent"))
 					.ColorAndOpacity(FSlateColor::UseSubduedForeground())
-					.ToolTipText(LOCTEXT("NodeTickEvent", "The task ticks on event at runtime."))
+					.ToolTipText(LOCTEXT("TaskTickEvent", "The task ticks on event at runtime."))
 				];
 			}
 
