StateTree: Fix crash when adding or duplicating a State.

#rnx

diff --git a/Engine/Plugins/Runtime/StateTree/Source/StateTreeEditorModule/Private/SStateTreeViewRow.cpp b/Engine/Plugins/Runtime/StateTree/Source/StateTreeEditorModule/Private/SStateTreeViewRow.cpp
--- a/Engine/Plugins/Runtime/StateTree/Source/StateTreeEditorModule/Private/SStateTreeViewRow.cpp
+++ b/Engine/Plugins/Runtime/StateTree/Source/StateTreeEditorModule/Private/SStateTreeViewRow.cpp
@@ -471,51 +471,53 @@ void SStateTreeViewRow::Construct(const FArguments& InArgs, const TSharedRef<STa
 		const UStateTree* StateTree = StateTreeViewModel->GetStateTree();
 		if (State && StateTree)
 		{
-			const FCompactStateTreeState* RuntimeState = StateTree->GetStateFromHandle(StateTree->GetStateHandleFromId(State->ID));
-			if (RuntimeState->bHasTickTasks || RuntimeState->bHasTickTasksOnlyOnEvents || RuntimeState->bHasTransitionTasks)
+			if (const FCompactStateTreeState* RuntimeState = StateTree->GetStateFromHandle(StateTree->GetStateHandleFromId(State->ID)))
 			{
-				TSharedRef<SHorizontalBox> Box = SNew(SHorizontalBox);
-				if (RuntimeState->bHasTickTasks)
+				if (RuntimeState->bHasTickTasks || RuntimeState->bHasTickTasksOnlyOnEvents || RuntimeState->bHasTransitionTasks)
 				{
-					Box->AddSlot()
-					.AutoWidth()
-					.VAlign(VAlign_Center)
-					[
-						SNew(SImage)
-						.Image(FStateTreeEditorStyle::Get().GetBrush("StateTreeEditor.Flags.Tick"))
-						.ColorAndOpacity(IconTint)
-						.ToolTipText(LOCTEXT("NodeTick", "The node ticks at runtime."))
-					];
-				}
+					TSharedRef<SHorizontalBox> Box = SNew(SHorizontalBox);
+					if (RuntimeState->bHasTickTasks)
+					{
+						Box->AddSlot()
+						.AutoWidth()
+						.VAlign(VAlign_Center)
+						[
+							SNew(SImage)
+							.Image(FStateTreeEditorStyle::Get().GetBrush("StateTreeEditor.Flags.Tick"))
+							.ColorAndOpacity(IconTint)
+							.ToolTipText(LOCTEXT("NodeTick", "The node ticks at runtime."))
+						];
+					}
 
-				if (RuntimeState->bHasTickTasksOnlyOnEvents)
-				{
-					Box->AddSlot()
-					.AutoWidth()
-					.VAlign(VAlign_Center)
-					[
-						SNew(SImage)
-						.Image(FStateTreeEditorStyle::Get().GetBrush("StateTreeEditor.Flags.TickOnEvent"))
-						.ColorAndOpacity(IconTint)
-						.ToolTipText(LOCTEXT("NodeTickEvent", "The node ticks on event at runtime."))
-					];
-				}
+					if (RuntimeState->bHasTickTasksOnlyOnEvents)
+					{
+						Box->AddSlot()
+						.AutoWidth()
+						.VAlign(VAlign_Center)
+						[
+							SNew(SImage)
+							.Image(FStateTreeEditorStyle::Get().GetBrush("StateTreeEditor.Flags.TickOnEvent"))
+							.ColorAndOpacity(IconTint)
+							.ToolTipText(LOCTEXT("NodeTickEvent", "The node ticks on event at runtime."))
+						];
+					}
 
-				if (RuntimeState->bHasTransitionTasks)
-				{
-					Box->AddSlot()
-					.AutoWidth()
-					.VAlign(VAlign_Center)
-					[
-						SNew(SImage)
-						.Image(FStateTreeEditorStyle::Get().GetBrush("StateTreeEditor.Transitions"))
-						.ColorAndOpacity(IconTint)
-						.ToolTipText(LOCTEXT("NodeTickTransition", "The node ticks on to trigger transition."))
-					];
-				}
+					if (RuntimeState->bHasTransitionTasks)
+					{
+						Box->AddSlot()
+						.AutoWidth()
+						.VAlign(VAlign_Center)
+						[
+							SNew(SImage)
+							.Image(FStateTreeEditorStyle::Get().GetBrush("StateTreeEditor.Transitions"))
+							.ColorAndOpacity(IconTint)
+							.ToolTipText(LOCTEXT("NodeTickTransition", "The node ticks on to trigger transition."))
+						];
+					}
 
-				FlagBorder->SetPadding(FMargin(4.f, 4.f, 4.f, 4.f));
-				FlagBorder->SetContent(Box);
+					FlagBorder->SetPadding(FMargin(4.f, 4.f, 4.f, 4.f));
+					FlagBorder->SetContent(Box);
+				}
 			}
 		}
 	}
