StateTree: Split test for transitions and linkedasset.

diff --git a/Engine/Plugins/Runtime/StateTree/Source/StateTreeTestSuite/Private/StateTreeLinkedStateTest.cpp b/Engine/Plugins/Runtime/StateTree/Source/StateTreeTestSuite/Private/StateTreeLinkedStateTest.cpp
new file mode 100644
--- /dev/null
+++ b/Engine/Plugins/Runtime/StateTree/Source/StateTreeTestSuite/Private/StateTreeLinkedStateTest.cpp
@@ -0,0 +1,343 @@
+// Copyright Epic Games, Inc. All Rights Reserved.
+
+#include "StateTreeTest.h"
+#include "StateTreeTestBase.h"
+#include "StateTreeTestTypes.h"
+
+#include "StateTreeCompilerLog.h"
+#include "StateTreeEditorData.h"
+#include "StateTreeCompiler.h"
+#include "Conditions/StateTreeCommonConditions.h"
+
+#define LOCTEXT_NAMESPACE "AITestSuite_StateTreeTest"
+
+UE_DISABLE_OPTIMIZATION_SHIP
+
+namespace UE::StateTree::Tests
+{
+
+struct FStateTreeTest_FailEnterLinkedAsset : FStateTreeTestBase
+{
+	virtual bool InstantTest() override
+	{
+		FStateTreeCompilerLog Log;
+
+		// Asset 2
+		UStateTree& StateTree2 = NewStateTree();
+		UStateTreeEditorData& EditorData2 = *Cast<UStateTreeEditorData>(StateTree2.EditorData);
+		UStateTreeState& Root2 = EditorData2.AddSubTree(FName(TEXT("Root2")));
+		TStateTreeEditorNode<FTestTask_Stand>& Task2 = Root2.AddTask<FTestTask_Stand>(FName(TEXT("Task2")));
+		TStateTreeEditorNode<FTestTask_Stand>& GlobalTask2 = EditorData2.AddGlobalTask<FTestTask_Stand>(FName(TEXT("GlobalTask2")));
+		GlobalTask2.GetInstanceData().Value = 123;
+
+		// Always failing enter condition
+		TStateTreeEditorNode<FStateTreeCompareIntCondition>& IntCond2 = Root2.AddEnterCondition<FStateTreeCompareIntCondition>();
+		EditorData2.AddPropertyBinding(GlobalTask2, TEXT("Value"), IntCond2, TEXT("Left"));
+		IntCond2.GetInstanceData().Right = 0;
+
+		FStateTreeCompiler Compiler2(Log);
+		const bool bResult2 = Compiler2.Compile(StateTree2);
+		AITEST_TRUE("StateTree2 should get compiled", bResult2);
+
+		// Main asset
+		UStateTree& StateTree = NewStateTree();
+		UStateTreeEditorData& EditorData = *Cast<UStateTreeEditorData>(StateTree.EditorData);
+
+		UStateTreeState& Root = EditorData.AddSubTree(FName(TEXT("Root1")));
+		UStateTreeState& A1 = Root.AddChildState(FName(TEXT("A1")), EStateTreeStateType::LinkedAsset);
+		A1.SetLinkedStateAsset(&StateTree2);
+
+		UStateTreeState& B1 = Root.AddChildState(FName(TEXT("B1")), EStateTreeStateType::State);
+		TStateTreeEditorNode<FTestTask_Stand>& Task1 = B1.AddTask<FTestTask_Stand>(FName(TEXT("Task1")));
+
+		FStateTreeCompiler Compiler(Log);
+		const bool bResult = Compiler.Compile(StateTree);
+		AITEST_TRUE("StateTree should get compiled", bResult);
+
+		const FString EnterStateStr(TEXT("EnterState"));
+		const FString ExitStateStr(TEXT("ExitState"));
+
+		{
+			EStateTreeRunStatus Status = EStateTreeRunStatus::Unset;
+			FStateTreeInstanceData InstanceData;
+			FTestStateTreeExecutionContext Exec(StateTree, StateTree, InstanceData);
+			const bool bInitSucceeded = Exec.IsValid();
+			AITEST_TRUE("StateTree should init", bInitSucceeded);
+
+			Status = Exec.Start();
+			AITEST_EQUAL("Start should complete with Running", Status, EStateTreeRunStatus::Running);
+			AITEST_TRUE("StateTree should enter GlobalTask2", Exec.Expect(GlobalTask2.GetName(), EnterStateStr));
+			AITEST_TRUE("StateTree should exit GlobalTask2", Exec.Expect(GlobalTask2.GetName(), ExitStateStr));
+			AITEST_FALSE("StateTree should not enter Task2", Exec.Expect(Task2.GetName(), EnterStateStr));
+			AITEST_TRUE("StateTree should enter Task1", Exec.Expect(Task1.GetName(), EnterStateStr));
+
+			Exec.Stop();
+		}
+
+		return true;
+	}
+};
+IMPLEMENT_AI_INSTANT_TEST(FStateTreeTest_FailEnterLinkedAsset, "System.StateTree.LinkedAsset.FailEnter");
+
+struct FStateTreeTest_EnterAndExitLinkedAsset : FStateTreeTestBase
+{
+	virtual bool InstantTest() override
+	{
+		FStateTreeCompilerLog Log;
+
+		// Asset 2
+		UStateTree& StateTree2 = NewStateTree();
+		UStateTreeEditorData& EditorData2 = *Cast<UStateTreeEditorData>(StateTree2.EditorData);
+		UStateTreeState& Root2 = EditorData2.AddSubTree(FName(TEXT("Root2")));
+		TStateTreeEditorNode<FTestTask_Stand>& Task2 = Root2.AddTask<FTestTask_Stand>(FName(TEXT("Task2")));
+		TStateTreeEditorNode<FTestTask_Stand>& GlobalTask2 = EditorData2.AddGlobalTask<FTestTask_Stand>(FName(TEXT("GlobalTask2")));
+		GlobalTask2.GetNode().TicksToCompletion = 2;
+
+		FStateTreeCompiler Compiler2(Log);
+		const bool bResult2 = Compiler2.Compile(StateTree2);
+		AITEST_TRUE("StateTree2 should get compiled", bResult2);
+
+		// Main asset
+		UStateTree& StateTree = NewStateTree();
+		UStateTreeEditorData& EditorData = *Cast<UStateTreeEditorData>(StateTree.EditorData);
+
+		UStateTreeState& Root = EditorData.AddSubTree(FName(TEXT("Root1")));
+		UStateTreeState& A1 = Root.AddChildState(FName(TEXT("A1")), EStateTreeStateType::LinkedAsset);
+		A1.AddTransition(EStateTreeTransitionTrigger::OnStateCompleted, EStateTreeTransitionType::NextState);
+		A1.SetLinkedStateAsset(&StateTree2);
+
+		UStateTreeState& B1 = Root.AddChildState(FName(TEXT("B1")), EStateTreeStateType::State);
+		TStateTreeEditorNode<FTestTask_Stand>& Task1 = B1.AddTask<FTestTask_Stand>(FName(TEXT("Task1")));
+
+		FStateTreeCompiler Compiler(Log);
+		const bool bResult = Compiler.Compile(StateTree);
+		AITEST_TRUE("StateTree should get compiled", bResult);
+
+		const FString EnterStateStr(TEXT("EnterState"));
+		const FString ExitStateStr(TEXT("ExitState"));
+
+		{
+			EStateTreeRunStatus Status = EStateTreeRunStatus::Unset;
+			FStateTreeInstanceData InstanceData;
+			FTestStateTreeExecutionContext Exec(StateTree, StateTree, InstanceData);
+			const bool bInitSucceeded = Exec.IsValid();
+			AITEST_TRUE("StateTree should init", bInitSucceeded);
+
+			Status = Exec.Start();
+			AITEST_EQUAL("Start should complete with Running", Status, EStateTreeRunStatus::Running);
+			AITEST_TRUE("StateTree should enter GlobalTask2", Exec.Expect(GlobalTask2.GetName(), EnterStateStr));
+			AITEST_FALSE("StateTree should not exit GlobalTask2", Exec.Expect(GlobalTask2.GetName(), ExitStateStr));
+			AITEST_TRUE("StateTree should enter Task2", Exec.Expect(Task2.GetName(), EnterStateStr));
+			AITEST_FALSE("StateTree should not exit Task2", Exec.Expect(Task2.GetName(), ExitStateStr));
+			AITEST_FALSE("StateTree should not enter Task1", Exec.Expect(Task1.GetName(), EnterStateStr));
+			Exec.LogClear();
+
+			Status = Exec.Tick(0.1f);
+			AITEST_EQUAL("Tick should complete with Running", Status, EStateTreeRunStatus::Running);
+			AITEST_FALSE("StateTree should not enter GlobalTask2", Exec.Expect(GlobalTask2.GetName(), EnterStateStr));
+			AITEST_TRUE("StateTree should exit GlobalTask2", Exec.Expect(GlobalTask2.GetName(), ExitStateStr));
+			AITEST_FALSE("StateTree should not enter Task2", Exec.Expect(Task2.GetName(), EnterStateStr));
+			AITEST_TRUE("StateTree should exit Task2", Exec.Expect(Task2.GetName(), ExitStateStr));
+			AITEST_TRUE("StateTree should enter Task1", Exec.Expect(Task1.GetName(), EnterStateStr));
+			Exec.LogClear();
+
+			Exec.Stop();
+		}
+
+		return true;
+	}
+};
+IMPLEMENT_AI_INSTANT_TEST(FStateTreeTest_EnterAndExitLinkedAsset, "System.StateTree.LinkedAsset.EnterAndExit");
+
+struct FStateTreeTest_MultipleSameLinkedAsset : FStateTreeTestBase
+{
+	virtual bool InstantTest() override
+	{
+		//Tree 1
+		//	RootA
+		//		StateB -> Next
+		//		StateC -> Next
+		//		StateLinkedTreeD1 (Tree2) -> Next
+		//		StateLinkedTreeD2 (Tree2) -> Next
+		//		StateLinkedTreeD3 (Tree2) -> Next
+		//		StateE -> Succeeded
+		//Tree 2
+		//  Global task and parameter
+		//	RootE
+		//		StateF (with transition OnTick to succeeded)
+
+		FStateTreeCompilerLog Log;
+
+		// Asset 2
+		UStateTree& StateTree2 = NewStateTree();
+		FGuid RootParameter_ValueID;
+		{
+			UStateTreeEditorData& EditorData2 = *Cast<UStateTreeEditorData>(StateTree2.EditorData);
+			{
+				// Parameters
+				EditorData2.RootParameters.Parameters.AddProperty("Value", EPropertyBagPropertyType::Int32);
+				EditorData2.RootParameters.Parameters.SetValueInt32("Value", -111);
+				RootParameter_ValueID = EditorData2.RootParameters.Parameters.FindPropertyDescByName("Value")->ID;
+
+				TStateTreeEditorNode<FTestTask_PrintValue>& GlobalTask = EditorData2.AddGlobalTask<FTestTask_PrintValue>("Tree2Global");
+				GlobalTask.GetInstanceData().Value = -1;
+				EditorData2.AddPropertyBinding(FStateTreePropertyPath(EditorData2.RootParameters.ID, TEXT("Value")), FStateTreePropertyPath(GlobalTask.ID, TEXT("Value")));
+			}
+
+			UStateTreeState& RootC = EditorData2.AddSubTree("RootE");
+			{
+				UStateTreeState& StateE = RootC.AddChildState("StateF", EStateTreeStateType::State);
+				FStateTreeTransition& Transition = StateE.AddTransition(EStateTreeTransitionTrigger::OnTick, EStateTreeTransitionType::Succeeded);
+				Transition.bDelayTransition = true;
+				Transition.DelayDuration = 1.0;
+
+				TStateTreeEditorNode<FTestTask_PrintValue>& Task1 = StateE.AddTask<FTestTask_PrintValue>("Task1");
+				Task1.GetInstanceData().Value = -2;
+				EditorData2.AddPropertyBinding(FStateTreePropertyPath(EditorData2.RootParameters.ID, TEXT("Value")), FStateTreePropertyPath(Task1.ID, TEXT("Value")));
+			}
+
+			FStateTreeCompiler Compiler2(Log);
+			const bool bResult2 = Compiler2.Compile(StateTree2);
+			AITEST_TRUE("StateTree2 should get compiled", bResult2);
+		}
+
+		// Main asset
+		UStateTree& StateTree1 = NewStateTree();
+		{
+			UStateTreeEditorData& EditorData = *Cast<UStateTreeEditorData>(StateTree1.EditorData);
+
+			UStateTreeState& Root = EditorData.AddSubTree("RootA");
+			{
+				TStateTreeEditorNode<FTestTask_PrintValue>& GlobalTask = EditorData.AddGlobalTask<FTestTask_PrintValue>("Tree1Global");
+				GlobalTask.GetInstanceData().Value = 99;
+			}
+			{
+				UStateTreeState& StateB = Root.AddChildState("StateB", EStateTreeStateType::State);
+				TStateTreeEditorNode<FTestTask_PrintValue>& Task = StateB.AddTask<FTestTask_PrintValue>("TaskB");
+				Task.GetInstanceData().Value = 1;
+				FStateTreeTransition& Transition = StateB.AddTransition(EStateTreeTransitionTrigger::OnTick, EStateTreeTransitionType::NextState);
+				Transition.bDelayTransition = true;
+				Transition.DelayDuration = 1.0;
+			}
+			{
+				UStateTreeState& StateB = Root.AddChildState("StateC", EStateTreeStateType::State);
+				TStateTreeEditorNode<FTestTask_PrintValue>& Task = StateB.AddTask<FTestTask_PrintValue>("TaskC");
+				Task.GetInstanceData().Value = 2;
+				FStateTreeTransition& Transition = StateB.AddTransition(EStateTreeTransitionTrigger::OnTick, EStateTreeTransitionType::NextState);
+				Transition.bDelayTransition = true;
+				Transition.DelayDuration = 1.0;
+			}
+			{
+				UStateTreeState& C1 = Root.AddChildState("StateLinkedTreeD1", EStateTreeStateType::LinkedAsset);
+				C1.SetLinkedStateAsset(&StateTree2);
+				C1.AddTransition(EStateTreeTransitionTrigger::OnStateSucceeded, EStateTreeTransitionType::NextState);
+				C1.SetParametersPropertyOverridden(RootParameter_ValueID, true);
+				C1.Parameters.Parameters.SetValueInt32("Value", 111);
+			}
+			{
+				UStateTreeState& C2 = Root.AddChildState("StateLinkedTreeD2", EStateTreeStateType::LinkedAsset);
+				C2.SetLinkedStateAsset(&StateTree2);
+				C2.AddTransition(EStateTreeTransitionTrigger::OnStateSucceeded, EStateTreeTransitionType::NextState);
+				C2.SetParametersPropertyOverridden(RootParameter_ValueID, true);
+				C2.Parameters.Parameters.SetValueInt32("Value", 222);
+			}
+			{
+				UStateTreeState& C3 = Root.AddChildState("StateLinkedTreeD3", EStateTreeStateType::LinkedAsset);
+				C3.SetLinkedStateAsset(&StateTree2);
+				C3.AddTransition(EStateTreeTransitionTrigger::OnStateSucceeded, EStateTreeTransitionType::NextState);
+				C3.SetParametersPropertyOverridden(RootParameter_ValueID, true);
+				C3.Parameters.Parameters.SetValueInt32("Value", 333);
+			}
+			{
+				UStateTreeState& StateD = Root.AddChildState("StateE", EStateTreeStateType::State);
+				TStateTreeEditorNode<FTestTask_PrintValue>& Task = StateD.AddTask<FTestTask_PrintValue>("TaskE");
+				Task.GetInstanceData().Value = 3;
+				FStateTreeTransition& Transition = StateD.AddTransition(EStateTreeTransitionTrigger::OnTick, EStateTreeTransitionType::GotoState, &Root);
+				Transition.bDelayTransition = true;
+				Transition.DelayDuration = 1.0;
+			}
+
+			FStateTreeCompiler Compiler(Log);
+			const bool bResult = Compiler.Compile(StateTree1);
+			AITEST_TRUE("StateTree should get compiled", bResult);
+		}
+
+		const FString EnterStateStr(TEXT("EnterState"));
+		const FString ExitStateStr(TEXT("ExitState"));
+
+		{
+			EStateTreeRunStatus Status = EStateTreeRunStatus::Unset;
+			FStateTreeInstanceData InstanceData;
+			FTestStateTreeExecutionContext Exec(StateTree1, StateTree1, InstanceData);
+			const bool bInitSucceeded = Exec.IsValid();
+			AITEST_TRUE("StateTree should init", bInitSucceeded);
+
+			Status = Exec.Start();
+			AITEST_EQUAL("Start should complete with Running", Status, EStateTreeRunStatus::Running);
+			AITEST_TRUE("Start should enter Global tasks", Exec.Expect("Tree1Global", TEXT("EnterState99")));
+			AITEST_TRUE("Start should enter StateB", Exec.Expect("TaskB", TEXT("EnterState1")));
+			Exec.LogClear();
+
+			Status = Exec.Tick(1.5f); // over tick, should trigger
+			AITEST_EQUAL("1st Tick should complete with Running", Status, EStateTreeRunStatus::Running);
+			AITEST_TRUE("1st should tick Global tasks", Exec.Expect("Tree1Global", TEXT("Tick99")));
+			AITEST_TRUE("1st Tick tick the StateB", Exec.Expect("TaskB", TEXT("Tick1")));
+			Exec.LogClear();
+
+			Status = Exec.Tick(1.0f);
+			AITEST_EQUAL("2nd Tick should complete with Running", Status, EStateTreeRunStatus::Running);
+			AITEST_TRUE("2nd Tick should tick Global tasks", Exec.Expect("Tree1Global", TEXT("Tick99")));
+			AITEST_TRUE("2nd Tick should tick the StateB", Exec.Expect("TaskB", TEXT("Tick1")));
+			AITEST_TRUE("2nd Tick should exit the StateB", Exec.Expect("TaskB", TEXT("ExitState1")));
+			AITEST_TRUE("2nd Tick should enter the StateC", Exec.Expect("TaskC", TEXT("EnterState2")));
+			Exec.LogClear();
+
+			Status = Exec.Tick(1.0f);
+			AITEST_EQUAL("3rd Tick should complete with Running", Status, EStateTreeRunStatus::Running);
+			AITEST_TRUE("3rd should tick Global tasks", Exec.Expect("Tree1Global", TEXT("Tick99")));
+			AITEST_TRUE("3rd Tick tick the StateB", Exec.Expect("TaskC", TEXT("Tick2")));
+			Exec.LogClear();
+
+			Status = Exec.Tick(1.0f);
+			AITEST_EQUAL("4th Tick should complete with Running", Status, EStateTreeRunStatus::Running);
+			AITEST_TRUE("4th Tick should tick Global tasks", Exec.Expect("Tree1Global", TEXT("Tick99")));
+			AITEST_TRUE("4th Tick should tick the StateB", Exec.Expect("TaskC", TEXT("Tick2")));
+			//AITEST_TRUE("4th Tick should enter the Global", Exec.Expect("Tree2Global", TEXT("EnterState111")));
+			AITEST_TRUE("4th Tick should exit the StateB", Exec.Expect("TaskC", TEXT("ExitState2")));
+			AITEST_TRUE("4th Tick should enter the StateLinkedTreeC1", Exec.Expect("Task1", TEXT("EnterState111")));
+			Exec.LogClear();
+
+			Status = Exec.Tick(0.001f);
+			AITEST_EQUAL("5th Tick should complete with Running", Status, EStateTreeRunStatus::Running);
+			AITEST_TRUE("5h Tick should tick the global tasks", Exec.Expect("Tree1Global", TEXT("Tick99")));
+			AITEST_TRUE("5th Tick should tick the global tasks", Exec.Expect("Tree2Global", TEXT("Tick111")));
+			AITEST_TRUE("5th Tick should tick the StateLinkedTreeC1", Exec.Expect("Task1", TEXT("Tick111")));
+			Exec.LogClear();
+
+			Status = Exec.Tick(1.0f);
+			AITEST_EQUAL("5th Tick should complete with Running", Status, EStateTreeRunStatus::Running);
+			AITEST_TRUE("5th Tick should tick the global tasks", Exec.Expect("Tree1Global", TEXT("Tick99")));
+			AITEST_TRUE("5th Tick should tick the global tasks", Exec.Expect("Tree2Global", TEXT("Tick111")));
+			AITEST_TRUE("5th Tick should tick the StateLinkedTreeC1", Exec.Expect("Task1", TEXT("Tick111")));
+			//AITEST_TRUE("5th Tick should exit the global tasks", Exec.Expect("Tree2Global", TEXT("ExitState111")));
+			//AITEST_TRUE("5th Tick should exit the StateLinkedTreeC1", Exec.Expect("Task1", TEXT("ExitState111")));
+			//AITEST_TRUE("5th Tick should enter the global tasks", Exec.Expect("Tree2Global", TEXT("EnterState222")));
+			AITEST_TRUE("5th Tick should enter the StateLinkedTreeC2", Exec.Expect("Task1", TEXT("EnterState222")));
+			Exec.LogClear();
+
+			Status = Exec.Tick(1.0f);
+			Exec.LogClear();
+
+			Exec.Stop();
+		}
+
+		return true;
+	}
+};
+IMPLEMENT_AI_INSTANT_TEST(FStateTreeTest_MultipleSameLinkedAsset, "System.StateTree.LinkedAsset.MultipleSameTree");
+
+} // namespace UE::StateTree::Tests
+
+UE_ENABLE_OPTIMIZATION_SHIP
+
+#undef LOCTEXT_NAMESPACE
diff --git a/Engine/Plugins/Runtime/StateTree/Source/StateTreeTestSuite/Private/StateTreeTest.cpp b/Engine/Plugins/Runtime/StateTree/Source/StateTreeTestSuite/Private/StateTreeTest.cpp
--- a/Engine/Plugins/Runtime/StateTree/Source/StateTreeTestSuite/Private/StateTreeTest.cpp
+++ b/Engine/Plugins/Runtime/StateTree/Source/StateTreeTestSuite/Private/StateTreeTest.cpp
@@ -1462,138 +1462,6 @@ struct FStateTreeTest_DeferredStop_ExitTask : FStateTreeTest_DeferredStop
 };
 IMPLEMENT_AI_INSTANT_TEST(FStateTreeTest_DeferredStop_ExitTask, "System.StateTree.DeferredStop.ExitTask");
 
-struct FStateTreeTest_FailEnterLinkedAsset : FStateTreeTestBase
-{
-	virtual bool InstantTest() override
-	{
-		FStateTreeCompilerLog Log;
-
-		// Asset 2
-		UStateTree& StateTree2 = NewStateTree();
-		UStateTreeEditorData& EditorData2 = *Cast<UStateTreeEditorData>(StateTree2.EditorData);
-		UStateTreeState& Root2 = EditorData2.AddSubTree(FName(TEXT("Root2")));
-		TStateTreeEditorNode<FTestTask_Stand>& Task2 = Root2.AddTask<FTestTask_Stand>(FName(TEXT("Task2")));
-		TStateTreeEditorNode<FTestTask_Stand>& GlobalTask2 = EditorData2.AddGlobalTask<FTestTask_Stand>(FName(TEXT("GlobalTask2")));
-		GlobalTask2.GetInstanceData().Value = 123;
-
-		// Always failing enter condition
-		TStateTreeEditorNode<FStateTreeCompareIntCondition>& IntCond2 = Root2.AddEnterCondition<FStateTreeCompareIntCondition>();
-		EditorData2.AddPropertyBinding(GlobalTask2, TEXT("Value"), IntCond2, TEXT("Left"));
-		IntCond2.GetInstanceData().Right = 0;
-
-		FStateTreeCompiler Compiler2(Log);
-		const bool bResult2 = Compiler2.Compile(StateTree2);
-		AITEST_TRUE("StateTree2 should get compiled", bResult2);
-
-		// Main asset
-		UStateTree& StateTree = NewStateTree();
-		UStateTreeEditorData& EditorData = *Cast<UStateTreeEditorData>(StateTree.EditorData);
-		
-		UStateTreeState& Root = EditorData.AddSubTree(FName(TEXT("Root1")));
-		UStateTreeState& A1 = Root.AddChildState(FName(TEXT("A1")), EStateTreeStateType::LinkedAsset);
-		A1.SetLinkedStateAsset(&StateTree2);
-
-		UStateTreeState& B1 = Root.AddChildState(FName(TEXT("B1")), EStateTreeStateType::State);
-		TStateTreeEditorNode<FTestTask_Stand>& Task1 = B1.AddTask<FTestTask_Stand>(FName(TEXT("Task1")));
-
-		FStateTreeCompiler Compiler(Log);
-		const bool bResult = Compiler.Compile(StateTree);
-		AITEST_TRUE("StateTree should get compiled", bResult);
-
-		const FString EnterStateStr(TEXT("EnterState"));
-		const FString ExitStateStr(TEXT("ExitState"));
-
-		{
-			EStateTreeRunStatus Status = EStateTreeRunStatus::Unset;
-			FStateTreeInstanceData InstanceData;
-			FTestStateTreeExecutionContext Exec(StateTree, StateTree, InstanceData);
-			const bool bInitSucceeded = Exec.IsValid();
-			AITEST_TRUE("StateTree should init", bInitSucceeded);
-
-			Status = Exec.Start();
-			AITEST_EQUAL("Start should complete with Running", Status, EStateTreeRunStatus::Running);
-			AITEST_TRUE("StateTree should enter GlobalTask2", Exec.Expect(GlobalTask2.GetName(), EnterStateStr));
-			AITEST_TRUE("StateTree should exit GlobalTask2", Exec.Expect(GlobalTask2.GetName(), ExitStateStr));
-			AITEST_FALSE("StateTree should not enter Task2", Exec.Expect(Task2.GetName(), EnterStateStr));
-			AITEST_TRUE("StateTree should enter Task1", Exec.Expect(Task1.GetName(), EnterStateStr));
-
-			Exec.Stop();
-		}
-
-		return true;
-	}
-};
-IMPLEMENT_AI_INSTANT_TEST(FStateTreeTest_FailEnterLinkedAsset, "System.StateTree.FailEnterLinkedAsset");
-
-struct FStateTreeTest_EnterAndExitLinkedAsset : FStateTreeTestBase
-{
-	virtual bool InstantTest() override
-	{
-		FStateTreeCompilerLog Log;
-
-		// Asset 2
-		UStateTree& StateTree2 = NewStateTree();
-		UStateTreeEditorData& EditorData2 = *Cast<UStateTreeEditorData>(StateTree2.EditorData);
-		UStateTreeState& Root2 = EditorData2.AddSubTree(FName(TEXT("Root2")));
-		TStateTreeEditorNode<FTestTask_Stand>& Task2 = Root2.AddTask<FTestTask_Stand>(FName(TEXT("Task2")));
-		TStateTreeEditorNode<FTestTask_Stand>& GlobalTask2 = EditorData2.AddGlobalTask<FTestTask_Stand>(FName(TEXT("GlobalTask2")));
-		GlobalTask2.GetNode().TicksToCompletion = 2;
-
-		FStateTreeCompiler Compiler2(Log);
-		const bool bResult2 = Compiler2.Compile(StateTree2);
-		AITEST_TRUE("StateTree2 should get compiled", bResult2);
-
-		// Main asset
-		UStateTree& StateTree = NewStateTree();
-		UStateTreeEditorData& EditorData = *Cast<UStateTreeEditorData>(StateTree.EditorData);
-		
-		UStateTreeState& Root = EditorData.AddSubTree(FName(TEXT("Root1")));
-		UStateTreeState& A1 = Root.AddChildState(FName(TEXT("A1")), EStateTreeStateType::LinkedAsset);
-		A1.AddTransition(EStateTreeTransitionTrigger::OnStateCompleted, EStateTreeTransitionType::NextState);
-		A1.SetLinkedStateAsset(&StateTree2);
-
-		UStateTreeState& B1 = Root.AddChildState(FName(TEXT("B1")), EStateTreeStateType::State);
-		TStateTreeEditorNode<FTestTask_Stand>& Task1 = B1.AddTask<FTestTask_Stand>(FName(TEXT("Task1")));
-
-		FStateTreeCompiler Compiler(Log);
-		const bool bResult = Compiler.Compile(StateTree);
-		AITEST_TRUE("StateTree should get compiled", bResult);
-
-		const FString EnterStateStr(TEXT("EnterState"));
-		const FString ExitStateStr(TEXT("ExitState"));
-
-		{
-			EStateTreeRunStatus Status = EStateTreeRunStatus::Unset;
-			FStateTreeInstanceData InstanceData;
-			FTestStateTreeExecutionContext Exec(StateTree, StateTree, InstanceData);
-			const bool bInitSucceeded = Exec.IsValid();
-			AITEST_TRUE("StateTree should init", bInitSucceeded);
-
-			Status = Exec.Start();
-			AITEST_EQUAL("Start should complete with Running", Status, EStateTreeRunStatus::Running);
-			AITEST_TRUE("StateTree should enter GlobalTask2", Exec.Expect(GlobalTask2.GetName(), EnterStateStr));
-			AITEST_FALSE("StateTree should not exit GlobalTask2", Exec.Expect(GlobalTask2.GetName(), ExitStateStr));
-			AITEST_TRUE("StateTree should enter Task2", Exec.Expect(Task2.GetName(), EnterStateStr));
-			AITEST_FALSE("StateTree should not exit Task2", Exec.Expect(Task2.GetName(), ExitStateStr));
-			AITEST_FALSE("StateTree should not enter Task1", Exec.Expect(Task1.GetName(), EnterStateStr));
-			Exec.LogClear();
-
-			Status = Exec.Tick(0.1f);
-			AITEST_EQUAL("Tick should complete with Running", Status, EStateTreeRunStatus::Running);
-			AITEST_FALSE("StateTree should not enter GlobalTask2", Exec.Expect(GlobalTask2.GetName(), EnterStateStr));
-			AITEST_TRUE("StateTree should exit GlobalTask2", Exec.Expect(GlobalTask2.GetName(), ExitStateStr));
-			AITEST_FALSE("StateTree should not enter Task2", Exec.Expect(Task2.GetName(), EnterStateStr));
-			AITEST_TRUE("StateTree should exit Task2", Exec.Expect(Task2.GetName(), ExitStateStr));
-			AITEST_TRUE("StateTree should enter Task1", Exec.Expect(Task1.GetName(), EnterStateStr));
-			Exec.LogClear();
-
-			Exec.Stop();
-		}
-
-		return true;
-	}
-};
-IMPLEMENT_AI_INSTANT_TEST(FStateTreeTest_EnterAndExitLinkedAsset, "System.StateTree.EnterAndExitLinkedAsset");
 
 // Test nested tree overrides
 struct FStateTreeTest_NestedOverride : FStateTreeTestBase
diff --git a/Engine/Plugins/Runtime/StateTree/Source/StateTreeTestSuite/Private/StateTreeTransitionTest.cpp b/Engine/Plugins/Runtime/StateTree/Source/StateTreeTestSuite/Private/StateTreeTransitionTest.cpp
--- a/Engine/Plugins/Runtime/StateTree/Source/StateTreeTestSuite/Private/StateTreeTransitionTest.cpp
+++ b/Engine/Plugins/Runtime/StateTree/Source/StateTreeTestSuite/Private/StateTreeTransitionTest.cpp
@@ -690,7 +690,6 @@ struct FStateTreeTest_FollowTransitions : FStateTreeTestBase
 };
 IMPLEMENT_AI_INSTANT_TEST(FStateTreeTest_FollowTransitions, "System.StateTree.Transition.FollowTransitions");
 
-
 struct FStateTreeTest_InfiniteLoop : FStateTreeTestBase
 {
 	virtual bool InstantTest() override
@@ -768,6 +767,166 @@ struct FStateTreeTest_InfiniteLoop : FStateTreeTestBase
 };
 IMPLEMENT_AI_INSTANT_TEST(FStateTreeTest_InfiniteLoop, "System.StateTree.Transition.InfiniteLoop");
 
+struct FStateTreeTest_RegularTransitions : FStateTreeTestBase
+{
+	virtual bool InstantTest() override
+	{
+		//Tree 1
+		//  Global task and parameter
+		//	RootA
+		//		StateB -> Next
+		//		StateC -> Next
+		//		StateD -> Next
+		//		StateE -> Next
+		//		StateF -> Next
+		//		StateG -> Succeeded
+
+		FStateTreeCompilerLog Log;
+
+		// Main asset
+		UStateTree& StateTree = NewStateTree();
+		{
+			UStateTreeEditorData& EditorData = *Cast<UStateTreeEditorData>(StateTree.EditorData);
+			FGuid RootParameter_ValueID;
+			{
+				// Parameters
+				EditorData.RootParameters.Parameters.AddProperty("Value", EPropertyBagPropertyType::Int32);
+				EditorData.RootParameters.Parameters.SetValueInt32("Value", -111);
+				RootParameter_ValueID = EditorData.RootParameters.Parameters.FindPropertyDescByName("Value")->ID;
+
+				TStateTreeEditorNode<FTestTask_PrintValue>& GlobalTask = EditorData.AddGlobalTask<FTestTask_PrintValue>("GlobalTask");
+				GlobalTask.GetInstanceData().Value = -1;
+				EditorData.AddPropertyBinding(FStateTreePropertyPath(EditorData.RootParameters.ID, TEXT("Value")), FStateTreePropertyPath(GlobalTask.ID, TEXT("Value")));
+			}
+
+			UStateTreeState& Root = EditorData.AddSubTree("RootA");
+			{
+				TStateTreeEditorNode<FTestTask_PrintValue>& Task = Root.AddTask<FTestTask_PrintValue>("TaskA");
+				Task.GetInstanceData().Value = -1;
+				EditorData.AddPropertyBinding(FStateTreePropertyPath(EditorData.RootParameters.ID, TEXT("Value")), FStateTreePropertyPath(Task.ID, TEXT("Value")));
+			}
+			{
+				UStateTreeState& StateB = Root.AddChildState("StateB", EStateTreeStateType::State);
+				TStateTreeEditorNode<FTestTask_PrintValue>& Task = StateB.AddTask<FTestTask_PrintValue>("TaskB");
+				Task.GetInstanceData().Value = 1;
+				FStateTreeTransition& Transition = StateB.AddTransition(EStateTreeTransitionTrigger::OnTick, EStateTreeTransitionType::NextState);
+				Transition.bDelayTransition = true;
+				Transition.DelayDuration = 1.0;
+			}
+			{
+				UStateTreeState& StateB = Root.AddChildState("StateC", EStateTreeStateType::State);
+				TStateTreeEditorNode<FTestTask_PrintValue>& Task = StateB.AddTask<FTestTask_PrintValue>("TaskC");
+				Task.GetInstanceData().Value = 2;
+				FStateTreeTransition& Transition = StateB.AddTransition(EStateTreeTransitionTrigger::OnTick, EStateTreeTransitionType::NextState);
+				Transition.bDelayTransition = true;
+				Transition.DelayDuration = 1.0;
+			}
+			{
+				UStateTreeState& StateD = Root.AddChildState("StateD", EStateTreeStateType::State);
+				TStateTreeEditorNode<FTestTask_PrintValue>& Task = StateD.AddTask<FTestTask_PrintValue>("TaskD");
+				Task.GetInstanceData().Value = 3;
+				FStateTreeTransition& Transition = StateD.AddTransition(EStateTreeTransitionTrigger::OnTick, EStateTreeTransitionType::GotoState, &Root);
+				Transition.bDelayTransition = true;
+				Transition.DelayDuration = 1.0;
+			}
+
+			FStateTreeCompiler Compiler(Log);
+			const bool bResult = Compiler.Compile(StateTree);
+			AITEST_TRUE("StateTree should get compiled", bResult);
+		}
+
+		const FString EnterStateStr(TEXT("EnterState"));
+		const FString ExitStateStr(TEXT("ExitState"));
+
+		{
+			EStateTreeRunStatus Status = EStateTreeRunStatus::Unset;
+			FStateTreeInstanceData InstanceData;
+			FTestStateTreeExecutionContext Exec(StateTree, StateTree, InstanceData);
+			const bool bInitSucceeded = Exec.IsValid();
+			AITEST_TRUE("StateTree should init", bInitSucceeded);
+
+			FInstancedPropertyBag Parameters;
+			Parameters.MigrateToNewBagInstance(StateTree.GetDefaultParameters());
+			Parameters.SetValueInt32("Value", 111);
+
+			Status = Exec.Start(&Parameters);
+			AITEST_EQUAL("Start should complete with Running", Status, EStateTreeRunStatus::Running);
+			AITEST_TRUE("Start should enter Global tasks", Exec.Expect("GlobalTask", TEXT("EnterState111")));
+			AITEST_TRUE("Start should enter StateA", Exec.Expect("TaskA", TEXT("EnterState111")));
+			AITEST_TRUE("Start should enter StateB", Exec.Expect("TaskB", TEXT("EnterState1")));
+			Exec.LogClear();
+
+			Status = Exec.Tick(1.5f); // over tick, should trigger
+			AITEST_EQUAL("1st Tick should complete with Running", Status, EStateTreeRunStatus::Running);
+			AITEST_TRUE("1st Tick should tick Global tasks", Exec.Expect("GlobalTask", TEXT("Tick111")));
+			AITEST_TRUE("1st Tick should tick StateA", Exec.Expect("TaskA", TEXT("Tick111")));
+			AITEST_TRUE("1st Tick should tick StateB", Exec.Expect("TaskB", TEXT("Tick1")));
+			Exec.LogClear();
+
+			Status = Exec.Tick(1.0f);
+			AITEST_EQUAL("2nd Tick should complete with Running", Status, EStateTreeRunStatus::Running);
+			AITEST_TRUE("2nd Tick should tick Global tasks", Exec.Expect("GlobalTask", TEXT("Tick111")));
+			AITEST_TRUE("2nd Tick should tick StateA", Exec.Expect("TaskA", TEXT("Tick111")));
+			AITEST_TRUE("2nd Tick should tick the StateB", Exec.Expect("TaskB", TEXT("Tick1")));
+			AITEST_TRUE("2nd Tick should exit the StateB", Exec.Expect("TaskB", TEXT("ExitState1")));
+			AITEST_TRUE("2nd Tick should enter the StateC", Exec.Expect("TaskC", TEXT("EnterState2")));
+			Exec.LogClear();
+
+			Status = Exec.Tick(1.0f);
+			AITEST_EQUAL("3rd Tick should complete with Running", Status, EStateTreeRunStatus::Running);
+			AITEST_TRUE("3rd Tick should tick Global tasks", Exec.Expect("GlobalTask", TEXT("Tick111")));
+			AITEST_TRUE("3rd Tick should tick StateA", Exec.Expect("TaskA", TEXT("Tick111")));
+			AITEST_TRUE("3rd Tick should tick StateC", Exec.Expect("TaskC", TEXT("Tick2")));
+			Exec.LogClear();
+
+			Status = Exec.Tick(1.0f);
+			AITEST_EQUAL("4th Tick should complete with Running", Status, EStateTreeRunStatus::Running);
+			AITEST_TRUE("4th Tick should tick Global tasks", Exec.Expect("GlobalTask", TEXT("Tick111")));
+			AITEST_TRUE("244ththnd Tick should tick StateA", Exec.Expect("TaskA", TEXT("Tick111")));
+			AITEST_TRUE("4th Tick should tick the StateC", Exec.Expect("TaskC", TEXT("Tick2")));
+			AITEST_TRUE("4th Tick should exit the StateC", Exec.Expect("TaskC", TEXT("ExitState2")));
+			AITEST_TRUE("4th Tick should enter the StateD", Exec.Expect("TaskD", TEXT("EnterState3")));
+			Exec.LogClear();
+
+			Status = Exec.Tick(0.001f);
+			AITEST_EQUAL("5th Tick should complete with Running", Status, EStateTreeRunStatus::Running);
+			AITEST_TRUE("5th Tick should tick Global tasks", Exec.Expect("GlobalTask", TEXT("Tick111")));
+			AITEST_TRUE("5th Tick should tick StateA", Exec.Expect("TaskA", TEXT("Tick111")));
+			AITEST_TRUE("5th Tick should tick StateD", Exec.Expect("TaskD", TEXT("Tick3")));
+			Exec.LogClear();
+
+			Status = Exec.Tick(1.0f);
+			AITEST_EQUAL("6th Tick should complete with Running", Status, EStateTreeRunStatus::Running);
+			AITEST_TRUE("6th Tick should tick Global tasks", Exec.Expect("GlobalTask", TEXT("Tick111")));
+			AITEST_TRUE("6th Tick should tick StateA", Exec.Expect("TaskA", TEXT("Tick111")));
+			AITEST_TRUE("6th Tick should tick StateD", Exec.Expect("TaskD", TEXT("Tick3")));
+			AITEST_TRUE("6th Tick should exit the StateD", Exec.Expect("TaskD", TEXT("ExitState3")));
+			//AITEST_FALSE("6th Tick should not exit the Global tasks", Exec.Expect("GlobalTask", TEXT("ExitState111")));
+			//AITEST_FALSE("6th Tick should not enter the Global tasks", Exec.Expect("GlobalTask", TEXT("EnterState111")));
+			//AITEST_FALSE("6th Tick should not exit the StateA", Exec.Expect("TaskA", TEXT("ExitState111")));
+			//AITEST_FALSE("6th Tick should not enter the StateA", Exec.Expect("TaskA", TEXT("EnterState111")));
+			AITEST_TRUE("6th Tick should enter the StateB", Exec.Expect("TaskB", TEXT("EnterState1")));
+			Exec.LogClear();
+
+			Status = Exec.Tick(1.0f);
+			AITEST_EQUAL("7th Tick should complete with Running", Status, EStateTreeRunStatus::Running);
+			AITEST_TRUE("7th Tick should tick Global tasks", Exec.Expect("GlobalTask", TEXT("Tick111")));
+			AITEST_TRUE("7th Tick should tick StateA", Exec.Expect("TaskA", TEXT("Tick111")));
+			AITEST_TRUE("7th Tick should tick StateB", Exec.Expect("TaskB", TEXT("Tick1")));
+			Exec.LogClear();
+
+			Exec.Stop();
+			AITEST_TRUE("Stop Tick should exit the StateB", Exec.Expect("TaskB", TEXT("ExitState1")));
+			AITEST_TRUE("Stop Tick should exit the StateA", Exec.Expect("TaskA", TEXT("ExitState111")));
+			AITEST_TRUE("Stop should tick Global tasks", Exec.Expect("GlobalTask", TEXT("ExitState111")));
+			Exec.LogClear();
+		}
+
+		return true;
+	}
+};
+IMPLEMENT_AI_INSTANT_TEST(FStateTreeTest_RegularTransitions, "System.StateTree.Transition.RegularTransitions");
+
 } // namespace UE::StateTree::Tests
 
 UE_ENABLE_OPTIMIZATION_SHIP
