StateTree: Only show the TickOnEvent icons when the Tick icon is not shown.

diff --git a/Engine/Plugins/Runtime/StateTree/Source/StateTreeEditorModule/Private/Customizations/StateTreeEditorNodeDetails.cpp b/Engine/Plugins/Runtime/StateTree/Source/StateTreeEditorModule/Private/Customizations/StateTreeEditorNodeDetails.cpp
--- a/Engine/Plugins/Runtime/StateTree/Source/StateTreeEditorModule/Private/Customizations/StateTreeEditorNodeDetails.cpp
+++ b/Engine/Plugins/Runtime/StateTree/Source/StateTreeEditorModule/Private/Customizations/StateTreeEditorNodeDetails.cpp
@@ -695,7 +695,7 @@ void FStateTreeEditorNodeDetails::CustomizeHeader(TSharedRef<class IPropertyHand
 				];
 			}
 
-			if (bShouldCallTickOnlyOnEvents)
+			if (!bShowCallTick && bShouldCallTickOnlyOnEvents)
 			{
 				Box->AddSlot()
 				.AutoWidth()
diff --git a/Engine/Plugins/Runtime/StateTree/Source/StateTreeEditorModule/Private/SStateTreeViewRow.cpp b/Engine/Plugins/Runtime/StateTree/Source/StateTreeEditorModule/Private/SStateTreeViewRow.cpp
--- a/Engine/Plugins/Runtime/StateTree/Source/StateTreeEditorModule/Private/SStateTreeViewRow.cpp
+++ b/Engine/Plugins/Runtime/StateTree/Source/StateTreeEditorModule/Private/SStateTreeViewRow.cpp
@@ -489,7 +489,7 @@ void SStateTreeViewRow::Construct(const FArguments& InArgs, const TSharedRef<STa
 						];
 					}
 
-					if (RuntimeState->bHasTickTasksOnlyOnEvents)
+					if (!RuntimeState->bHasTickTasks && RuntimeState->bHasTickTasksOnlyOnEvents)
 					{
 						Box->AddSlot()
 						.AutoWidth()
