Closing the last opened tab in the StateTree Editor hides the toolbar buttons on the remaining open tabs
- Retrieve current state-tree editor mode through UAssetEditorToolkitMenuContext rather than capturing
- Added menu profile to disable state-tree toolbar entries (used by workspace editor)
- Added delegate to _extend_ workspace toolmenu context (used by AnimNextStateTree)

#swarm 37091915

diff --git a/Engine/Plugins/Experimental/AnimNextStateTree/Source/AnimNextStateTreeEditor/AnimNextStateTreeEditor.Build.cs b/Engine/Plugins/Experimental/AnimNextStateTree/Source/AnimNextStateTreeEditor/AnimNextStateTreeEditor.Build.cs
--- a/Engine/Plugins/Experimental/AnimNextStateTree/Source/AnimNextStateTreeEditor/AnimNextStateTreeEditor.Build.cs
+++ b/Engine/Plugins/Experimental/AnimNextStateTree/Source/AnimNextStateTreeEditor/AnimNextStateTreeEditor.Build.cs
@@ -21,6 +21,7 @@ namespace UnrealBuildTool.Rules
 				    "RigVM",
 				    "StructUtilsEditor",
 				    "PropertyEditor",
+					"ToolMenus"
 			    });
 
 			PublicDependencyModuleNames.AddRange(
diff --git a/Engine/Plugins/Experimental/AnimNextStateTree/Source/AnimNextStateTreeEditor/Private/AnimNextStateTreeEditorModule.cpp b/Engine/Plugins/Experimental/AnimNextStateTree/Source/AnimNextStateTreeEditor/Private/AnimNextStateTreeEditorModule.cpp
--- a/Engine/Plugins/Experimental/AnimNextStateTree/Source/AnimNextStateTreeEditor/Private/AnimNextStateTreeEditorModule.cpp
+++ b/Engine/Plugins/Experimental/AnimNextStateTree/Source/AnimNextStateTreeEditor/Private/AnimNextStateTreeEditorModule.cpp
@@ -19,6 +19,8 @@
 #include "StateTreeEditorStyle.h"
 #include "Framework/Docking/LayoutExtender.h"
 #include "Toolkits/AssetEditorModeUILayer.h"
+#include "ToolMenus.h"
+#include "Framework/MultiBox/ToolMenuBase.h"
 
 #define LOCTEXT_NAMESPACE "AnimNextStateTreeEditorModule"
 
@@ -43,6 +45,17 @@ void FAnimNextStateTreeEditorModule::StartupModule()
 		FTabManager::FTab DebuggerTab(FTabId(UAssetEditorUISubsystem::TopRightTabID), ETabState::ClosedTab);
 		InLayoutExtender.ExtendLayout(FTabId(Workspace::WorkspaceTabs::BottomMiddleDocumentArea), ELayoutExtensionPosition::After, DebuggerTab);		
 	});
+
+	WorkspaceEditorModule.OnExtendToolMenuContext().AddLambda([](const TWeakPtr<Workspace::IWorkspaceEditor>& InWorkspaceEditor, FToolMenuContext& InContext)
+	{
+		if(!InWorkspaceEditor.Pin()->GetEditorModeManager().IsModeActive(UStateTreeEditorMode::EM_StateTree))
+		{
+			UToolMenuProfileContext* ProfileContext = NewObject<UToolMenuProfileContext>();
+			ProfileContext->ActiveProfiles.Add( TEXT("StateTreeEditModeDisabledProfile") );
+			InContext.AddObject(ProfileContext);
+		}
+	});
+
 	
 	// --- AnimNextStateTree ---
 	Workspace::FObjectDocumentArgs StateTreeDocumentArgs(
diff --git a/Engine/Plugins/Experimental/Workspace/Source/WorkspaceEditor/Private/WorkspaceEditor.cpp b/Engine/Plugins/Experimental/Workspace/Source/WorkspaceEditor/Private/WorkspaceEditor.cpp
--- a/Engine/Plugins/Experimental/Workspace/Source/WorkspaceEditor/Private/WorkspaceEditor.cpp
+++ b/Engine/Plugins/Experimental/Workspace/Source/WorkspaceEditor/Private/WorkspaceEditor.cpp
@@ -218,6 +218,8 @@ void FWorkspaceEditor::OnToolkitHostingFinished(const TSharedRef<IToolkit>& Tool
 		
 	UToolMenus::UnregisterOwner(&(*Toolkit));	
 	HostedToolkits.Remove( Toolkit );
+		
+	RegenerateMenusAndToolbars();
 }
 
 void FWorkspaceEditor::OnToolkitHostingStarted(const TSharedRef<IToolkit>& Toolkit)
@@ -235,6 +237,8 @@ void FWorkspaceEditor::OnToolkitHostingStarted(const TSharedRef<IToolkit>& Toolk
 
 	ModeUILayers.Add(Toolkit->GetToolkitFName(), ModeUILayer);	
 	HostedToolkits.Add(Toolkit);
+
+	RegenerateMenusAndToolbars();
 }
 
 void FWorkspaceEditor::RestoreEditedObjectState()
@@ -551,6 +555,9 @@ void FWorkspaceEditor::InitToolMenuContext(FToolMenuContext& InMenuContext)
 	InMenuContext.AddObject(ToolkitMenuContext);
 	
 	IWorkspaceEditor::InitToolMenuContext(InMenuContext);
+
+	Workspace::IWorkspaceEditorModule& WorkspaceEditorModule = FModuleManager::Get().LoadModuleChecked<Workspace::IWorkspaceEditorModule>("WorkspaceEditor");
+	WorkspaceEditorModule.OnExtendToolMenuContext().Broadcast(StaticCastSharedRef<UE::Workspace::IWorkspaceEditor>(AsShared()), InMenuContext);
 }
 
 void FWorkspaceEditor::SaveAsset_Execute()
diff --git a/Engine/Plugins/Experimental/Workspace/Source/WorkspaceEditor/Private/WorkspaceEditorModeUILayer.cpp b/Engine/Plugins/Experimental/Workspace/Source/WorkspaceEditor/Private/WorkspaceEditorModeUILayer.cpp
--- a/Engine/Plugins/Experimental/Workspace/Source/WorkspaceEditor/Private/WorkspaceEditorModeUILayer.cpp
+++ b/Engine/Plugins/Experimental/Workspace/Source/WorkspaceEditor/Private/WorkspaceEditorModeUILayer.cpp
@@ -20,17 +20,14 @@ void FWorkspaceEditorModeUILayer::OnToolkitHostingStarted(const TSharedRef<ITool
 
 		OnToolkitHostReadyForUI.Execute();
 
-		// Set up an owner for the current scope so that we can cleanly clean up the toolbar extension on hosting finish
-		FToolMenuOwnerScoped Owner(&(*Toolkit));
 		UToolMenu* SecondaryModeToolbar = UToolMenus::Get()->ExtendMenu(GetSecondaryModeToolbarName());
-		OnRegisterSecondaryModeToolbarExtension.ExecuteIfBound(SecondaryModeToolbar);		
+		OnRegisterSecondaryModeToolbarExtension.ExecuteIfBound(SecondaryModeToolbar);
 	}
 }
 
 void FWorkspaceEditorModeUILayer::OnToolkitHostingFinished(const TSharedRef<IToolkit>& Toolkit)
 {	
 	FAssetEditorModeUILayer::OnToolkitHostingFinished(Toolkit);
-	UToolMenus::UnregisterOwner(&(*Toolkit));
 }
 
 void FWorkspaceEditorModeUILayer::SetModeMenuCategory(const TSharedPtr<FWorkspaceItem>& MenuCategoryIn)
diff --git a/Engine/Plugins/Experimental/Workspace/Source/WorkspaceEditor/Private/WorkspaceEditorModule.h b/Engine/Plugins/Experimental/Workspace/Source/WorkspaceEditor/Private/WorkspaceEditorModule.h
--- a/Engine/Plugins/Experimental/Workspace/Source/WorkspaceEditor/Private/WorkspaceEditorModule.h
+++ b/Engine/Plugins/Experimental/Workspace/Source/WorkspaceEditor/Private/WorkspaceEditorModule.h
@@ -36,6 +36,7 @@ private:
 	virtual void UnregisterWorkspaceItemDetails(const FOutlinerItemDetailsId& InItemDetails) override;
 	virtual FOnRegisterTabs& OnRegisterTabsForEditor() override { return RegisterTabsForEditor; }
 	virtual FOnExtendTabs& OnExtendTabs() override { return ExtendTabsForEditor; }
+	virtual FOnExtendToolMenuContext& OnExtendToolMenuContext() override { return ExtendToolMenuContext; }
 
 	// Find an existing registered object document type. Note this redirects based on FObjectDocumentArgs::OnRedirectWorkspaceContext
 	const FObjectDocumentArgs* FindObjectDocumentType(const UObject* InObject) const;
@@ -65,6 +66,9 @@ private:
 	/** Event called to allow external clients to extend the tab layout for the specified editor */
 	FOnExtendTabs ExtendTabsForEditor;
 
+	/** Event called to allow external clients to add details to the ToolMenuContext used for the toolbar and context menus */
+	FOnExtendToolMenuContext ExtendToolMenuContext;
+
 	friend struct FAssetDocumentSummoner;
 	friend struct FWorkspaceOutlinerTreeItem;
 	friend class SWorkspaceOutlinerTreelabel;
diff --git a/Engine/Plugins/Experimental/Workspace/Source/WorkspaceEditor/Public/IWorkspaceEditorModule.h b/Engine/Plugins/Experimental/Workspace/Source/WorkspaceEditor/Public/IWorkspaceEditorModule.h
--- a/Engine/Plugins/Experimental/Workspace/Source/WorkspaceEditor/Public/IWorkspaceEditorModule.h
+++ b/Engine/Plugins/Experimental/Workspace/Source/WorkspaceEditor/Public/IWorkspaceEditorModule.h
@@ -216,6 +216,10 @@ public:
 	// Event to allow extending the layout
 	DECLARE_EVENT_TwoParams(IWorkspaceEditorModule, FOnExtendTabs, FLayoutExtender&, TSharedPtr<IWorkspaceEditor>);
 	virtual FOnExtendTabs& OnExtendTabs() = 0;
+
+	// Event to allow extending the FToolMenuContext
+	DECLARE_EVENT_TwoParams(IWorkspaceEditorModule, FOnExtendToolMenuContext, const TWeakPtr<IWorkspaceEditor>&, FToolMenuContext&);
+	virtual FOnExtendToolMenuContext& OnExtendToolMenuContext() = 0;
 };
 
 }
diff --git a/Engine/Plugins/Runtime/StateTree/Source/StateTreeEditorModule/Private/StateTreeEditor.cpp b/Engine/Plugins/Runtime/StateTree/Source/StateTreeEditorModule/Private/StateTreeEditor.cpp
--- a/Engine/Plugins/Runtime/StateTree/Source/StateTreeEditorModule/Private/StateTreeEditor.cpp
+++ b/Engine/Plugins/Runtime/StateTree/Source/StateTreeEditorModule/Private/StateTreeEditor.cpp
@@ -191,6 +191,8 @@ void FStateTreeEditor::InitEditor( const EToolkitMode::Type Mode, const TSharedP
 	StateTreeEditorModule.OnRegisterLayoutExtensions().Broadcast(LayoutExtender);
 	StandaloneDefaultLayout->ProcessExtensions(LayoutExtender);
 
+	CreateEditorModeManager();
+	
 	constexpr bool bCreateDefaultStandaloneMenu = true;
 	constexpr bool bCreateDefaultToolbar = true;
 	FAssetEditorToolkit::InitAssetEditor(Mode, InitToolkitHost, StateTreeEditorAppName, StandaloneDefaultLayout, bCreateDefaultStandaloneMenu, bCreateDefaultToolbar, StateTree);
diff --git a/Engine/Plugins/Runtime/StateTree/Source/StateTreeEditorModule/Private/StateTreeEditorModeToolkit.cpp b/Engine/Plugins/Runtime/StateTree/Source/StateTreeEditorModule/Private/StateTreeEditorModeToolkit.cpp
--- a/Engine/Plugins/Runtime/StateTree/Source/StateTreeEditorModule/Private/StateTreeEditorModeToolkit.cpp
+++ b/Engine/Plugins/Runtime/StateTree/Source/StateTreeEditorModule/Private/StateTreeEditorModeToolkit.cpp
@@ -13,6 +13,7 @@
 
 #include "ClassViewerFilter.h"
 #include "ContentBrowserModule.h"
+#include "EditorModeManager.h"
 #include "IContentBrowserSingleton.h"
 #include "SStateTreeOutliner.h"
 #include "StateTreeEditingSubsystem.h"
@@ -26,6 +27,8 @@
 #include "ToolMenus.h"
 #include "Subsystems/AssetEditorSubsystem.h"
 #include "Framework/Application/SlateApplication.h"
+#include "Toolkits/AssetEditorToolkit.h"
+#include "Toolkits/AssetEditorToolkitMenuContext.h"
 
 #define LOCTEXT_NAMESPACE "StateTreeModeToolkit"
 
@@ -209,7 +212,6 @@ static void GenerateCompileOptionsMenu(UToolMenu* InMenu)
 }
 }
 
-//
 void FStateTreeEditorModeToolkit::ExtendSecondaryModeToolbar(UToolMenu* ToolBar)
 {
 	ToolBar->Context.AppendCommandList(ToolkitCommands);
@@ -221,90 +223,125 @@ void FStateTreeEditorModeToolkit::ExtendSecondaryModeToolbar(UToolMenu* ToolBar)
 
 	FToolMenuSection& CompileSection = ToolBar->AddSection("Compile", FText(), InsertLast);
 
-
-	TWeakPtr<FStateTreeEditorModeToolkit> WeakToolkit = StaticCastSharedRef<FStateTreeEditorModeToolkit>(AsShared());
-	CompileSection.AddDynamicEntry("CompileCommands", FNewToolMenuSectionDelegate::CreateLambda([WeakToolkit, WeakMode = WeakEditorMode](FToolMenuSection& InSection)
+	auto GetToolkitFromAssetEditorContext = [](const UAssetEditorToolkitMenuContext* InContext) -> TSharedPtr<FStateTreeEditorModeToolkit>
 	{
-		if (UStateTreeEditorMode* Mode = WeakMode.Get())
+		if (InContext)
 		{
-			UInteractiveToolManager* ToolManager = Mode->GetToolManager();
-			if (const UContextObjectStore* ContextObjectStore = ToolManager ? ToolManager->GetContextObjectStore() : nullptr)
+			if (TSharedPtr<FAssetEditorToolkit> SharedToolkit = InContext->Toolkit.Pin())
 			{
-				if (ContextObjectStore->FindContext<UStateTreeEditorContext>())
+				if(UStateTreeEditorMode* Mode = Cast<UStateTreeEditorMode>(SharedToolkit->GetEditorModeManager().GetActiveScriptableMode(UStateTreeEditorMode::EM_StateTree)))
 				{
-					const FStateTreeEditorCommands& Commands = FStateTreeEditorCommands::Get();
-
-					FToolMenuEntry& CompileButton = InSection.AddEntry(FToolMenuEntry::InitToolBarButton(
-						Commands.Compile,
-						TAttribute<FText>(),
-						TAttribute<FText>(),
-						TAttribute<FSlateIcon>::CreateLambda([WeakToolkit]() -> FSlateIcon
-						{
-							if (TSharedPtr<FStateTreeEditorModeToolkit> SharedThis = WeakToolkit.Pin())
-							{
-								return SharedThis->GetCompileStatusImage();
-							}
-
-							return FSlateIcon();
-						}))
-					);
-					CompileButton.StyleNameOverride = "CalloutToolbar";
-
-					FToolMenuEntry& CompileOptions = InSection.AddEntry(FToolMenuEntry::InitComboButton(
-						"CompileComboButton",
-						FUIAction(),
-						FNewToolMenuDelegate::CreateStatic(&UE::StateTree::Editor::Internal::GenerateCompileOptionsMenu),
-						LOCTEXT("CompileOptions_ToolbarTooltip", "Options to customize how State Trees compile")
-					));
-					CompileOptions.StyleNameOverride = "CalloutToolbar";
-					CompileOptions.ToolBarData.bSimpleComboBox = true;
+					if (TSharedPtr<FModeToolkit> Toolkit = Mode->GetToolkit().Pin())
+					{
+						return StaticCastSharedPtr<FStateTreeEditorModeToolkit>(Toolkit);
+					}
 				}
 			}
 		}
+
+		return nullptr;
+	};
+	
+	CompileSection.AddDynamicEntry("CompileCommands", FNewToolMenuSectionDelegate::CreateLambda([GetToolkitFromAssetEditorContext](FToolMenuSection& InSection)
+	{
+		if (UAssetEditorToolkitMenuContext* ToolkitContextObject = InSection.FindContext<UAssetEditorToolkitMenuContext>())
+		{
+			const FStateTreeEditorCommands& Commands = FStateTreeEditorCommands::Get();
+			FToolMenuEntry& CompileButton = InSection.AddEntry(FToolMenuEntry::InitToolBarButton(
+				Commands.Compile,
+				TAttribute<FText>(),
+				TAttribute<FText>(),
+				TAttribute<FSlateIcon>::CreateLambda([GetToolkitFromAssetEditorContext, ToolkitContextObject]() -> FSlateIcon
+				{
+					if(TSharedPtr<FStateTreeEditorModeToolkit> SharedToolkit = GetToolkitFromAssetEditorContext(ToolkitContextObject))
+					{
+						return SharedToolkit->GetCompileStatusImage();
+					}
+				
+					return FSlateIcon();
+				}))
+			);
+			CompileButton.StyleNameOverride = "CalloutToolbar";
+
+			FToolMenuEntry& CompileOptions = InSection.AddEntry(FToolMenuEntry::InitComboButton(
+				"CompileComboButton",
+				FUIAction(),
+				FNewToolMenuDelegate::CreateStatic(&UE::StateTree::Editor::Internal::GenerateCompileOptionsMenu),
+				LOCTEXT("CompileOptions_ToolbarTooltip", "Options to customize how State Trees compile")
+			));
+			CompileOptions.StyleNameOverride = "CalloutToolbar";
+			CompileOptions.ToolBarData.bSimpleComboBox = true;
+		}
 	}));
 
 	static const FToolMenuInsert InsertAfterCompileSection("Compile", EToolMenuInsertType::After);
 
 	FToolMenuSection& CreateNewNodeSection = ToolBar->AddSection("CreateNewNodes", TAttribute<FText>(), InsertAfterCompileSection);
-
-	CreateNewNodeSection.AddDynamicEntry("CreateNewNodes", FNewToolMenuSectionDelegate::CreateLambda([this, WeakMode = WeakEditorMode](FToolMenuSection& InSection)
+	CreateNewNodeSection.AddDynamicEntry("CreateNewNodes", FNewToolMenuSectionDelegate::CreateLambda([GetToolkitFromAssetEditorContext](FToolMenuSection& InSection)
 	{
-		if (UStateTreeEditorMode* Mode = WeakMode.Get())
+		if (UAssetEditorToolkitMenuContext* ToolkitContextObject = InSection.FindContext<UAssetEditorToolkitMenuContext>())
 		{
-			if (UContextObjectStore* ContextObjectStore = Mode->GetToolManager()->GetContextObjectStore())
-			{
-				if (ContextObjectStore->FindContext<UStateTreeEditorContext>())
+			InSection.AddEntry(FToolMenuEntry::InitComboButton(
+				 "CreateNewTaskComboButton",
+				 FUIAction(),
+				 FOnGetContent::CreateLambda([ToolkitContextObject, GetToolkitFromAssetEditorContext]() -> TSharedRef<SWidget>
+				 {
+				 	if(TSharedPtr<FStateTreeEditorModeToolkit> SharedToolkit = GetToolkitFromAssetEditorContext(ToolkitContextObject))
+					{
+						return SharedToolkit->GenerateTaskBPBaseClassesMenu();
+					}
+
+					return SNullWidget::NullWidget;
+				 }),
+				 LOCTEXT("CreateNewTask_Title", "New Task"),
+				 LOCTEXT("CreateNewTask_ToolbarTooltip", "Create a new Blueprint State Tree Task"),
+				 GetNewTaskButtonImage()
+			 ));
+				
+			InSection.AddEntry(FToolMenuEntry::InitComboButton(
+				 "CreateNewConditionComboButton",
+				 FUIAction(),
+				FOnGetContent::CreateLambda([ToolkitContextObject, GetToolkitFromAssetEditorContext]() -> TSharedRef<SWidget>
 				{
-					InSection.AddEntry(FToolMenuEntry::InitComboButton(
-						 "CreateNewTaskComboButton",
-						 FUIAction(),
-						 FOnGetContent::CreateSP(this, &FStateTreeEditorModeToolkit::GenerateTaskBPBaseClassesMenu),
-						 LOCTEXT("CreateNewTask_Title", "New Task"),
-						 LOCTEXT("CreateNewTask_ToolbarTooltip", "Create a new Blueprint State Tree Task"),
-						 GetNewTaskButtonImage()
-					 ));
-						
-					InSection.AddEntry(FToolMenuEntry::InitComboButton(
-						 "CreateNewConditionComboButton",
-						 FUIAction(),
-						 FOnGetContent::CreateSP(this, &FStateTreeEditorModeToolkit::GenerateConditionBPBaseClassesMenu),
-						 LOCTEXT("CreateNewCondition_Title", "New Condition"),
-						 LOCTEXT("CreateNewCondition_ToolbarTooltip", "Create a new Blueprint State Tree Condition"),
-						 GetNewConditionButtonImage()
-					 ));
-						
-					 InSection.AddEntry(FToolMenuEntry::InitComboButton(
-						 "CreateNewConsiderationComboButton",
-						 FUIAction(),
-						 FOnGetContent::CreateSP(this, &FStateTreeEditorModeToolkit::GenerateConsiderationBPBaseClassesMenu),
-						 LOCTEXT("CreateNewConsideration_Title", "New Consideration"),
-						 LOCTEXT("CreateNewConsideration_ToolbarTooltip", "Create a new Blueprint State Tree Utility Consideration"),
-						 GetNewConsiderationButtonImage()
-					 ));
-				}
-			}
+					if(TSharedPtr<FStateTreeEditorModeToolkit> SharedToolkit = GetToolkitFromAssetEditorContext(ToolkitContextObject))
+					{
+					   return SharedToolkit->GenerateConditionBPBaseClassesMenu();
+					}
+
+					return SNullWidget::NullWidget;
+				}),
+				 LOCTEXT("CreateNewCondition_Title", "New Condition"),
+				 LOCTEXT("CreateNewCondition_ToolbarTooltip", "Create a new Blueprint State Tree Condition"),
+				 GetNewConditionButtonImage()
+			 ));
+				
+			 InSection.AddEntry(FToolMenuEntry::InitComboButton(
+				 "CreateNewConsiderationComboButton",
+				 FUIAction(),
+				 FOnGetContent::CreateLambda([ToolkitContextObject, GetToolkitFromAssetEditorContext]() -> TSharedRef<SWidget>
+				{
+				 	if(TSharedPtr<FStateTreeEditorModeToolkit> SharedToolkit = GetToolkitFromAssetEditorContext(ToolkitContextObject))
+				 	{
+						return SharedToolkit->GenerateConsiderationBPBaseClassesMenu();
+					}
+
+					return SNullWidget::NullWidget;
+				}),
+				 LOCTEXT("CreateNewConsideration_Title", "New Consideration"),
+				 LOCTEXT("CreateNewConsideration_ToolbarTooltip", "Create a new Blueprint State Tree Utility Consideration"),
+				 GetNewConsiderationButtonImage()
+			 ));
 		}
 	}));
+
+	const FName StateTreeEditModeProfile = TEXT("StateTreeEditModeDisabledProfile");
+	FToolMenuProfile* ToolbarProfile = UToolMenus::Get()->AddRuntimeMenuProfile(ToolBar->GetMenuName(), StateTreeEditModeProfile);
+	{
+		ToolbarProfile->MenuPermissions.AddDenyListItem("CompileCommands", Commands.Compile->GetCommandName());
+		ToolbarProfile->MenuPermissions.AddDenyListItem("CreateNewNodes", "CreateNewTaskComboButton");
+		ToolbarProfile->MenuPermissions.AddDenyListItem("CreateNewNodes", "CreateNewConditionComboButton");
+		ToolbarProfile->MenuPermissions.AddDenyListItem("CreateNewNodes", "CreateNewConsiderationComboButton");
+	}
 }
 
 void FStateTreeEditorModeToolkit::OnStateTreeChanged()
@@ -341,7 +378,7 @@ using FStateTreeConditionBPClassFilter = FEditorNodeClassFilter<UStateTreeCondit
 using FStateTreeConsiderationBPClassFilter = FEditorNodeClassFilter<UStateTreeConsiderationBlueprintBase>;
 }; // UE::StateTree::Editor
 
-FSlateIcon FStateTreeEditorModeToolkit::GetNewTaskButtonImage() const
+FSlateIcon FStateTreeEditorModeToolkit::GetNewTaskButtonImage()
 {
 	return FSlateIcon("StateTreeEditorStyle", "StateTreeEditor.Tasks.Large");
 }
@@ -357,7 +394,7 @@ TSharedRef<SWidget> FStateTreeEditorModeToolkit::GenerateTaskBPBaseClassesMenu()
 	return FModuleManager::LoadModuleChecked<FClassViewerModule>("ClassViewer").CreateClassViewer(Options, OnPicked);
 }
 
-FSlateIcon FStateTreeEditorModeToolkit::GetNewConditionButtonImage() const
+FSlateIcon FStateTreeEditorModeToolkit::GetNewConditionButtonImage()
 {
 	return FSlateIcon("StateTreeEditorStyle", "StateTreeEditor.Conditions.Large");
 }
@@ -373,7 +410,7 @@ TSharedRef<SWidget> FStateTreeEditorModeToolkit::GenerateConditionBPBaseClassesM
 	return FModuleManager::LoadModuleChecked<FClassViewerModule>("ClassViewer").CreateClassViewer(Options, OnPicked);
 }
 
-FSlateIcon FStateTreeEditorModeToolkit::GetNewConsiderationButtonImage() const
+FSlateIcon FStateTreeEditorModeToolkit::GetNewConsiderationButtonImage()
 {
     return FSlateIcon("StateTreeEditorStyle", "StateTreeEditor.Utility.Large");
 }
diff --git a/Engine/Plugins/Runtime/StateTree/Source/StateTreeEditorModule/Private/StateTreeEditorUILayer.cpp b/Engine/Plugins/Runtime/StateTree/Source/StateTreeEditorModule/Private/StateTreeEditorUILayer.cpp
--- a/Engine/Plugins/Runtime/StateTree/Source/StateTreeEditorModule/Private/StateTreeEditorUILayer.cpp
+++ b/Engine/Plugins/Runtime/StateTree/Source/StateTreeEditorModule/Private/StateTreeEditorUILayer.cpp
@@ -24,9 +24,8 @@ void FStateTreeEditorModeUILayer::OnToolkitHostingStarted(const TSharedRef<ITool
 		OnToolkitHostReadyForUI.Execute();
 
 		// Set up an owner for the current scope so that we can cleanly clean up the toolbar extension on hosting finish
-		FToolMenuOwnerScoped Owner(&(*Toolkit));
 		UToolMenu* SecondaryModeToolbar = UToolMenus::Get()->ExtendMenu(GetSecondaryModeToolbarName());
-		OnRegisterSecondaryModeToolbarExtension.ExecuteIfBound(SecondaryModeToolbar);		
+		OnRegisterSecondaryModeToolbarExtension.ExecuteIfBound(SecondaryModeToolbar);
 	}
 }
 
diff --git a/Engine/Plugins/Runtime/StateTree/Source/StateTreeEditorModule/Public/StateTreeEditorModeToolkit.h b/Engine/Plugins/Runtime/StateTree/Source/StateTreeEditorModule/Public/StateTreeEditorModeToolkit.h
--- a/Engine/Plugins/Runtime/StateTree/Source/StateTreeEditorModule/Public/StateTreeEditorModeToolkit.h
+++ b/Engine/Plugins/Runtime/StateTree/Source/StateTreeEditorModule/Public/StateTreeEditorModeToolkit.h
@@ -28,13 +28,13 @@ public:
 protected:	
 	FSlateIcon GetCompileStatusImage() const;
 
-	FSlateIcon GetNewTaskButtonImage() const;
+	static FSlateIcon GetNewTaskButtonImage();
 	TSharedRef<SWidget> GenerateTaskBPBaseClassesMenu() const;
 
-	FSlateIcon GetNewConditionButtonImage() const;
+	static FSlateIcon GetNewConditionButtonImage();
 	TSharedRef<SWidget> GenerateConditionBPBaseClassesMenu() const;
     
-	FSlateIcon GetNewConsiderationButtonImage() const;
+	static FSlateIcon GetNewConsiderationButtonImage();
 	TSharedRef<SWidget> GenerateConsiderationBPBaseClassesMenu() const;
 
 	void OnNodeBPBaseClassPicked(UClass* NodeClass) const;
@@ -44,9 +44,6 @@ protected:
 	void UpdateStateTreeOutliner();
 protected:
 	TWeakObjectPtr<UStateTreeEditorMode> WeakEditorMode;
-	bool bLastCompileSucceeded = true;
-	uint32 EditorDataHash = 0;
-
 	TSharedPtr<IStateTreeEditorHost> EditorHost;
 
 	/** Tree Outliner */
