[StateTree] Fix for hitting the ensure "ExitState executed on node %s without an EnterState" (missing call to UE_STATETREE_DEBUG_EVALUATOR_ENTER_TREE)
The ensure was hit for global evaluators, since it seems their OnEnter call was never registered

diff --git a/Engine/Plugins/Runtime/StateTree/Source/StateTreeModule/Private/StateTreeExecutionContext.cpp b/Engine/Plugins/Runtime/StateTree/Source/StateTreeModule/Private/StateTreeExecutionContext.cpp
--- a/Engine/Plugins/Runtime/StateTree/Source/StateTreeModule/Private/StateTreeExecutionContext.cpp
+++ b/Engine/Plugins/Runtime/StateTree/Source/StateTreeModule/Private/StateTreeExecutionContext.cpp
@@ -2492,6 +2492,7 @@ EStateTreeRunStatus FStateTreeExecutionContext::StartTemporaryEvaluatorsAndGloba
 		}
 
 		STATETREE_LOG(Verbose, TEXT("  Start: '%s'"), *Eval.Name.ToString());
+		UE_STATETREE_DEBUG_EVALUATOR_ENTER_TREE(this, CurrentStateTree, FStateTreeIndex16(EvalIndex));
 		{
 			QUICK_SCOPE_CYCLE_COUNTER(StateTree_Eval_TreeStart);
 			Eval.TreeStart(*this);
