[StateTree] States from nested trees should never be considered as 'sustained' once a parent state has changed
- Minor refactor to EnterState/ExitState to make them as symmetrical as possible
- Updated log indentations to better represent nested trees

diff --git a/Engine/Plugins/Runtime/StateTree/Source/StateTreeModule/Private/StateTreeExecutionContext.cpp b/Engine/Plugins/Runtime/StateTree/Source/StateTreeModule/Private/StateTreeExecutionContext.cpp
--- a/Engine/Plugins/Runtime/StateTree/Source/StateTreeModule/Private/StateTreeExecutionContext.cpp
+++ b/Engine/Plugins/Runtime/StateTree/Source/StateTreeModule/Private/StateTreeExecutionContext.cpp
@@ -1614,6 +1614,10 @@ EStateTreeRunStatus FStateTreeExecutionContext::EnterState(FStateTreeTransitionR
 	// Reset the current active frames, new ones are added one by one.
 	Exec.ActiveFrames.Reset();
 
+	// Track any changed state (i.e., not sustained) to prevent reused subtrees from being considered as sustained states since the whole
+	// tree is reused and should be considered as a new instance.
+	bool bAnyParentStateChanged = false;
+
 	for (int32 FrameIndex = 0; FrameIndex < Transition.NextActiveFrames.Num() && Result != EStateTreeRunStatus::Failed; FrameIndex++)
 	{
 		const FStateTreeExecutionFrame& NextFrame = Transition.NextActiveFrames[FrameIndex];
@@ -1641,8 +1645,6 @@ EStateTreeRunStatus FStateTreeExecutionContext::EnterState(FStateTreeTransitionR
 			const FStateTreeStateHandle PreviousHandle = PreviousActiveStates.GetStateSafe(StateIndex);
 			const FCompactStateTreeState& State = CurrentStateTree->States[CurrentHandle.Index];
 
-			FCurrentlyProcessedStateScope StateScope(*this, CurrentHandle);
-			
 			// Add only enabled States to the list of active States
 			if (State.bEnabled && !CurrentFrame.ActiveStates.Push(CurrentHandle))
 			{
@@ -1652,6 +1654,8 @@ EStateTreeRunStatus FStateTreeExecutionContext::EnterState(FStateTreeTransitionR
 			}
 			CurrentFrame.NumCurrentlyActiveStates = static_cast<uint8>(CurrentFrame.ActiveStates.Num());
 
+			FCurrentlyProcessedStateScope StateScope(*this, CurrentHandle);
+
 			if (State.Type == EStateTreeStateType::Linked
 				|| State.Type == EStateTreeStateType::LinkedAsset)
 			{
@@ -1663,21 +1667,33 @@ EStateTreeRunStatus FStateTreeExecutionContext::EnterState(FStateTreeTransitionR
 				}
 			}
 
-			bOnTargetBranch |= CurrentFrame.StateTree == Transition.SourceStateTree && CurrentFrame.RootState == Transition.SourceRootState && CurrentHandle == Transition.TargetState;
-			const bool bWasActive = PreviousHandle == CurrentHandle;
+			if (CurrentFrame.StateTree == Transition.SourceStateTree
+				&& CurrentFrame.RootState == Transition.SourceRootState
+				&& CurrentHandle == Transition.TargetState)
+			{
+				bOnTargetBranch = true;
+			}
 
-			// Do not enter a disabled State tasks but maintain property bindings
-			const bool bIsEnteringState = (!bWasActive || bOnTargetBranch) && State.bEnabled;
+			const bool bWasActive = (PreviousHandle == CurrentHandle) && !bAnyParentStateChanged;
+
+			const EStateTreeStateChangeType ChangeType = bWasActive ? EStateTreeStateChangeType::Sustained : EStateTreeStateChangeType::Changed;
+			if (ChangeType == EStateTreeStateChangeType::Changed)
+			{
+				bAnyParentStateChanged = true;
+			}
 
 			CurrentTransition.CurrentState = CurrentHandle;
-			CurrentTransition.ChangeType = bWasActive ? EStateTreeStateChangeType::Sustained : EStateTreeStateChangeType::Changed;
+			CurrentTransition.ChangeType = ChangeType;
+
+			// Do not enter a disabled State tasks but maintain property bindings
+			const bool bIsEnteringState = (!bWasActive || bOnTargetBranch) && State.bEnabled;
 
 			if (bIsEnteringState)
 			{
 				STATETREE_TRACE_STATE_EVENT(CurrentHandle, EStateTreeTraceEventType::OnEntering);
-				STATETREE_LOG(Log, TEXT("%*sState '%s' %s"), StateIndex*UE::StateTree::DebugIndentSize, TEXT(""),
-					*DebugGetStatePath(Transition.NextActiveFrames, &NextFrame, StateIndex),
-					*UEnum::GetDisplayValueAsText(CurrentTransition.ChangeType).ToString());
+				STATETREE_LOG(Log, TEXT("%*sState '%s' (%s)"), (FrameIndex + StateIndex + 1)*UE::StateTree::DebugIndentSize, TEXT("")
+					, *GetSafeStateName(NextFrame, CurrentHandle)
+					, *UEnum::GetDisplayValueAsText(CurrentTransition.ChangeType).ToString());
 			}
 
 			// Call state change events on conditions if needed.
@@ -1688,10 +1704,10 @@ EStateTreeRunStatus FStateTreeExecutionContext::EnterState(FStateTreeTransitionR
 					const FStateTreeConditionBase& Cond = CurrentFrame.StateTree->Nodes[ConditionIndex].Get<const FStateTreeConditionBase>();
 					if (Cond.bHasShouldCallStateChangeEvents)
 					{
-						const bool bShouldCallStateChange = CurrentTransition.ChangeType == EStateTreeStateChangeType::Changed
+						const bool bShouldCallEnterState = CurrentTransition.ChangeType == EStateTreeStateChangeType::Changed
 															|| (CurrentTransition.ChangeType == EStateTreeStateChangeType::Sustained && Cond.bShouldStateChangeOnReselect);
 
-						if (bShouldCallStateChange)
+						if (bShouldCallEnterState)
 						{
 							const FStateTreeDataView ConditionInstanceView = GetDataView(CurrentParentFrame, CurrentFrame, Cond.InstanceDataHandle);
 							FNodeInstanceDataScope DataScope(*this, Cond.InstanceDataHandle, ConditionInstanceView);
@@ -1736,12 +1752,12 @@ EStateTreeRunStatus FStateTreeExecutionContext::EnterState(FStateTreeTransitionR
 					continue;
 				}
 
-				const bool bShouldCallStateChange = CurrentTransition.ChangeType == EStateTreeStateChangeType::Changed
+				const bool bShouldCallEnterState = CurrentTransition.ChangeType == EStateTreeStateChangeType::Changed
 													|| (CurrentTransition.ChangeType == EStateTreeStateChangeType::Sustained && Task.bShouldStateChangeOnReselect);
 
-				if (bIsEnteringState && bShouldCallStateChange)
+				if (bIsEnteringState && bShouldCallEnterState)
 				{
-					STATETREE_LOG(Verbose, TEXT("%*s  Task '%s'"), StateIndex*UE::StateTree::DebugIndentSize, TEXT(""), *Task.Name.ToString());
+					STATETREE_LOG(Verbose, TEXT("%*sTask '%s'.EnterState()"), (FrameIndex + StateIndex + 1)*UE::StateTree::DebugIndentSize, TEXT(""), *Task.Name.ToString());
 					UE_STATETREE_DEBUG_TASK_ENTER_STATE(this, NextFrame.StateTree, FStateTreeIndex16(TaskIndex));
 
 					EStateTreeRunStatus Status = EStateTreeRunStatus::Unset;
@@ -1819,7 +1835,11 @@ void FStateTreeExecutionContext::ExitState(const FStateTreeTransitionResult& Tra
 	};
 
 	TArray<FExitStateCall, TConcurrentLinearArrayAllocator<FDefaultBlockAllocationTag>> ExitStateCalls;
-	
+
+	// Track any changed state (i.e., not sustained) to prevent reused subtrees from being considered as sustained states since the whole
+	// tree is reused and should be considered as a new instance.
+	bool bAnyParentStateChanged = false;
+
 	for (int32 FrameIndex = 0; FrameIndex < Exec.ActiveFrames.Num(); FrameIndex++)
 	{
 		FStateTreeExecutionFrame* CurrentParentFrame = FrameIndex > 0 ? &Exec.ActiveFrames[FrameIndex - 1] : nullptr;
@@ -1884,15 +1904,25 @@ void FStateTreeExecutionContext::ExitState(const FStateTreeTransitionResult& Tra
 				}
 			}
 
-			const bool bRemainsActive = NextHandle == CurrentHandle;
-			const bool bCurrentIsTarget = CurrentFrame.StateTree == Transition.SourceStateTree && CurrentFrame.RootState == Transition.SourceRootState && CurrentHandle == Transition.TargetState;
-			bOnTargetBranch = bOnTargetBranch || bCurrentIsTarget;
+			if (CurrentFrame.StateTree == Transition.SourceStateTree
+				&& CurrentFrame.RootState == Transition.SourceRootState
+				&& CurrentHandle == Transition.TargetState)
+			{
+				bOnTargetBranch = true;
+			}
+
+			const bool bRemainsActive = (NextHandle == CurrentHandle) && !bAnyParentStateChanged;
 
 			const EStateTreeStateChangeType ChangeType = bRemainsActive ? EStateTreeStateChangeType::Sustained : EStateTreeStateChangeType::Changed;
+			if (ChangeType == EStateTreeStateChangeType::Changed)
+			{
+				bAnyParentStateChanged = true;
+			}
 			
 			// Should call ExitState() on this state.
-			const bool bShouldCall = !bRemainsActive || bOnTargetBranch; 
-			ExitStateCalls.Emplace(ChangeType, bShouldCall);
+			// (No need to test for 'State.bEnabled' like EnterState since we can't enter a disabled state)
+			const bool bShouldCallExitState = (!bRemainsActive || bOnTargetBranch);
+			ExitStateCalls.Emplace(ChangeType, bShouldCallExitState);
 
 			// Do property copies, ExitState() is called below.
 			for (int32 TaskIndex = State.TasksBegin; TaskIndex < (State.TasksBegin + State.TasksNum); TaskIndex++)
@@ -1924,18 +1954,20 @@ void FStateTreeExecutionContext::ExitState(const FStateTreeTransitionResult& Tra
 
 		FCurrentlyProcessedFrameScope FrameScope(*this, CurrentParentFrame, CurrentFrame);
 
-		for (int32 Index = CurrentFrame.ActiveStates.Num() - 1; Index >= 0; Index--)
+		for (int32 StateIndex = CurrentFrame.ActiveStates.Num() - 1; StateIndex >= 0; StateIndex--)
 		{
-			const FStateTreeStateHandle CurrentHandle = CurrentFrame.ActiveStates[Index];
+			const FStateTreeStateHandle CurrentHandle = CurrentFrame.ActiveStates[StateIndex];
 			const FCompactStateTreeState& State = CurrentStateTree->States[CurrentHandle.Index];
 
 			const FExitStateCall& ExitCall = ExitStateCalls[CallIndex--];
 			CurrentTransition.ChangeType = ExitCall.ChangeType;
 
-			STATETREE_LOG(Log, TEXT("%*sState '%s' %s"), Index*UE::StateTree::DebugIndentSize, TEXT(""), *DebugGetStatePath(Exec.ActiveFrames, &CurrentFrame, CurrentHandle.Index), *UEnum::GetDisplayValueAsText(CurrentTransition.ChangeType).ToString());
+			STATETREE_LOG(Log, TEXT("%*sState '%s' (%s)"), (FrameIndex + StateIndex + 1)*UE::StateTree::DebugIndentSize, TEXT("")
+				, *GetSafeStateName(CurrentFrame, CurrentHandle)
+				, *UEnum::GetDisplayValueAsText(CurrentTransition.ChangeType).ToString());
 
 			STATETREE_TRACE_STATE_EVENT(CurrentHandle, EStateTreeTraceEventType::OnExiting);
-				
+
 			if (ExitCall.bShouldCall)
 			{
 				FCurrentlyProcessedStateScope StateScope(*this, CurrentHandle);
@@ -1975,7 +2007,7 @@ void FStateTreeExecutionContext::ExitState(const FStateTreeTransitionResult& Tra
 
 						if (bShouldCallStateChange)
 						{
-							STATETREE_LOG(Verbose, TEXT("%*s  Task '%s'"), Index*UE::StateTree::DebugIndentSize, TEXT(""), *Task.Name.ToString());
+							STATETREE_LOG(Verbose, TEXT("%*sTask '%s'.ExitState()"), (FrameIndex + StateIndex + 1)*UE::StateTree::DebugIndentSize, TEXT(""), *Task.Name.ToString());
 							UE_STATETREE_DEBUG_TASK_EXIT_STATE(this, CurrentStateTree, FStateTreeIndex16(TaskIndex));
 							{
 								QUICK_SCOPE_CYCLE_COUNTER(StateTree_Task_ExitState);
@@ -2066,14 +2098,14 @@ void FStateTreeExecutionContext::StateCompleted()
 
 		if (FrameIndex <= Exec.EnterStateFailedFrameIndex.Get())
 		{
-			for (int32 Index = CurrentFrame.ActiveStates.Num() - 1; Index >= 0; Index--)
+			for (int32 StateIndex = CurrentFrame.ActiveStates.Num() - 1; StateIndex >= 0; StateIndex--)
 			{
-				const FStateTreeStateHandle CurrentHandle = CurrentFrame.ActiveStates[Index];
+				const FStateTreeStateHandle CurrentHandle = CurrentFrame.ActiveStates[StateIndex];
 				const FCompactStateTreeState& State = CurrentStateTree->States[CurrentHandle.Index];
 
 				FCurrentlyProcessedStateScope StateScope(*this, CurrentHandle);
 				
-				STATETREE_LOG(Verbose, TEXT("%*sState '%s'"), Index*UE::StateTree::DebugIndentSize, TEXT(""), *DebugGetStatePath(Exec.ActiveFrames, &CurrentFrame, Index));
+				STATETREE_LOG(Verbose, TEXT("%*sState '%s'"), (FrameIndex + StateIndex + 1)*UE::StateTree::DebugIndentSize, TEXT(""), *GetSafeStateName(CurrentFrame, CurrentHandle));
 				STATETREE_TRACE_STATE_EVENT(CurrentHandle, EStateTreeTraceEventType::OnStateCompleted);
 
 				// Notify Tasks
@@ -2095,7 +2127,7 @@ void FStateTreeExecutionContext::StateCompleted()
 							continue;
 						}
 
-						STATETREE_LOG(Verbose, TEXT("%*s  Task '%s'"), Index*UE::StateTree::DebugIndentSize, TEXT(""), *Task.Name.ToString());
+						STATETREE_LOG(Verbose, TEXT("%*sTask '%s'.StateCompleted()"), (FrameIndex + StateIndex + 1)*UE::StateTree::DebugIndentSize, TEXT(""), *Task.Name.ToString());
 						Task.StateCompleted(*this, Exec.LastTickStatus, CurrentFrame.ActiveStates);
 					}
 				}
@@ -2642,15 +2674,16 @@ EStateTreeRunStatus FStateTreeExecutionContext::TickTasks(const float DeltaTime)
 
 		FCurrentlyProcessedFrameScope FrameScope(*this, CurrentParentFrame, CurrentFrame);
 
-		for (int32 Index = 0; Index < CurrentFrame.ActiveStates.Num(); Index++)
+		for (int32 StateIndex = 0; StateIndex < CurrentFrame.ActiveStates.Num(); StateIndex++)
 		{
-			const FStateTreeStateHandle CurrentHandle = CurrentFrame.ActiveStates[Index];
+			const FStateTreeStateHandle CurrentHandle = CurrentFrame.ActiveStates[StateIndex];
 			const FCompactStateTreeState& State = CurrentStateTree->States[CurrentHandle.Index];
 
 			FCurrentlyProcessedStateScope StateScope(*this, CurrentHandle);
 			STATETREE_TRACE_SCOPED_STATE(CurrentHandle);
 
-			STATETREE_CLOG(State.TasksNum > 0, VeryVerbose, TEXT("%*sState '%s'"), Index*UE::StateTree::DebugIndentSize, TEXT(""), *DebugGetStatePath(Exec.ActiveFrames, &CurrentFrame, Index));
+			STATETREE_CLOG(State.TasksNum > 0, VeryVerbose, TEXT("%*sState '%s'"), (FrameIndex + StateIndex + 1)*UE::StateTree::DebugIndentSize, TEXT("")
+				, *DebugGetStatePath(Exec.ActiveFrames, &CurrentFrame, StateIndex));
 
 			if (State.ShouldTickTasks(bHasEvents))
 			{
@@ -2683,7 +2716,10 @@ EStateTreeRunStatus FStateTreeExecutionContext::TickTasks(const float DeltaTime)
 					}
 
 					const bool bNeedsTick = bShouldTickTasks && (Task.bShouldCallTick || (bHasEvents && Task.bShouldCallTickOnlyOnEvents));
-					STATETREE_LOG(VeryVerbose, TEXT("%*s  Tick: '%s' %s"), Index*UE::StateTree::DebugIndentSize, TEXT(""), *Task.Name.ToString(), !bNeedsTick ? TEXT("[not ticked]") : TEXT(""));
+					STATETREE_LOG(VeryVerbose, TEXT("%*s  Tick: '%s' %s"), (FrameIndex + StateIndex + 1)*UE::StateTree::DebugIndentSize, TEXT("")
+						, *Task.Name.ToString()
+						, !bNeedsTick ? TEXT("[not ticked]") : TEXT(""));
+
 					if (!bNeedsTick)
 					{
 						continue;
@@ -4153,8 +4189,13 @@ bool FStateTreeExecutionContext::SelectStateInternal(
 				// Try to prevent (infinite) loops in the selection.
 				if (CurrentFrame.ActiveStates.Contains(Transition.State))
 				{
-					STATETREE_LOG(Error, TEXT("%hs: Loop detected when trying to select state %s from '%s'. Prior states: %s.  '%s' using StateTree '%s'."),
-						__FUNCTION__, *GetSafeStateName(CurrentFrame, NextStateHandle), *GetStateStatusString(Exec), *DebugGetStatePath(OutSelectionResult.GetSelectedFrames(), &CurrentFrame), *GetNameSafe(&Owner), *GetFullNameSafe(CurrentFrame.StateTree));
+					STATETREE_LOG(Error, TEXT("%hs: Loop detected when trying to select state %s from '%s'. Prior states: %s.  '%s' using StateTree '%s'.")
+						, __FUNCTION__
+						, *GetSafeStateName(CurrentFrame, NextStateHandle)
+						, *GetStateStatusString(Exec)
+						, *DebugGetStatePath(OutSelectionResult.GetSelectedFrames(), &CurrentFrame)
+						, *GetNameSafe(&Owner)
+						, *GetFullNameSafe(CurrentFrame.StateTree));
 					continue;
 				}
 
