[State Tree Editor]
- Update TreeView when task or enter condition array changes, or any state transition or task transition changes
- Added comment

diff --git a/Engine/Plugins/Runtime/StateTree/Source/StateTreeEditorModule/Private/Customizations/StateTreeEditorNodeUtils.cpp b/Engine/Plugins/Runtime/StateTree/Source/StateTreeEditorModule/Private/Customizations/StateTreeEditorNodeUtils.cpp
--- a/Engine/Plugins/Runtime/StateTree/Source/StateTreeEditorModule/Private/Customizations/StateTreeEditorNodeUtils.cpp
+++ b/Engine/Plugins/Runtime/StateTree/Source/StateTreeEditorModule/Private/Customizations/StateTreeEditorNodeUtils.cpp
@@ -653,7 +653,8 @@ void OnArrayNodePicked(const UStruct* InStruct, TSharedPtr<SComboButton> PickerC
 				NewNodeHandle->SetExpanded(true);
 			}
 		}
-		
+
+		// We initialized the new element, so broadcasting an extra callback with ValueSet type, besides the one from AddItem()
 		ArrayPropertyHandle->NotifyPostChange(EPropertyChangeType::ValueSet);
 		ArrayPropertyHandle->NotifyFinishedChangingProperties();
 		GEditor->EndTransaction();
diff --git a/Engine/Plugins/Runtime/StateTree/Source/StateTreeEditorModule/Private/SStateTreeView.cpp b/Engine/Plugins/Runtime/StateTree/Source/StateTreeEditorModule/Private/SStateTreeView.cpp
--- a/Engine/Plugins/Runtime/StateTree/Source/StateTreeEditorModule/Private/SStateTreeView.cpp
+++ b/Engine/Plugins/Runtime/StateTree/Source/StateTreeEditorModule/Private/SStateTreeView.cpp
@@ -445,24 +445,21 @@ void SStateTreeView::HandleModelStateAdded(UStateTreeState* ParentState, UStateT
 
 void SStateTreeView::HandleModelStatesChanged(const TSet<UStateTreeState*>& AffectedStates, const FPropertyChangedEvent& PropertyChangedEvent)
 {
-	bool bArraysChanged = false;
-
-	// The purpose of the rebuild below is to update the task visualization (number of widgets change).
-	// This method is called when anything in a state changes, make sure to only rebuild when needed.
-	if (PropertyChangedEvent.MemberProperty->GetFName() == GET_MEMBER_NAME_CHECKED(UStateTreeState, Tasks))
-	{
-		bArraysChanged = true;
-	}
-	if (PropertyChangedEvent.MemberProperty->GetFName() == GET_MEMBER_NAME_CHECKED(UStateTreeState, Transitions))
+	// When the tasks or conditions array changed(this includes both normal array operations: Add, Remove. Clear, Move,
+	// and Paste or Duplicate an element in the array), The TreeView needs to be rebuilt because new elements came in or old elements have gone or both.
+	// This will not rebuild the list when we change an inner property in a condition or in a task node because of InstanceStruct wrapper
+	if (PropertyChangedEvent.MemberProperty->GetFName() == GET_MEMBER_NAME_CHECKED(UStateTreeState, Tasks)
+		|| PropertyChangedEvent.MemberProperty->GetFName() == GET_MEMBER_NAME_CHECKED(UStateTreeState, EnterConditions))
 	{
-		bArraysChanged = true;
-	}
-	if (PropertyChangedEvent.MemberProperty->GetFName() == GET_MEMBER_NAME_CHECKED(UStateTreeState, EnterConditions))
-	{
-		bArraysChanged = true;
+		TreeView->RebuildList();
 	}
-		
-	if (bArraysChanged)
+	// todo: find a better way to fetch/update Transition information so we don't have to hardcode the rules here
+	// We cache the information of Transitions when making SStateTreeViewRow and using cached info to display
+	// So we need to rebuild the list whenever the relevant information changes to update the view
+	// For State Transitions, this will cover any changes made on the Array and on each Transition
+	// For Task Transitions, LinkType is the only exposed one in the editor
+	else if (PropertyChangedEvent.MemberProperty->GetFName() == GET_MEMBER_NAME_CHECKED(UStateTreeState, Transitions)
+		|| PropertyChangedEvent.Property->GetFName() == GET_MEMBER_NAME_CHECKED(FStateTreeStateLink, LinkType))
 	{
 		TreeView->RebuildList();
 	}
