Fix crash during editor teardown while a State Tree is open

diff --git a/Engine/Plugins/Runtime/StateTree/Source/StateTreeEditorModule/Private/StateTreeViewModel.cpp b/Engine/Plugins/Runtime/StateTree/Source/StateTreeEditorModule/Private/StateTreeViewModel.cpp
--- a/Engine/Plugins/Runtime/StateTree/Source/StateTreeEditorModule/Private/StateTreeViewModel.cpp
+++ b/Engine/Plugins/Runtime/StateTree/Source/StateTreeEditorModule/Private/StateTreeViewModel.cpp
@@ -241,7 +241,10 @@ FStateTreeViewModel::FStateTreeViewModel()
 
 FStateTreeViewModel::~FStateTreeViewModel()
 {
-	GEditor->UnregisterForUndo(this);
+	if (GEditor)
+	{
+	    GEditor->UnregisterForUndo(this);
+	}
 
 	UE::StateTree::Delegates::OnIdentifierChanged.RemoveAll(this);
 }
