[State Tree] Fixed Random Selector not using the correct StateIndex

diff --git a/Engine/Plugins/Runtime/StateTree/Source/StateTreeModule/Private/StateTreeExecutionContext.cpp b/Engine/Plugins/Runtime/StateTree/Source/StateTreeModule/Private/StateTreeExecutionContext.cpp
--- a/Engine/Plugins/Runtime/StateTree/Source/StateTreeModule/Private/StateTreeExecutionContext.cpp
+++ b/Engine/Plugins/Runtime/StateTree/Source/StateTreeModule/Private/StateTreeExecutionContext.cpp
@@ -4306,7 +4306,8 @@ bool FStateTreeExecutionContext::SelectStateInternal(
 				while (!NextLevelChildStates.IsEmpty())
 				{
 					const int32 ChildStateIndex = Exec.RandomStream.RandRange(0, NextLevelChildStates.Num() - 1);
-					if (SelectStateInternal(CurrentParentFrame, CurrentFrame, CurrentFrameInActiveFrames, { FStateTreeStateHandle(ChildStateIndex) }, OutSelectionResult))
+					if (SelectStateInternal(CurrentParentFrame, CurrentFrame, CurrentFrameInActiveFrames,
+					                        {FStateTreeStateHandle(NextLevelChildStates[ChildStateIndex])}, OutSelectionResult))
 					{
 						// Selection succeeded
 						bSucceededToSelectState = true;
