StateTree: Copy ID of nodes and bindable struct after a copy. Remap the old value to the new value.

diff --git a/Engine/Plugins/Runtime/StateTree/Source/StateTreeEditorModule/Private/StateTreeEditorData.cpp b/Engine/Plugins/Runtime/StateTree/Source/StateTreeEditorModule/Private/StateTreeEditorData.cpp
--- a/Engine/Plugins/Runtime/StateTree/Source/StateTreeEditorModule/Private/StateTreeEditorData.cpp
+++ b/Engine/Plugins/Runtime/StateTree/Source/StateTreeEditorModule/Private/StateTreeEditorData.cpp
@@ -295,6 +295,12 @@ void UStateTreeEditorData::PostEditChangeChainProperty(FPropertyChangedChainEven
 
 	UE::StateTree::PropertyHelpers::DispatchPostEditToNodes(*this, PropertyChangedEvent, *this);
 }
+
+void UStateTreeEditorData::PostDuplicate(EDuplicateMode::Type DuplicateMode)
+{
+	Super::PostDuplicate(DuplicateMode);
+	DuplicateIDs();
+}
 #endif // WITH_EDITOR
 
 void UStateTreeEditorData::GetAccessibleStructs(const FGuid TargetStructID, TArray<FStateTreeBindableStructDesc>& OutStructDescs) const
@@ -973,6 +979,134 @@ void UStateTreeEditorData::FixDuplicateIDs()
 	EditorBindings.RemoveUnusedBindings(AllValues);
 }
 
+void UStateTreeEditorData::DuplicateIDs()
+{
+	TMap<FGuid, FGuid> OldToNewIDs;
+
+	// Visit and create new ids
+	{
+		auto AddId = [&OldToNewIDs](const FGuid& OldID, bool bTestIfContains = true)
+		{
+			ensureAlwaysMsgf(bTestIfContains == false || !OldToNewIDs.Contains(OldID), TEXT("The id is duplicated and FixDuplicateIDs failed to fix it."));
+
+			if (OldID.IsValid())
+			{
+				const FGuid NewID = FGuid::NewGuid();
+				OldToNewIDs.Add(OldID, NewID);
+			}
+		};
+
+		auto AddIds = [&AddId](TArrayView<FStateTreeEditorNode> Nodes)
+		{
+			for (FStateTreeEditorNode& Node : Nodes)
+			{
+				AddId(Node.ID);
+			}
+		};
+
+		// Do not use the VisitGlobalNodes because the schema should not be included in OldToNewIDs 
+		OldToNewIDs.Add(RootParameters.ID, FGuid::NewGuid());
+		AddIds(Evaluators);
+		AddIds(GlobalTasks);
+		for (FStateTreeEditorColor& Color : Colors)
+		{
+			AddId(Color.ColorRef.ID);
+		}
+
+		VisitHierarchy([Self = this, &AddId, &OldToNewIDs](const UStateTreeState& State, UStateTreeState* ParentState)
+		{
+			AddId(State.ID);
+			AddId(State.Parameters.ID);
+
+			for (const FStateTreeTransition& Transition : State.Transitions)
+			{
+				AddId(Transition.ID);
+			}
+
+			return Self->VisitStateNodes(State, [&AddId](const UStateTreeState* State, const FStateTreeBindableStructDesc& Desc, const FStateTreeDataView Value)
+			{
+				AddId(Desc.ID, /*bTestIfContains*/false);
+				return EStateTreeVisitor::Continue;
+			});
+		});
+
+		// Confirms that we collected everything.
+		{
+			TMap<FGuid, const FStateTreeDataView> AllStructValues;
+			GetAllStructValues(AllStructValues);
+			// Schema ids are not duplicated
+			if (Schema != nullptr)
+			{
+				for (const FStateTreeExternalDataDesc& ContextDesc : Schema->GetContextDataDescs())
+				{
+					AllStructValues.Remove(ContextDesc.ID);
+				}
+			}
+			for (auto& StructValue : AllStructValues)
+			{
+				ensureMsgf(OldToNewIDs.Contains(StructValue.Key), TEXT("An ID container was not duplicated for asset '%s'."), *GetOutermost()->GetName());
+			}
+		}
+	}
+
+	// Remap ids properties to the new generated ids
+	{
+		TArray<const UObject*> ObjectToSearch;
+		TSet<const UObject*> ObjectSearched;
+		ObjectToSearch.Add(this);
+		ObjectSearched.Add(this);
+		while (ObjectToSearch.Num())
+		{
+			const UObject* CurrentObject = ObjectToSearch.Pop();
+			for (TPropertyValueIterator<FProperty> It(CurrentObject->GetClass(), CurrentObject, EPropertyValueIteratorFlags::FullRecursion, EFieldIteratorFlags::ExcludeDeprecated); It; ++It)
+			{
+				const FProperty* Property = It.Key();
+				if (const FStructProperty* StructProperty = CastField<FStructProperty>(Property))
+				{
+					if (StructProperty->Struct == TBaseStructure<FGuid>::Get())
+					{
+						// Skip the guid properties.
+						It.SkipRecursiveProperty();
+
+						// Modify the value if needed.
+						// We const_cast because we don't change the layout and the iterator advance is still going to work.
+						FGuid* GuidValue = reinterpret_cast<FGuid*>(const_cast<void*>(It.Value()));
+						if (FGuid* NewGuidValue = OldToNewIDs.Find(*GuidValue))
+						{
+							*GuidValue = *NewGuidValue;
+						}
+					}
+				}
+				else if (const FObjectProperty* ObjectProperty = CastField<FObjectProperty>(Property))
+				{
+					if (const UObject* ObjectValue = ObjectProperty->GetPropertyValue(It.Value()))
+					{
+						if (!ObjectSearched.Contains(ObjectValue))
+						{
+							// Add the inner properties of this instanced object.
+							bool bAddObject = false;
+							if (Property->HasAllPropertyFlags(CPF_ExportObject) && ObjectValue->GetClass()->HasAllClassFlags(CLASS_EditInlineNew))
+							{
+								bAddObject = true;
+							}
+							else if (Property->HasAllPropertyFlags(CPF_InstancedReference))
+							{
+								bAddObject = true;
+							}
+
+							if (bAddObject)
+							{
+								ObjectSearched.Add(ObjectValue);
+								ObjectToSearch.Add(ObjectValue);
+							}
+						}
+					}
+				}
+			}
+		}
+	}
+}
+
 void UStateTreeEditorData::UpdateBindingsInstanceStructs()
 {
 	TMap<FGuid, const FStateTreeDataView> AllValues;
diff --git a/Engine/Plugins/Runtime/StateTree/Source/StateTreeEditorModule/Public/StateTreeEditorData.h b/Engine/Plugins/Runtime/StateTree/Source/StateTreeEditorModule/Public/StateTreeEditorData.h
--- a/Engine/Plugins/Runtime/StateTree/Source/StateTreeEditorModule/Public/StateTreeEditorData.h
+++ b/Engine/Plugins/Runtime/StateTree/Source/StateTreeEditorModule/Public/StateTreeEditorData.h
@@ -94,6 +94,7 @@ public:
 	virtual void BeginDestroy() override;
 	virtual void PostLoad() override;
 	virtual void PostEditChangeChainProperty(FPropertyChangedChainEvent& PropertyChangedEvent) override;
+	virtual void PostDuplicate(EDuplicateMode::Type DuplicateMode) override;
 #endif
 
 	/** @returns parent state of a struct, or nullptr if not found. */
@@ -283,6 +284,7 @@ private:
 	void FixObjectInstance(TSet<UObject*>& SeenObjects, UObject& Outer, FStateTreeEditorNode& Node);
 	void FixObjectNodes();
 	void FixDuplicateIDs();
+	void DuplicateIDs();
 	void UpdateBindingsInstanceStructs();
 	void CallPostLoadOnNodes();
 
diff --git a/Engine/Plugins/Runtime/StateTree/Source/StateTreeModule/Private/StateTree.cpp b/Engine/Plugins/Runtime/StateTree/Source/StateTreeModule/Private/StateTree.cpp
--- a/Engine/Plugins/Runtime/StateTree/Source/StateTreeModule/Private/StateTree.cpp
+++ b/Engine/Plugins/Runtime/StateTree/Source/StateTreeModule/Private/StateTree.cpp
@@ -329,6 +329,12 @@ void UStateTree::BeginDestroy()
 	Super::BeginDestroy();
 }
 
+void UStateTree::PostDuplicate(EDuplicateMode::Type DuplicateMode)
+{
+	Super::PostDuplicate(DuplicateMode);
+	ResetCompiled();
+}
+
 void UStateTree::GetAssetRegistryTags(TArray<FAssetRegistryTag>& OutTags) const
 {
 	PRAGMA_DISABLE_DEPRECATION_WARNINGS;
diff --git a/Engine/Plugins/Runtime/StateTree/Source/StateTreeModule/Public/StateTree.h b/Engine/Plugins/Runtime/StateTree/Source/StateTreeModule/Public/StateTree.h
--- a/Engine/Plugins/Runtime/StateTree/Source/StateTreeModule/Public/StateTree.h
+++ b/Engine/Plugins/Runtime/StateTree/Source/StateTreeModule/Public/StateTree.h
@@ -227,6 +227,7 @@ protected:
 	void OnUserDefinedStructReinstanced(const UUserDefinedStruct& UserDefinedStruct);
 	virtual void PostInitProperties() override;
 	virtual void BeginDestroy() override;
+	virtual void PostDuplicate(EDuplicateMode::Type DuplicateMode) override;
 	virtual void GetAssetRegistryTags(FAssetRegistryTagsContext Context) const override;
 	UE_DEPRECATED(5.4, "Implement the version that takes FAssetRegistryTagsContext instead.")
 	virtual void GetAssetRegistryTags(TArray<FAssetRegistryTag>& OutTags) const override;
