Make new-Task/Condition/Consideration toolbar buttons conditional on whether the classes are valid within the state-tree schema (this hides it for cases where the buttons don't apply)

diff --git a/Engine/Plugins/Runtime/StateTree/Source/StateTreeEditorModule/Private/StateTreeEditorModeToolkit.cpp b/Engine/Plugins/Runtime/StateTree/Source/StateTreeEditorModule/Private/StateTreeEditorModeToolkit.cpp
--- a/Engine/Plugins/Runtime/StateTree/Source/StateTreeEditorModule/Private/StateTreeEditorModeToolkit.cpp
+++ b/Engine/Plugins/Runtime/StateTree/Source/StateTreeEditorModule/Private/StateTreeEditorModeToolkit.cpp
@@ -309,7 +309,25 @@ void FStateTreeEditorModeToolkit::ExtendSecondaryModeToolbar(UToolMenu* ToolBar)
 		{
 			InSection.AddEntry(FToolMenuEntry::InitComboButton(
 				 "CreateNewTaskComboButton",
-				 FUIAction(),
+				 FUIAction(FExecuteAction(), FCanExecuteAction(), FIsActionChecked(),
+				 	FIsActionButtonVisible::CreateLambda([ToolkitContextObject, GetToolkitFromAssetEditorContext]() -> bool
+					{
+				 		if(TSharedPtr<FStateTreeEditorModeToolkit> SharedToolkit = GetToolkitFromAssetEditorContext(ToolkitContextObject))
+					 	{
+							 if(UStateTree* StateTree = SharedToolkit->EditorHost->GetStateTree())
+							 {
+								 if (const UStateTreeEditorData* EditorData = Cast<UStateTreeEditorData>(StateTree->EditorData))
+								 {
+									 if (const UStateTreeSchema* Schema = EditorData->Schema.Get())
+									 {
+										 return Schema->IsClassAllowed(UStateTreeTaskBlueprintBase::StaticClass());
+									 }
+								 }
+							 }
+						}
+					 	
+						return false;
+					})),
 				 FOnGetContent::CreateLambda([ToolkitContextObject, GetToolkitFromAssetEditorContext]() -> TSharedRef<SWidget>
 				 {
 				 	if(TSharedPtr<FStateTreeEditorModeToolkit> SharedToolkit = GetToolkitFromAssetEditorContext(ToolkitContextObject))
@@ -326,7 +344,25 @@ void FStateTreeEditorModeToolkit::ExtendSecondaryModeToolbar(UToolMenu* ToolBar)
 				
 			InSection.AddEntry(FToolMenuEntry::InitComboButton(
 				 "CreateNewConditionComboButton",
-				 FUIAction(),
+				 FUIAction(FExecuteAction(), FCanExecuteAction(), FIsActionChecked(),
+					 FIsActionButtonVisible::CreateLambda([ToolkitContextObject, GetToolkitFromAssetEditorContext]() -> bool
+					{
+						if(TSharedPtr<FStateTreeEditorModeToolkit> SharedToolkit = GetToolkitFromAssetEditorContext(ToolkitContextObject))
+						{
+							if(UStateTree* StateTree = SharedToolkit->EditorHost->GetStateTree())
+							{
+								if (const UStateTreeEditorData* EditorData = Cast<UStateTreeEditorData>(StateTree->EditorData))
+								{
+									if (const UStateTreeSchema* Schema = EditorData->Schema.Get())
+									{
+										return Schema->IsClassAllowed(UStateTreeConditionBlueprintBase::StaticClass());
+									}
+								}
+							}
+						}
+				 		
+						return false;
+					})),
 				FOnGetContent::CreateLambda([ToolkitContextObject, GetToolkitFromAssetEditorContext]() -> TSharedRef<SWidget>
 				{
 					if(TSharedPtr<FStateTreeEditorModeToolkit> SharedToolkit = GetToolkitFromAssetEditorContext(ToolkitContextObject))
@@ -343,7 +379,25 @@ void FStateTreeEditorModeToolkit::ExtendSecondaryModeToolbar(UToolMenu* ToolBar)
 				
 			 InSection.AddEntry(FToolMenuEntry::InitComboButton(
 				 "CreateNewConsiderationComboButton",
-				 FUIAction(),
+				 FUIAction(FExecuteAction(), FCanExecuteAction(), FIsActionChecked(),
+					 FIsActionButtonVisible::CreateLambda([ToolkitContextObject, GetToolkitFromAssetEditorContext]() -> bool
+					{
+					 	if(TSharedPtr<FStateTreeEditorModeToolkit> SharedToolkit = GetToolkitFromAssetEditorContext(ToolkitContextObject))
+					 	{
+							 if(UStateTree* StateTree = SharedToolkit->EditorHost->GetStateTree())
+							 {
+								 if (const UStateTreeEditorData* EditorData = Cast<UStateTreeEditorData>(StateTree->EditorData))
+								 {
+									 if (const UStateTreeSchema* Schema = EditorData->Schema.Get())
+									 {
+										 return Schema->IsClassAllowed(UStateTreeConsiderationBlueprintBase::StaticClass());
+									 }
+								 }
+							 }
+						}
+					 	
+						return false;
+					})),
 				 FOnGetContent::CreateLambda([ToolkitContextObject, GetToolkitFromAssetEditorContext]() -> TSharedRef<SWidget>
 				{
 				 	if(TSharedPtr<FStateTreeEditorModeToolkit> SharedToolkit = GetToolkitFromAssetEditorContext(ToolkitContextObject))
