[StateTree] exposed currently processed node and selection result for debugging purposes

diff --git a/Engine/Plugins/Runtime/StateTree/Source/StateTreeModule/Private/StateTreeExecutionContext.cpp b/Engine/Plugins/Runtime/StateTree/Source/StateTreeModule/Private/StateTreeExecutionContext.cpp
--- a/Engine/Plugins/Runtime/StateTree/Source/StateTreeModule/Private/StateTreeExecutionContext.cpp
+++ b/Engine/Plugins/Runtime/StateTree/Source/StateTreeModule/Private/StateTreeExecutionContext.cpp
@@ -3626,6 +3626,8 @@ bool FStateTreeExecutionContext::SelectState(const FStateTreeExecutionFrame& Cur
 											const FStateTreeSharedEvent* TransitionEvent,
 											const EStateTreeSelectionFallback Fallback)
 {
+	TGuardValue<const FStateSelectionResult*> GuardValue(CurrentSelectionResult, &OutSelectionResult);
+
 	const FStateTreeExecutionState& Exec = GetExecState();
 
 	if (Exec.ActiveFrames.IsEmpty())
diff --git a/Engine/Plugins/Runtime/StateTree/Source/StateTreeModule/Public/StateTreeExecutionContext.h b/Engine/Plugins/Runtime/StateTree/Source/StateTreeModule/Public/StateTreeExecutionContext.h
--- a/Engine/Plugins/Runtime/StateTree/Source/StateTreeModule/Public/StateTreeExecutionContext.h
+++ b/Engine/Plugins/Runtime/StateTree/Source/StateTreeModule/Public/StateTreeExecutionContext.h
@@ -141,8 +141,15 @@ public:
 	/** @return The world of the owner or nullptr if the owner is not set. */ 
 	UWorld* GetWorld() const { return Owner.GetWorld(); };
 
-	/** @return True of the the execution context is valid and initialized. */ 
-	bool IsValid() const { return RootStateTree.IsReadyToRun(); }
+	/**
+	 * Indicates if the context is valid and able to run an instance of the associated StateTree asset.
+	 * @return True if the StateTree asset assigned to the execution context is valid
+	 * (i.e., not empty) and successfully initialized (i.e., linked and all bindings resolved).
+	 */
+	bool IsValid() const
+	{
+		return RootStateTree.IsReadyToRun();
+	}
 	
 	/**
 	 * Start executing.
@@ -284,15 +291,30 @@ public:
 		});
 	}
 
+	/** @return the currently processed node if applicable. */
+	FStateTreeDataHandle GetCurrentlyProcessedNode() const
+	{
+		return CurrentNodeDataHandle;
+	}
+
 	/** @return the currently processed state if applicable. */
-	FStateTreeStateHandle GetCurrentlyProcessedState() const { return CurrentlyProcessedState; }
+	FStateTreeStateHandle GetCurrentlyProcessedState() const
+	{
+		return CurrentlyProcessedState;
+	}
 
 	/** @return the currently processed execution frame if applicable. */
-	const FStateTreeExecutionFrame* GetCurrentlyProcessedFrame() const { return CurrentlyProcessedFrame; }
+	const FStateTreeExecutionFrame* GetCurrentlyProcessedFrame() const
+	{
+		return CurrentlyProcessedFrame;
+	}
 
 	/** @return the currently processed execution parent frame if applicable. */
-	const FStateTreeExecutionFrame* GetCurrentlyProcessedParentFrame() const { return CurrentlyProcessedParentFrame; }
-	
+	const FStateTreeExecutionFrame* GetCurrentlyProcessedParentFrame() const
+	{
+		return CurrentlyProcessedParentFrame;
+	}
+
 	/** @return Pointer to a State or null if state not found */ 
 	const FCompactStateTreeState* GetStateFromHandle(const FStateTreeStateHandle StateHandle) const
 	{
@@ -505,6 +527,11 @@ protected:
 			return SelectedFrames;
 		}
 
+		TConstArrayView<FStateTreeExecutionFrame> GetSelectedFrames() const
+		{
+			return SelectedFrames;
+		}
+
 		TArrayView<FStateTreeFrameStateSelectionEvents> GetFramesStateSelectionEvents()
 		{
 			return FramesStateSelectionEvents;
@@ -543,7 +570,7 @@ protected:
 	void ExitState(const FStateTreeTransitionResult& Transition);
 
 	/**
-	 * Handles logic for signalling State completed. StateCompleted is called on current active Evaluators and Tasks in reverse order (from leaf to root).
+	 * Handles logic for signaling State completed. StateCompleted is called on current active Evaluators and Tasks in reverse order (from leaf to root).
 	 */
 	void StateCompleted();
 
@@ -792,10 +819,13 @@ protected:
 		const FStateTreeExecutionFrame* SavedFrame = nullptr;
 		const FStateTreeExecutionFrame* SavedParentFrame = nullptr;
 	};
-	
+
+	/** Current state selection result when performing recursive state selection, or nullptr if not applicable. */
+	const FStateSelectionResult* CurrentSelectionResult = nullptr;
+
 	/** Current state we're processing, or invalid if not applicable. */
 	FStateTreeStateHandle CurrentlyProcessedState;
-	
+
 	/** Helper struct to track currently processed state. */
 	struct FCurrentlyProcessedStateScope
 	{
