[StateTree] Merge bindings related types with PropertyBindingUtils types

diff --git a/Engine/Plugins/Runtime/PropertyBindingUtils/Source/PropertyBindingUtils/Private/PropertyBindingPath.cpp b/Engine/Plugins/Runtime/PropertyBindingUtils/Source/PropertyBindingUtils/Private/PropertyBindingPath.cpp
--- a/Engine/Plugins/Runtime/PropertyBindingUtils/Source/PropertyBindingUtils/Private/PropertyBindingPath.cpp
+++ b/Engine/Plugins/Runtime/PropertyBindingUtils/Source/PropertyBindingUtils/Private/PropertyBindingPath.cpp
@@ -143,10 +143,10 @@ bool FPropertyBindingPath::UpdateSegmentsFromValue(const FPropertyBindingDataVie
 	return true;
 }
 
-FString FPropertyBindingPath::ToString(const int32 HighlightedSegment, const TCHAR* HighlightPrefix, const TCHAR* HighlightPostfix, const bool bOutputInstances) const
+FString FPropertyBindingPath::ToString(const int32 HighlightedSegment, const TCHAR* HighlightPrefix, const TCHAR* HighlightPostfix, const bool bOutputInstances, const int32 FirstSegment) const
 {
-	FString Result;
-	for (TEnumerateRef<const FPropertyBindingPathSegment> Segment : EnumerateRange(Segments))
+	FStringBuilderBase Result;
+	for (TEnumerateRef<const FPropertyBindingPathSegment> Segment : EnumerateRange(TConstArrayView<FPropertyBindingPathSegment>(Segments).Mid(FMath::Max(FirstSegment, 0))))
 	{
 		if (Segment.GetIndex() > 0)
 		{
@@ -162,24 +162,66 @@ FString FPropertyBindingPath::ToString(const int32 HighlightedSegment, const TCH
 			Result += FString::Printf(TEXT("(%s)"), *GetNameSafe(Segment->GetInstanceStruct()));
 		}
 
-		if (Segment->GetArrayIndex() >= 0)
+#if WITH_EDITORONLY_DATA
+		const UStruct* ParentInstanceStruct = Segment.GetIndex() > 0 ? Segments[Segment.GetIndex() - 1].GetInstanceStruct() : nullptr;
+		if (const UUserDefinedStruct* ParentUserDefinedStruct = Cast<const UUserDefinedStruct>(ParentInstanceStruct))
 		{
-			Result += FString::Printf(TEXT("%s[%d]"), *Segment->GetName().ToString(), Segment->GetArrayIndex());
+			// Find friendly names for UDS properties (the property name itself has hash in it). 
+			const FString FriendlyName = FStructureEditorUtils::GetVariableFriendlyName(ParentUserDefinedStruct, Segment->GetPropertyGuid());
+			if (!FriendlyName.IsEmpty())
+			{
+				Result += FriendlyName;
+			}
+			else
+			{
+				Result += Segment->GetName().ToString();
+			}
 		}
 		else
+#endif // WITH_EDITORONLY_DATA
 		{
 			Result += Segment->GetName().ToString();
 		}
 
+		if (Segment->GetArrayIndex() >= 0)
+		{
+			Result += FString::Printf(TEXT("[%d]"), Segment->GetArrayIndex());
+		}
+
 		if (Segment.GetIndex() == HighlightedSegment && HighlightPostfix)
 		{
 			Result += HighlightPostfix;
 		}
 	}
-	return Result;
+	return Result.ToString();
 }
 
 
+bool FPropertyBindingPath::Includes(const FPropertyBindingPath& Other) const
+{
+#if WITH_EDITORONLY_DATA
+	if (StructID != Other.StructID)
+	{
+		return false;
+	}
+#endif // WITH_EDITORONLY_DATA
+
+	if (Segments.Num() < Other.Segments.Num())
+	{
+		return false;
+	}
+
+	for (TEnumerateRef<const FPropertyBindingPathSegment> OtherSegment : EnumerateRange(Other.Segments))
+	{
+		if (*OtherSegment != Segments[OtherSegment.GetIndex()])
+		{
+			return false;
+		}
+	}
+
+	return true;
+}
+
 bool FPropertyBindingPath::ResolveIndirections(const UStruct* BaseStruct, TArray<FPropertyBindingPathIndirection>& OutIndirections, FString* OutError, bool bHandleRedirects) const
 {
 	return ResolveIndirectionsWithValue(FPropertyBindingDataView(BaseStruct, nullptr), OutIndirections, OutError, bHandleRedirects);
@@ -199,7 +241,7 @@ bool FPropertyBindingPath::ResolveIndirectionsWithValue(const FPropertyBindingDa
 		return true;
 	}
 
-	const void* CurrentAddress = BaseValueView.GetMemory();
+	const uint8* CurrentAddress = static_cast<const uint8*>(BaseValueView.GetMemory());
 	const UStruct* CurrentStruct = BaseValueView.GetStruct();
 	
 	for (const TEnumerateRef<const FPropertyBindingPathSegment> Segment : EnumerateRange(Segments))
@@ -339,7 +381,7 @@ bool FPropertyBindingPath::ResolveIndirectionsWithValue(const FPropertyBindingDa
 
 			if (bWithValue)
 			{
-				FScriptArrayHelper Helper(ArrayProperty, (uint8*)CurrentAddress + ArrayProperty->GetOffset_ForInternal());
+				FScriptArrayHelper Helper(ArrayProperty, CurrentAddress + ArrayProperty->GetOffset_ForInternal());
 				if (!Helper.IsValidIndex(Segment->GetArrayIndex()))
 				{
 					if (OutError)
@@ -477,8 +519,8 @@ bool FPropertyBindingPath::ResolveIndirectionsWithValue(const FPropertyBindingDa
 			{
 				if (bWithValue)
 				{
-					const UObject* Object = *reinterpret_cast<UObject* const*>((uint8*)CurrentAddress + Offset);
-					CurrentAddress = Object;
+					const UObject* Object = *reinterpret_cast<UObject* const*>(CurrentAddress + Offset);
+					CurrentAddress = reinterpret_cast<const uint8*>(Object);
 					
 					// The property path is pointing into the object, if the object is present use it's specific type, otherwise use the type of the pointer.
 					if (Object)
@@ -511,19 +553,45 @@ bool FPropertyBindingPath::ResolveIndirectionsWithValue(const FPropertyBindingDa
 			// Check to see if this is a simple weak object property (eg. not an array of weak objects).
 			else if (const FWeakObjectProperty* WeakObjectProperty = CastField<FWeakObjectProperty>(Property))
 			{
-				const TWeakObjectPtr<UObject>& WeakObjectPtr = *reinterpret_cast<const TWeakObjectPtr<UObject>*>((uint8*)CurrentAddress + Offset);
-				const UObject* Object = WeakObjectPtr.Get();
-				CurrentAddress = Object;
-				CurrentStruct = WeakObjectProperty->PropertyClass;
+				if (bWithValue)
+				{
+					const TWeakObjectPtr<UObject>& WeakObjectPtr = *reinterpret_cast<const TWeakObjectPtr<UObject>*>(CurrentAddress + Offset);
+					const UObject* Object = WeakObjectPtr.Get();
+					CurrentAddress = reinterpret_cast<const uint8*>(Object);
+
+					if (Object)
+					{
+						CurrentStruct = Object->GetClass();
+						Indirection.InstanceStruct = CurrentStruct;
+					}
+				}
+				else
+				{
+					CurrentStruct = WeakObjectProperty->PropertyClass;
+				}
+
 				Indirection.AccessType = EPropertyBindingAccessType::WeakObject;
 			}
 			// Check to see if this is a simple soft object property (eg. not an array of soft objects).
 			else if (const FSoftObjectProperty* SoftObjectProperty = CastField<FSoftObjectProperty>(Property))
 			{
-				const FSoftObjectPtr& SoftObjectPtr = *reinterpret_cast<const FSoftObjectPtr*>((uint8*)CurrentAddress + Offset);
-				const UObject* Object = SoftObjectPtr.Get();
-				CurrentAddress = Object;
-				CurrentStruct = SoftObjectProperty->PropertyClass;
+				if (bWithValue)
+				{
+					const FSoftObjectPtr& SoftObjectPtr = *reinterpret_cast<const FSoftObjectPtr*>(CurrentAddress + Offset);
+					const UObject* Object = SoftObjectPtr.Get();
+					CurrentAddress = reinterpret_cast<const uint8*>(Object);
+
+					if (Object)
+					{
+						CurrentStruct = Object->GetClass();
+						Indirection.InstanceStruct = CurrentStruct;
+					}
+				}
+				else
+				{
+					CurrentStruct = SoftObjectProperty->PropertyClass;
+				}
+
 				Indirection.AccessType = EPropertyBindingAccessType::SoftObject;
 			}
 			else
diff --git a/Engine/Plugins/Runtime/PropertyBindingUtils/Source/PropertyBindingUtils/Public/PropertyBindingPath.h b/Engine/Plugins/Runtime/PropertyBindingUtils/Source/PropertyBindingUtils/Public/PropertyBindingPath.h
--- a/Engine/Plugins/Runtime/PropertyBindingUtils/Source/PropertyBindingUtils/Public/PropertyBindingPath.h
+++ b/Engine/Plugins/Runtime/PropertyBindingUtils/Source/PropertyBindingUtils/Public/PropertyBindingPath.h
@@ -70,29 +70,37 @@ enum class EPropertyBindingAccessType : uint8
 };
 
 /**
- * Short lived pointer to an object or struct.
- * The data view expects a type (UStruct) when you pass in a valid memory. In case of null, the type can be empty too.
+ * Short-lived pointer to an object or struct.
+ * Can be constructed directly from a pointer to a UObject or from a StructView.
+ *
+ * When creating from a generic UStruct, the data view expects a valid type (UStruct) when a valid memory pointer is provided.
+ * It that case it is also fine to provide null for both and the constructed data view will be considered as invalid.
  */
 struct PROPERTYBINDINGUTILS_API FPropertyBindingDataView
 {
 	FPropertyBindingDataView() = default;
 
-	// Generic struct constructor.
+	/**
+	 * Generic struct constructor.
+	 * Valid UStruct is required when passing in valid memory pointer.
+	 * Both can be null and the constructed data view will be considered as invalid.
+	 * @see IsValid
+	 */
 	FPropertyBindingDataView(const UStruct* InStruct, void* InMemory) : Struct(InStruct), Memory(InMemory)
 	{
 		// Must have type with valid pointer.
 		check(!Memory || (Memory && Struct));
 	}
 
-	// UObject constructor.
+	/** UObject constructor. */
 	FPropertyBindingDataView(UObject* Object) : Struct(Object ? Object->GetClass() : nullptr), Memory(Object)
 	{
 		// Must have type with valid pointer.
 		check(!Memory || (Memory && Struct));
 	}
 
-	// Struct from a StructView.
-	FPropertyBindingDataView(FStructView StructView) : Struct(StructView.GetScriptStruct()), Memory(StructView.GetMemory())
+	/** Struct from a StructView. */
+	FPropertyBindingDataView(const FStructView StructView) : Struct(StructView.GetScriptStruct()), Memory(StructView.GetMemory())
 	{
 		// Must have type with valid pointer.
 		check(!Memory || (Memory && Struct));
@@ -101,13 +109,13 @@ struct PROPERTYBINDINGUTILS_API FPropertyBindingDataView
 	/**
 	 * Check is the view is valid (both pointer and type are set). On valid views it is safe to call the Get<>() methods returning a reference.
 	 * @return True if the view is valid.
-	*/
+	 */
 	bool IsValid() const
 	{
 		return Memory != nullptr && Struct != nullptr;
 	}
 
-	/*
+	/**
 	 * UObject getters (reference & pointer, const & mutable)
 	 */
 	template <typename T>
@@ -146,7 +154,7 @@ struct PROPERTYBINDINGUTILS_API FPropertyBindingDataView
 		return ((T*)Memory);
 	}
 
-	/*
+	/**
 	 * Struct getters (reference & pointer, const & mutable)
 	 */
 	template <typename T>
@@ -185,7 +193,7 @@ struct PROPERTYBINDINGUTILS_API FPropertyBindingDataView
 		return ((T*)Memory);
 	}
 
-	/*
+	/**
 	 * IInterface getters (reference & pointer, const & mutable)
 	 */
 	template <typename T>
@@ -494,8 +502,9 @@ struct PROPERTYBINDINGUTILS_API FPropertyBindingPath
 	 * @param HighlightPrefix String to append before highlighted segment
 	 * @param HighlightPostfix String to append after highlighted segment
 	 * @param bOutputInstances if true, the instance struct types will be output. 
+	 * @param FirstSegment Index of the first path segment to be stringified.
 	 */
-	FString ToString(const int32 HighlightedSegment = INDEX_NONE, const TCHAR* HighlightPrefix = nullptr, const TCHAR* HighlightPostfix = nullptr, const bool bOutputInstances = false) const;
+	FString ToString(const int32 HighlightedSegment = INDEX_NONE, const TCHAR* HighlightPrefix = nullptr, const TCHAR* HighlightPostfix = nullptr, const bool bOutputInstances = false, const int32 FirstSegment = 0) const;
 
 	/**
 	 * Resolves the property path against base struct type. The path is assumed to be relative to the BaseStruct.
@@ -607,6 +616,12 @@ struct PROPERTYBINDINGUTILS_API FPropertyBindingPath
 	/** Test if paths are equal. */
 	bool operator==(const FPropertyBindingPath& RHS) const;
 
+	/**
+	 * Test if this paths includes the provided path.
+	 * A path includes another one when they are == but this path can be longer.
+	 */
+	bool Includes(const FPropertyBindingPath& Other) const;
+
 private:
 #if WITH_EDITORONLY_DATA
 	/** ID of the struct this property path is relative to. */
diff --git a/Engine/Plugins/Runtime/StateTree/Config/BaseStateTree.ini b/Engine/Plugins/Runtime/StateTree/Config/BaseStateTree.ini
--- a/Engine/Plugins/Runtime/StateTree/Config/BaseStateTree.ini
+++ b/Engine/Plugins/Runtime/StateTree/Config/BaseStateTree.ini
@@ -27,7 +27,8 @@
 +PropertyRedirects=(OldName="StateTreeEditorNode.Item",NewName="StateTreeEditorNode.Node")
 
 +EnumRedirects=(OldName="EStateTreePropertyCopyType",ValueChanges=(("Plain","CopyPlain"),("Complex","CopyComplex"),("Bool","CopyBool"),("Struct","CopyStruct"),("Object","CopyObject"),("Name","CopyName"),("FixedArray","CopyFixedArray")))
-+EnumRedirects=(OldName="EStateTreePropertyAccessType",ValueChanges=(("Array","IndexArray")))
++EnumRedirects=(OldName="EStateTreePropertyAccessType",NewName="/Script/PropertyBindingUtilsModule.EPropertyBindingAccessType")
++EnumRedirects=(OldName="EPropertyBindingAccessType",ValueChanges=(("Array","IndexArray")))
 
 +StructRedirects=(OldName="StateTreeCondition_CompareIntInstanceData",NewName="/script/StateTreeModule.StateTreeCompareIntConditionInstanceData")
 +StructRedirects=(OldName="StateTreeCondition_CompareInt",NewName="/script/StateTreeModule.StateTreeCompareIntCondition")
@@ -69,4 +70,6 @@
 
 +EnumRedirects=(OldName="EStateTreeConditionOperand",NewName="/Script/StateTreeModule.EStateTreeExpressionOperand")
 +PropertyRedirects=(OldName="StateTreeEditorNode.ConditionIndent",NewName="StateTreeEditorNode.ExpressionIndent")
-+PropertyRedirects=(OldName="StateTreeEditorNode.ConditionOperand",NewName="StateTreeEditorNode.ExpressionOperand")
\ No newline at end of file
++PropertyRedirects=(OldName="StateTreeEditorNode.ConditionOperand",NewName="StateTreeEditorNode.ExpressionOperand")
+
++StructRedirects=(OldName="StateTreePropertyPathSegment",NewName="/script/PropertyBindingUtilsModule.PropertyBindingPathSegment")
diff --git a/Engine/Plugins/Runtime/StateTree/Source/StateTreeEditorModule/Private/AsyncStateTreeDiff.cpp b/Engine/Plugins/Runtime/StateTree/Source/StateTreeEditorModule/Private/AsyncStateTreeDiff.cpp
--- a/Engine/Plugins/Runtime/StateTree/Source/StateTreeEditorModule/Private/AsyncStateTreeDiff.cpp
+++ b/Engine/Plugins/Runtime/StateTree/Source/StateTreeEditorModule/Private/AsyncStateTreeDiff.cpp
@@ -274,7 +274,7 @@ static FPropertySoftPath GetPropertyPath(const FStateTreePropertyPath& StateTree
 		|| CheckNodes(StateTreeState->Tasks, "Tasks")
 		|| CheckTransitions(StateTreeState->Transitions, "Transitions"))
 	{
-		for (const FStateTreePropertyPathSegment& PropertySegment : StateTreePropertyPath.GetSegments())
+		for (const FPropertyBindingPathSegment& PropertySegment : StateTreePropertyPath.GetSegments())
 		{
 			Path.Add(PropertySegment.GetName());
 		}
diff --git a/Engine/Plugins/Runtime/StateTree/Source/StateTreeEditorModule/Private/Customizations/StateTreeBindingExtension.cpp b/Engine/Plugins/Runtime/StateTree/Source/StateTreeEditorModule/Private/Customizations/StateTreeBindingExtension.cpp
--- a/Engine/Plugins/Runtime/StateTree/Source/StateTreeEditorModule/Private/Customizations/StateTreeBindingExtension.cpp
+++ b/Engine/Plugins/Runtime/StateTree/Source/StateTreeEditorModule/Private/Customizations/StateTreeBindingExtension.cpp
@@ -83,7 +83,7 @@ UStruct* ResolveLeafValueStructType(FStateTreeDataView ValueView, TConstArrayVie
 		}
 	}
 
-	TArray<FStateTreePropertyPathIndirection> Indirections;
+	TArray<FPropertyBindingPathIndirection> Indirections;
 	if (!Path.ResolveIndirectionsWithValue(ValueView, Indirections)
 		|| Indirections.IsEmpty())
 	{
@@ -91,7 +91,7 @@ UStruct* ResolveLeafValueStructType(FStateTreeDataView ValueView, TConstArrayVie
 	}
 
 	// Last indirection points to the value of the leaf property, check the type.
-	const FStateTreePropertyPathIndirection& LastIndirection = Indirections.Last();
+	const FPropertyBindingPathIndirection& LastIndirection = Indirections.Last();
 
 	UStruct* Result = nullptr;
 
@@ -144,7 +144,7 @@ EStateTreePropertyUsage MakeStructPropertyPathFromPropertyHandle(TSharedPtr<cons
 	OutPath.Reset();
 
 	FGuid StructID;
-	TArray<FStateTreePropertyPathSegment> PathSegments;
+	TArray<FPropertyBindingPathSegment> PathSegments;
 	EStateTreePropertyUsage ResultUsage = EStateTreePropertyUsage::Invalid; 
 
 	TSharedPtr<const IPropertyHandle> CurrentPropertyHandle = InPropertyHandle;
@@ -153,7 +153,7 @@ EStateTreePropertyUsage MakeStructPropertyPathFromPropertyHandle(TSharedPtr<cons
 		const FProperty* Property = CurrentPropertyHandle->GetProperty();
 		if (Property)
 		{
-			FStateTreePropertyPathSegment& Segment = PathSegments.InsertDefaulted_GetRef(0); // Traversing from leaf to root, insert in reverse.
+			FPropertyBindingPathSegment& Segment = PathSegments.InsertDefaulted_GetRef(0); // Traversing from leaf to root, insert in reverse.
 
 			// Store path up to the property which has ID.
 			Segment.SetName(Property->GetFName());
@@ -446,7 +446,7 @@ struct FCachedBindingData : public TSharedFromThis<FCachedBindingData>
 					const FProperty* SingleOutputProperty = UE::StateTree::GetStructSingleOutputProperty(*BindableStruct.Struct);
 					check(SingleOutputProperty);
 
-					FStateTreePropertyPathSegment SingleOutputPropertySegment = FStateTreePropertyPathSegment(SingleOutputProperty->GetFName());
+					FPropertyBindingPathSegment SingleOutputPropertySegment = FPropertyBindingPathSegment(SingleOutputProperty->GetFName());
 					SourcePath = EditorBindings->AddFunctionPropertyBinding(PropertyFunctionNodeStruct, {SingleOutputPropertySegment}, TargetPath);
 				}
 				else
@@ -629,10 +629,10 @@ struct FCachedBindingData : public TSharedFromThis<FCachedBindingData>
 		FStateTreeDataView TargetDataView;
 		if (BindingOwner->GetDataViewByID(TargetPath.GetStructID(), TargetDataView) && TargetDataView.IsValid())
 		{
-			TArray<FStateTreePropertyPathIndirection> TargetIndirections;
+			TArray<FPropertyBindingPathIndirection> TargetIndirections;
 			if (ensure(TargetPath.ResolveIndirectionsWithValue(TargetDataView, TargetIndirections)))
 			{
-				const FStateTreePropertyPathIndirection& LastIndirection = TargetIndirections.Last();
+				const FPropertyBindingPathIndirection& LastIndirection = TargetIndirections.Last();
 				TargetProperty = LastIndirection.GetProperty();
 				TargetContainerAddress = LastIndirection.GetContainerAddress();
 			}
@@ -737,10 +737,10 @@ struct FCachedBindingData : public TSharedFromThis<FCachedBindingData>
 		if (bIsPropertyRef && TargetDataView.IsValid())
 		{
 			// Use internal type to construct PinType if it's property of PropertyRef type.
-			TArray<FStateTreePropertyPathIndirection> TargetIndirections;
+			TArray<FPropertyBindingPathIndirection> TargetIndirections;
 			if (ensure(TargetPath.ResolveIndirectionsWithValue(TargetDataView, TargetIndirections)))
 			{
-				const uint8* PropertyRef = TargetIndirections.Last().GetPropertyAddress();
+				const void* PropertyRef = TargetIndirections.Last().GetPropertyAddress();
 				PinType = UE::StateTree::PropertyRefHelpers::GetPropertyRefInternalTypeAsPin(*Property, PropertyRef);
 			}
 		}
@@ -779,8 +779,8 @@ struct FCachedBindingData : public TSharedFromThis<FCachedBindingData>
 				if (BindingOwner->GetDataViewByID(SourcePath.GetStructID(), SourceDataView)
 					&& TargetDataView.IsValid())
 				{
-					TArray<FStateTreePropertyPathIndirection> SourceIndirections;
-					TArray<FStateTreePropertyPathIndirection> TargetIndirections;
+					TArray<FPropertyBindingPathIndirection> SourceIndirections;
+					TArray<FPropertyBindingPathIndirection> TargetIndirections;
 
 					// Resolve source and target properties.
 					// Source path can be empty, when the binding binds directly to a context struct/class.
@@ -789,11 +789,11 @@ struct FCachedBindingData : public TSharedFromThis<FCachedBindingData>
 						&& TargetPath.ResolveIndirectionsWithValue(TargetDataView, TargetIndirections)
 						&& !TargetIndirections.IsEmpty())
 					{
-						const FStateTreePropertyPathIndirection TargetLeafIndirection = TargetIndirections.Last();
+						const FPropertyBindingPathIndirection TargetLeafIndirection = TargetIndirections.Last();
 						if (SourceIndirections.Num() > 0)
 						{
 							// Binding to a source property.
-							const FStateTreePropertyPathIndirection SourceLeafIndirection = SourceIndirections.Last();
+							const FPropertyBindingPathIndirection SourceLeafIndirection = SourceIndirections.Last();
 							SourceLeafProperty = SourceLeafIndirection.GetProperty();
 							bIsValidBinding = ArePropertiesCompatible(SourceLeafProperty, TargetLeafIndirection.GetProperty(), SourceLeafIndirection.GetPropertyAddress(), TargetLeafIndirection.GetPropertyAddress());
 						}
@@ -1004,7 +1004,7 @@ struct FCachedBindingData : public TSharedFromThis<FCachedBindingData>
 		FStateTreePropertyPath SourcePath;
 		UE::StateTree::PropertyBinding::MakeStructPropertyPathFromBindingChain(AccessibleStructs[SourceStructIndex].ID, InBindingChain, SourceDataView, SourcePath);
 
-		TArray<FStateTreePropertyPathIndirection> SourceIndirections;
+		TArray<FPropertyBindingPathIndirection> SourceIndirections;
 		void* TargetValueAddress = nullptr;
 		if (PropertyHandle->GetValueData(TargetValueAddress) == FPropertyAccess::Success && SourcePath.ResolveIndirectionsWithValue(SourceDataView, SourceIndirections))
 		{
diff --git a/Engine/Plugins/Runtime/StateTree/Source/StateTreeEditorModule/Private/Customizations/StateTreeEditorNodeDetails.cpp b/Engine/Plugins/Runtime/StateTree/Source/StateTreeEditorModule/Private/Customizations/StateTreeEditorNodeDetails.cpp
--- a/Engine/Plugins/Runtime/StateTree/Source/StateTreeEditorModule/Private/Customizations/StateTreeEditorNodeDetails.cpp
+++ b/Engine/Plugins/Runtime/StateTree/Source/StateTreeEditorModule/Private/Customizations/StateTreeEditorNodeDetails.cpp
@@ -175,7 +175,7 @@ namespace UE::StateTreeEditor::Internal
 					FStateTreeDataView TargetDataView;
 					if (ensure(EditorData->GetDataViewByID(ID, TargetDataView)))
 					{
-						TArray<FStateTreePropertyPathIndirection> TargetIndirections;
+						TArray<FPropertyBindingPathIndirection> TargetIndirections;
 						if (ensure(Path.ResolveIndirectionsWithValue(TargetDataView, TargetIndirections)))
 						{
 							const void* PropertyRef = TargetIndirections.Last().GetPropertyAddress();
diff --git a/Engine/Plugins/Runtime/StateTree/Source/StateTreeEditorModule/Private/StateTreeCompiler.cpp b/Engine/Plugins/Runtime/StateTree/Source/StateTreeEditorModule/Private/StateTreeCompiler.cpp
--- a/Engine/Plugins/Runtime/StateTree/Source/StateTreeEditorModule/Private/StateTreeCompiler.cpp
+++ b/Engine/Plugins/Runtime/StateTree/Source/StateTreeEditorModule/Private/StateTreeCompiler.cpp
@@ -1814,7 +1814,7 @@ bool FStateTreeCompiler::CreateEvaluator(const FStateTreeEditorNode& EvalNode, c
 
 bool FStateTreeCompiler::IsPropertyOfTypeOrChild(UScriptStruct& Type, const FStateTreeBindableStructDesc& Struct, FStateTreePropertyPath Path) const
 {
-	TArray<FStateTreePropertyPathIndirection> Indirection;
+	TArray<FPropertyBindingPathIndirection> Indirection;
 	const bool bResolved = Path.ResolveIndirections(Struct.Struct, Indirection);
 	
 	if (bResolved && Indirection.Num() > 0)
@@ -1835,7 +1835,7 @@ bool FStateTreeCompiler::ValidateStructRef(const FStateTreeBindableStructDesc& S
 											const FStateTreeBindableStructDesc& TargetStruct, FStateTreePropertyPath TargetPath) const
 {
 	FString ResolveError;
-	TArray<FStateTreePropertyPathIndirection> TargetIndirection;
+	TArray<FPropertyBindingPathIndirection> TargetIndirection;
 	if (!TargetPath.ResolveIndirections(TargetStruct.Struct, TargetIndirection, &ResolveError))
 	{
 		// This will later be reported by the bindings compiler.
@@ -1861,7 +1861,7 @@ bool FStateTreeCompiler::ValidateStructRef(const FStateTreeBindableStructDesc& S
 		return false;
 	}
 
-	TArray<FStateTreePropertyPathIndirection> SourceIndirection;
+	TArray<FPropertyBindingPathIndirection> SourceIndirection;
 	if (!SourcePath.ResolveIndirections(SourceStruct.Struct, SourceIndirection, &ResolveError))
 	{
 		// This will later be reported by the bindings compiler.
@@ -2076,7 +2076,7 @@ bool FStateTreeCompiler::GetAndValidateBindings(const FStateTreeBindableStructDe
 
 		if (UE::StateTree::PropertyRefHelpers::IsPropertyRef(*Property))
 		{
-			TArray<FStateTreePropertyPathIndirection> TargetIndirections;
+			TArray<FPropertyBindingPathIndirection> TargetIndirections;
 			FStateTreePropertyPath TargetPath(TargetStruct.ID, PropertyName);
 			if (!TargetPath.ResolveIndirectionsWithValue(TargetValue, TargetIndirections))
 			{
diff --git a/Engine/Plugins/Runtime/StateTree/Source/StateTreeEditorModule/Private/StateTreeEditorData.cpp b/Engine/Plugins/Runtime/StateTree/Source/StateTreeEditorModule/Private/StateTreeEditorData.cpp
--- a/Engine/Plugins/Runtime/StateTree/Source/StateTreeEditorModule/Private/StateTreeEditorData.cpp
+++ b/Engine/Plugins/Runtime/StateTree/Source/StateTreeEditorModule/Private/StateTreeEditorData.cpp
@@ -1737,11 +1737,4 @@ bool UStateTreeEditorData::RemoveBreakpoint(const FGuid ID, const EStateTreeBrea
 
 #endif // WITH_STATETREE_TRACE_DEBUGGER
 
-PRAGMA_DISABLE_DEPRECATION_WARNINGS
-void UStateTreeEditorData::AddPropertyBinding(const FStateTreeEditorPropertyPath& SourcePath, const FStateTreeEditorPropertyPath& TargetPath)
-{
-	EditorBindings.AddPropertyBinding(UE::StateTree::Private::ConvertEditorPath(SourcePath), UE::StateTree::Private::ConvertEditorPath(TargetPath));
-}
-PRAGMA_ENABLE_DEPRECATION_WARNINGS
-
 #undef LOCTEXT_NAMESPACE
\ No newline at end of file
diff --git a/Engine/Plugins/Runtime/StateTree/Source/StateTreeEditorModule/Private/StateTreeEditorPropertyBindings.cpp b/Engine/Plugins/Runtime/StateTree/Source/StateTreeEditorModule/Private/StateTreeEditorPropertyBindings.cpp
--- a/Engine/Plugins/Runtime/StateTree/Source/StateTreeEditorModule/Private/StateTreeEditorPropertyBindings.cpp
+++ b/Engine/Plugins/Runtime/StateTree/Source/StateTreeEditorModule/Private/StateTreeEditorPropertyBindings.cpp
@@ -48,7 +48,7 @@ void FStateTreeEditorPropertyBindings::AddPropertyBinding(const FStateTreeProper
 	PropertyBindings.Add(Binding);
 }
 
-FStateTreePropertyPath FStateTreeEditorPropertyBindings::AddFunctionPropertyBinding(const UScriptStruct* PropertyFunctionNodeStruct, TConstArrayView<FStateTreePropertyPathSegment> SourcePathSegments, const FStateTreePropertyPath& TargetPath)
+FStateTreePropertyPath FStateTreeEditorPropertyBindings::AddFunctionPropertyBinding(const UScriptStruct* PropertyFunctionNodeStruct, TConstArrayView<FPropertyBindingPathSegment> SourcePathSegments, const FStateTreePropertyPath& TargetPath)
 {
 	check(PropertyFunctionNodeStruct->IsChildOf<FStateTreePropertyFunctionBase>());
 
@@ -188,7 +188,7 @@ void FStateTreeEditorPropertyBindings::RemoveUnusedBindings(const TMap<FGuid, co
 				if (SourceValue && SourceValue->IsValid())
 				{
 					FString Error;
-					TArray<FStateTreePropertyPathIndirection> Indirections;
+					TArray<FPropertyBindingPathIndirection> Indirections;
 					if (!Binding.GetSourcePath().ResolveIndirectionsWithValue(*SourceValue, Indirections, &Error))
 					{
 						UE_LOG(LogStateTree, Verbose, TEXT("Removing binding to %s because binding source path cannot be resolved: %s"),
@@ -203,7 +203,7 @@ void FStateTreeEditorPropertyBindings::RemoveUnusedBindings(const TMap<FGuid, co
 			{
 				const FStateTreeDataView TargetValue = ValidStructs.FindChecked(Binding.GetTargetPath().GetStructID());
 				FString Error;
-				TArray<FStateTreePropertyPathIndirection> Indirections;
+				TArray<FPropertyBindingPathIndirection> Indirections;
 				if (!Binding.GetTargetPath().ResolveIndirectionsWithValue(TargetValue, Indirections, &Error))
 				{
 					UE_LOG(LogStateTree, Verbose, TEXT("Removing binding to %s because binding target path cannot be resolved: %s"),
@@ -222,7 +222,7 @@ bool FStateTreeEditorPropertyBindings::ContainsAnyStruct(const TSet<const UStruc
 {
 	auto PathContainsStruct = [&Structs](const FStateTreePropertyPath& PropertyPath)
 	{
-		for (const FStateTreePropertyPathSegment& Segment : PropertyPath.GetSegments())
+		for (const FPropertyBindingPathSegment& Segment : PropertyPath.GetSegments())
 		{
 			if (Structs.Contains(Segment.GetInstanceStruct()))
 			{
@@ -248,33 +248,6 @@ bool FStateTreeEditorPropertyBindings::ContainsAnyStruct(const TSet<const UStruc
 }
 
 PRAGMA_DISABLE_DEPRECATION_WARNINGS
-void FStateTreeEditorPropertyBindings::AddPropertyBinding(const FStateTreeEditorPropertyPath& SourcePath, const FStateTreeEditorPropertyPath& TargetPath)
-{
-	AddPropertyBinding(UE::StateTree::Private::ConvertEditorPath(SourcePath), UE::StateTree::Private::ConvertEditorPath(TargetPath));
-}
-
-void FStateTreeEditorPropertyBindings::RemovePropertyBindings(const FStateTreeEditorPropertyPath& TargetPath)
-{
-	RemovePropertyBindings(UE::StateTree::Private::ConvertEditorPath(TargetPath), ESearchMode::Exact);
-}
-
-bool FStateTreeEditorPropertyBindings::HasPropertyBinding(const FStateTreeEditorPropertyPath& TargetPath) const
-{
-	return HasPropertyBinding(UE::StateTree::Private::ConvertEditorPath(TargetPath), ESearchMode::Exact);
-}
-
-const FStateTreeEditorPropertyPath* FStateTreeEditorPropertyBindings::GetPropertyBindingSource(const FStateTreeEditorPropertyPath& TargetPath) const
-{
-	static FStateTreeEditorPropertyPath Dummy;
-	const FStateTreePropertyPath* SourcePath = GetPropertyBindingSource(UE::StateTree::Private::ConvertEditorPath(TargetPath));
-	if (SourcePath != nullptr)
-	{
-		Dummy = UE::StateTree::Private::ConvertEditorPath(*SourcePath);
-		return &Dummy;
-	}
-	return nullptr;
-}
-
 void FStateTreeEditorPropertyBindings::GetPropertyBindingsFor(const FGuid StructID, TArray<FStateTreePropertyPathBinding>& OutBindings) const
 {
 	TArray<const FStateTreePropertyPathBinding*> NodeBindings;
@@ -283,36 +256,6 @@ void FStateTreeEditorPropertyBindings::GetPropertyBindingsFor(const FGuid Struct
 }
 PRAGMA_ENABLE_DEPRECATION_WARNINGS
 
-void FStateTreeEditorPropertyBindings::RemoveUnusedBindings(const TMap<FGuid, const UStruct*>& ValidStructs)
-{
-	PropertyBindings.RemoveAllSwap([ValidStructs](const FStateTreePropertyPathBinding& Binding)
-		{
-			// Remove binding if it's target struct has been removed
-			if (!ValidStructs.Contains(Binding.GetTargetPath().GetStructID()))
-			{
-				return true;
-			}
-
-			// Target path should always have at least one segment (copy bind directly on a target struct/object). 
-			if (Binding.GetTargetPath().IsPathEmpty())
-			{
-				return true;
-			}
-		
-			// Remove binding if path containing instanced indirections (e.g. instance struct or object) cannot be resolved.
-			// TODO: Try to use core redirect to find new name.
-			const UStruct* Struct = ValidStructs.FindChecked(Binding.GetTargetPath().GetStructID());
-			TArray<FStateTreePropertyPathIndirection> Indirections;
-			if (!Binding.GetTargetPath().ResolveIndirections(Struct, Indirections))
-			{
-				// Remove
-				return true;
-			}
-
-			return false;
-		});
-}
-
 
 //////////////////////////////////////////////////////////////////////////
 
@@ -410,11 +353,11 @@ FText FStateTreeBindingLookup::GetBindingSourceDisplayName(const FStateTreePrope
 	FStateTreeBindableStructDesc TargetStructDesc;
 	if (BindingOwner->GetStructByID(InTargetPath.GetStructID(), TargetStructDesc))
 	{
-		TArray<FStateTreePropertyPathIndirection> Indirection;
+		TArray<FPropertyBindingPathIndirection> Indirection;
 		if (InTargetPath.ResolveIndirections(TargetStructDesc.Struct, Indirection)
 			&& Indirection.Num() > 0)
 		{
-			const FStateTreePropertyPathIndirection& Leaf = Indirection.Last(); 
+			const FPropertyBindingPathIndirection& Leaf = Indirection.Last();
 			TargetProperty = Leaf.GetProperty();
 			if (TargetProperty)
 			{
@@ -459,7 +402,7 @@ const FProperty* FStateTreeBindingLookup::GetPropertyPathLeafProperty(const FSta
 	FStateTreeBindableStructDesc Struct;
 	if (BindingOwner->GetStructByID(InPath.GetStructID(), Struct))
 	{
-		TArray<FStateTreePropertyPathIndirection> Indirection;
+		TArray<FPropertyBindingPathIndirection> Indirection;
 		if (InPath.ResolveIndirections(Struct.Struct, Indirection) && Indirection.Num() > 0)
 		{
 			return Indirection.Last().GetProperty();
diff --git a/Engine/Plugins/Runtime/StateTree/Source/StateTreeEditorModule/Private/StateTreePropertyBindingCompiler.cpp b/Engine/Plugins/Runtime/StateTree/Source/StateTreeEditorModule/Private/StateTreePropertyBindingCompiler.cpp
--- a/Engine/Plugins/Runtime/StateTree/Source/StateTreeEditorModule/Private/StateTreePropertyBindingCompiler.cpp
+++ b/Engine/Plugins/Runtime/StateTree/Source/StateTreeEditorModule/Private/StateTreePropertyBindingCompiler.cpp
@@ -35,7 +35,7 @@ bool FStateTreePropertyBindingCompiler::CompileBatch(const FStateTreeBindableStr
 	struct FSortedBinding
 	{
 		FStateTreePropertyPathBinding Binding;
-		TArray<FStateTreePropertyPathIndirection> TargetIndirections;
+		TArray<FPropertyBindingPathIndirection> TargetIndirections;
 	};
 	TArray<FSortedBinding> NewBindings;
 
@@ -55,8 +55,8 @@ bool FStateTreePropertyBindingCompiler::CompileBatch(const FStateTreeBindableStr
 		}
 
 		FString Error;
-		TArray<FStateTreePropertyPathIndirection> SourceIndirections;
-		TArray<FStateTreePropertyPathIndirection> TargetIndirections;
+		TArray<FPropertyBindingPathIndirection> SourceIndirections;
+		TArray<FPropertyBindingPathIndirection> TargetIndirections;
 		
 		if (!Binding.GetSourcePath().ResolveIndirections(SourceStruct->Struct, SourceIndirections, &Error))
 		{
@@ -71,8 +71,8 @@ bool FStateTreePropertyBindingCompiler::CompileBatch(const FStateTreeBindableStr
 		}
 
 		FStateTreePropertyCopy DummyCopy;
-		FStateTreePropertyPathIndirection LastSourceIndirection = !SourceIndirections.IsEmpty() ? SourceIndirections.Last() : FStateTreePropertyPathIndirection(SourceStruct->Struct);
-		FStateTreePropertyPathIndirection LastTargetIndirection = !TargetIndirections.IsEmpty() ? TargetIndirections.Last() : FStateTreePropertyPathIndirection(TargetStruct.Struct);
+		FPropertyBindingPathIndirection LastSourceIndirection = !SourceIndirections.IsEmpty() ? SourceIndirections.Last() : FPropertyBindingPathIndirection(SourceStruct->Struct);
+		FPropertyBindingPathIndirection LastTargetIndirection = !TargetIndirections.IsEmpty() ? TargetIndirections.Last() : FPropertyBindingPathIndirection(TargetStruct.Struct);
 		if (!PropertyBindings->ResolveCopyType(LastSourceIndirection, LastTargetIndirection, DummyCopy))
 		{
 			Log->Reportf(EMessageSeverity::Error, TargetStruct,
@@ -160,7 +160,7 @@ bool FStateTreePropertyBindingCompiler::CompileReferences(const FStateTreeBindab
 		}
 
 		FString Error;
-		TArray<FStateTreePropertyPathIndirection> SourceIndirections;
+		TArray<FPropertyBindingPathIndirection> SourceIndirections;
 		if (!Binding.GetSourcePath().ResolveIndirectionsWithValue(*SourceDataView, SourceIndirections, &Error))
 		{
 			Log->Reportf(EMessageSeverity::Error, TargetStruct, TEXT("Resolving path in %s: %s"), *SourceStruct->ToString(), *Error);
@@ -176,18 +176,18 @@ bool FStateTreePropertyBindingCompiler::CompileReferences(const FStateTreeBindab
 			return false;
 		}
 
-		TArray<FStateTreePropertyPathIndirection> TargetIndirections;
+		TArray<FPropertyBindingPathIndirection> TargetIndirections;
 		if (!Binding.GetTargetPath().ResolveIndirectionsWithValue(InstanceDataView, TargetIndirections, &Error))
 		{
 			Log->Reportf(EMessageSeverity::Error, TargetStruct, TEXT("Resolving path in %s: %s"), *TargetStruct.ToString(), *Error);
 			return false;
 		}
 
-		FStateTreePropertyPathIndirection& TargetLeafIndirection = TargetIndirections.Last();
-		FStateTreePropertyRef* PropertyRef = reinterpret_cast<FStateTreePropertyRef*>(const_cast<uint8*>(TargetLeafIndirection.GetPropertyAddress()));
+		FPropertyBindingPathIndirection& TargetLeafIndirection = TargetIndirections.Last();
+		FStateTreePropertyRef* PropertyRef = static_cast<FStateTreePropertyRef*>(const_cast<void*>(TargetLeafIndirection.GetPropertyAddress()));
 		check(PropertyRef);
 
-		FStateTreePropertyPathIndirection& SourceLeafIndirection = SourceIndirections.Last();
+		FPropertyBindingPathIndirection& SourceLeafIndirection = SourceIndirections.Last();
 		if (!UE::StateTree::PropertyRefHelpers::IsPropertyRefCompatibleWithProperty(*TargetLeafIndirection.GetProperty(), *SourceLeafIndirection.GetProperty(), PropertyRef, SourceLeafIndirection.GetPropertyAddress()))
 		{
 			Log->Reportf(EMessageSeverity::Error, TargetStruct,
@@ -283,208 +283,6 @@ int32 FStateTreePropertyBindingCompiler::GetSourceStructIndexByID(const FGuid& I
 	return SourceStructs.IndexOfByPredicate([ID](const FStateTreeBindableStructDesc& Structs) { return (Structs.ID == ID); });
 }
 
-PRAGMA_DISABLE_DEPRECATION_WARNINGS
-bool FStateTreePropertyBindingCompiler::ResolvePropertyPath(const FStateTreeBindableStructDesc& InStructDesc, const FStateTreeEditorPropertyPath& InPath,
-															TArray<FStateTreePropertySegment>& OutSegments, const FProperty*& OutLeafProperty, int32& OutLeafArrayIndex,
-															FStateTreeCompilerLog* InLog, const FStateTreeBindableStructDesc* InLogContextStruct)
-{
-	if (!InPath.IsValid())
-	{
-		if (InLog != nullptr && InLogContextStruct != nullptr)
-		{
-			InLog->Reportf(EMessageSeverity::Error, *InLogContextStruct,
-					TEXT("Invalid path '%s:%s'."),
-					*InStructDesc.Name.ToString(), *InPath.ToString());
-		}
-		return false;
-	}
-
-	// If the path is empty, we're pointing directly at the source struct.
-	if (InPath.Path.Num() == 0)
-	{
-		OutLeafProperty = nullptr;
-		OutLeafArrayIndex = INDEX_NONE;
-		return true;
-	}
-
-	const UStruct* CurrentStruct = InStructDesc.Struct;
-	const FProperty* LeafProperty = nullptr;
-	int32 LeafArrayIndex = INDEX_NONE;
-	bool bResult = true;
-
-	for (int32 SegmentIndex = 0; SegmentIndex < InPath.Path.Num(); SegmentIndex++)
-	{
-		const FString& SegmentString = InPath.Path[SegmentIndex];
-		const TCHAR* PropertyNamePtr = nullptr;
-		int32 PropertyNameLength = 0;
-		int32 ArrayIndex = INDEX_NONE;
-		PropertyPathHelpers::FindFieldNameAndArrayIndex(SegmentString.Len(), *SegmentString, PropertyNameLength, &PropertyNamePtr, ArrayIndex);
-		ensure(PropertyNamePtr != nullptr);
-		FString PropertyNameString(PropertyNameLength, PropertyNamePtr);
-		const FName PropertyName = FName(*PropertyNameString, FNAME_Find);
-
-		const bool bFinalSegment = SegmentIndex == (InPath.Path.Num() - 1);
-
-		if (CurrentStruct == nullptr)
-		{
-			if (InLog != nullptr && InLogContextStruct != nullptr)
-			{
-				InLog->Reportf(EMessageSeverity::Error, *InLogContextStruct,
-						TEXT("Malformed path '%s:%s'."),
-						*InStructDesc.Name.ToString(), *InPath.ToString(SegmentIndex, TEXT("<"), TEXT(">")));
-			}
-			bResult = false;
-			break;
-		}
-
-		const FProperty* Property = CurrentStruct->FindPropertyByName(PropertyName);
-		if (Property == nullptr)
-		{
-			// TODO: use core redirects to fix up the name.
-			if (InLog != nullptr && InLogContextStruct != nullptr)
-			{
-				InLog->Reportf(EMessageSeverity::Error, *InLogContextStruct,
-						TEXT("Malformed path '%s:%s', could not find property '%s%s.%s'."),
-						*InStructDesc.Name.ToString(), *InPath.ToString(SegmentIndex, TEXT("<"), TEXT(">")),
-						CurrentStruct->GetPrefixCPP(), *CurrentStruct->GetName(), *PropertyName.ToString());
-			}
-			bResult = false;
-			break;
-		}
-
-		if (const auto Validation = UE::StateTree::Compiler::IsValidIndex16(ArrayIndex); Validation.DidFail())
-		{
-			if (InLog != nullptr)
-			{
-				Validation.Log(*InLog, TEXT("ArrayIndex"), InStructDesc);
-			}
-			return false;
-		}
-
-		FStateTreePropertySegment& Segment = OutSegments.AddDefaulted_GetRef();
-		Segment.Name = PropertyName;
-		Segment.ArrayIndex = FStateTreeIndex16(ArrayIndex);
-
-		// Check to see if it is an array access first
-		const FArrayProperty* ArrayProperty = CastField<FArrayProperty>(Property);
-		if (ArrayProperty != nullptr && ArrayIndex != INDEX_NONE)
-		{
-			// It is an array, now check to see if this is an array of structures
-			if (const FStructProperty* ArrayOfStructsProperty = CastField<FStructProperty>(ArrayProperty->Inner))
-			{
-				Segment.Type = EStateTreePropertyAccessType::IndexArray;
-				CurrentStruct = ArrayOfStructsProperty->Struct;
-			}
-			// if it's not an array of structs, maybe it's an array of objects
-			else if (const FObjectPropertyBase* ArrayOfObjectsProperty = CastField<FObjectPropertyBase>(ArrayProperty->Inner))
-			{
-				Segment.Type = EStateTreePropertyAccessType::IndexArray;
-				CurrentStruct = ArrayOfObjectsProperty->PropertyClass;
-
-				if (!bFinalSegment)
-				{
-					// Object arrays need an object dereference adding if non-leaf
-					FStateTreePropertySegment& ExtraSegment = OutSegments.AddDefaulted_GetRef();
-					ExtraSegment.ArrayIndex = FStateTreeIndex16(0);
-					ExtraSegment.Type = EStateTreePropertyAccessType::Object;
-					const FProperty* InnerProperty = ArrayProperty->Inner;
-					if (const FObjectProperty* ObjectProperty = CastField<FObjectProperty>(InnerProperty))
-					{
-						ExtraSegment.Type = EStateTreePropertyAccessType::Object;
-					}
-					else if (const FWeakObjectProperty* WeakObjectProperty = CastField<FWeakObjectProperty>(InnerProperty))
-					{
-						ExtraSegment.Type = EStateTreePropertyAccessType::WeakObject;
-					}
-					else if (const FSoftObjectProperty* SoftObjectProperty = CastField<FSoftObjectProperty>(InnerProperty))
-					{
-						ExtraSegment.Type = EStateTreePropertyAccessType::SoftObject;
-					}
-				}
-			}
-			else
-			{
-				Segment.Type = EStateTreePropertyAccessType::IndexArray;
-				Segment.ArrayIndex = FStateTreeIndex16(ArrayIndex);
-				CurrentStruct = nullptr;
-			}
-		}
-		// Leaf segments all get treated the same, plain, array, struct or object. Copy type is figured out separately.
-		else if (bFinalSegment)
-		{
-			Segment.Type = EStateTreePropertyAccessType::Offset;
-			CurrentStruct = nullptr;
-		}
-		// Check to see if this is a simple structure (eg. not an array of structures)
-		else if (const FStructProperty* StructProperty = CastField<FStructProperty>(Property))
-		{
-			Segment.Type = EStateTreePropertyAccessType::Offset;
-			CurrentStruct = StructProperty->Struct;
-		}
-		// Check to see if this is a simple object (eg. not an array of objects)
-		else if (const FObjectProperty* ObjectProperty = CastField<FObjectProperty>(Property))
-		{
-			Segment.Type = EStateTreePropertyAccessType::Object;
-			CurrentStruct = ObjectProperty->PropertyClass;
-		}
-		// Check to see if this is a simple weak object property (eg. not an array of weak objects).
-		else if (const FWeakObjectProperty* WeakObjectProperty = CastField<FWeakObjectProperty>(Property))
-		{
-			Segment.Type = EStateTreePropertyAccessType::WeakObject;
-			CurrentStruct = WeakObjectProperty->PropertyClass;
-		}
-		// Check to see if this is a simple soft object property (eg. not an array of soft objects).
-		else if (const FSoftObjectProperty* SoftObjectProperty = CastField<FSoftObjectProperty>(Property))
-		{
-			Segment.Type = EStateTreePropertyAccessType::SoftObject;
-			CurrentStruct = SoftObjectProperty->PropertyClass;
-		}
-		else
-		{
-			if (InLog != nullptr && InLogContextStruct != nullptr)
-			{
-				InLog->Reportf(EMessageSeverity::Error, *InLogContextStruct,
-						TEXT("Unsupported segment %s in path '%s:%s'."),
-						*InStructDesc.Name.ToString(), *InPath.ToString(SegmentIndex, TEXT("<"), TEXT(">")),
-						*Property->GetCPPType(), *InStructDesc.Name.ToString(), *InPath.ToString(SegmentIndex, TEXT("<"), TEXT(">")));
-			}
-			bResult = false;
-			break;
-		}
-
-		if (bFinalSegment)
-		{
-			LeafProperty = Property;
-			LeafArrayIndex = ArrayIndex;
-		}
-	}
-
-	if (!bResult)
-	{
-		return false;
-	}
-
-	OutLeafProperty = LeafProperty;
-	OutLeafArrayIndex = LeafArrayIndex;
-
-	return true;
-}
-
-EPropertyAccessCompatibility FStateTreePropertyBindingCompiler::GetPropertyCompatibility(const FProperty* FromProperty, const FProperty* ToProperty)
-{
-	const EStateTreePropertyAccessCompatibility Result = FStateTreePropertyBindings::GetPropertyCompatibility(FromProperty, ToProperty);
-	if (Result == EStateTreePropertyAccessCompatibility::Compatible)
-	{
-		return EPropertyAccessCompatibility::Compatible;
-	}
-	if (Result == EStateTreePropertyAccessCompatibility::Promotable)
-	{
-		return EPropertyAccessCompatibility::Promotable;
-	}
-	return EPropertyAccessCompatibility::Incompatible;
-}
-PRAGMA_ENABLE_DEPRECATION_WARNINGS
-
 void FStateTreePropertyBindingCompiler::StoreSourceStructs()
 {
 	// Check that existing structs are compatible
diff --git a/Engine/Plugins/Runtime/StateTree/Source/StateTreeEditorModule/Private/StateTreePropertyHelpers.cpp b/Engine/Plugins/Runtime/StateTree/Source/StateTreeEditorModule/Private/StateTreePropertyHelpers.cpp
--- a/Engine/Plugins/Runtime/StateTree/Source/StateTreeEditorModule/Private/StateTreePropertyHelpers.cpp
+++ b/Engine/Plugins/Runtime/StateTree/Source/StateTreeEditorModule/Private/StateTreePropertyHelpers.cpp
@@ -117,7 +117,7 @@ void DispatchPostEditToNodes(UObject& Owner, FPropertyChangedChainEvent& InPrope
 			CurrentProperty = ArrayProperty->Inner;
 		}
 
-		FStateTreePropertyPathSegment PathSegment(CurrentProperty->GetFName());
+		FPropertyBindingPathSegment PathSegment(CurrentProperty->GetFName());
 		if (const FStructProperty* StructProperty = CastField<FStructProperty>(CurrentProperty))
 		{
 			if (StructProperty->Struct == FInstancedStruct::StaticStruct())
@@ -155,7 +155,7 @@ void DispatchPostEditToNodes(UObject& Owner, FPropertyChangedChainEvent& InPrope
 					const FName EditorNodeChildMemberName = CurrentPropNode->GetValue()->GetFName();
 					if (EditorNodeChildMemberName == GET_MEMBER_NAME_CHECKED(FStateTreeEditorNode, Instance) || EditorNodeChildMemberName == GET_MEMBER_NAME_CHECKED(FStateTreeEditorNode, InstanceObject))
 					{
-						CurrentAddress = LastEditorNode->GetInstance().GetMutableMemory();
+						CurrentAddress = static_cast<uint8*>(LastEditorNode->GetInstance().GetMutableMemory());
 						CurrentPropNode = CurrentPropNode->GetNextNode();
 						continue;
 					}
diff --git a/Engine/Plugins/Runtime/StateTree/Source/StateTreeEditorModule/Public/StateTreeEditorData.h b/Engine/Plugins/Runtime/StateTree/Source/StateTreeEditorModule/Public/StateTreeEditorData.h
--- a/Engine/Plugins/Runtime/StateTree/Source/StateTreeEditorModule/Public/StateTreeEditorData.h
+++ b/Engine/Plugins/Runtime/StateTree/Source/StateTreeEditorModule/Public/StateTreeEditorData.h
@@ -235,16 +235,11 @@ public:
 	/**
 	 * Adds property binding to PropertyFunction of provided type.
 	 */
-	void AddPropertyBinding(const UScriptStruct* PropertyFunctionNodeStruct, TConstArrayView<FStateTreePropertyPathSegment> SourcePathSegments, const FStateTreePropertyPath& TargetPath)
+	void AddPropertyBinding(const UScriptStruct* PropertyFunctionNodeStruct, TConstArrayView<FPropertyBindingPathSegment> SourcePathSegments, const FStateTreePropertyPath& TargetPath)
 	{
 		EditorBindings.AddFunctionPropertyBinding(PropertyFunctionNodeStruct, SourcePathSegments, TargetPath);
 	}
 
-PRAGMA_DISABLE_DEPRECATION_WARNINGS
-	UE_DEPRECATED(5.3, "Use version with FStateTreePropertyPath instead.")
-	void AddPropertyBinding(const FStateTreeEditorPropertyPath& SourcePath, const FStateTreeEditorPropertyPath& TargetPath);
-PRAGMA_ENABLE_DEPRECATION_WARNINGS
-
 	/**
 	 * Adds property binding between two structs.
 	 */
diff --git a/Engine/Plugins/Runtime/StateTree/Source/StateTreeEditorModule/Public/StateTreeEditorPropertyBindings.h b/Engine/Plugins/Runtime/StateTree/Source/StateTreeEditorModule/Public/StateTreeEditorPropertyBindings.h
--- a/Engine/Plugins/Runtime/StateTree/Source/StateTreeEditorModule/Public/StateTreeEditorPropertyBindings.h
+++ b/Engine/Plugins/Runtime/StateTree/Source/StateTreeEditorModule/Public/StateTreeEditorPropertyBindings.h
@@ -42,7 +42,7 @@ struct STATETREEEDITORMODULE_API FStateTreeEditorPropertyBindings
 	 * @param TargetPath Binding target property path.
 	 * @return Constructed binding source property path.
 	 */
-	FStateTreePropertyPath AddFunctionPropertyBinding(const UScriptStruct* PropertyFunctionNodeStruct, TConstArrayView<FStateTreePropertyPathSegment> SourcePathSegments, const FStateTreePropertyPath& TargetPath);
+	FStateTreePropertyPath AddFunctionPropertyBinding(const UScriptStruct* PropertyFunctionNodeStruct, TConstArrayView<FPropertyBindingPathSegment> SourcePathSegments, const FStateTreePropertyPath& TargetPath);
 
 	enum class ESearchMode
 	{
@@ -108,25 +108,10 @@ struct STATETREEEDITORMODULE_API FStateTreeEditorPropertyBindings
 	}
 
 PRAGMA_DISABLE_DEPRECATION_WARNINGS
-	UE_DEPRECATED(5.3, "Use version with FStateTreePropertyPath instead.")
-	void AddPropertyBinding(const FStateTreeEditorPropertyPath& SourcePath, const FStateTreeEditorPropertyPath& TargetPath);
-
-	UE_DEPRECATED(5.3, "Use version with FStateTreePropertyPath instead.")
-	void RemovePropertyBindings(const FStateTreeEditorPropertyPath& TargetPath);
-
-	UE_DEPRECATED(5.3, "Use version with FStateTreePropertyPath instead.")
-	bool HasPropertyBinding(const FStateTreeEditorPropertyPath& TargetPath) const;
-
-	UE_DEPRECATED(5.3, "Use version with FStateTreePropertyPath instead.")
-	const FStateTreeEditorPropertyPath* GetPropertyBindingSource(const FStateTreeEditorPropertyPath& TargetPath) const;
-
-	UE_DEPRECATED(5.3, "Use RemoveUnusedBindings with values instead.")
-	void RemoveUnusedBindings(const TMap<FGuid, const UStruct*>& ValidStructs);
-
 	UE_DEPRECATED(5.5, "Use GetPropertyBindingsFor returning pointers instead.")
 	void GetPropertyBindingsFor(const FGuid StructID, TArray<FStateTreePropertyPathBinding>& OutBindings) const;
 PRAGMA_ENABLE_DEPRECATION_WARNINGS
-	
+
 private:
 	UPROPERTY()
 	TArray<FStateTreePropertyPathBinding> PropertyBindings;
diff --git a/Engine/Plugins/Runtime/StateTree/Source/StateTreeEditorModule/Public/StateTreePropertyBindingCompiler.h b/Engine/Plugins/Runtime/StateTree/Source/StateTreeEditorModule/Public/StateTreePropertyBindingCompiler.h
--- a/Engine/Plugins/Runtime/StateTree/Source/StateTreeEditorModule/Public/StateTreePropertyBindingCompiler.h
+++ b/Engine/Plugins/Runtime/StateTree/Source/StateTreeEditorModule/Public/StateTreePropertyBindingCompiler.h
@@ -76,32 +76,6 @@ struct STATETREEEDITORMODULE_API FStateTreePropertyBindingCompiler
 	}
 
 PRAGMA_DISABLE_DEPRECATION_WARNINGS
-	/**
-	 * Resolves a string based property path in specified struct into segments of property names and access types.
-	 * If logging is required both, Log and LogContextStruct needs to be non-null.
-	 * @param InStructDesc Description of the struct in which the property path is valid.
-	 * @param InPath The property path in string format.
-	 * @param OutSegments The resolved property access path as segments.
-	 * @param OutLeafProperty The leaf property of the resolved path.
-	 * @param OutLeafArrayIndex The left array index (or INDEX_NONE if not applicable) of the resolved path.
-	 * @param Log Pointer to compiler log, or null if no logging needed.
-	 * @param LogContextStruct Pointer to bindable struct desc where the property path belongs to.
-	 * @return True of the property was solved successfully.
-	 */
-	UE_DEPRECATED(5.3, "Use FStateTreePropertyPath resolve indirection methods instead.")
-	[[nodiscard]] static bool ResolvePropertyPath(const FStateTreeBindableStructDesc& InStructDesc, const FStateTreeEditorPropertyPath& InPath,
-												  TArray<FStateTreePropertySegment>& OutSegments, const FProperty*& OutLeafProperty, int32& OutLeafArrayIndex,
-												  FStateTreeCompilerLog* Log = nullptr, const FStateTreeBindableStructDesc* LogContextStruct = nullptr);
-
-	/**
-	 * Checks if two property types can are compatible for copying.
-	 * @param FromProperty Property to copy from.
-	 * @param ToProperty Property to copy to.
-	 * @return Incompatible if the properties cannot be copied, Compatible if they are trivially copyable, or Promotable if numeric values can be promoted to another numeric type.
-	 */
-	UE_DEPRECATED(5.3, "Use FStateTreePropertyBindings::GetPropertyCompatibility instead.")
-	static EPropertyAccessCompatibility GetPropertyCompatibility(const FProperty* FromProperty, const FProperty* ToProperty);
-
 	UE_DEPRECATED(5.5, "Use CompileBatch with PropertyFuncsBegin and PropertyFuncsEnd instead.")
 	[[nodiscard]] bool CompileBatch(const FStateTreeBindableStructDesc& InTargetStruct, TConstArrayView<FStateTreePropertyPathBinding> InPropertyBindings, int32& OutBatchIndex) 
 	{
diff --git a/Engine/Plugins/Runtime/StateTree/Source/StateTreeEditorModule/StateTreeEditorModule.Build.cs b/Engine/Plugins/Runtime/StateTree/Source/StateTreeEditorModule/StateTreeEditorModule.Build.cs
--- a/Engine/Plugins/Runtime/StateTree/Source/StateTreeEditorModule/StateTreeEditorModule.Build.cs
+++ b/Engine/Plugins/Runtime/StateTree/Source/StateTreeEditorModule/StateTreeEditorModule.Build.cs
@@ -30,6 +30,7 @@ namespace UnrealBuildTool.Rules
 				"SourceControl",
 				"Projects",
 				"BlueprintGraph",
+				"PropertyBindingUtils",
 				"PropertyAccessEditor",
 				"StructUtilsEditor",
 				"GameplayTags",
diff --git a/Engine/Plugins/Runtime/StateTree/Source/StateTreeModule/Private/StateTree.cpp b/Engine/Plugins/Runtime/StateTree/Source/StateTreeModule/Private/StateTree.cpp
--- a/Engine/Plugins/Runtime/StateTree/Source/StateTreeModule/Private/StateTree.cpp
+++ b/Engine/Plugins/Runtime/StateTree/Source/StateTreeModule/Private/StateTree.cpp
@@ -765,7 +765,7 @@ bool UStateTree::PatchBindings()
 
 	auto PatchPropertyPath = [](FStateTreePropertyPath& PropertyPath)
 	{
-		for (FStateTreePropertyPathSegment& Segment : PropertyPath.GetMutableSegments())
+		for (FPropertyBindingPathSegment& Segment : PropertyPath.GetMutableSegments())
 		{
 			if (const UClass* InstanceStruct = Cast<UClass>(Segment.GetInstanceStruct()))
 			{
diff --git a/Engine/Plugins/Runtime/StateTree/Source/StateTreeModule/Private/StateTreePropertyBindings.cpp b/Engine/Plugins/Runtime/StateTree/Source/StateTreeModule/Private/StateTreePropertyBindings.cpp
--- a/Engine/Plugins/Runtime/StateTree/Source/StateTreeModule/Private/StateTreePropertyBindings.cpp
+++ b/Engine/Plugins/Runtime/StateTree/Source/StateTreeModule/Private/StateTreePropertyBindings.cpp
@@ -11,7 +11,6 @@
 #include "UObject/Package.h"
 #include "Engine/BlueprintGeneratedClass.h"
 #include "StructUtils/UserDefinedStruct.h"
-#include "Kismet2/StructureEditorUtils.h"
 #include "UObject/Field.h"
 #endif
 
@@ -92,50 +91,6 @@ namespace UE::StateTree
 
 } // UE::StateTree
 
-namespace UE::StateTree::Private
-{
-
-#if WITH_EDITORONLY_DATA
-PRAGMA_DISABLE_DEPRECATION_WARNINGS
-	FStateTreePropertyPath ConvertEditorPath(const FStateTreeEditorPropertyPath& InEditorPath)
-	{
-		FStateTreePropertyPath Path;
-		Path.SetStructID(InEditorPath.StructID);
-
-		for (const FString& Segment : InEditorPath.Path)
-		{
-			const TCHAR* PropertyNamePtr = nullptr;
-			int32 PropertyNameLength = 0;
-			int32 ArrayIndex = INDEX_NONE;
-			PropertyPathHelpers::FindFieldNameAndArrayIndex(Segment.Len(), *Segment, PropertyNameLength, &PropertyNamePtr, ArrayIndex);
-			FString PropertyNameString(PropertyNameLength, PropertyNamePtr);
-			const FName PropertyName(*PropertyNameString, FNAME_Find);
-			Path.AddPathSegment(PropertyName, ArrayIndex);
-		}
-		return Path;
-	}
-
-	FStateTreeEditorPropertyPath ConvertEditorPath(const FStateTreePropertyPath& InPath)
-	{
-		FStateTreeEditorPropertyPath Path;
-		Path.StructID = InPath.GetStructID();
-		for (const FStateTreePropertyPathSegment& Segment : InPath.GetSegments())
-		{
-			if (Segment.GetArrayIndex() != INDEX_NONE)
-			{
-				Path.Path.Add(FString::Printf(TEXT("%s[%d]"), *Segment.GetName().ToString(), Segment.GetArrayIndex()));
-			}
-			else
-			{
-				Path.Path.Add(Segment.GetName().ToString());
-			}
-		}
-		return Path;
-	}
-PRAGMA_ENABLE_DEPRECATION_WARNINGS
-#endif // WITH_EDITORONLY_DATA
-} // UE::StateTree::Private 
-
 
 //----------------------------------------------------------------//
 //  FStateTreeBindableStructDesc
@@ -156,32 +111,6 @@ FString FStateTreeBindableStructDesc::ToString() const
 	return Result.ToString();
 }
 
-//----------------------------------------------------------------//
-//  FStateTreePropertyPathBinding
-//----------------------------------------------------------------//
-void FStateTreePropertyPathBinding::PostSerialize(const FArchive& Ar)
-{
-#if WITH_EDITORONLY_DATA
-PRAGMA_DISABLE_DEPRECATION_WARNINGS
-	if (SourcePath_DEPRECATED.IsValid())
-	{
-		SourcePropertyPath = UE::StateTree::Private::ConvertEditorPath(SourcePath_DEPRECATED);
-		SourcePath_DEPRECATED.StructID = FGuid();
-		SourcePath_DEPRECATED.Path.Reset();
-	}
-
-	if (TargetPath_DEPRECATED.IsValid())
-	{
-		TargetPropertyPath = UE::StateTree::Private::ConvertEditorPath(TargetPath_DEPRECATED);
-		TargetPath_DEPRECATED.StructID = FGuid();
-		TargetPath_DEPRECATED.Path.Reset();
-	}
-PRAGMA_ENABLE_DEPRECATION_WARNINGS
-#endif // WITH_EDITORONLY_DATA
-
-}
-
-
 //----------------------------------------------------------------//
 //  FStateTreePropertyBindings
 //----------------------------------------------------------------//
@@ -272,8 +201,8 @@ bool FStateTreePropertyBindings::ResolvePaths()
 
 			// Resolve paths and validate the copy. Stops on first failure.
 			bool bSuccess = true;
-			FStateTreePropertyPathIndirection SourceLeafIndirection;
-			FStateTreePropertyPathIndirection TargetLeafIndirection;
+			FPropertyBindingPathIndirection SourceLeafIndirection;
+			FPropertyBindingPathIndirection TargetLeafIndirection;
 			bSuccess = bSuccess && ResolvePath(SourceStruct, Binding.GetSourcePath(), Copy.SourceIndirection, SourceLeafIndirection);
 			bSuccess = bSuccess && ResolvePath(TargetStruct, Binding.GetTargetPath(), Copy.TargetIndirection, TargetLeafIndirection);
 			bSuccess = bSuccess && ResolveCopyType(SourceLeafIndirection, TargetLeafIndirection, Copy);
@@ -297,7 +226,7 @@ bool FStateTreePropertyBindings::ResolvePaths()
 		const FStateTreeBindableStructDesc* SourceDesc = GetSourceDescByHandle(PropertyAccess.SourceDataHandle);
 		PropertyAccess.SourceStructType = SourceDesc->Struct;
 
-		FStateTreePropertyPathIndirection SourceLeafIndirection;
+		FPropertyBindingPathIndirection SourceLeafIndirection;
 		if (!ResolvePath(SourceDesc->Struct, ReferencePath.GetSourcePath(), PropertyAccess.SourceIndirection, SourceLeafIndirection))
 		{
 			bResult = false;
@@ -309,7 +238,14 @@ bool FStateTreePropertyBindings::ResolvePaths()
 	return bResult;
 }
 
+PRAGMA_DISABLE_DEPRECATION_WARNINGS
 bool FStateTreePropertyBindings::ResolvePath(const UStruct* Struct, const FStateTreePropertyPath& Path, FStateTreePropertyIndirection& OutFirstIndirection, FStateTreePropertyPathIndirection& OutLeafIndirection)
+{
+	return ResolvePath(Struct, Path, OutFirstIndirection, static_cast<FPropertyBindingPathIndirection&>(OutLeafIndirection));
+}
+PRAGMA_ENABLE_DEPRECATION_WARNINGS
+
+bool FStateTreePropertyBindings::ResolvePath(const UStruct* Struct, const FStateTreePropertyPath& Path, FStateTreePropertyIndirection& OutFirstIndirection, FPropertyBindingPathIndirection& OutLeafIndirection)
 {
 	if (!Struct)
  	{
@@ -318,7 +254,7 @@ bool FStateTreePropertyBindings::ResolvePath(const UStruct* Struct, const FState
 	}
 
 	FString Error;
-	TArray<FStateTreePropertyPathIndirection> PathIndirections;
+	TArray<FPropertyBindingPathIndirection> PathIndirections;
 	if (!Path.ResolveIndirections(Struct, PathIndirections, &Error))
 	{
 		UE_LOG(LogStateTree, Error, TEXT("%hs: %s"), __FUNCTION__, *Error);
@@ -337,7 +273,7 @@ bool FStateTreePropertyBindings::ResolvePath(const UStruct* Struct, const FState
 	};
 
 	TArray<FStateTreePropertyIndirection, TInlineAllocator<16>> TempIndirections;
-	for (FStateTreePropertyPathIndirection& PathIndirection : PathIndirections)
+	for (FPropertyBindingPathIndirection& PathIndirection : PathIndirections)
 	{
 		FStateTreePropertyIndirection& Indirection = TempIndirections.AddDefaulted_GetRef();
 
@@ -346,7 +282,7 @@ bool FStateTreePropertyBindings::ResolvePath(const UStruct* Struct, const FState
 		Indirection.Offset = static_cast<uint16>(PathIndirection.GetPropertyOffset());
 		Indirection.Type = PathIndirection.GetAccessType();
 
-		if (Indirection.Type == EStateTreePropertyAccessType::IndexArray)
+		if (Indirection.Type == EPropertyBindingAccessType::IndexArray)
 		{
 			if (const FArrayProperty* ArrayProperty = CastField<FArrayProperty>(PathIndirection.GetProperty()))
 			{
@@ -366,8 +302,8 @@ bool FStateTreePropertyBindings::ResolvePath(const UStruct* Struct, const FState
 				return false;
 			}
 		}
-		else if (Indirection.Type == EStateTreePropertyAccessType::StructInstance
-				|| Indirection.Type == EStateTreePropertyAccessType::ObjectInstance)
+		else if (Indirection.Type == EPropertyBindingAccessType::StructInstance
+				|| Indirection.Type == EPropertyBindingAccessType::ObjectInstance)
 		{
 			if (PathIndirection.GetInstanceStruct())
 			{
@@ -390,33 +326,33 @@ bool FStateTreePropertyBindings::ResolvePath(const UStruct* Struct, const FState
 			if ((Index + 1) < TempIndirections.Num())
 			{
 				const FStateTreePropertyIndirection& NextIndirection = TempIndirections[Index + 1];
-				if (Indirection.Type == EStateTreePropertyAccessType::Offset
-					&& NextIndirection.Type == EStateTreePropertyAccessType::Offset)
+				if (Indirection.Type == EPropertyBindingAccessType::Offset
+					&& NextIndirection.Type == EPropertyBindingAccessType::Offset)
 				{
 					// Collapse adjacent offset indirections
 					Indirection.Offset += NextIndirection.Offset;
 					TempIndirections.RemoveAt(Index + 1);
 					Index--;
 				}
-				else if (Indirection.Type == EStateTreePropertyAccessType::IndexArray
-					&& NextIndirection.Type == EStateTreePropertyAccessType::Offset
+				else if (Indirection.Type == EPropertyBindingAccessType::IndexArray
+					&& NextIndirection.Type == EPropertyBindingAccessType::Offset
 					&& NextIndirection.Offset == 0)
 				{
 					// Remove empty offset after array indexing.
 					TempIndirections.RemoveAt(Index + 1);
 					Index--;
 				}
-				else if (Indirection.Type == EStateTreePropertyAccessType::StructInstance
-					&& NextIndirection.Type == EStateTreePropertyAccessType::Offset
+				else if (Indirection.Type == EPropertyBindingAccessType::StructInstance
+					&& NextIndirection.Type == EPropertyBindingAccessType::Offset
 					&& NextIndirection.Offset == 0)
 				{
 					// Remove empty offset after struct indirection.
 					TempIndirections.RemoveAt(Index + 1);
 					Index--;
 				}
-				else if ((Indirection.Type == EStateTreePropertyAccessType::Object
-						|| Indirection.Type == EStateTreePropertyAccessType::ObjectInstance)
-					&& NextIndirection.Type == EStateTreePropertyAccessType::Offset
+				else if ((Indirection.Type == EPropertyBindingAccessType::Object
+						|| Indirection.Type == EPropertyBindingAccessType::ObjectInstance)
+					&& NextIndirection.Type == EPropertyBindingAccessType::Offset
 					&& NextIndirection.Offset == 0)
 				{
 					// Remove empty offset after object indirection.
@@ -444,15 +380,22 @@ bool FStateTreePropertyBindings::ResolvePath(const UStruct* Struct, const FState
 		// Indirections can be empty in case we're directly binding to source structs.
 		// Zero offset will return the struct itself.
 		OutFirstIndirection.Offset = 0;
-		OutFirstIndirection.Type = EStateTreePropertyAccessType::Offset;
+		OutFirstIndirection.Type = EPropertyBindingAccessType::Offset;
 
-		OutLeafIndirection = FStateTreePropertyPathIndirection(Struct);
+		OutLeafIndirection = FPropertyBindingPathIndirection(Struct);
 	}
 
 	return true;
 }
 
-bool FStateTreePropertyBindings::ResolveCopyType(const FStateTreePropertyPathIndirection& SourceIndirection,const FStateTreePropertyPathIndirection& TargetIndirection, FStateTreePropertyCopy& OutCopy)
+PRAGMA_DISABLE_DEPRECATION_WARNINGS
+bool FStateTreePropertyBindings::ResolveCopyType(const FStateTreePropertyPathIndirection& SourceIndirection, const FStateTreePropertyPathIndirection& TargetIndirection, FStateTreePropertyCopy& OutCopy)
+{
+	return ResolveCopyType(static_cast<const FPropertyBindingPathIndirection&>(SourceIndirection), TargetIndirection, OutCopy);
+}
+PRAGMA_ENABLE_DEPRECATION_WARNINGS
+
+bool FStateTreePropertyBindings::ResolveCopyType(const FPropertyBindingPathIndirection& SourceIndirection, const FPropertyBindingPathIndirection& TargetIndirection, FStateTreePropertyCopy& OutCopy)
 {
 	// @todo: see if GetPropertyCompatibility() can be implemented as call to ResolveCopyType() instead so that we write this logic just once.
 	
@@ -835,7 +778,7 @@ EStateTreePropertyAccessCompatibility FStateTreePropertyBindings::GetPropertyCom
 
 uint8* FStateTreePropertyBindings::GetAddress(FStateTreeDataView InStructView, const FStateTreePropertyIndirection& FirstIndirection, const FProperty* LeafProperty) const
 {
-	uint8* Address = InStructView.GetMutableMemory();
+	uint8* Address = (uint8*)InStructView.GetMutableMemory();
 	if (Address == nullptr)
 	{
 		// Failed indirection, will be reported by caller.
@@ -848,32 +791,32 @@ uint8* FStateTreePropertyBindings::GetAddress(FStateTreeDataView InStructView, c
 	{
 		switch (Indirection->Type)
 		{
-		case EStateTreePropertyAccessType::Offset:
+		case EPropertyBindingAccessType::Offset:
 		{
 			Address = Address + Indirection->Offset;
 			break;
 		}
-		case EStateTreePropertyAccessType::Object:
+		case EPropertyBindingAccessType::Object:
 		{
 			UObject* Object = *reinterpret_cast<UObject**>(Address + Indirection->Offset);
 			Address = reinterpret_cast<uint8*>(Object);
 			break;
 		}
-		case EStateTreePropertyAccessType::WeakObject:
+		case EPropertyBindingAccessType::WeakObject:
 		{
 			TWeakObjectPtr<UObject>& WeakObjectPtr = *reinterpret_cast<TWeakObjectPtr<UObject>*>(Address + Indirection->Offset);
 			UObject* Object = WeakObjectPtr.Get();
 			Address = reinterpret_cast<uint8*>(Object);
 			break;
 		}
-		case EStateTreePropertyAccessType::SoftObject:
+		case EPropertyBindingAccessType::SoftObject:
 		{
 			FSoftObjectPtr& SoftObjectPtr = *reinterpret_cast<FSoftObjectPtr*>(Address + Indirection->Offset);
 			UObject* Object = SoftObjectPtr.Get();
 			Address = reinterpret_cast<uint8*>(Object);
 			break;
 		}
-		case EStateTreePropertyAccessType::ObjectInstance:
+		case EPropertyBindingAccessType::ObjectInstance:
 		{
 			check(Indirection->InstanceStruct);
 			UObject* Object = *reinterpret_cast<UObject**>(Address + Indirection->Offset);
@@ -889,7 +832,7 @@ uint8* FStateTreePropertyBindings::GetAddress(FStateTreeDataView InStructView, c
 			}
 			break;
 		}
-		case EStateTreePropertyAccessType::StructInstance:
+		case EPropertyBindingAccessType::StructInstance:
 		{
 			check(Indirection->InstanceStruct);
 			FInstancedStruct& InstancedStruct = *reinterpret_cast<FInstancedStruct*>(Address + Indirection->Offset);
@@ -906,7 +849,7 @@ uint8* FStateTreePropertyBindings::GetAddress(FStateTreeDataView InStructView, c
 			}
 			break;
 		}
-		case EStateTreePropertyAccessType::IndexArray:
+		case EPropertyBindingAccessType::IndexArray:
 		{
 			check(Indirection->ArrayProperty);
 			FScriptArrayHelper Helper(Indirection->ArrayProperty, Address + Indirection->Offset);
@@ -923,13 +866,13 @@ uint8* FStateTreePropertyBindings::GetAddress(FStateTreeDataView InStructView, c
 		}
 		default:
 			ensureMsgf(false, TEXT("FStateTreePropertyBindings::GetAddress: Unhandled indirection type %s for '%s'"),
-				*StaticEnum<EStateTreePropertyAccessType>()->GetValueAsString(Indirection->Type), *LeafProperty->GetNameCPP());
+				*StaticEnum<EPropertyBindingAccessType>()->GetValueAsString(Indirection->Type), *LeafProperty->GetNameCPP());
 		}
 
 		Indirection = Indirection->NextIndex.IsValid() ? &PropertyIndirections[Indirection->NextIndex.Get()] : nullptr;
 	}
 
-	return Address;
+	return (uint8*)Address;
 }
 
 void FStateTreePropertyBindings::PerformCopy(const FStateTreePropertyCopy& Copy, uint8* SourceAddress, uint8* TargetAddress) const
@@ -1230,7 +1173,7 @@ bool FStateTreePropertyBindings::ContainsAnyStruct(const TSet<const UStruct*>& S
 
 	auto PathContainsStruct = [&Structs](const FStateTreePropertyPath& PropertyPath)
 	{
-		for (const FStateTreePropertyPathSegment& Segment : PropertyPath.GetSegments())
+		for (const FPropertyBindingPathSegment& Segment : PropertyPath.GetSegments())
 		{
 			if (Structs.Contains(Segment.GetInstanceStruct()))
 			{
@@ -1356,626 +1299,3 @@ FString FStateTreePropertyBindings::DebugInternalLayoutAsString() const
 	return DebugString.ToString();
 }
 #endif // WITH_EDITOR || WITH_STATETREE_DEBUG
-
-
-//----------------------------------------------------------------//
-//  FStateTreePropertyPath
-//----------------------------------------------------------------//
-
-bool FStateTreePropertyPath::FromString(const FStringView InPath)
-{
-	Segments.Reset();
-	
-	if (InPath.IsEmpty())
-	{
-		return true;
-	}
-	
-	auto ParseSegment = [This = this](const FStringView Segment)
-	{
-		if (Segment.IsEmpty())
-		{
-			return false;
-		}
-
-		int32 FirstBracket = INDEX_NONE;
-		int32 LastBracket = INDEX_NONE;
-		if (Segment.FindChar(TEXT('['), FirstBracket) && Segment.FindLastChar(TEXT(']'), LastBracket))
-		{
-			const int32 NameStringLength = FirstBracket;
-			const int32 IndexStringLength = LastBracket - FirstBracket - 1;
-			if (NameStringLength < 1
-				|| IndexStringLength <= 0)
-			{
-				return false;
-			}
-
-			const FStringView NameString = Segment.Left(FirstBracket);
-			const FStringView IndexString = Segment.Mid(FirstBracket + 1, IndexStringLength);
-			int32 ArrayIndex = INDEX_NONE;
-			LexFromString(ArrayIndex, IndexString.GetData());
-			if (ArrayIndex < 0)
-			{
-				return false;
-			}
-			
-			This->AddPathSegment(FName(NameString), ArrayIndex);
-		}
-		else
-		{
-			This->AddPathSegment(FName(Segment));
-		}
-		return true;
-	};
-
-	bool bResult = true;
-	FStringView Path = InPath;
-	int32 FoundChar = INDEX_NONE;
-	while (Path.FindChar(TEXT('.'), FoundChar))
-	{
-		const FStringView Segment = Path.SubStr(0, FoundChar);
-		if (!ParseSegment(Segment))
-		{
-			bResult = false;
-			break;
-		}
-		Path = Path.Mid(FoundChar+1);
-	}
-	if (bResult)
-	{
-		bResult = ParseSegment(Path);
-	}
-
-	if (!bResult)
-	{
-		Segments.Reset();
-	}
-	
-	return bResult;
-}
-
-bool FStateTreePropertyPath::UpdateSegments(const UStruct* BaseStruct, FString* OutError)
-{
-	return UpdateSegmentsFromValue(FStateTreeDataView(BaseStruct, nullptr), OutError);
-}
-
-bool FStateTreePropertyPath::UpdateSegmentsFromValue(const FStateTreeDataView BaseValueView, FString* OutError)
-{
-	TArray<FStateTreePropertyPathIndirection> Indirections;
-	if (!ResolveIndirectionsWithValue(BaseValueView, Indirections, OutError, /*bHandleRedirects*/true))
-	{
-		return false;
-	}
-
-	for (FStateTreePropertyPathSegment& Segment : Segments)
-	{
-		Segment.SetInstanceStruct(nullptr);
-	}
-	
-	for (const FStateTreePropertyPathIndirection& Indirection : Indirections)
-	{
-		if (Indirection.InstanceStruct != nullptr)
-		{
-			Segments[Indirection.PathSegmentIndex].SetInstanceStruct(Indirection.InstanceStruct);
-		}
-#if WITH_EDITORONLY_DATA		
-		if (!Indirection.GetRedirectedName().IsNone())
-		{
-			Segments[Indirection.PathSegmentIndex].SetName(Indirection.GetRedirectedName());
-		}
-		Segments[Indirection.PathSegmentIndex].SetPropertyGuid(Indirection.GetPropertyGuid());
-#endif			
-	}
-
-	return true;
-}
-
-FString FStateTreePropertyPath::ToString(const int32 HighlightedSegment, const TCHAR* HighlightPrefix, const TCHAR* HighlightPostfix, const bool bOutputInstances, const int32 FirstSegment) const
-{
-	FStringBuilderBase Result;
-	for (TEnumerateRef<const FStateTreePropertyPathSegment> Segment : EnumerateRange(TConstArrayView<FStateTreePropertyPathSegment>(Segments).Mid(FMath::Max(FirstSegment, 0))))
-	{
-		if (Segment.GetIndex() > 0)
-		{
-			Result += TEXT(".");
-		}
-		if (Segment.GetIndex() == HighlightedSegment && HighlightPrefix)
-		{
-			Result += HighlightPrefix;
-		}
-
-		if (bOutputInstances && Segment->GetInstanceStruct())
-		{
-			Result += FString::Printf(TEXT("(%s)"), *GetNameSafe(Segment->GetInstanceStruct()));
-		}
-
-#if WITH_EDITORONLY_DATA
-		const UStruct* ParentInstanceStruct = Segment.GetIndex() > 0 ? Segments[Segment.GetIndex() - 1].GetInstanceStruct() : nullptr;
-		if (const UUserDefinedStruct* ParentUserDefinedStruct = Cast<const UUserDefinedStruct>(ParentInstanceStruct))
-		{
-			// Find friendly names for UDS properties (the property name itself has hash in it). 
-			const FString FriendlyName = FStructureEditorUtils::GetVariableFriendlyName(ParentUserDefinedStruct, Segment->GetPropertyGuid());
-			if (!FriendlyName.IsEmpty())
-			{
-				Result += FriendlyName;
-			}
-			else
-			{
-				Result += Segment->GetName().ToString();
-			}
-		}
-		else
-#endif			
-		{
-			Result += Segment->GetName().ToString();
-		}
-
-		if (Segment->GetArrayIndex() >= 0)
-		{
-			Result += FString::Printf(TEXT("[%d]"), Segment->GetArrayIndex());
-		}
-
-		if (Segment.GetIndex() == HighlightedSegment && HighlightPostfix)
-		{
-			Result += HighlightPostfix;
-		}
-	}
-	return Result.ToString();
-}
-
-
-bool FStateTreePropertyPath::ResolveIndirections(const UStruct* BaseStruct, TArray<FStateTreePropertyPathIndirection>& OutIndirections, FString* OutError, bool bHandleRedirects) const
-{
-	return ResolveIndirectionsWithValue(FStateTreeDataView(BaseStruct, nullptr), OutIndirections, OutError, bHandleRedirects);
-}
-
-bool FStateTreePropertyPath::ResolveIndirectionsWithValue(const FStateTreeDataView BaseValueView, TArray<FStateTreePropertyPathIndirection>& OutIndirections, FString* OutError, bool bHandleRedirects) const
-{
-	OutIndirections.Reset();
-	if (OutError)
-	{
-		OutError->Reset();
-	}
-	
-	// Nothing to do for an empty path.
-	if (IsPathEmpty())
-	{
-		return true;
-	}
-
-	const uint8* CurrentAddress = BaseValueView.GetMemory();
-	const UStruct* CurrentStruct = BaseValueView.GetStruct();
-	
-	for (const TEnumerateRef<const FStateTreePropertyPathSegment> Segment : EnumerateRange(Segments))
-	{
-		if (CurrentStruct == nullptr)
-		{
-			if (OutError)
-			{
-				*OutError = FString::Printf(TEXT("Malformed path '%s'."),
-					*ToString(Segment.GetIndex(), TEXT("<"), TEXT(">")));
-			}
-			OutIndirections.Reset();
-			return false;
-		}
-
-		const FProperty* Property = CurrentStruct->FindPropertyByName(Segment->GetName());
-		const bool bWithValue = CurrentAddress != nullptr;
-
-#if WITH_EDITORONLY_DATA
-		FName RedirectedName;
-		FGuid PropertyGuid = Segment->GetPropertyGuid();
-
-		// Try to fix the path in editor.
-		if (bHandleRedirects)
-		{
-			
-			// Check if there's a core redirect for it.
-			if (!Property)
-			{
-				// Try to match by property ID (Blueprint or User Defined Struct).
-				if (Segment->GetPropertyGuid().IsValid())
-				{
-					if (const UBlueprintGeneratedClass* BlueprintClass = Cast<UBlueprintGeneratedClass>(CurrentStruct))
-					{
-						if (const FName* Name = BlueprintClass->PropertyGuids.FindKey(Segment->GetPropertyGuid()))
-						{
-							RedirectedName = *Name;
-							Property = CurrentStruct->FindPropertyByName(RedirectedName);
-						}
-					}
-					else if (const UUserDefinedStruct* UserDefinedStruct = Cast<UUserDefinedStruct>(CurrentStruct))
-					{
-						if (FProperty* FoundProperty = FStructureEditorUtils::GetPropertyByGuid(UserDefinedStruct, Segment->GetPropertyGuid()))
-						{
-							RedirectedName = FoundProperty->GetFName();
-							Property = FoundProperty;
-						}
-					}
-					else if (const UPropertyBag* PropertyBag = Cast<UPropertyBag>(CurrentStruct))
-					{
-						if (const FPropertyBagPropertyDesc* Desc = PropertyBag->FindPropertyDescByID(Segment->GetPropertyGuid()))
-						{
-							if (Desc->CachedProperty)
-							{
-								RedirectedName = Desc->CachedProperty->GetFName();
-								Property = Desc->CachedProperty;
-							}
-						}
-					}
-				}
-				else
-				{
-					// Try core redirect
-					const FCoreRedirectObjectName OldPropertyName(Segment->GetName(), CurrentStruct->GetFName(), *CurrentStruct->GetOutermost()->GetPathName());
-					const FCoreRedirectObjectName NewPropertyName = FCoreRedirects::GetRedirectedName(ECoreRedirectFlags::Type_Property, OldPropertyName);
-					if (OldPropertyName != NewPropertyName)
-					{
-						// Cached the result for later use.
-						RedirectedName = NewPropertyName.ObjectName;
-
-						Property = CurrentStruct->FindPropertyByName(RedirectedName);
-					}
-				}
-			}
-
-			// Update PropertyGuid 
-			if (Property)
-			{
-				const FName PropertyName = !RedirectedName.IsNone() ? RedirectedName : Segment->GetName();
-				if (const UBlueprintGeneratedClass* BlueprintClass = Cast<UBlueprintGeneratedClass>(CurrentStruct))
-				{
-					if (const FGuid* VarGuid = BlueprintClass->PropertyGuids.Find(PropertyName))
-					{
-						PropertyGuid = *VarGuid;
-					}
-				}
-				else if (const UUserDefinedStruct* UserDefinedStruct = Cast<UUserDefinedStruct>(CurrentStruct))
-				{
-					// Parse Guid from UDS property name.
-					PropertyGuid = FStructureEditorUtils::GetGuidFromPropertyName(PropertyName);
-				}
-				else if (const UPropertyBag* PropertyBag = Cast<UPropertyBag>(CurrentStruct))
-				{
-					if (const FPropertyBagPropertyDesc* Desc = PropertyBag->FindPropertyDescByPropertyName(PropertyName))
-					{
-						PropertyGuid = Desc->ID;
-					}
-				}
-			}
-		}
-#endif // WITH_EDITORONLY_DATA
-
-		if (!Property)
-		{
-			if (OutError)
-			{
-				*OutError = FString::Printf(TEXT("Malformed path '%s', could not find property '%s%s::%s'."),
-					*ToString(Segment.GetIndex(), TEXT("<"), TEXT(">")),
-					CurrentStruct->GetPrefixCPP(), *CurrentStruct->GetName(), *Segment->GetName().ToString());
-			}
-			OutIndirections.Reset();
-			return false;
-		}
-
-		const FArrayProperty* ArrayProperty = CastField<FArrayProperty>(Property);
-		int ArrayIndex = 0;
-		int32 Offset = 0;
-		if (ArrayProperty && Segment->GetArrayIndex() != INDEX_NONE)
-		{
-			FStateTreePropertyPathIndirection& Indirection = OutIndirections.AddDefaulted_GetRef();
-			Indirection.Property = Property;
-			Indirection.ContainerAddress = CurrentAddress;
-			Indirection.ContainerStruct = CurrentStruct;
-			Indirection.InstanceStruct = nullptr;
-			Indirection.ArrayIndex = Segment->GetArrayIndex();
-			Indirection.PropertyOffset = ArrayProperty->GetOffset_ForInternal();
-			Indirection.PathSegmentIndex = Segment.GetIndex();
-			Indirection.AccessType = EStateTreePropertyAccessType::IndexArray;
-#if WITH_EDITORONLY_DATA
-			Indirection.RedirectedName = RedirectedName;
-			Indirection.PropertyGuid = PropertyGuid;
-#endif
-			
-			ArrayIndex = 0;
-			Offset = 0;
-			Property = ArrayProperty->Inner;
-
-			if (bWithValue)
-			{
-				FScriptArrayHelper Helper(ArrayProperty, CurrentAddress + ArrayProperty->GetOffset_ForInternal());
-				if (!Helper.IsValidIndex(Segment->GetArrayIndex()))
-				{
-					if (OutError)
-					{
-						*OutError = FString::Printf(TEXT("Index %d out of range (num elements %d) trying to access dynamic array '%s'."),
-							Segment->GetArrayIndex(), Helper.Num(), *ToString(Segment.GetIndex(), TEXT("<"), TEXT(">")));
-					}
-					OutIndirections.Reset();
-					return false;
-				}
-				CurrentAddress = Helper.GetRawPtr(Segment->GetArrayIndex());
-			}
-		}
-		else
-		{
-			if (Segment->GetArrayIndex() > Property->ArrayDim)
-			{
-				if (OutError)
-				{
-					*OutError = FString::Printf(TEXT("Index %d out of range %d trying to access static array '%s'."),
-						Segment->GetArrayIndex(), Property->ArrayDim, *ToString(Segment.GetIndex(), TEXT("<"), TEXT(">")));
-				}
-				OutIndirections.Reset();
-				return false;
-			}
-			ArrayIndex = FMath::Max(0, Segment->GetArrayIndex());
-			Offset = Property->GetOffset_ForInternal() + Property->GetElementSize() * ArrayIndex;
-		}
-
-		FStateTreePropertyPathIndirection& Indirection = OutIndirections.AddDefaulted_GetRef();
-		Indirection.Property = Property;
-		Indirection.ContainerAddress = CurrentAddress;
-		Indirection.ContainerStruct = CurrentStruct;
-		Indirection.ArrayIndex = ArrayIndex;
-		Indirection.PropertyOffset = Offset;
-		Indirection.PathSegmentIndex = Segment.GetIndex();
-		Indirection.AccessType = EStateTreePropertyAccessType::Offset; 
-#if WITH_EDITORONLY_DATA
-		Indirection.RedirectedName = RedirectedName;
-		Indirection.PropertyGuid = PropertyGuid;
-#endif
-		const bool bLastSegment = Segment.GetIndex() == (Segments.Num() - 1);
-
-		if (!bLastSegment)
-		{
-			if (const FStructProperty* StructProperty = CastField<FStructProperty>(Property))
-			{
-				if (bWithValue)
-				{
-					if (StructProperty->Struct == TBaseStructure<FInstancedStruct>::Get())
-					{
-						// The property path is pointing into the instanced struct, it must be present.
-						// @TODO:	We could potentially check the BaseStruct metadata in editor (for similar behavior as objects)
-						//			Omitting for now to have matching functionality in editor and runtime.
-						const FInstancedStruct& InstancedStruct = *reinterpret_cast<const FInstancedStruct*>(CurrentAddress + Offset);
-						if (!InstancedStruct.IsValid())
-						{
-							if (OutError)
-							{
-								*OutError = FString::Printf(TEXT("Expecting valid instanced struct value at path '%s'."),
-									*ToString(Segment.GetIndex(), TEXT("<"), TEXT(">")));
-							}
-							OutIndirections.Reset();
-							return false;
-						}
-						const UScriptStruct* ValueInstanceStructType = InstancedStruct.GetScriptStruct();
-
-						CurrentAddress = InstancedStruct.GetMemory();
-						CurrentStruct = ValueInstanceStructType;
-						Indirection.InstanceStruct = CurrentStruct;
-						Indirection.AccessType = EStateTreePropertyAccessType::StructInstance; 
-					}
-					else
-					{
-						CurrentAddress = CurrentAddress + Offset;
-						CurrentStruct = StructProperty->Struct;
-						Indirection.AccessType = EStateTreePropertyAccessType::Offset;
-					}
-				}
-				else
-				{
-					if (Segment->GetInstanceStruct())
-					{
-						CurrentStruct = Segment->GetInstanceStruct();
-						Indirection.InstanceStruct = CurrentStruct;
-						Indirection.AccessType = EStateTreePropertyAccessType::StructInstance;
-					}
-					else
-					{
-						CurrentStruct = StructProperty->Struct;
-						Indirection.AccessType = EStateTreePropertyAccessType::Offset;
-					}
-				}
-			}
-			else if (const FObjectProperty* ObjectProperty = CastField<FObjectProperty>(Property))
-			{
-				if (bWithValue)
-				{
-					const UObject* Object = *reinterpret_cast<UObject* const*>(CurrentAddress + Offset);
-					CurrentAddress = reinterpret_cast<const uint8*>(Object);
-					
-					// The property path is pointing into the object, if the object is present use it's specific type, otherwise use the type of the pointer.
-					if (Object)
-					{
-						CurrentStruct = Object->GetClass();
-						Indirection.InstanceStruct = CurrentStruct;
-						Indirection.AccessType = EStateTreePropertyAccessType::ObjectInstance;
-					}
-					else
-					{
-						CurrentStruct = ObjectProperty->PropertyClass;
-						Indirection.AccessType = EStateTreePropertyAccessType::Object;
-					}
-				}
-				else
-				{
-					if (Segment->GetInstanceStruct())
-					{
-						CurrentStruct = Segment->GetInstanceStruct();
-						Indirection.InstanceStruct = CurrentStruct;
-						Indirection.AccessType = EStateTreePropertyAccessType::ObjectInstance;
-					}
-					else
-					{
-						CurrentStruct = ObjectProperty->PropertyClass;
-						Indirection.AccessType = EStateTreePropertyAccessType::Object;
-					}
-				}
-			}
-			// Check to see if this is a simple weak object property (eg. not an array of weak objects).
-			else if (const FWeakObjectProperty* WeakObjectProperty = CastField<FWeakObjectProperty>(Property))
-			{
-				if (bWithValue)
-				{
-					const TWeakObjectPtr<UObject>& WeakObjectPtr = *reinterpret_cast<const TWeakObjectPtr<UObject>*>(CurrentAddress + Offset);
-					const UObject* Object = WeakObjectPtr.Get();
-					CurrentAddress = reinterpret_cast<const uint8*>(Object);
-
-					if (Object)
-					{
-						CurrentStruct = Object->GetClass();
-						Indirection.InstanceStruct = CurrentStruct;
-					}
-				}
-				else
-				{
-					CurrentStruct = WeakObjectProperty->PropertyClass;
-				}
-				
-				Indirection.AccessType = EStateTreePropertyAccessType::WeakObject;
-			}
-			// Check to see if this is a simple soft object property (eg. not an array of soft objects).
-			else if (const FSoftObjectProperty* SoftObjectProperty = CastField<FSoftObjectProperty>(Property))
-			{
-				if (bWithValue)
-				{
-					const FSoftObjectPtr& SoftObjectPtr = *reinterpret_cast<const FSoftObjectPtr*>(CurrentAddress + Offset);
-					const UObject* Object = SoftObjectPtr.Get();
-					CurrentAddress = reinterpret_cast<const uint8*>(Object);
-
-					if (Object)
-					{
-						CurrentStruct = Object->GetClass();
-						Indirection.InstanceStruct = CurrentStruct;
-					}			
-				}
-				else
-				{			
-					CurrentStruct = SoftObjectProperty->PropertyClass;
-				}
-
-				Indirection.AccessType = EStateTreePropertyAccessType::SoftObject;
-			}
-			else
-			{
-				// We get here if we encounter a property type that is not supported for indirection (e.g. Map or Set).
-				if (OutError)
-				{
-					*OutError = FString::Printf(TEXT("Unsupported property indirection type %s in path '%s'."),
-						*Property->GetCPPType(), *ToString(Segment.GetIndex(), TEXT("<"), TEXT(">")));
-				}
-				OutIndirections.Reset();
-				return false;
-			}
-		}
-	}
-
-	return true;
-}
-
-bool FStateTreePropertyPath::Includes(const FStateTreePropertyPath& Other) const
-{
-#if WITH_EDITORONLY_DATA
-	if (StructID != Other.StructID)
-	{
-		return false;
-	}
-#endif // WITH_EDITORONLY_DATA
-
-	if (Segments.Num() < Other.Segments.Num())
-	{
-		return false;
-	}
-
-	for (TEnumerateRef<const FStateTreePropertyPathSegment> OtherSegment : EnumerateRange(Other.Segments))
-	{
-		if (*OtherSegment != Segments[OtherSegment.GetIndex()])
-		{
-			return false;
-		}
-	}
-
-	return true;
-}
-
-
-
-bool FStateTreePropertyPath::operator==(const FStateTreePropertyPath& RHS) const
-{
-#if WITH_EDITORONLY_DATA
-	if (StructID != RHS.StructID)
-	{
-		return false;
-	}
-#endif // WITH_EDITORONLY_DATA
-
-	if (Segments.Num() != RHS.Segments.Num())
-	{
-		return false;
-	}
-
-	for (TEnumerateRef<const FStateTreePropertyPathSegment> Segment : EnumerateRange(Segments))
-	{
-		if (*Segment != RHS.Segments[Segment.GetIndex()])
-		{
-			return false;
-		}
-	}
-
-	return true;
-}
-
-
-//----------------------------------------------------------------//
-//  FStateTreeEditorPropertyPath
-//----------------------------------------------------------------//
-
-FString FStateTreeEditorPropertyPath::ToString(const int32 HighlightedSegment, const TCHAR* HighlightPrefix, const TCHAR* HighlightPostfix) const
-{
-	FString Result;
-	for (int32 i = 0; i < Path.Num(); i++)
-	{
-		if (i > 0)
-		{
-			Result += TEXT(".");
-		}
-		if (i == HighlightedSegment && HighlightPrefix)
-		{
-			Result += HighlightPrefix;
-		}
-
-		Result += Path[i];
-
-		if (i == HighlightedSegment && HighlightPostfix)
-		{
-			Result += HighlightPostfix;
-		}
-	}
-	return Result;
-}
-
-PRAGMA_DISABLE_DEPRECATION_WARNINGS
-bool FStateTreeEditorPropertyPath::operator==(const FStateTreeEditorPropertyPath& RHS) const
-{
-	if (StructID != RHS.StructID)
-	{
-		return false;
-	}
-
-	if (Path.Num() != RHS.Path.Num())
-	{
-		return false;
-	}
-
-	for (int32 i = 0; i < Path.Num(); i++)
-	{
-		if (Path[i] != RHS.Path[i])
-		{
-			return false;
-		}
-	}
-
-	return true;
-}
-PRAGMA_ENABLE_DEPRECATION_WARNINGS
diff --git a/Engine/Plugins/Runtime/StateTree/Source/StateTreeModule/Private/StateTreePropertyRefHelpers.cpp b/Engine/Plugins/Runtime/StateTree/Source/StateTreeModule/Private/StateTreePropertyRefHelpers.cpp
--- a/Engine/Plugins/Runtime/StateTree/Source/StateTreeModule/Private/StateTreePropertyRefHelpers.cpp
+++ b/Engine/Plugins/Runtime/StateTree/Source/StateTreeModule/Private/StateTreePropertyRefHelpers.cpp
@@ -249,10 +249,17 @@ namespace UE::StateTree::PropertyRefHelpers
 		return false;
 	}
 
+PRAGMA_DISABLE_DEPRECATION_WARNINGS
 	bool IsPropertyAccessibleForPropertyRef(TConstArrayView<FStateTreePropertyPathIndirection> SourcePropertyPathIndirections, FStateTreeBindableStructDesc SourceStruct)
+	{
+		return false;
+	}
+PRAGMA_ENABLE_DEPRECATION_WARNINGS
+
+	bool IsPropertyAccessibleForPropertyRef(TConstArrayView<FPropertyBindingPathIndirection> SourcePropertyPathIndirections, FStateTreeBindableStructDesc SourceStruct)
 	{
 		bool bIsOutput = false;
-		for (const FStateTreePropertyPathIndirection& Indirection : SourcePropertyPathIndirections)
+		for (const FPropertyBindingPathIndirection& Indirection : SourcePropertyPathIndirections)
 		{
 			if (UE::StateTree::GetUsageFromMetaData(Indirection.GetProperty()) == EStateTreePropertyUsage::Output)
 			{
diff --git a/Engine/Plugins/Runtime/StateTree/Source/StateTreeModule/Public/StateTreePropertyBindings.h b/Engine/Plugins/Runtime/StateTree/Source/StateTreeModule/Public/StateTreePropertyBindings.h
--- a/Engine/Plugins/Runtime/StateTree/Source/StateTreeModule/Public/StateTreePropertyBindings.h
+++ b/Engine/Plugins/Runtime/StateTree/Source/StateTreeModule/Public/StateTreePropertyBindings.h
@@ -2,6 +2,7 @@
 
 #pragma once
 
+#include "PropertyBindingPath.h"
 #include "StateTreeNodeBase.h"
 #include "StateTreeTypes.h"
 #include "StateTreePropertyRefHelpers.h"
@@ -10,7 +11,6 @@
 #include "StateTreePropertyBindings.generated.h"
 
 class FProperty;
-struct FStateTreePropertyPath;
 struct FStateTreePropertyBindingCompiler;
 struct FStateTreePropertyRef;
 class UStateTree;
@@ -44,21 +44,6 @@ enum class EStateTreeBindableStructSource : uint8
 };
 
 
-/**
- * The type of access a property path describes.
- */
-UENUM()
-enum class EStateTreePropertyAccessType : uint8
-{
-	Offset,			// Access node is a simple basePtr + offset
-	Object,			// Access node needs to dereference an object at its current address
-	WeakObject,		// Access is a weak object
-	SoftObject,		// Access is a soft object
-	ObjectInstance,	// Access node needs to dereference an object of specific type at its current address
-	StructInstance,	// Access node needs to dereference an instanced struct of specific type at its current address
-	IndexArray,		// Access node indexes a dynamic array
-};
-
 /**
  * Describes how the copy should be performed.
  */
@@ -208,313 +193,21 @@ struct STATETREEMODULE_API FStateTreeBindableStructDesc
 #endif
 };
 
-/** Struct describing a path segment in FStateTreePropertyPath. */
-USTRUCT()
-struct STATETREEMODULE_API FStateTreePropertyPathSegment
-{
-	GENERATED_BODY()
-
-	FStateTreePropertyPathSegment() = default;
-
-	explicit FStateTreePropertyPathSegment(const FName InName, const int32 InArrayIndex = INDEX_NONE, const UStruct* InInstanceStruct = nullptr)
-		: Name(InName)
-		, ArrayIndex(InArrayIndex)
-		, InstanceStruct(InInstanceStruct)
-	{
-	}
-
-	bool operator==(const FStateTreePropertyPathSegment& RHS) const
-	{
-		return Name == RHS.Name && InstanceStruct == RHS.InstanceStruct && ArrayIndex == RHS.ArrayIndex;
-	}
-
-	bool operator!=(const FStateTreePropertyPathSegment& RHS) const
-	{
-		return !(*this == RHS);
-	}
-
-	void SetName(const FName InName)
-	{
-		Name = InName;
-	}
-
-	FName GetName() const
-	{
-		return Name;
-	}
-
-	void SetArrayIndex(const int32 InArrayIndex)
-	{
-		ArrayIndex = InArrayIndex;
-	}
-
-	int32 GetArrayIndex() const
-	{
-		return ArrayIndex;
-	}
-
-	void SetInstanceStruct(const UStruct* InInstanceStruct)
-	{
-		InstanceStruct = InInstanceStruct;
-	}
-
-	const UStruct* GetInstanceStruct() const
-	{
-		return InstanceStruct;
-	}
-
-#if WITH_EDITORONLY_DATA
-	FGuid GetPropertyGuid() const
-	{
-		return PropertyGuid;
-	}
-
-	void SetPropertyGuid(const FGuid NewGuid)
-	{
-		PropertyGuid = NewGuid;
-	}
-#endif
-	
-private:
-	/** Name of the property */
-	UPROPERTY()
-	FName Name;
-
-#if WITH_EDITORONLY_DATA
-	/** Guid of the property for Blueprint classes or User Defined Structs. */
-	FGuid PropertyGuid;
-#endif
-
-	/** Array index if the property is dynamic or static array. */
-	UPROPERTY()
-	int32 ArrayIndex = INDEX_NONE;
-
-	/** Type of the instanced struct or object reference by the property at the segment. This allows the path to be resolved when it points to a specific instance. */
-	UPROPERTY()
-	TObjectPtr<const UStruct> InstanceStruct = nullptr;
-};
-
-/**
- * Struct describing an indirection at specific segment at path.
- * Returned by FStateTreePropertyPath::ResolveIndirections() and FStateTreePropertyPath::ResolveIndirectionsWithValue().
- * Generally there's one indirection per FProperty. Containers have one path segment but two indirection (container + inner type).
- */
-struct FStateTreePropertyPathIndirection
+struct UE_DEPRECATED(5.6, "Use FPropertyBindingPathIndirection instead") FStateTreePropertyPathIndirection : FPropertyBindingPathIndirection
 {
-	FStateTreePropertyPathIndirection() = default;
-	explicit FStateTreePropertyPathIndirection(const UStruct* InContainerStruct)
-		: ContainerStruct(InContainerStruct)
-	{
-	}
-
-	const FProperty* GetProperty() const { return Property; }
-	const uint8* GetContainerAddress() const { return ContainerAddress; }
-	const UStruct* GetInstanceStruct() const { return InstanceStruct; }
-	const UStruct* GetContainerStruct() const { return ContainerStruct; }
-	int32 GetArrayIndex() const { return ArrayIndex; }
-	int32 GetPropertyOffset() const { return PropertyOffset; }
-	int32 GetPathSegmentIndex() const { return PathSegmentIndex; }
-	EStateTreePropertyAccessType GetAccessType() const { return AccessType; }
-	const uint8* GetPropertyAddress() const { return ContainerAddress + PropertyOffset; }
-
-#if WITH_EDITORONLY_DATA
-	FName GetRedirectedName() const { return RedirectedName; }
-	FGuid GetPropertyGuid() const { return PropertyGuid; }
-#endif
-	
-private:
-	/** Property at the indirection. */
-	const FProperty* Property = nullptr;
-	
-	/** Address of the container class/struct where the property belongs to. Only valid if created with ResolveIndirectionsWithValue() */
-	const uint8* ContainerAddress = nullptr;
-	
-	/** Type of the container class/struct. */
-	const UStruct* ContainerStruct = nullptr;
-	
-	/** Type of the instance class/struct of when AccessType is ObjectInstance or StructInstance. */
-	const UStruct* InstanceStruct = nullptr;
-
-#if WITH_EDITORONLY_DATA
-	/** Redirected name, if the give property name was not found but was reconciled using core redirect or property Guid. Requires ResolveIndirections/WithValue() to be called with bHandleRedirects = true. */
-	FName RedirectedName;
-
-	/** Guid of the property for Blueprint classes or User Defined Structs. Requires ResolveIndirections/WithValue() to be called with bHandleRedirects = true. */
-	FGuid PropertyGuid;
-#endif
-	
-	/** Array index for static and dynamic arrays. Note: static array indexing is baked in the PropertyOffset. */
-	int32 ArrayIndex = 0;
-	
-	/** Offset of the property relative to ContainerAddress. Includes static array indexing. */
-	int32 PropertyOffset = 0;
-	
-	/** Index of the path segment where indirection originated from. */
-	int32 PathSegmentIndex = 0;
-	
-	/** How to access the data through the indirection. */
-	EStateTreePropertyAccessType AccessType = EStateTreePropertyAccessType::Offset;
-
-	friend FStateTreePropertyPath;
 };
-	
 
 /**
  * Representation of a property path used for property binding in StateTree.
  */
 USTRUCT()
-struct STATETREEMODULE_API FStateTreePropertyPath
+struct STATETREEMODULE_API FStateTreePropertyPath : public FPropertyBindingPath
 {
 	GENERATED_BODY()
 
-	FStateTreePropertyPath() = default;
-
-#if WITH_EDITORONLY_DATA
-	explicit FStateTreePropertyPath(const FGuid InStructID)
-		: StructID(InStructID)
-	{
-	}
-
-	explicit FStateTreePropertyPath(const FGuid InStructID, const FName PropertyName)
-		: StructID(InStructID)
-	{
-		Segments.Emplace(PropertyName);
-	}
-	
-	explicit FStateTreePropertyPath(const FGuid InStructID, TConstArrayView<FStateTreePropertyPathSegment> InSegments)
-		: StructID(InStructID)
-		, Segments(InSegments)
-	{
-	}
-#endif // WITH_EDITORONLY_DATA
-	
-	/**
-	 * Parses path from string. The path should be in format: Foo.Bar[1].Baz
-	 * @param InPath Path string to parse
-	 * @return true if path was parsed successfully.
-	 */
-	bool FromString(const FStringView InPath);
-
-	/**
-	 * Returns the property path as a one string. Highlight allows to decorate a specific segment.
-	 * @param HighlightedSegment Index of the highlighted path segment
-	 * @param HighlightPrefix String to append before highlighted segment
-	 * @param HighlightPostfix String to append after highlighted segment
-	 * @param bOutputInstances if true, the instance struct types will be output. 
-	 * @param FirstSegment Index of the first path segment to be stringified.
-	 */
-	FString ToString(const int32 HighlightedSegment = INDEX_NONE, const TCHAR* HighlightPrefix = nullptr, const TCHAR* HighlightPostfix = nullptr, const bool bOutputInstances = false, const int32 FirstSegment = 0) const;
-
-	/**
-	 * Resolves the property path against base struct type. The path is assumed to be relative to the BaseStruct.
-	 * @param BaseStruct Base struct/class type the path is relative to.
-	 * @param OutIndirections Indirections describing how the properties were accessed. 
-	 * @param OutError Optional, pointer to string where error will be logged if update fails.
-	 * @param bHandleRedirects If true, the method will try to resolve missing properties using core redirects, and properties on Blueprint and User Defined Structs by ID. Available only in editor builds!
-	 * @return true of path could be resolved against the base value, and instance types were updated.
-	 */
-	bool ResolveIndirections(const UStruct* BaseStruct, TArray<FStateTreePropertyPathIndirection>& OutIndirections, FString* OutError = nullptr, bool bHandleRedirects = false) const;
-	
-	/**
-	 * Resolves the property path against base value. The path is assumed to be relative to the BaseValueView.
-	 * @param BaseValueView Base value the path is relative to.
-	 * @param OutIndirections Indirections describing how the properties were accessed. 
-	 * @param OutError Optional, pointer to string where error will be logged if update fails.
-	 * @param bHandleRedirects If true, the method will try to resolve missing properties using core redirects, and properties on Blueprint and User Defined Structs by ID. Available only in editor builds!
-	 * @return true of path could be resolved against the base value, and instance types were updated.
-	 */
-	bool ResolveIndirectionsWithValue(const FStateTreeDataView BaseValueView, TArray<FStateTreePropertyPathIndirection>& OutIndirections, FString* OutError = nullptr, bool bHandleRedirects = false) const;
-
-	/**
-	 * Updates property segments from base struct type. The path is expected to be relative to the BaseStruct.
-	 * The method handles renamed properties (core redirect, Blueprint and User Defined Structs by ID).
-	 * @param BaseValueView Base value the path is relative to.
-	 * @param OutError Optional, pointer to string where error will be logged if update fails.
-	 * @return true of path could be resolved against the base value, and instance types were updated.
-	 */
-	bool UpdateSegments(const UStruct* BaseStruct, FString* OutError = nullptr);
-
-	/**
-	 * Updates property segments from base value. The path is expected to be relative to the base value.
-	 * The method updates instance types, and handles renamed properties (core redirect, Blueprint and User Defined Structs by ID).
-	 * By storing the instance types on the path, we can resolve the path without the base value later.
-	 * @param BaseValueView Base value the path is relative to.
-	 * @param OutError Optional, pointer to string where error will be logged if update fails.
-	 * @return true of path could be resolved against the base value, and instance types were updated.
-	 */
-	bool UpdateSegmentsFromValue(const FStateTreeDataView BaseValueView, FString* OutError = nullptr);
-
-	UE_DEPRECATED(5.3, "Use UpdateSegmentsFromValue instead")
-	bool UpdateInstanceStructsFromValue(const FStateTreeDataView BaseValueView, FString* OutError = nullptr)
-	{
-		return UpdateSegmentsFromValue(BaseValueView, OutError);
-	}
-	
-	/** @return true if the path is empty. In that case the path points to the struct. */
-	bool IsPathEmpty() const { return Segments.IsEmpty(); }
-
-	/** @return true if any of the path segments is and indirection via instanced struct or object. */
-	bool HasAnyInstancedIndirection() const
-	{
-		return Segments.ContainsByPredicate([](const FStateTreePropertyPathSegment& Segment)
-		{
-			return Segment.GetInstanceStruct() != nullptr;
-		});
-	}
-
-	/** Reset the path to empty. */
-	void Reset()
-	{
-#if WITH_EDITORONLY_DATA
-		StructID = FGuid();
-#endif
-		Segments.Reset();
-	}
-
-#if WITH_EDITORONLY_DATA
-	const FGuid& GetStructID() const { return StructID; }
-	void SetStructID(const FGuid NewStructID) { StructID = NewStructID; }
-#endif // WITH_EDITORONLY_DATA
-
-	TConstArrayView<FStateTreePropertyPathSegment> GetSegments() const { return Segments; }
-	TArrayView<FStateTreePropertyPathSegment> GetMutableSegments() { return Segments; }
-	int32 NumSegments() const { return Segments.Num(); }
-	const FStateTreePropertyPathSegment& GetSegment(const int32 Index) const { return Segments[Index]; }
-
-	/** Adds a path segment to the path. */
-	void AddPathSegment(const FName InName, const int32 InArrayIndex = INDEX_NONE, const UStruct* InInstanceType = nullptr)
-	{
-		Segments.Emplace(InName, InArrayIndex, InInstanceType);
-	}
-
-	/** Adds a path segment to the path. */
-	void AddPathSegment(const FStateTreePropertyPathSegment& PathSegment)
-	{
-		Segments.Add(PathSegment);
-	}
-
-	/**
-	 * Test if this paths includes the provided path.
-	 * A path includes another one when they are == but this path can be longer.
-	 */
-	bool Includes(const FStateTreePropertyPath& Other) const;
-
-	/** Test if paths are equal. */
-	bool operator==(const FStateTreePropertyPath& RHS) const;
-
-private:
-#if WITH_EDITORONLY_DATA
-	/** ID of the struct this property path is relative to. */
-	UPROPERTY()
-	FGuid StructID;
-#endif // WITH_EDITORONLY_DATA
-
-	/** Path segments pointing to a specific property on the path. */
-	UPROPERTY()
-	TArray<FStateTreePropertyPathSegment> Segments;
+	using FPropertyBindingPath::FPropertyBindingPath;
 };
 
-
 USTRUCT()
 struct UE_DEPRECATED(5.3, "Use FStateTreePropertyPath instead.") STATETREEMODULE_API FStateTreeEditorPropertyPath
 {
@@ -592,8 +285,6 @@ FStateTreePropertyPathBinding(FConstStructView InFunctionNodeStruct, const FStat
 	}
 PRAGMA_ENABLE_DEPRECATION_WARNINGS
 
-	void PostSerialize(const FArchive& Ar);
-
 	const FStateTreePropertyPath& GetSourcePath() const { return SourcePropertyPath; }
 	const FStateTreePropertyPath& GetTargetPath() const { return TargetPropertyPath; }
 
@@ -639,15 +330,6 @@ PRAGMA_ENABLE_DEPRECATION_WARNINGS
 #endif // WITH_EDITORONLY_DATA
 };
 
-template<>
-struct TStructOpsTypeTraits<FStateTreePropertyPathBinding> : public TStructOpsTypeTraitsBase2<FStateTreePropertyPathBinding>
-{
-	enum 
-	{
-		WithPostSerialize = true,
-	};
-};
-
 using FStateTreeEditorPropertyBinding UE_DEPRECATED(5.3, "Deprecated struct. Please use FStateTreePropertyPathBinding instead.") = FStateTreePropertyPathBinding;
 
 /**
@@ -708,7 +390,7 @@ struct UE_DEPRECATED(5.3, "Use FStateTreePropertyPath instead.") STATETREEMODULE
 
 	/** Type of access/indirection. */
 	UPROPERTY()
-	EStateTreePropertyAccessType Type = EStateTreePropertyAccessType::Offset;
+	EPropertyBindingAccessType Type = EPropertyBindingAccessType::Offset;
 };
 
 
@@ -765,7 +447,7 @@ struct STATETREEMODULE_API FStateTreePropertyIndirection
 
 	/** Type of access/indirection. */
 	UPROPERTY()
-	EStateTreePropertyAccessType Type = EStateTreePropertyAccessType::Offset;
+	EPropertyBindingAccessType Type = EPropertyBindingAccessType::Offset;
 
 	/** Type of the struct or object instance in case the segment is pointing into an instanced data. */
 	UPROPERTY()
@@ -997,6 +679,10 @@ struct STATETREEMODULE_API FStateTreePropertyBindings
 	 * @param OutCopy Resulting copy type.
 	 * @return true if copy was resolved, or false if no copy could be resolved between paths.
 	 */
+	[[nodiscard]] static bool ResolveCopyType(const FPropertyBindingPathIndirection& SourceIndirection, const FPropertyBindingPathIndirection& TargetIndirection, FStateTreePropertyCopy& OutCopy);
+
+PRAGMA_DISABLE_DEPRECATION_WARNINGS
+	UE_DEPRECATED(5.6, "Use the overload taking FPropertyBindingPathIndirection instead")
 	[[nodiscard]] static bool ResolveCopyType(const FStateTreePropertyPathIndirection& SourceIndirection, const FStateTreePropertyPathIndirection& TargetIndirection, FStateTreePropertyCopy& OutCopy);
 
 	UE_DEPRECATED(5.3, "Should not be used, will be removed in a future version.")
@@ -1013,13 +699,15 @@ struct STATETREEMODULE_API FStateTreePropertyBindings
 
 	UE_DEPRECATED(5.5, "DebugPrintInternalLayout is deprecated. Use GetDebugInfo instead.")
 	void DebugPrintInternalLayout(FString& OutString) const;
+PRAGMA_ENABLE_DEPRECATION_WARNINGS
 
 #if WITH_EDITOR || WITH_STATETREE_DEBUG
 	[[nodiscard]] FString DebugInternalLayoutAsString() const;
 #endif
 
 private:
-	[[nodiscard]] bool ResolvePath(const UStruct* Struct, const FStateTreePropertyPath& Path, FStateTreePropertyIndirection& OutFirstIndirection, FStateTreePropertyPathIndirection& OutLeafIndirection);
+	[[nodiscard]] bool ResolvePath(const UStruct* Struct, const FStateTreePropertyPath& Path, FStateTreePropertyIndirection& OutFirstIndirection, FPropertyBindingPathIndirection& OutLeafIndirection);
+
 	const FStateTreeBindableStructDesc* GetSourceDescByHandle(const FStateTreeDataHandle SourceDataHandle);
 
 	void PerformCopy(const FStateTreePropertyCopy& Copy, uint8* SourceAddress, uint8* TargetAddress) const;
@@ -1057,8 +745,11 @@ private:
 	/** Flag indicating if the properties has been resolved successfully . */
 	bool bBindingsResolved = false;
 
-#if WITH_EDITORONLY_DATA
 PRAGMA_DISABLE_DEPRECATION_WARNINGS
+	UE_DEPRECATED(5.6, "Use the overload taking FPropertyBindingPathIndirection instead")
+	[[nodiscard]] bool ResolvePath(const UStruct* Struct, const FStateTreePropertyPath& Path, FStateTreePropertyIndirection& OutFirstIndirection, FStateTreePropertyPathIndirection& OutLeafIndirection);
+
+#if WITH_EDITORONLY_DATA
 	
 	UE_DEPRECATED_FORGAME(5.3, "Use PropertyPathBindings instead.")
 	UPROPERTY()
@@ -1067,9 +758,8 @@ PRAGMA_DISABLE_DEPRECATION_WARNINGS
 	UE_DEPRECATED_FORGAME(5.3, "Use PropertyPathBindings instead.")
 	UPROPERTY()
 	TArray<FStateTreePropertySegment> PropertySegments_DEPRECATED;
-	
-PRAGMA_ENABLE_DEPRECATION_WARNINGS
 #endif // WITH_EDITORONLY_DATA
+PRAGMA_ENABLE_DEPRECATION_WARNINGS
 
 	friend FStateTreePropertyBindingCompiler;
 	friend UStateTree;
diff --git a/Engine/Plugins/Runtime/StateTree/Source/StateTreeModule/Public/StateTreePropertyRefHelpers.h b/Engine/Plugins/Runtime/StateTree/Source/StateTreeModule/Public/StateTreePropertyRefHelpers.h
--- a/Engine/Plugins/Runtime/StateTree/Source/StateTreeModule/Public/StateTreePropertyRefHelpers.h
+++ b/Engine/Plugins/Runtime/StateTree/Source/StateTreeModule/Public/StateTreePropertyRefHelpers.h
@@ -11,6 +11,7 @@
 #endif
 
 class FProperty;
+struct FPropertyBindingPathIndirection;
 struct FBindingChainElement;
 struct FStateTreeBindableStructDesc;
 struct FStateTreePropertyPathIndirection;
@@ -26,7 +27,7 @@ namespace UE::StateTree::PropertyRefHelpers
 	STATETREEMODULE_API extern const FName RefTypeName;
 	/**
 	 * @param RefProperty Property of PropertyRef type.
-	 * @param SourceProperty Property to check it's type compatibility.
+	 * @param SourceProperty Property to check its type compatibility.
 	 * @param PropertyRefAddress Address of PropertyRef
 	 * @param SourceAddress Address of checked property value.
 	 * @return true if SourceProperty type is compatible with PropertyRef.
@@ -38,6 +39,9 @@ namespace UE::StateTree::PropertyRefHelpers
 	 * @param SourceStruct Bindable owner of referenced property.
 	 * @return true if property can be referenced by PropertyRef.
 	 */
+	bool STATETREEMODULE_API IsPropertyAccessibleForPropertyRef(TConstArrayView<FPropertyBindingPathIndirection> SourcePropertyPathIndirections, FStateTreeBindableStructDesc SourceStruct);
+
+	UE_DEPRECATED(5.6, "Use the overload taking FPropertyBindingPathIndirection instead")
 	bool STATETREEMODULE_API IsPropertyAccessibleForPropertyRef(TConstArrayView<FStateTreePropertyPathIndirection> SourcePropertyPathIndirections, FStateTreeBindableStructDesc SourceStruct);
 
 	/**
diff --git a/Engine/Plugins/Runtime/StateTree/Source/StateTreeModule/Public/StateTreeTypes.h b/Engine/Plugins/Runtime/StateTree/Source/StateTreeModule/Public/StateTreeTypes.h
--- a/Engine/Plugins/Runtime/StateTree/Source/StateTreeModule/Public/StateTreeTypes.h
+++ b/Engine/Plugins/Runtime/StateTree/Source/StateTreeModule/Public/StateTreeTypes.h
@@ -2,6 +2,7 @@
 
 #pragma once
 
+#include "PropertyBindingPath.h"
 #include "StructUtils/PropertyBag.h"
 #include "GameplayTagContainer.h"
 #include "StateTreeIndexTypes.h"
@@ -990,171 +991,12 @@ protected:
 
 
 /**
- * Short lived pointer to an UOBJECT() or USTRUCT().
- * The data view expects a type (UStruct) when you pass in a valid memory. In case of null, the type can be empty too.
+ * Helper struct to facilitate transition to FPropertyBindingDataView in the StateTree module.
+ * @see FPropertyBindingDataView
  */
-struct STATETREEMODULE_API FStateTreeDataView
+struct STATETREEMODULE_API FStateTreeDataView : public FPropertyBindingDataView
 {
-	FStateTreeDataView() = default;
-
-	// USTRUCT() constructor.
-	FStateTreeDataView(const UStruct* InStruct, uint8* InMemory) : Struct(InStruct), Memory(InMemory)
-	{
-		// Must have type with valid pointer.
-		check(!Memory || (Memory && Struct));
-	}
-
-	// UOBJECT() constructor.
-	FStateTreeDataView(UObject* Object) : Struct(Object ? Object->GetClass() : nullptr), Memory(reinterpret_cast<uint8*>(Object))
-	{
-		// Must have type with valid pointer.
-		check(!Memory || (Memory && Struct));
-	}
-
-	// USTRUCT() from a StructView.
-	FStateTreeDataView(FStructView StructView) : Struct(StructView.GetScriptStruct()), Memory(StructView.GetMemory())
-	{
-		// Must have type with valid pointer.
-		check(!Memory || (Memory && Struct));
-	}
-
-	/**
-	 * Check is the view is valid (both pointer and type are set). On valid views it is safe to call the Get<>() methods returning a reference.
-	 * @return True if the view is valid.
-	*/
-	bool IsValid() const
-	{
-		return Memory != nullptr && Struct != nullptr;
-	}
-
-	/*
-	 * UOBJECT() getters (reference & pointer, const & mutable)
-	 */
-	template <typename T>
-    typename TEnableIf<TIsDerivedFrom<T, UObject>::IsDerived, const T&>::Type Get() const
-	{
-		check(Memory != nullptr);
-		check(Struct != nullptr);
-		check(Struct->IsChildOf(T::StaticClass()));
-		return *((T*)Memory);
-	}
-
-	template <typename T>
-	typename TEnableIf<TIsDerivedFrom<T, UObject>::IsDerived, T&>::Type GetMutable() const
-	{
-		check(Memory != nullptr);
-		check(Struct != nullptr);
-		check(Struct->IsChildOf(T::StaticClass()));
-		return *((T*)Memory);
-	}
-
-	template <typename T>
-	typename TEnableIf<TIsDerivedFrom<T, UObject>::IsDerived, const T*>::Type GetPtr() const
-	{
-		// If Memory is set, expect Struct too. Otherwise, let nulls pass through.
-		check(!Memory || (Memory && Struct));
-		check(!Struct || Struct->IsChildOf(T::StaticClass()));
-		return ((T*)Memory);
-	}
-
-	template <typename T>
-    typename TEnableIf<TIsDerivedFrom<T, UObject>::IsDerived, T*>::Type GetMutablePtr() const
-	{
-		// If Memory is set, expect Struct too. Otherwise, let nulls pass through.
-		check(!Memory || (Memory && Struct));
-		check(!Struct || Struct->IsChildOf(T::StaticClass()));
-		return ((T*)Memory);
-	}
-
-	/*
-	 * USTRUCT() getters (reference & pointer, const & mutable)
-	 */
-	template <typename T>
-	typename TEnableIf<!TIsDerivedFrom<T, UObject>::IsDerived && !TIsIInterface<T>::Value, const T&>::Type Get() const
-	{
-		check(Memory != nullptr);
-		check(Struct != nullptr);
-		check(Struct->IsChildOf(T::StaticStruct()));
-		return *((T*)Memory);
-	}
-
-	template <typename T>
-    typename TEnableIf<!TIsDerivedFrom<T, UObject>::IsDerived && !TIsIInterface<T>::Value, T&>::Type GetMutable() const
-	{
-		check(Memory != nullptr);
-		check(Struct != nullptr);
-		check(Struct->IsChildOf(T::StaticStruct()));
-		return *((T*)Memory);
-	}
-
-	template <typename T>
-    typename TEnableIf<!TIsDerivedFrom<T, UObject>::IsDerived && !TIsIInterface<T>::Value, const T*>::Type GetPtr() const
-	{
-		// If Memory is set, expect Struct too. Otherwise, let nulls pass through.
-		check(!Memory || (Memory && Struct));
-		check(!Struct || Struct->IsChildOf(T::StaticStruct()));
-		return ((T*)Memory);
-	}
-
-	template <typename T>
-    typename TEnableIf<!TIsDerivedFrom<T, UObject>::IsDerived && !TIsIInterface<T>::Value, T*>::Type GetMutablePtr() const
-	{
-		// If Memory is set, expect Struct too. Otherwise, let nulls pass through.
-		check(!Memory || (Memory && Struct));
-		check(!Struct || Struct->IsChildOf(T::StaticStruct()));
-		return ((T*)Memory);
-	}
-
-	/*
-	 * IInterface() getters (reference & pointer, const & mutable)
-	 */
-	template <typename T>
-	typename TEnableIf<TIsIInterface<T>::Value, const T&>::Type Get() const
-	{
-		check(!Memory || (Memory && Struct));
-		check(Struct->IsChildOf(UObject::StaticClass()) && ((UClass*)Struct)->ImplementsInterface(T::UClassType::StaticClass()));
-		return *(T*)((UObject*)Memory)->GetInterfaceAddress(T::UClassType::StaticClass());
-	}
-
-	template <typename T>
-    typename TEnableIf<TIsIInterface<T>::Value, T&>::Type GetMutable() const
-	{
-		check(!Memory || (Memory && Struct));
-		check(Struct->IsChildOf(UObject::StaticClass()) && ((UClass*)Struct)->ImplementsInterface(T::UClassType::StaticClass()));
-		return *(T*)((UObject*)Memory)->GetInterfaceAddress(T::UClassType::StaticClass());
-	}
-
-	template <typename T>
-    typename TEnableIf<TIsIInterface<T>::Value, const T*>::Type GetPtr() const
-	{
-		check(!Memory || (Memory && Struct));
-		check(Struct->IsChildOf(UObject::StaticClass()) && ((UClass*)Struct)->ImplementsInterface(T::UClassType::StaticClass()));
-		return (T*)((UObject*)Memory)->GetInterfaceAddress(T::UClassType::StaticClass());
-	}
-
-	template <typename T>
-    typename TEnableIf<TIsIInterface<T>::Value, T*>::Type GetMutablePtr() const
-	{
-		check(!Memory || (Memory && Struct));
-		check(Struct->IsChildOf(UObject::StaticClass()) && ((UClass*)Struct)->ImplementsInterface(T::UClassType::StaticClass()));
-		return (T*)((UObject*)Memory)->GetInterfaceAddress(T::UClassType::StaticClass());
-	}
-
-	/** @return Struct describing the data type. */
-	const UStruct* GetStruct() const { return Struct; }
-
-	/** @return Raw const pointer to the data. */
-	const uint8* GetMemory() const { return Memory; }
-
-	/** @return Raw mutable pointer to the data. */
-	uint8* GetMutableMemory() const { return Memory; }
-	
-protected:
-	/** UClass or UScriptStruct of the data. */
-	const UStruct* Struct = nullptr;
-
-	/** Memory pointing at the class or struct */
-	uint8* Memory = nullptr;
+	using FPropertyBindingDataView::FPropertyBindingDataView;
 };
 
 /**
diff --git a/Engine/Plugins/Runtime/StateTree/Source/StateTreeModule/StateTreeModule.Build.cs b/Engine/Plugins/Runtime/StateTree/Source/StateTreeModule/StateTreeModule.Build.cs
--- a/Engine/Plugins/Runtime/StateTree/Source/StateTreeModule/StateTreeModule.Build.cs
+++ b/Engine/Plugins/Runtime/StateTree/Source/StateTreeModule/StateTreeModule.Build.cs
@@ -20,7 +20,8 @@ namespace UnrealBuildTool.Rules
 					"DeveloperSettings",
 					"Engine",
 					"AIModule",
-					"GameplayTags"
+					"GameplayTags",
+					"PropertyBindingUtils"
 				}
 			);
 
diff --git a/Engine/Plugins/Runtime/StateTree/Source/StateTreeTestSuite/Private/StateTreeBindingTest.cpp b/Engine/Plugins/Runtime/StateTree/Source/StateTreeTestSuite/Private/StateTreeBindingTest.cpp
--- a/Engine/Plugins/Runtime/StateTree/Source/StateTreeTestSuite/Private/StateTreeBindingTest.cpp
+++ b/Engine/Plugins/Runtime/StateTree/Source/StateTreeTestSuite/Private/StateTreeBindingTest.cpp
@@ -115,7 +115,7 @@ struct FStateTreeTest_PropertyFunctions : FStateTreeTestBase
 		UStateTree& StateTree = NewStateTree();
 		UStateTreeEditorData& EditorData = *Cast<UStateTreeEditorData>(StateTree.EditorData);
 		UStateTreeState& Root = EditorData.AddSubTree(FName(TEXT("Root")));
-		FStateTreePropertyPathSegment PathSegmentToFuncResult = FStateTreePropertyPathSegment(TEXT("Result"));
+		FPropertyBindingPathSegment PathSegmentToFuncResult = FPropertyBindingPathSegment(TEXT("Result"));
 
 		// Condition with property function binding.
 		{
diff --git a/Engine/Plugins/Runtime/StateTree/Source/StateTreeTestSuite/Private/StateTreePropertyPathTest.cpp b/Engine/Plugins/Runtime/StateTree/Source/StateTreeTestSuite/Private/StateTreePropertyPathTest.cpp
--- a/Engine/Plugins/Runtime/StateTree/Source/StateTreeTestSuite/Private/StateTreePropertyPathTest.cpp
+++ b/Engine/Plugins/Runtime/StateTree/Source/StateTreeTestSuite/Private/StateTreePropertyPathTest.cpp
@@ -4,9 +4,7 @@
 #include "StateTreeTestBase.h"
 #include "StateTreeTestTypes.h"
 
-#include "StateTreeCompilerLog.h"
 #include "StateTreeEditorData.h"
-#include "StateTreeCompiler.h"
 #include "Conditions/StateTreeCommonConditions.h"
 
 #define LOCTEXT_NAMESPACE "AITestSuite_StateTreeTest"
@@ -27,15 +25,15 @@ struct FStateTreeTest_PropertyPathOffset : FStateTreeTestBase
 		AITEST_EQUAL("Should have 2 path segments", Path.NumSegments(), 2);
 
 		FString ResolveErrors;
-		TArray<FStateTreePropertyPathIndirection> Indirections;
+		TArray<FPropertyBindingPathIndirection> Indirections;
 		const bool bResolveResult = Path.ResolveIndirections(FStateTreeTest_PropertyStruct::StaticStruct(), Indirections, &ResolveErrors);
 
 		AITEST_TRUE("Resolve path should succeeed", bResolveResult);
 		AITEST_EQUAL("Should have no resolve errors", ResolveErrors.Len(), 0);
 		
 		AITEST_EQUAL("Should have 2 indirections", Indirections.Num(), 2);
-		AITEST_EQUAL("Indirection 0 should be Offset type", Indirections[0].GetAccessType(), EStateTreePropertyAccessType::Offset);
-		AITEST_EQUAL("Indirection 1 should be Offset type", Indirections[1].GetAccessType(), EStateTreePropertyAccessType::Offset);
+		AITEST_EQUAL("Indirection 0 should be Offset type", Indirections[0].GetAccessType(), EPropertyBindingAccessType::Offset);
+		AITEST_EQUAL("Indirection 1 should be Offset type", Indirections[1].GetAccessType(), EPropertyBindingAccessType::Offset);
 
 		return true;
 	}
@@ -92,7 +90,7 @@ struct FStateTreeTest_PropertyPathOffsetFail : FStateTreeTestBase
 		AITEST_EQUAL("Should have 2 path segments", Path.NumSegments(), 2);
 
 		FString ResolveErrors;
-		TArray<FStateTreePropertyPathIndirection> Indirections;
+		TArray<FPropertyBindingPathIndirection> Indirections;
 		const bool bResolveResult = Path.ResolveIndirections(FStateTreeTest_PropertyStruct::StaticStruct(), Indirections, &ResolveErrors);
 
 		AITEST_FALSE("Resolve path should not succeeed", bResolveResult);
@@ -144,12 +142,12 @@ struct FStateTreeTest_PropertyPathWrongObject : FStateTreeTestBase
 		Object->InstancedObject = NewObject<UStateTreeTest_PropertyObjectInstancedWithB>();
 		{
 			FString ResolveErrors;
-			TArray<FStateTreePropertyPathIndirection> Indirections;
+			TArray<FPropertyBindingPathIndirection> Indirections;
 			const bool bResolveResult = Path.ResolveIndirectionsWithValue(FStateTreeDataView(Object), Indirections, &ResolveErrors);
 
 			AITEST_TRUE("Resolve path should succeeed", bResolveResult);
 			AITEST_EQUAL("Should have 2 indirections", Indirections.Num(), 2);
-			AITEST_TRUE("Object ", Indirections[0].GetAccessType() == EStateTreePropertyAccessType::ObjectInstance);
+			AITEST_TRUE("Object ", Indirections[0].GetAccessType() == EPropertyBindingAccessType::ObjectInstance);
 			AITEST_TRUE("Object ", Indirections[0].GetContainerStruct() == Object->GetClass());
 			AITEST_TRUE("Object ", Indirections[0].GetInstanceStruct() == UStateTreeTest_PropertyObjectInstancedWithB::StaticClass());
 			AITEST_EQUAL("Should not have error", ResolveErrors.Len(), 0);
@@ -158,7 +156,7 @@ struct FStateTreeTest_PropertyPathWrongObject : FStateTreeTestBase
 		Object->InstancedObject = NewObject<UStateTreeTest_PropertyObjectInstanced>();
 		{
 			FString ResolveErrors;
-			TArray<FStateTreePropertyPathIndirection> Indirections;
+			TArray<FPropertyBindingPathIndirection> Indirections;
 			const bool bResolveResult = Path.ResolveIndirectionsWithValue(FStateTreeDataView(Object), Indirections, &ResolveErrors);
 
 			AITEST_FALSE("Resolve path should fail", bResolveResult);
@@ -186,14 +184,14 @@ struct FStateTreeTest_PropertyPathArray : FStateTreeTestBase
 		Object->ArrayOfInts.Add(123);
 
 		FString ResolveErrors;
-		TArray<FStateTreePropertyPathIndirection> Indirections;
+		TArray<FPropertyBindingPathIndirection> Indirections;
 		const bool bResolveResult = Path.ResolveIndirectionsWithValue(FStateTreeDataView(Object), Indirections, &ResolveErrors);
 
 		AITEST_TRUE("Resolve path should succeeed", bResolveResult);
 		AITEST_EQUAL("Should have no resolve errors", ResolveErrors.Len(), 0);
 		AITEST_EQUAL("Should have 2 indirections", Indirections.Num(), 2);
-		AITEST_EQUAL("Indirection 0 should be IndexArray type", Indirections[0].GetAccessType(), EStateTreePropertyAccessType::IndexArray);
-		AITEST_EQUAL("Indirection 1 should be Offset type", Indirections[1].GetAccessType(), EStateTreePropertyAccessType::Offset);
+		AITEST_EQUAL("Indirection 0 should be IndexArray type", Indirections[0].GetAccessType(), EPropertyBindingAccessType::IndexArray);
+		AITEST_EQUAL("Indirection 1 should be Offset type", Indirections[1].GetAccessType(), EPropertyBindingAccessType::Offset);
 
 		const int32 Value = *reinterpret_cast<const int32*>(Indirections[1].GetPropertyAddress());
 		AITEST_EQUAL("Value should be 123", Value, 123);
@@ -218,7 +216,7 @@ struct FStateTreeTest_PropertyPathArrayInvalidIndex : FStateTreeTestBase
 		Object->ArrayOfInts.Add(123);
 
 		FString ResolveErrors;
-		TArray<FStateTreePropertyPathIndirection> Indirections;
+		TArray<FPropertyBindingPathIndirection> Indirections;
 		const bool bResolveResult = Path.ResolveIndirectionsWithValue(FStateTreeDataView(Object), Indirections, &ResolveErrors);
 
 		AITEST_FALSE("Resolve path should fail", bResolveResult);
@@ -245,15 +243,15 @@ struct FStateTreeTest_PropertyPathArrayOfStructs : FStateTreeTestBase
 
 		{
 			FString ResolveErrors;
-			TArray<FStateTreePropertyPathIndirection> Indirections;
+			TArray<FPropertyBindingPathIndirection> Indirections;
 			const bool bResolveResult = Path1.ResolveIndirectionsWithValue(FStateTreeDataView(Object), Indirections, &ResolveErrors);
 
 			AITEST_TRUE("Resolve path1 should succeeed", bResolveResult);
 			AITEST_EQUAL("Should have no resolve errors", ResolveErrors.Len(), 0);
 			AITEST_EQUAL("Should have 3 indirections", Indirections.Num(), 3);
-			AITEST_EQUAL("Indirection 0 should be ArrayIndex type", Indirections[0].GetAccessType(), EStateTreePropertyAccessType::IndexArray);
-			AITEST_EQUAL("Indirection 1 should be Offset type", Indirections[1].GetAccessType(), EStateTreePropertyAccessType::Offset);
-			AITEST_EQUAL("Indirection 2 should be Offset type", Indirections[2].GetAccessType(), EStateTreePropertyAccessType::Offset);
+			AITEST_EQUAL("Indirection 0 should be ArrayIndex type", Indirections[0].GetAccessType(), EPropertyBindingAccessType::IndexArray);
+			AITEST_EQUAL("Indirection 1 should be Offset type", Indirections[1].GetAccessType(), EPropertyBindingAccessType::Offset);
+			AITEST_EQUAL("Indirection 2 should be Offset type", Indirections[2].GetAccessType(), EPropertyBindingAccessType::Offset);
 
 			const int32 Value = *reinterpret_cast<const int32*>(Indirections[2].GetPropertyAddress());
 			AITEST_EQUAL("Value should be 3", Value, 3);
@@ -261,16 +259,16 @@ struct FStateTreeTest_PropertyPathArrayOfStructs : FStateTreeTestBase
 
 		{
 			FString ResolveErrors;
-			TArray<FStateTreePropertyPathIndirection> Indirections;
+			TArray<FPropertyBindingPathIndirection> Indirections;
 			const bool bResolveResult = Path2.ResolveIndirectionsWithValue(FStateTreeDataView(Object), Indirections, &ResolveErrors);
 
 			AITEST_TRUE("Resolve path2 should succeeed", bResolveResult);
 			AITEST_EQUAL("Should have no resolve errors", ResolveErrors.Len(), 0);
 			AITEST_EQUAL("Should have 4 indirections", Indirections.Num(), 4);
-			AITEST_EQUAL("Indirection 0 should be ArrayIndex type", Indirections[0].GetAccessType(), EStateTreePropertyAccessType::IndexArray);
-			AITEST_EQUAL("Indirection 1 should be Offset type", Indirections[1].GetAccessType(), EStateTreePropertyAccessType::Offset);
-			AITEST_EQUAL("Indirection 2 should be Offset type", Indirections[2].GetAccessType(), EStateTreePropertyAccessType::Offset);
-			AITEST_EQUAL("Indirection 3 should be Offset type", Indirections[3].GetAccessType(), EStateTreePropertyAccessType::Offset);
+			AITEST_EQUAL("Indirection 0 should be ArrayIndex type", Indirections[0].GetAccessType(), EPropertyBindingAccessType::IndexArray);
+			AITEST_EQUAL("Indirection 1 should be Offset type", Indirections[1].GetAccessType(), EPropertyBindingAccessType::Offset);
+			AITEST_EQUAL("Indirection 2 should be Offset type", Indirections[2].GetAccessType(), EPropertyBindingAccessType::Offset);
+			AITEST_EQUAL("Indirection 3 should be Offset type", Indirections[3].GetAccessType(), EPropertyBindingAccessType::Offset);
 
 			const int32 Value = *reinterpret_cast<const int32*>(Indirections[3].GetPropertyAddress());
 			AITEST_EQUAL("Value should be 42", Value, 42);
@@ -301,28 +299,28 @@ struct FStateTreeTest_PropertyPathArrayOfInstancedObjects : FStateTreeTestBase
 
 		{
 			FString ResolveErrors;
-			TArray<FStateTreePropertyPathIndirection> Indirections;
+			TArray<FPropertyBindingPathIndirection> Indirections;
 			const bool bResolveResult = Path.ResolveIndirections(UStateTreeTest_PropertyObject::StaticClass(), Indirections, &ResolveErrors);
 
 			AITEST_TRUE("Resolve path should succeeed", bResolveResult);
 			AITEST_EQUAL("Should have no resolve errors", ResolveErrors.Len(), 0);
 			AITEST_EQUAL("Should have 3 indirections", Indirections.Num(), 3);
-			AITEST_EQUAL("Indirection 0 should be ArrayIndex type", Indirections[0].GetAccessType(), EStateTreePropertyAccessType::IndexArray);
-			AITEST_EQUAL("Indirection 1 should be StructInstance type", Indirections[1].GetAccessType(), EStateTreePropertyAccessType::StructInstance);
-			AITEST_EQUAL("Indirection 2 should be Offset type", Indirections[2].GetAccessType(), EStateTreePropertyAccessType::Offset);
+			AITEST_EQUAL("Indirection 0 should be ArrayIndex type", Indirections[0].GetAccessType(), EPropertyBindingAccessType::IndexArray);
+			AITEST_EQUAL("Indirection 1 should be StructInstance type", Indirections[1].GetAccessType(), EPropertyBindingAccessType::StructInstance);
+			AITEST_EQUAL("Indirection 2 should be Offset type", Indirections[2].GetAccessType(), EPropertyBindingAccessType::Offset);
 		}
 
 		{
 			FString ResolveErrors;
-			TArray<FStateTreePropertyPathIndirection> Indirections;
+			TArray<FPropertyBindingPathIndirection> Indirections;
 			const bool bResolveResult = Path.ResolveIndirectionsWithValue(FStateTreeDataView(Object), Indirections, &ResolveErrors);
 
 			AITEST_TRUE("Resolve path should succeeed", bResolveResult);
 			AITEST_EQUAL("Should have no resolve errors", ResolveErrors.Len(), 0);
 			AITEST_EQUAL("Should have 3 indirections", Indirections.Num(), 3);
-			AITEST_EQUAL("Indirection 0 should be ArrayIndex type", Indirections[0].GetAccessType(), EStateTreePropertyAccessType::IndexArray);
-			AITEST_EQUAL("Indirection 1 should be StructInstance type", Indirections[1].GetAccessType(), EStateTreePropertyAccessType::StructInstance);
-			AITEST_EQUAL("Indirection 2 should be Offset type", Indirections[2].GetAccessType(), EStateTreePropertyAccessType::Offset);
+			AITEST_EQUAL("Indirection 0 should be ArrayIndex type", Indirections[0].GetAccessType(), EPropertyBindingAccessType::IndexArray);
+			AITEST_EQUAL("Indirection 1 should be StructInstance type", Indirections[1].GetAccessType(), EPropertyBindingAccessType::StructInstance);
+			AITEST_EQUAL("Indirection 2 should be Offset type", Indirections[2].GetAccessType(), EPropertyBindingAccessType::Offset);
 
 			const int32 Value = *reinterpret_cast<const int32*>(Indirections[2].GetPropertyAddress());
 			AITEST_EQUAL("Value should be 123", Value, 123);
diff --git a/Engine/Plugins/Runtime/StateTree/Source/StateTreeTestSuite/StateTreeTestSuite.Build.cs b/Engine/Plugins/Runtime/StateTree/Source/StateTreeTestSuite/StateTreeTestSuite.Build.cs
--- a/Engine/Plugins/Runtime/StateTree/Source/StateTreeTestSuite/StateTreeTestSuite.Build.cs
+++ b/Engine/Plugins/Runtime/StateTree/Source/StateTreeTestSuite/StateTreeTestSuite.Build.cs
@@ -25,6 +25,7 @@ namespace UnrealBuildTool.Rules
 					"StateTreeEditorModule",
 					"AITestSuite",
 					"GameplayTags",
+					"PropertyBindingUtils"
 				}
 			);
 
diff --git a/Engine/Plugins/Runtime/StateTree/StateTree.uplugin b/Engine/Plugins/Runtime/StateTree/StateTree.uplugin
--- a/Engine/Plugins/Runtime/StateTree/StateTree.uplugin
+++ b/Engine/Plugins/Runtime/StateTree/StateTree.uplugin
@@ -35,6 +35,10 @@
 		{
 			"Name": "PropertyAccessEditor",
 			"Enabled": true
+		},
+		{
+			"Name": "PropertyBindingUtils",
+			"Enabled": true
 		}
 	]
 }
\ No newline at end of file
