[StateTree] deprecation cleanup for headers, methods, non-reflected properties and reflected properties not used in deprecation code of exposed to blueprint

diff --git a/Engine/Plugins/Runtime/StateTree/Source/StateTreeEditorModule/Private/Customizations/StateTreeBindingExtension.h b/Engine/Plugins/Runtime/StateTree/Source/StateTreeEditorModule/Private/Customizations/StateTreeBindingExtension.h
--- a/Engine/Plugins/Runtime/StateTree/Source/StateTreeEditorModule/Private/Customizations/StateTreeBindingExtension.h
+++ b/Engine/Plugins/Runtime/StateTree/Source/StateTreeEditorModule/Private/Customizations/StateTreeBindingExtension.h
@@ -10,7 +10,6 @@ enum class EStateTreePropertyUsage : uint8;
 class IPropertyHandle;
 class IDetailLayoutBuilder;
 class IPropertyAccessEditor;
-struct FStateTreeEditorPropertyPath;
 struct FStateTreePropertyPath;
 class FProperty;
 
diff --git a/Engine/Plugins/Runtime/StateTree/Source/StateTreeEditorModule/Private/StateTreeEditorData.cpp b/Engine/Plugins/Runtime/StateTree/Source/StateTreeEditorModule/Private/StateTreeEditorData.cpp
--- a/Engine/Plugins/Runtime/StateTree/Source/StateTreeEditorModule/Private/StateTreeEditorData.cpp
+++ b/Engine/Plugins/Runtime/StateTree/Source/StateTreeEditorModule/Private/StateTreeEditorData.cpp
@@ -692,20 +692,6 @@ UStateTreeState* UStateTreeEditorData::GetMutableStateByID(const FGuid StateID)
 	return Result;
 }
 
-void UStateTreeEditorData::GetAllStructIDs(TMap<FGuid, const UStruct*>& AllStructs) const
-{
-	AllStructs.Reset();
-
-	const UStateTree* StateTree = GetTypedOuter<UStateTree>();
-	checkf(StateTree, TEXT("UStateTreeEditorData should only be allocated within a UStateTree"));
-
-	VisitAllNodes([&AllStructs](const UStateTreeState* State, const FStateTreeBindableStructDesc& Desc, const FStateTreeDataView Value)
-		{
-			AllStructs.Emplace(Desc.ID, Desc.Struct);
-			return EStateTreeVisitor::Continue;
-		});
-}
-
 void UStateTreeEditorData::GetAllStructValues(TMap<FGuid, const FStateTreeDataView>& AllValues) const
 {
 	AllValues.Reset();
@@ -1146,117 +1132,6 @@ void UStateTreeEditorData::CallPostLoadOnNodes()
 	}
 }
 
-EStateTreeVisitor UStateTreeEditorData::VisitStateNodes(const UStateTreeState& State, TFunctionRef<EStateTreeVisitor(const UStateTreeState* State, const FGuid& ID, const FName& Name, const EStateTreeNodeType NodeType, const UScriptStruct* NodeStruct, const UStruct* InstanceStruct)> InFunc) const
-{
-	bool bContinue = true;
-
-	auto VisitFuncNodes = [&InFunc, &State, this](const FGuid StructID, const FName NodeName)
-	{
-		const FString StatePath = FString::Printf(TEXT("%s/%s"), *State.GetPath(), *NodeName.ToString());
-		return VisitStructBoundPropertyFunctions(StructID, StatePath, [&InFunc, &State](const FStateTreeEditorNode& EditorNode, const FStateTreeBindableStructDesc& Desc, FStateTreeDataView Value)
-		{
-			return InFunc(&State, EditorNode.ID, EditorNode.GetName(), EStateTreeNodeType::PropertyFunction, EditorNode.Node.GetScriptStruct(), EditorNode.GetInstance().GetStruct());
-		});
-	};
-
-	if (bContinue)
-	{
-		// Enter conditions
-		for (const FStateTreeEditorNode& Node : State.EnterConditions)
-		{
-			if (VisitFuncNodes(Node.ID, Node.GetName()) == EStateTreeVisitor::Break)
-			{
-				bContinue = false;
-				break;
-			}
-			else if (const FStateTreeConditionBase* Cond = Node.Node.GetPtr<FStateTreeConditionBase>())
-			{
-				if (InFunc(&State, Node.ID, Node.Node.GetScriptStruct()->GetFName(), EStateTreeNodeType::EnterCondition, Node.Node.GetScriptStruct(), Cond->GetInstanceDataType()) == EStateTreeVisitor::Break)
-				{
-					bContinue = false;
-					break;
-				}
-			}
-		}
-	}
-	if (bContinue)
-	{
-		// Tasks
-		for (const FStateTreeEditorNode& Node : State.Tasks)
-		{
-			if (VisitFuncNodes(Node.ID, Node.GetName()) == EStateTreeVisitor::Break)
-			{
-				bContinue = false;
-				break;
-			}
-			else if (const FStateTreeTaskBase* Task = Node.Node.GetPtr<FStateTreeTaskBase>())
-			{
-				if (InFunc(&State, Node.ID, Task->Name, EStateTreeNodeType::Task, Node.Node.GetScriptStruct(), Task->GetInstanceDataType()) == EStateTreeVisitor::Break)
-				{
-					bContinue = false;
-					break;
-				}
-			}
-		}
-	}
-	if (bContinue)
-	{
-		if (const FStateTreeTaskBase* Task = State.SingleTask.Node.GetPtr<FStateTreeTaskBase>())
-		{
-			if (VisitFuncNodes(State.SingleTask.ID, State.SingleTask.GetName()) == EStateTreeVisitor::Break)
-			{
-				bContinue = false;
-			}
-			else if (InFunc(&State, State.SingleTask.ID, Task->Name, EStateTreeNodeType::Task, State.SingleTask.Node.GetScriptStruct(), Task->GetInstanceDataType()) == EStateTreeVisitor::Break)
-			{
-				bContinue = false;
-			}
-		}
-
-	}
-	if (bContinue)
-	{
-		// Transitions
-		for (const FStateTreeTransition& Transition : State.Transitions)
-		{
-			for (const FStateTreeEditorNode& Node : Transition.Conditions)
-			{
-				if (VisitFuncNodes(Node.ID, Node.GetName()) == EStateTreeVisitor::Break)
-				{
-					bContinue = false;
-					break;
-				}
-				else if (const FStateTreeConditionBase* Cond = Node.Node.GetPtr<FStateTreeConditionBase>())
-				{
-					if (InFunc(&State, Node.ID, Node.Node.GetScriptStruct()->GetFName(), EStateTreeNodeType::TransitionCondition, Node.Node.GetScriptStruct(), Cond->GetInstanceDataType()) == EStateTreeVisitor::Break)
-					{
-						bContinue = false;
-						break;
-					}
-				}
-			}
-		}
-	}
-	if (bContinue)
-	{
-		// Bindable state parameters
-		if (State.Type != EStateTreeStateType::Subtree
-			&& State.Parameters.Parameters.IsValid())
-		{
-			if (VisitFuncNodes(State.Parameters.ID, UE::StateTree::Editor::ParametersNodeName) == EStateTreeVisitor::Break)
-			{
-				bContinue = false;
-			}
-			else if (InFunc(&State, State.Parameters.ID, State.Name, EStateTreeNodeType::StateParameters, nullptr, State.Parameters.Parameters.GetPropertyBagStruct()) == EStateTreeVisitor::Break)
-			{
-				bContinue = false;
-			}
-		}
-	}
-
-	return bContinue ? EStateTreeVisitor::Continue : EStateTreeVisitor::Break;
-}
-
 EStateTreeVisitor UStateTreeEditorData::VisitStateNodes(const UStateTreeState& State, TFunctionRef<EStateTreeVisitor(const UStateTreeState* State, const FStateTreeBindableStructDesc& Desc, const FStateTreeDataView Value)> InFunc) const
 {
 	auto VisitFuncNodes = [&InFunc, &State, this](const FGuid StructID, const FName NodeName)
@@ -1558,16 +1433,6 @@ EStateTreeVisitor UStateTreeEditorData::VisitHierarchy(TFunctionRef<EStateTreeVi
 	return EStateTreeVisitor::Continue;
 }
 
-void UStateTreeEditorData::VisitHierarchyNodes(TFunctionRef<EStateTreeVisitor(const UStateTreeState* State, const FGuid& ID, const FName& Name, const EStateTreeNodeType NodeType, const UScriptStruct* NodeStruct, const UStruct* InstanceStruct)> InFunc) const
-{
-	PRAGMA_DISABLE_DEPRECATION_WARNINGS
-	VisitHierarchy([this, &InFunc](const UStateTreeState& State, UStateTreeState* /*ParentState*/)
-	{
-		return VisitStateNodes(State, InFunc);
-	});
-	PRAGMA_ENABLE_DEPRECATION_WARNINGS
-}
-
 EStateTreeVisitor UStateTreeEditorData::VisitGlobalNodes(TFunctionRef<EStateTreeVisitor(const UStateTreeState* State, const FStateTreeBindableStructDesc& Desc, const FStateTreeDataView Value)> InFunc) const
 {
 	// Root parameters
diff --git a/Engine/Plugins/Runtime/StateTree/Source/StateTreeEditorModule/Private/StateTreePropertyBindingCompiler.cpp b/Engine/Plugins/Runtime/StateTree/Source/StateTreeEditorModule/Private/StateTreePropertyBindingCompiler.cpp
--- a/Engine/Plugins/Runtime/StateTree/Source/StateTreeEditorModule/Private/StateTreePropertyBindingCompiler.cpp
+++ b/Engine/Plugins/Runtime/StateTree/Source/StateTreeEditorModule/Private/StateTreePropertyBindingCompiler.cpp
@@ -278,11 +278,6 @@ int32 FStateTreePropertyBindingCompiler::AddSourceStruct(const FStateTreeBindabl
 	return SourceStructs.Num() - 1;
 }
 
-int32 FStateTreePropertyBindingCompiler::GetSourceStructIndexByID(const FGuid& ID) const
-{
-	return SourceStructs.IndexOfByPredicate([ID](const FStateTreeBindableStructDesc& Structs) { return (Structs.ID == ID); });
-}
-
 void FStateTreePropertyBindingCompiler::StoreSourceStructs()
 {
 	// Check that existing structs are compatible
diff --git a/Engine/Plugins/Runtime/StateTree/Source/StateTreeEditorModule/Public/IStateTreeEditor.h b/Engine/Plugins/Runtime/StateTree/Source/StateTreeEditorModule/Public/IStateTreeEditor.h
--- a/Engine/Plugins/Runtime/StateTree/Source/StateTreeEditorModule/Public/IStateTreeEditor.h
+++ b/Engine/Plugins/Runtime/StateTree/Source/StateTreeEditorModule/Public/IStateTreeEditor.h
@@ -11,9 +11,3 @@ class STATETREEEDITORMODULE_API IStateTreeEditor : public FAssetEditorToolkit
 public:
 
 };
-
-
-
-#if UE_ENABLE_INCLUDE_ORDER_DEPRECATED_IN_5_2
-#include "CoreMinimal.h"
-#endif
diff --git a/Engine/Plugins/Runtime/StateTree/Source/StateTreeEditorModule/Public/StateTreeCompiler.h b/Engine/Plugins/Runtime/StateTree/Source/StateTreeEditorModule/Public/StateTreeCompiler.h
--- a/Engine/Plugins/Runtime/StateTree/Source/StateTreeEditorModule/Public/StateTreeCompiler.h
+++ b/Engine/Plugins/Runtime/StateTree/Source/StateTreeEditorModule/Public/StateTreeCompiler.h
@@ -167,7 +167,3 @@ namespace UE::StateTree::Compiler
 	const UScriptStruct* GetBaseStructFromMetaData(const FProperty* Property, FString& OutBaseStructName);
 
 }; // UE::StateTree::Compiler
-
-#if UE_ENABLE_INCLUDE_ORDER_DEPRECATED_IN_5_2
-#include "StateTreeCompilerLog.h"
-#endif
diff --git a/Engine/Plugins/Runtime/StateTree/Source/StateTreeEditorModule/Public/StateTreeCompilerLog.h b/Engine/Plugins/Runtime/StateTree/Source/StateTreeEditorModule/Public/StateTreeCompilerLog.h
--- a/Engine/Plugins/Runtime/StateTree/Source/StateTreeEditorModule/Public/StateTreeCompilerLog.h
+++ b/Engine/Plugins/Runtime/StateTree/Source/StateTreeEditorModule/Public/StateTreeCompilerLog.h
@@ -123,7 +123,3 @@ struct FStateTreeCompilerLogStateScope
 	FStateTreeCompilerLog& Log; 
 	const UStateTreeState* State = nullptr;
 };
-
-#if UE_ENABLE_INCLUDE_ORDER_DEPRECATED_IN_5_2
-#include "CoreMinimal.h"
-#endif
diff --git a/Engine/Plugins/Runtime/StateTree/Source/StateTreeEditorModule/Public/StateTreeEditorData.h b/Engine/Plugins/Runtime/StateTree/Source/StateTreeEditorModule/Public/StateTreeEditorData.h
--- a/Engine/Plugins/Runtime/StateTree/Source/StateTreeEditorModule/Public/StateTreeEditorData.h
+++ b/Engine/Plugins/Runtime/StateTree/Source/StateTreeEditorModule/Public/StateTreeEditorData.h
@@ -9,7 +9,6 @@
 #include "StateTreeEditorData.generated.h"
 
 struct FStateTreeBindableStructDesc;
-struct FStateTreeEditorPropertyPath;
 
 class UStateTreeSchema;
 
@@ -153,15 +152,6 @@ public:
 	 */
 	void GetAccessibleStructs(const TConstArrayView<const UStateTreeState*> Path, const FGuid TargetStructID, TArray<FStateTreeBindableStructDesc>& OutStructDescs) const;
 
-	UE_DEPRECATED(5.3, "Use VisitHierarchyNodes with State, Desc, Value instead.")
-	void VisitHierarchyNodes(TFunctionRef<EStateTreeVisitor(const UStateTreeState* State, const FGuid& ID, const FName& Name, const EStateTreeNodeType NodeType, const UScriptStruct* NodeStruct, const UStruct* InstanceStruct)> InFunc) const;
-
-	UE_DEPRECATED(5.3, "Use VisitStateNodes with State, Desc, Value instead.")
-	EStateTreeVisitor VisitStateNodes(const UStateTreeState& State, TFunctionRef<EStateTreeVisitor(const UStateTreeState* State, const FGuid& ID, const FName& Name, const EStateTreeNodeType NodeType, const UScriptStruct* NodeStruct, const UStruct* InstanceStruct)> InFunc) const;
-
-	UE_DEPRECATED(5.3, "Use GetAllStructValues with values instead.")
-	void GetAllStructIDs(TMap<FGuid, const UStruct*>& AllStructs) const;
-
 	void ReparentStates();
 	
 	// StateTree Builder API
diff --git a/Engine/Plugins/Runtime/StateTree/Source/StateTreeEditorModule/Public/StateTreeEditorModule.h b/Engine/Plugins/Runtime/StateTree/Source/StateTreeEditorModule/Public/StateTreeEditorModule.h
--- a/Engine/Plugins/Runtime/StateTree/Source/StateTreeEditorModule/Public/StateTreeEditorModule.h
+++ b/Engine/Plugins/Runtime/StateTree/Source/StateTreeEditorModule/Public/StateTreeEditorModule.h
@@ -45,8 +45,3 @@ protected:
 	FDelegateHandle OnUserDefinedStructReinstancedHandle;
 	FOnRegisterLayoutExtensions RegisterLayoutExtensions;
 };
-
-#if UE_ENABLE_INCLUDE_ORDER_DEPRECATED_IN_5_2
-#include "CoreMinimal.h"
-#include "Toolkits/IToolkitHost.h"
-#endif
diff --git a/Engine/Plugins/Runtime/StateTree/Source/StateTreeEditorModule/Public/StateTreeEditorPropertyBindings.h b/Engine/Plugins/Runtime/StateTree/Source/StateTreeEditorModule/Public/StateTreeEditorPropertyBindings.h
--- a/Engine/Plugins/Runtime/StateTree/Source/StateTreeEditorModule/Public/StateTreeEditorPropertyBindings.h
+++ b/Engine/Plugins/Runtime/StateTree/Source/StateTreeEditorModule/Public/StateTreeEditorPropertyBindings.h
@@ -214,7 +214,3 @@ protected:
 	virtual const FProperty* GetPropertyPathLeafProperty(const FStateTreePropertyPath& InPath) const override;
 
 };
-
-#if UE_ENABLE_INCLUDE_ORDER_DEPRECATED_IN_5_2
-#include "CoreMinimal.h"
-#endif
diff --git a/Engine/Plugins/Runtime/StateTree/Source/StateTreeEditorModule/Public/StateTreeEditorSettings.h b/Engine/Plugins/Runtime/StateTree/Source/StateTreeEditorModule/Public/StateTreeEditorSettings.h
--- a/Engine/Plugins/Runtime/StateTree/Source/StateTreeEditorModule/Public/StateTreeEditorSettings.h
+++ b/Engine/Plugins/Runtime/StateTree/Source/StateTreeEditorModule/Public/StateTreeEditorSettings.h
@@ -47,7 +47,3 @@ protected:
 
 	virtual FName GetCategoryName() const override;
 };
-
-#if UE_ENABLE_INCLUDE_ORDER_DEPRECATED_IN_5_2
-#include "CoreMinimal.h"
-#endif
diff --git a/Engine/Plugins/Runtime/StateTree/Source/StateTreeEditorModule/Public/StateTreeNodeClassCache.h b/Engine/Plugins/Runtime/StateTree/Source/StateTreeEditorModule/Public/StateTreeNodeClassCache.h
--- a/Engine/Plugins/Runtime/StateTree/Source/StateTreeEditorModule/Public/StateTreeNodeClassCache.h
+++ b/Engine/Plugins/Runtime/StateTree/Source/StateTreeEditorModule/Public/StateTreeNodeClassCache.h
@@ -121,7 +121,3 @@ private:
 	TArray<FRootClassContainer> RootClasses;
 	TMap<FString, int32> ClassNameToRootIndex;
 };
-
-#if UE_ENABLE_INCLUDE_ORDER_DEPRECATED_IN_5_2
-#include "CoreMinimal.h"
-#endif
diff --git a/Engine/Plugins/Runtime/StateTree/Source/StateTreeEditorModule/Public/StateTreePropertyBindingCompiler.h b/Engine/Plugins/Runtime/StateTree/Source/StateTreeEditorModule/Public/StateTreePropertyBindingCompiler.h
--- a/Engine/Plugins/Runtime/StateTree/Source/StateTreeEditorModule/Public/StateTreePropertyBindingCompiler.h
+++ b/Engine/Plugins/Runtime/StateTree/Source/StateTreeEditorModule/Public/StateTreePropertyBindingCompiler.h
@@ -56,20 +56,6 @@ struct STATETREEEDITORMODULE_API FStateTreePropertyBindingCompiler
 	  */
 	int32 AddSourceStruct(const FStateTreeBindableStructDesc& SourceStruct);
 
-	/** @return Index of a source struct by specified ID, or INDEX_NONE if not found. */
-	UE_DEPRECATED(5.4, "Use GetSourceStructDescByID() instead.")
-	int32 GetSourceStructIndexByID(const FGuid& ID) const;
-
-	/** @return Reference to a source struct based on ID. */
-	UE_DEPRECATED(5.4, "Use GetSourceStructDescByID() instead.")
-	const FStateTreeBindableStructDesc& GetSourceStructDesc(const int32 Index) const
-	{
-		return SourceStructs[Index];
-	}
-
-	UE_DEPRECATED(5.4, "Use CompileReferences with additional IDToStructValue parameter instead.")
-	[[nodiscard]] bool CompileReferences(const FStateTreeBindableStructDesc& TargetStruct, TConstArrayView<FStateTreePropertyPathBinding> PropertyReferenceBindings, FStateTreeDataView InstanceDataView) { return false; }
-
 	const FStateTreeBindableStructDesc* GetSourceStructDescByID(const FGuid& ID) const
 	{
 		return SourceStructs.FindByPredicate([ID](const FStateTreeBindableStructDesc& Structs) { return (Structs.ID == ID); });
@@ -105,11 +91,3 @@ protected:
 
 	FStateTreeCompilerLog* Log = nullptr;
 };
-
-#if UE_ENABLE_INCLUDE_ORDER_DEPRECATED_IN_5_2
-#include "CoreMinimal.h"
-#include "IPropertyAccessEditor.h"
-#include "StateTreeCompilerLog.h"
-#include "StateTreeEditorPropertyBindings.h"
-#include "UObject/Interface.h"
-#endif
diff --git a/Engine/Plugins/Runtime/StateTree/Source/StateTreeModule/Private/StateTree.cpp b/Engine/Plugins/Runtime/StateTree/Source/StateTreeModule/Private/StateTree.cpp
--- a/Engine/Plugins/Runtime/StateTree/Source/StateTreeModule/Private/StateTree.cpp
+++ b/Engine/Plugins/Runtime/StateTree/Source/StateTreeModule/Private/StateTree.cpp
@@ -335,13 +335,6 @@ void UStateTree::PostDuplicate(EDuplicateMode::Type DuplicateMode)
 	ResetCompiled();
 }
 
-void UStateTree::GetAssetRegistryTags(TArray<FAssetRegistryTag>& OutTags) const
-{
-	PRAGMA_DISABLE_DEPRECATION_WARNINGS;
-	Super::GetAssetRegistryTags(OutTags);
-	PRAGMA_ENABLE_DEPRECATION_WARNINGS;
-}
-
 void UStateTree::GetAssetRegistryTags(FAssetRegistryTagsContext Context) const
 {
 	const FString SchemaClassName = Schema ? Schema->GetClass()->GetPathName() : TEXT("");
diff --git a/Engine/Plugins/Runtime/StateTree/Source/StateTreeModule/Private/StateTreeExecutionContext.cpp b/Engine/Plugins/Runtime/StateTree/Source/StateTreeModule/Private/StateTreeExecutionContext.cpp
--- a/Engine/Plugins/Runtime/StateTree/Source/StateTreeModule/Private/StateTreeExecutionContext.cpp
+++ b/Engine/Plugins/Runtime/StateTree/Source/StateTreeModule/Private/StateTreeExecutionContext.cpp
@@ -4549,16 +4549,6 @@ EStateTreeRunStatus FStateTreeExecutionContext::GetLastTickStatus() const
 	return Exec.LastTickStatus;
 }
 
-void FStateTreeExecutionContext::SetDefaultParameters()
-{
-	SetGlobalParameters(RootStateTree.GetDefaultParameters());
-}
-
-void FStateTreeExecutionContext::SetParameters(const FInstancedPropertyBag& Parameters)
-{
-	SetGlobalParameters(Parameters);
-}
-
 FString FStateTreeExecutionContext::GetInstanceDescription() const
 {
 	return Owner.GetName();
diff --git a/Engine/Plugins/Runtime/StateTree/Source/StateTreeModule/Private/StateTreeReference.cpp b/Engine/Plugins/Runtime/StateTree/Source/StateTreeModule/Private/StateTreeReference.cpp
--- a/Engine/Plugins/Runtime/StateTree/Source/StateTreeModule/Private/StateTreeReference.cpp
+++ b/Engine/Plugins/Runtime/StateTree/Source/StateTreeModule/Private/StateTreeReference.cpp
@@ -33,18 +33,6 @@ void FStateTreeReference::SyncParameters()
 	}
 }
 
-void FStateTreeReference::SyncParametersToMatchStateTree(FInstancedPropertyBag& ParametersToSync) const
-{
-	if (StateTree == nullptr)
-	{
-		ParametersToSync.Reset();
-	}
-	else
-	{
-		ParametersToSync.MigrateToNewBagInstance(StateTree->GetDefaultParameters());
-	}
-}
-
 bool FStateTreeReference::RequiresParametersSync() const
 {
 	bool bShouldSync = false;
diff --git a/Engine/Plugins/Runtime/StateTree/Source/StateTreeModule/Public/Blueprint/StateTreeConditionBlueprintBase.h b/Engine/Plugins/Runtime/StateTree/Source/StateTreeModule/Public/Blueprint/StateTreeConditionBlueprintBase.h
--- a/Engine/Plugins/Runtime/StateTree/Source/StateTreeModule/Public/Blueprint/StateTreeConditionBlueprintBase.h
+++ b/Engine/Plugins/Runtime/StateTree/Source/StateTreeModule/Public/Blueprint/StateTreeConditionBlueprintBase.h
@@ -49,7 +49,3 @@ struct STATETREEMODULE_API FStateTreeBlueprintConditionWrapper : public FStateTr
 	UPROPERTY()
 	TSubclassOf<UStateTreeConditionBlueprintBase> ConditionClass = nullptr;
 };
-
-#if UE_ENABLE_INCLUDE_ORDER_DEPRECATED_IN_5_2
-#include "CoreMinimal.h"
-#endif
diff --git a/Engine/Plugins/Runtime/StateTree/Source/StateTreeModule/Public/Blueprint/StateTreeEvaluatorBlueprintBase.h b/Engine/Plugins/Runtime/StateTree/Source/StateTreeModule/Public/Blueprint/StateTreeEvaluatorBlueprintBase.h
--- a/Engine/Plugins/Runtime/StateTree/Source/StateTreeModule/Public/Blueprint/StateTreeEvaluatorBlueprintBase.h
+++ b/Engine/Plugins/Runtime/StateTree/Source/StateTreeModule/Public/Blueprint/StateTreeEvaluatorBlueprintBase.h
@@ -62,8 +62,3 @@ struct STATETREEMODULE_API FStateTreeBlueprintEvaluatorWrapper : public FStateTr
 	UPROPERTY()
 	TSubclassOf<UStateTreeEvaluatorBlueprintBase> EvaluatorClass = nullptr;
 };
-
-#if UE_ENABLE_INCLUDE_ORDER_DEPRECATED_IN_5_2
-#include "CoreMinimal.h"
-#include "StateTreeEvents.h"
-#endif
diff --git a/Engine/Plugins/Runtime/StateTree/Source/StateTreeModule/Public/Blueprint/StateTreeNodeBlueprintBase.h b/Engine/Plugins/Runtime/StateTree/Source/StateTreeModule/Public/Blueprint/StateTreeNodeBlueprintBase.h
--- a/Engine/Plugins/Runtime/StateTree/Source/StateTreeModule/Public/Blueprint/StateTreeNodeBlueprintBase.h
+++ b/Engine/Plugins/Runtime/StateTree/Source/StateTreeModule/Public/Blueprint/StateTreeNodeBlueprintBase.h
@@ -86,12 +86,7 @@ protected:
 	/** These methods are const as they set mutable variables and need to be called from a const method. */
 	void SetCachedInstanceDataFromContext(const FStateTreeExecutionContext& Context) const;
 	void ClearCachedInstanceData() const;
-	
-	UE_DEPRECATED(5.2, "Use SetCachedInstanceDataFromContext() instead.")
-	void SetCachedEventQueueFromContext(const FStateTreeExecutionContext& Context) const { SetCachedInstanceDataFromContext(Context); }
-	UE_DEPRECATED(5.2, "Use ClearCachedInstanceData() instead.")
-	void ClearCachedEventQueue() const { ClearCachedInstanceData(); }
-	
+
 private:
 	DECLARE_FUNCTION(execGetPropertyReference);
 
@@ -139,8 +134,3 @@ private:
 	static const IStateTreeBindingLookup* CachedBindingLookup;
 #endif	
 };
-
-#if UE_ENABLE_INCLUDE_ORDER_DEPRECATED_IN_5_2
-#include "StateTreeEvents.h"
-#include "StateTreeExecutionContext.h"
-#endif
diff --git a/Engine/Plugins/Runtime/StateTree/Source/StateTreeModule/Public/Blueprint/StateTreeTaskBlueprintBase.h b/Engine/Plugins/Runtime/StateTree/Source/StateTreeModule/Public/Blueprint/StateTreeTaskBlueprintBase.h
--- a/Engine/Plugins/Runtime/StateTree/Source/StateTreeModule/Public/Blueprint/StateTreeTaskBlueprintBase.h
+++ b/Engine/Plugins/Runtime/StateTree/Source/StateTreeModule/Public/Blueprint/StateTreeTaskBlueprintBase.h
@@ -61,23 +61,15 @@ public:
 	UFUNCTION(BlueprintImplementableEvent, meta = (DisplayName = "Tick"))
 	void ReceiveLatentTick(const float DeltaTime);
 
-	UE_DEPRECATED(5.3, "Use the new EnterState event without without return value instead. Task status is now controlled via FinishTask node, instead of a return value. Default status is running.")
+	UE_DEPRECATED(all, "Use the new EnterState event without without return value instead. Task status is now controlled via FinishTask node, instead of a return value. Default status is running.")
 	UFUNCTION(BlueprintImplementableEvent, meta = (DisplayName = "EnterState (Deprecated)", DeprecatedFunction, DeprecationMessage="Use the new EnterState event without without return value instead. Task status is now controlled via FinishTask node, instead of a return value. Default status is running."))
 	EStateTreeRunStatus ReceiveEnterState(const FStateTreeTransitionResult& Transition);
 
-	UE_DEPRECATED(5.3, "Use the new Tick event without without return value instead. Task status is now controlled via FinishTask node, instead of a return value. Default status is running.")
+	UE_DEPRECATED(all, "Use the new Tick event without without return value instead. Task status is now controlled via FinishTask node, instead of a return value. Default status is running.")
 	UFUNCTION(BlueprintImplementableEvent, meta = (DisplayName = "Tick (Deprecated)", DeprecatedFunction, DeprecationMessage="Use the new Tick event without without return value instead. Task status is now controlled via FinishTask node, instead of a return value. Default status is running."))
 	EStateTreeRunStatus ReceiveTick(const float DeltaTime);
 
 protected:
-	/**
-	 * Note: The API has been deprecated. ChangeType is moved into FStateTreeTransitionResult.
-	 * You can configure the task to be only called on state changes (that is, never call sustained changes) by setting bShouldStateChangeOnReselect to true.
-	 */
-	UE_DEPRECATED(5.2, "Use EnterState without ChangeType instead.")
-	virtual EStateTreeRunStatus EnterState(FStateTreeExecutionContext& Context, const EStateTreeStateChangeType ChangeType, const FStateTreeTransitionResult& Transition) final { return EStateTreeRunStatus::Running; }
-	UE_DEPRECATED(5.2, "Use ExitState without ChangeType instead.")
-	virtual void ExitState(FStateTreeExecutionContext& Context, const EStateTreeStateChangeType ChangeType, const FStateTreeTransitionResult& Transition) final {};
 
 	virtual EStateTreeRunStatus EnterState(FStateTreeExecutionContext& Context, const FStateTreeTransitionResult& Transition);
 	virtual void ExitState(FStateTreeExecutionContext& Context, const FStateTreeTransitionResult& Transition);
@@ -124,9 +116,9 @@ protected:
 	uint8 bHasStateCompleted : 1;
 	uint8 bHasLatentEnterState : 1;
 	uint8 bHasLatentTick : 1;
-	UE_DEPRECATED(5.3, "Use bHasLatentEnterState instead.")
+	UE_DEPRECATED(all, "Use bHasLatentEnterState instead.")
 	uint8 bHasEnterState_DEPRECATED : 1;
-	UE_DEPRECATED(5.3, "Use bHasLatentTick instead.")
+	UE_DEPRECATED(all, "Use bHasLatentTick instead.")
 	uint8 bHasTick_DEPRECATED : 1;
 
 	friend struct FStateTreeBlueprintTaskWrapper;
@@ -155,8 +147,3 @@ struct STATETREEMODULE_API FStateTreeBlueprintTaskWrapper : public FStateTreeTas
 	UPROPERTY()
 	TSubclassOf<UStateTreeTaskBlueprintBase> TaskClass = nullptr;
 };
-
-#if UE_ENABLE_INCLUDE_ORDER_DEPRECATED_IN_5_2
-#include "CoreMinimal.h"
-#include "StateTreeEvents.h"
-#endif
diff --git a/Engine/Plugins/Runtime/StateTree/Source/StateTreeModule/Public/StateTree.h b/Engine/Plugins/Runtime/StateTree/Source/StateTreeModule/Public/StateTree.h
--- a/Engine/Plugins/Runtime/StateTree/Source/StateTreeModule/Public/StateTree.h
+++ b/Engine/Plugins/Runtime/StateTree/Source/StateTreeModule/Public/StateTree.h
@@ -170,9 +170,6 @@ public:
 		return bHasGlobalTickTasks || (bHasEvent && bHasGlobalTickTasksOnlyOnEvents);
 	}
 
-	UE_DEPRECATED(5.4, "Replaced with GetNumContextDataViews() which contains context data and external data only.")
-	int32 GetNumDataViews() const { return 0; }
-
 #if WITH_EDITOR
 	/** Resets the compiled data to empty. */
 	void ResetCompiled();
@@ -229,8 +226,6 @@ protected:
 	virtual void BeginDestroy() override;
 	virtual void PostDuplicate(EDuplicateMode::Type DuplicateMode) override;
 	virtual void GetAssetRegistryTags(FAssetRegistryTagsContext Context) const override;
-	UE_DEPRECATED(5.4, "Implement the version that takes FAssetRegistryTagsContext instead.")
-	virtual void GetAssetRegistryTags(TArray<FAssetRegistryTag>& OutTags) const override;
 protected:
 	virtual void ThreadedPostLoadAssetRegistryTagsOverride(FPostLoadAssetRegistryTagsContext& Context) const override;
 public:
@@ -373,8 +368,3 @@ private:
 	friend struct FStateTreeCompiler;
 #endif
 };
-
-
-#if UE_ENABLE_INCLUDE_ORDER_DEPRECATED_IN_5_2
-#include "Misc/ScopeRWLock.h"
-#endif
diff --git a/Engine/Plugins/Runtime/StateTree/Source/StateTreeModule/Public/StateTreeAnyEnum.h b/Engine/Plugins/Runtime/StateTree/Source/StateTreeModule/Public/StateTreeAnyEnum.h
--- a/Engine/Plugins/Runtime/StateTree/Source/StateTreeModule/Public/StateTreeAnyEnum.h
+++ b/Engine/Plugins/Runtime/StateTree/Source/StateTreeModule/Public/StateTreeAnyEnum.h
@@ -37,8 +37,3 @@ struct STATETREEMODULE_API FStateTreeAnyEnum
 	UPROPERTY(EditAnywhere, Category = Enum)
 	TObjectPtr<UEnum> Enum = nullptr;
 };
-
-
-#if UE_ENABLE_INCLUDE_ORDER_DEPRECATED_IN_5_2
-#include "CoreMinimal.h"
-#endif
diff --git a/Engine/Plugins/Runtime/StateTree/Source/StateTreeModule/Public/StateTreeConditionBase.h b/Engine/Plugins/Runtime/StateTree/Source/StateTreeModule/Public/StateTreeConditionBase.h
--- a/Engine/Plugins/Runtime/StateTree/Source/StateTreeModule/Public/StateTreeConditionBase.h
+++ b/Engine/Plugins/Runtime/StateTree/Source/StateTreeModule/Public/StateTreeConditionBase.h
@@ -96,8 +96,3 @@ struct TStructOpsTypeTraits<Type> : public TStructOpsTypeTraitsBase2<Type> \
 		WithNoDestructor = true, \
 	}; \
 };
-
-#if UE_ENABLE_INCLUDE_ORDER_DEPRECATED_IN_5_2
-#include "StateTreeExecutionContext.h"
-#include "StateTreePropertyBindings.h"
-#endif
diff --git a/Engine/Plugins/Runtime/StateTree/Source/StateTreeModule/Public/StateTreeExecutionContext.h b/Engine/Plugins/Runtime/StateTree/Source/StateTreeModule/Public/StateTreeExecutionContext.h
--- a/Engine/Plugins/Runtime/StateTree/Source/StateTreeModule/Public/StateTreeExecutionContext.h
+++ b/Engine/Plugins/Runtime/StateTree/Source/StateTreeModule/Public/StateTreeExecutionContext.h
@@ -99,17 +99,6 @@ public:
 	FStateTreeExecutionContext(const FStateTreeExecutionContext& InContextToCopy, const UStateTree& InStateTree, FStateTreeInstanceData& InInstanceData);
 	virtual ~FStateTreeExecutionContext();
 
-	/** Updates data view of the parameters by using the default values defined in the StateTree asset. */
-	UE_DEPRECATED(5.4, "Not providing parameters to Start() leads to setting up default values now.")
-	void SetDefaultParameters();
-
-	/**
-	 * Updates data view of the parameters by replacing the default values defined in the StateTree asset by the provided values.
-	 * Note: caller is responsible to make sure external parameters lifetime matches the context.
-	 */
-	UE_DEPRECATED(5.4, "Provide parameters through Start() instead.")
-	void SetParameters(const FInstancedPropertyBag& Parameters);
-
 	/** Sets callback used to collect external data views during State Tree execution. */
 	void SetCollectExternalDataCallback(const FOnCollectStateTreeExternalData& Callback);
 
@@ -194,14 +183,6 @@ public:
 	/** @return the status of the last tick function */
 	EStateTreeRunStatus GetLastTickStatus() const;
 
-	/** @return reference to the list of currently active states. */
-	UE_DEPRECATED(5.4, "Use GetActiveFrames() instead.")
-	const FStateTreeActiveStates& GetActiveStates() const
-	{
-		static FStateTreeActiveStates Dummy;
-		return Dummy;
-	}
-
 	/** @return reference to the list of currently active frames and states. */
 	TConstArrayView<FStateTreeExecutionFrame> GetActiveFrames() const;
 
@@ -321,27 +302,12 @@ public:
 		return RootStateTree.GetStateFromHandle(StateHandle);
 	}
 
-	/** @return Array view to external data descriptors associated with this context. Note: Init() must be called before calling this method. */
-	UE_DEPRECATED(5.4, "Use CollectStateTreeExternalData delegate instead.")
-	TConstArrayView<FStateTreeExternalDataDesc> GetExternalDataDescs() const
-	{
-		return RootStateTree.ExternalDataDescs;
-	}
-
 	/** @return Array view to named external data descriptors associated with this context. Note: Init() must be called before calling this method. */
 	TConstArrayView<FStateTreeExternalDataDesc> GetContextDataDescs() const
 	{
 		return RootStateTree.GetContextDataDescs();
 	}
 
-	/** @return Handle to external data of type InStruct, or invalid handle if struct not found. */
-	UE_DEPRECATED(5.4, "Not supported anymore.")
-	FStateTreeExternalDataHandle GetExternalDataHandleByStruct(const UStruct* InStruct) const
-	{
-		const FStateTreeExternalDataDesc* DataDesc = RootStateTree.ExternalDataDescs.FindByPredicate([InStruct](const FStateTreeExternalDataDesc& Item) { return Item.Struct == InStruct; });
-		return DataDesc != nullptr ? DataDesc->Handle : FStateTreeExternalDataHandle::Invalid;
-	}
-
 	/** Sets context data view value for specific item. */
 	void SetContextData(const FStateTreeExternalDataHandle Handle, FStateTreeDataView DataView)
 	{
@@ -356,20 +322,6 @@ public:
 	/** @return True if all context data pointers are set. */ 
 	bool AreContextDataViewsValid() const;
 
-	/** @return True if all required external data pointers are set. */ 
-	UE_DEPRECATED(5.4, "Please AreContextDataViewsValid().")
-	bool AreExternalDataViewsValid() const
-	{
-		return AreContextDataViewsValid();
-	}
-
-	/** Sets external data view value for specific item. */
-	UE_DEPRECATED(5.4, "Use SetContextData() for context data, or set SetExternalDataDelegate() to provide external data.")
-	void SetExternalData(const FStateTreeExternalDataHandle Handle, FStateTreeDataView DataView)
-	{
-		SetContextData(Handle, DataView);
-	}
-
 	/**
 	 * Returns reference to external data based on provided handle. The return type is deduced from the handle's template type.
      * @param Handle Valid TStateTreeExternalDataHandle<> handle. 
diff --git a/Engine/Plugins/Runtime/StateTree/Source/StateTreeModule/Public/StateTreeExecutionTypes.h b/Engine/Plugins/Runtime/StateTree/Source/StateTreeModule/Public/StateTreeExecutionTypes.h
--- a/Engine/Plugins/Runtime/StateTree/Source/StateTreeModule/Public/StateTreeExecutionTypes.h
+++ b/Engine/Plugins/Runtime/StateTree/Source/StateTreeModule/Public/StateTreeExecutionTypes.h
@@ -107,17 +107,8 @@ struct STATETREEMODULE_API FStateTreeExternalDataHandle
 	
 	bool IsValid() const { return DataHandle.IsValid(); }
 
-	UE_DEPRECATED(5.4, "Index is deprecated, use DataHandle instead.")
-	static bool IsValidIndex(const int32 Index) { return FStateTreeDataHandle::IsValidIndex(Index); }
-
 	UPROPERTY()
 	FStateTreeDataHandle DataHandle = FStateTreeDataHandle::Invalid;
-
-#if WITH_EDITORONLY_DATA
-	UE_DEPRECATED(5.4, "Use DataHandle instead.")
-	UPROPERTY()
-	FStateTreeIndex16 DataViewIndex_DEPRECATED = FStateTreeIndex16::Invalid;
-#endif // WITH_EDITORONLY_DATA
 };
 
 /**
@@ -702,20 +693,6 @@ PRAGMA_ENABLE_DEPRECATION_WARNINGS
 	/** Used by state tree random-based operations. */
 	UPROPERTY()
 	FRandomStream RandomStream;
-
-#if WITH_EDITORONLY_DATA
-	UE_DEPRECATED(5.3, "Use DataHandle instead.")
-	UPROPERTY()
-	FStateTreeActiveStates CurrentActiveStates_DEPRECATED;
-	
-	/** Index of the first task struct in the currently initialized instance data. */
-	UE_DEPRECATED(5.4, "Superceded by State Tree Data Handles.")
-	FStateTreeIndex16 FirstTaskStructIndex_DEPRECATED = FStateTreeIndex16::Invalid;
-	
-	/** Index of the first task object in the currently initialized instance data. */
-	UE_DEPRECATED(5.4, "Superceded by State Tree Data Handles.")
-	FStateTreeIndex16 FirstTaskObjectIndex_DEPRECATED = FStateTreeIndex16::Invalid;
-#endif	
 };
 
 /** Contains StateTree events used during State Selection for a single execution frame. */
@@ -792,16 +769,6 @@ PRAGMA_ENABLE_DEPRECATION_WARNINGS
 	/** Root state the execution frame where the transition was requested. */
 	UPROPERTY(EditDefaultsOnly, Category = "Default", BlueprintReadOnly)
 	FStateTreeStateHandle SourceRootState = FStateTreeStateHandle::Invalid;
-	
-#if WITH_EDITORONLY_DATA
-	UE_DEPRECATED(5.4, "Use the ActiveFrames on FStateTreeExecutionState instead.")
-	UPROPERTY()
-	FStateTreeActiveStates CurrentActiveStates_DEPRECATED;
-
-	UE_DEPRECATED(5.4, "Use the NextActiveFrames instead.")
-	UPROPERTY()
-	FStateTreeActiveStates NextActiveStates_DEPRECATED;
-#endif	
 };
 
 /*
diff --git a/Engine/Plugins/Runtime/StateTree/Source/StateTreeModule/Public/StateTreeInstanceData.h b/Engine/Plugins/Runtime/StateTree/Source/StateTreeModule/Public/StateTreeInstanceData.h
--- a/Engine/Plugins/Runtime/StateTree/Source/StateTreeModule/Public/StateTreeInstanceData.h
+++ b/Engine/Plugins/Runtime/StateTree/Source/StateTreeModule/Public/StateTreeInstanceData.h
@@ -250,18 +250,6 @@ struct STATETREEMODULE_API FStateTreeInstanceStorage
 
 	/** Get the data used at runtime to confirm the inner working of the StateTree. */
 	UE::StateTree::Debug::FRuntimeValidation GetRuntimeValidation() const;
-	
-	UE_DEPRECATED(5.4, "Use Num() instead.")
-	int32 NumStructs() const { return 0; }
-
-	UE_DEPRECATED(5.4, "Use IsValidIndex() instead.")
-	bool IsValidStructIndex(const int32 Index) const { return false; }
-	
-	UE_DEPRECATED(5.4, "Use Num() instead.")
-	int32 NumObjects() const { return 0; }
-
-	UE_DEPRECATED(5.4, "Not used anymore, since ExecutionState is stored separately.")
-	bool IsValid() const { return false; }
 
 protected:
 	/** Struct instances (Not transient, as we use FStateTreeInstanceData to store default values for instance data) */
@@ -466,36 +454,6 @@ struct STATETREEMODULE_API FStateTreeInstanceData
 	{
 		return GetStorage().GetRuntimeValidation();
 	}
-	
-	UE_DEPRECATED(5.4, "Use the structs only Init(), objects should be wrapped in FStateTreeInstanceObjectWrapper.")
-	void Init(UObject& InOwner, TConstArrayView<FInstancedStruct> InStructs, TConstArrayView<const UObject*> InObjects);
-
-	UE_DEPRECATED(5.4, "Use the structs only Init(), objects should be wrapped in FStateTreeInstanceObjectWrapper.")
-	void Init(UObject& InOwner, TConstArrayView<FConstStructView> InStructs, TConstArrayView<const UObject*> InObjects);
-
-	UE_DEPRECATED(5.4, "Use the structs only Init(), objects should be wrapped in FStateTreeInstanceObjectWrapper.")
-	void Append(UObject& InOwner, TConstArrayView<FInstancedStruct> InStructs, TConstArrayView<const UObject*> InObjects);
-
-	UE_DEPRECATED(5.4, "Use the structs only Init(), objects should be wrapped in FStateTreeInstanceObjectWrapper.")
-	void Append(UObject& InOwner, TConstArrayView<FConstStructView> InStructs, TConstArrayView<const UObject*> InObjects);
-
-	UE_DEPRECATED(5.4, "Use the one param ShrinkTo().")
-	void ShrinkTo(const int32 NumStructs, const int32 NumObjects);
-
-	UE_DEPRECATED(5.4, "Use IsValidIndex() instead.")
-	bool IsValidStructIndex(const int32 Index) const { return false; }
-
-	UE_DEPRECATED(5.4, "Use Num() instead.")
-	int32 NumStructs() const { return GetStorage().InstanceStructs.Num(); }
-	
-	UE_DEPRECATED(5.4, "Use Num() instead.")
-	int32 NumObjects() const { return 0; }
-
-	UE_DEPRECATED(5.4, "Use Num() instead.")
-	int32 GetNumItems() const { return 0; }
-
-	UE_DEPRECATED(5.4, "InstanceData is always valid.")
-	bool IsValid() const { return true; }
 
 protected:
 	/** Storage for the actual instance data, always stores FStateTreeInstanceStorage. */
@@ -609,14 +567,6 @@ struct TStateTreeInstanceDataStructRef
 		return reinterpret_cast<T*>(Struct.GetMemory());
 	}
 
-	UE_DEPRECATED(5.4, "Please use GetPtr(), as the ref may be invalidated while in use.")
-	T& operator*()
-	{
-		T* Result = GetPtr();
-		check(Result);
-		return *Result;
-	}
-
 protected:
 
 	FStructView GetDataView(FStateTreeInstanceStorage& Storage, const FStateTreeExecutionFrame& CurrentFrame, const FStateTreeDataHandle Handle) const
diff --git a/Engine/Plugins/Runtime/StateTree/Source/StateTreeModule/Public/StateTreeLinker.h b/Engine/Plugins/Runtime/StateTree/Source/StateTreeModule/Public/StateTreeLinker.h
--- a/Engine/Plugins/Runtime/StateTree/Source/StateTreeModule/Public/StateTreeLinker.h
+++ b/Engine/Plugins/Runtime/StateTree/Source/StateTreeModule/Public/StateTreeLinker.h
@@ -91,9 +91,6 @@ struct FStateTreeLinker
 	/** @return linked external data descriptors. */
 	TConstArrayView<FStateTreeExternalDataDesc> GetExternalDataDescs() const { return ExternalDataDescs; }
 
-	UE_DEPRECATED(5.4, "Not used anymore.")
-	void SetExternalDataBaseIndex(const int32 InExternalDataBaseIndex) {}
-
 protected:
 
 	const UStateTreeSchema* Schema = nullptr;
diff --git a/Engine/Plugins/Runtime/StateTree/Source/StateTreeModule/Public/StateTreeNodeBase.h b/Engine/Plugins/Runtime/StateTree/Source/StateTreeModule/Public/StateTreeNodeBase.h
--- a/Engine/Plugins/Runtime/StateTree/Source/StateTreeModule/Public/StateTreeNodeBase.h
+++ b/Engine/Plugins/Runtime/StateTree/Source/StateTreeModule/Public/StateTreeNodeBase.h
@@ -6,7 +6,6 @@
 #include "StateTreeNodeBase.generated.h"
 
 struct FStateTreeLinker;
-struct FStateTreeEditorPropertyPath;
 struct FStateTreePropertyPath;
 struct IStateTreeBindingLookup;
 
@@ -69,10 +68,6 @@ PRAGMA_ENABLE_DEPRECATION_WARNINGS
 	virtual EDataValidationResult Compile(FStateTreeDataView InstanceDataView, TArray<FText>& ValidationMessages) { return EDataValidationResult::NotValidated; }
 	
 #if WITH_EDITOR
-	PRAGMA_DISABLE_DEPRECATION_WARNINGS
-	UE_DEPRECATED(5.3, "Use version with FStateTreePropertyPath instead.")
-	virtual void OnBindingChanged(const FGuid& ID, FStateTreeDataView InstanceData, const FStateTreeEditorPropertyPath& SourcePath, const FStateTreeEditorPropertyPath& TargetPath, const IStateTreeBindingLookup& BindingLookup) final {}
-	PRAGMA_ENABLE_DEPRECATION_WARNINGS
 
 	/**
 	 * Returns description for the node, use in the UI.
@@ -152,19 +147,5 @@ PRAGMA_ENABLE_DEPRECATION_WARNINGS
 
 	/** Data handle to access the instance data. */
 	UPROPERTY()
-	FStateTreeDataHandle InstanceDataHandle = FStateTreeDataHandle::Invalid; 
-
-#if WITH_EDITORONLY_DATA
-	UE_DEPRECATED(5.4, "InstanceDataHandle is used instead to reference to the instance data.")
-	UPROPERTY()
-	FStateTreeIndex16 DataViewIndex_DEPRECATED = FStateTreeIndex16::Invalid;
-
-	UE_DEPRECATED(5.4, "InstanceDataHandle is used instead to reference to the instance data.")
-	UPROPERTY()
-	FStateTreeIndex16 InstanceIndex_DEPRECATED = FStateTreeIndex16::Invalid;
-	
-	UE_DEPRECATED(5.4, "InstanceDataHandle is used to determine if the node has object data.")
-	UPROPERTY()
-	uint8 bInstanceIsObject_DEPRECATED : 1;
-#endif
+	FStateTreeDataHandle InstanceDataHandle = FStateTreeDataHandle::Invalid;
 };
diff --git a/Engine/Plugins/Runtime/StateTree/Source/StateTreeModule/Public/StateTreePropertyBindings.h b/Engine/Plugins/Runtime/StateTree/Source/StateTreeModule/Public/StateTreePropertyBindings.h
--- a/Engine/Plugins/Runtime/StateTree/Source/StateTreeModule/Public/StateTreePropertyBindings.h
+++ b/Engine/Plugins/Runtime/StateTree/Source/StateTreeModule/Public/StateTreePropertyBindings.h
@@ -143,15 +143,6 @@ struct STATETREEMODULE_API FStateTreeBindableStructDesc
 	{
 	}
 
-	UE_DEPRECATED(5.4, "Use constructor with DataHandle instead.")
-	FStateTreeBindableStructDesc(const FName InName, const UStruct* InStruct, const EStateTreeBindableStructSource InDataSource, const FGuid InGuid)
-		: Struct(InStruct)
-		, Name(InName)
-		, DataSource(InDataSource)
-		, ID(InGuid)
-	{
-	}
-
 	bool operator==(const FStateTreeBindableStructDesc& RHS) const
 	{
 		return ID == RHS.ID && Struct == RHS.Struct; // Not checking name, it's cosmetic.
@@ -208,36 +199,6 @@ struct STATETREEMODULE_API FStateTreePropertyPath : public FPropertyBindingPath
 	using FPropertyBindingPath::FPropertyBindingPath;
 };
 
-USTRUCT()
-struct UE_DEPRECATED(5.3, "Use FStateTreePropertyPath instead.") STATETREEMODULE_API FStateTreeEditorPropertyPath
-{
-	GENERATED_BODY()
-
-	FStateTreeEditorPropertyPath() = default;
-	FStateTreeEditorPropertyPath(const FGuid& InStructID, const TCHAR* PropertyName) : StructID(InStructID) { Path.Add(PropertyName); }
-	
-	/**
-	 * Returns the property path as a one string. Highlight allows to decorate a specific segment.
-	 * @param HighlightedSegment Index of the highlighted path segment
-	 * @param HighlightPrefix String to append before highlighted segment
-	 * @param HighlightPostfix String to append after highlighted segment
-	 */
-	FString ToString(const int32 HighlightedSegment = INDEX_NONE, const TCHAR* HighlightPrefix = nullptr, const TCHAR* HighlightPostfix = nullptr) const;
-
-	/** Handle of the struct this property path is relative to. */
-	UPROPERTY()
-	FGuid StructID;
-
-	/** Property path segments */
-	UPROPERTY()
-	TArray<FString> Path;
-
-	bool IsValid() const { return StructID.IsValid(); }
-
-PRAGMA_DISABLE_DEPRECATION_WARNINGS
-	bool operator==(const FStateTreeEditorPropertyPath& RHS) const;
-PRAGMA_ENABLE_DEPRECATION_WARNINGS
-};
 
 /**
  * Representation of a property binding in StateTree
@@ -272,17 +233,6 @@ FStateTreePropertyPathBinding(FConstStructView InFunctionNodeStruct, const FStat
 	}
 #endif
 
-	UE_DEPRECATED(5.4, "Use constructor with DataHandle instead.")
-	FStateTreePropertyPathBinding(const FStateTreeIndex16 InCompiledSourceStructIndex, const FStateTreePropertyPath& InSourcePath, const FStateTreePropertyPath& InTargetPath)
-		: SourcePropertyPath(InSourcePath)
-		, TargetPropertyPath(InTargetPath)
-	{
-	}
-
-	UE_DEPRECATED(5.3, "Use constructor with FStateTreePropertyPath instead.")
-	FStateTreePropertyPathBinding(const FStateTreeEditorPropertyPath& InSourcePath, const FStateTreeEditorPropertyPath& InTargetPath)
-	{
-	}
 PRAGMA_ENABLE_DEPRECATION_WARNINGS
 
 	const FStateTreePropertyPath& GetSourcePath() const { return SourcePropertyPath; }
@@ -291,9 +241,6 @@ PRAGMA_ENABLE_DEPRECATION_WARNINGS
 	FStateTreePropertyPath& GetMutableSourcePath() { return SourcePropertyPath; }
 	FStateTreePropertyPath& GetMutableTargetPath() { return TargetPropertyPath; }
 
-	UE_DEPRECATED(5.4, "Use GetSourceDataHandle instead.")
-	FStateTreeIndex16 GetCompiledSourceStructIndex() const { return {}; }
-
 	void SetSourceDataHandle(const FStateTreeDataHandle NewSourceDataHandle) { SourceDataHandle = NewSourceDataHandle; }
 	FStateTreeDataHandle GetSourceDataHandle() const { return SourceDataHandle; }
 
@@ -319,19 +266,9 @@ private:
 	/**	Instance of bound PropertyFunction. */
 	UPROPERTY()
 	FInstancedStruct PropertyFunctionNode;
-
-PRAGMA_DISABLE_DEPRECATION_WARNINGS
-	UPROPERTY()
-	FStateTreeEditorPropertyPath SourcePath_DEPRECATED;
-
-	UPROPERTY()
-	FStateTreeEditorPropertyPath TargetPath_DEPRECATED;
-PRAGMA_ENABLE_DEPRECATION_WARNINGS
 #endif // WITH_EDITORONLY_DATA
 };
 
-using FStateTreeEditorPropertyBinding UE_DEPRECATED(5.3, "Deprecated struct. Please use FStateTreePropertyPathBinding instead.") = FStateTreePropertyPathBinding;
-
 /**
  * Representation of a property reference binding in StateTree.
  */
@@ -365,64 +302,6 @@ private:
 	FStateTreeDataHandle SourceDataHandle = FStateTreeDataHandle::Invalid;
 };
 
-/**
- * Deprecated. Describes a segment of a property path. Used for storage only.
- */
-USTRUCT()
-struct UE_DEPRECATED(5.3, "Use FStateTreePropertyPath instead.") STATETREEMODULE_API FStateTreePropertySegment
-{
-	GENERATED_BODY()
-
-	/** @return true if the segment is empty. */
-	bool IsEmpty() const { return Name.IsNone(); }
-	
-	/** Property name. */
-	UPROPERTY()
-	FName Name;
-
-	/** Index in the array the property points at. */
-	UPROPERTY()
-	FStateTreeIndex16 ArrayIndex = FStateTreeIndex16::Invalid;
-
-	/** Index to next segment. */
-	UPROPERTY()
-	FStateTreeIndex16 NextIndex = FStateTreeIndex16::Invalid;
-
-	/** Type of access/indirection. */
-	UPROPERTY()
-	EPropertyBindingAccessType Type = EPropertyBindingAccessType::Offset;
-};
-
-
-/**
- * Deprecated. Describes property binding, the property from source path is copied into the property at the target path.
- */
-USTRUCT()
-struct UE_DEPRECATED(5.3, "Use FStateTreePropertyPath instead.") STATETREEMODULE_API FStateTreePropertyBinding
-{
-PRAGMA_DISABLE_DEPRECATION_WARNINGS
-
-	GENERATED_BODY()
-
-	/** Source property path. */
-	UPROPERTY()
-	FStateTreePropertySegment SourcePath;
-
-	/** Target property path. */
-	UPROPERTY()
-	FStateTreePropertySegment TargetPath;
-	
-	/** Index to the source struct the source path refers to, sources are stored in FStateTreePropertyBindings. */
-	UPROPERTY()
-	FStateTreeIndex16 SourceStructIndex = FStateTreeIndex16::Invalid;
-	
-	/** The type of copy to use between the properties. */
-	UPROPERTY()
-	EStateTreePropertyCopyType CopyType = EStateTreePropertyCopyType::None;
-
-PRAGMA_ENABLE_DEPRECATION_WARNINGS
-};
-
 
 /**
  * Used internally.
@@ -505,14 +384,8 @@ PRAGMA_ENABLE_DEPRECATION_WARNINGS
 	/** Type of the copy */
 	UPROPERTY()
 	EStateTreePropertyCopyType Type = EStateTreePropertyCopyType::None;
-
-	UE_DEPRECATED(5.4, "Use SourceDataHandle instead")
-	UPROPERTY()
-	FStateTreeIndex16 SourceStructIndex_DEPRECATED = FStateTreeIndex16::Invalid;
 };
 
-using FStateTreePropCopy UE_DEPRECATED(5.3, "Deprecated struct. Please use FStateTreePropertyCopy instead.") = FStateTreePropertyCopy;
-
 
 /**
  * Describes a batch of property copies from many sources to one target struct.
@@ -544,8 +417,6 @@ struct STATETREEMODULE_API FStateTreePropertyCopyBatch
 	FStateTreeIndex16 PropertyFunctionsEnd;
 };
 
-using FStateTreePropCopyBatch UE_DEPRECATED(5.3, "Deprecated struct. Please use FStateTreePropertyCopy instead.") = FStateTreePropertyCopyBatch;
-
 /**
  * Describes access to referenced property.
  */
@@ -685,18 +556,6 @@ PRAGMA_DISABLE_DEPRECATION_WARNINGS
 	UE_DEPRECATED(5.6, "Use the overload taking FPropertyBindingPathIndirection instead")
 	[[nodiscard]] static bool ResolveCopyType(const FStateTreePropertyPathIndirection& SourceIndirection, const FStateTreePropertyPathIndirection& TargetIndirection, FStateTreePropertyCopy& OutCopy);
 
-	UE_DEPRECATED(5.3, "Should not be used, will be removed in a future version.")
-	TArrayView<FStateTreeBindableStructDesc> GetSourceStructs() { return  SourceStructs; };
-
-	UE_DEPRECATED(5.3, "Use GetBatch() instead.")
-	TArrayView<FStateTreePropertyCopyBatch> GetCopyBatches() { return CopyBatches; };
-
-	UE_DEPRECATED(5.4, "Use GetBatchCopies() and Copy() instead.")
-	bool CopyTo(TConstArrayView<FStateTreeDataView> SourceStructViews, const FStateTreeIndex16 TargetBatchIndex, FStateTreeDataView TargetStructView) const
-	{
-		return false;
-	}
-
 	UE_DEPRECATED(5.5, "DebugPrintInternalLayout is deprecated. Use GetDebugInfo instead.")
 	void DebugPrintInternalLayout(FString& OutString) const;
 PRAGMA_ENABLE_DEPRECATION_WARNINGS
@@ -748,17 +607,6 @@ private:
 PRAGMA_DISABLE_DEPRECATION_WARNINGS
 	UE_DEPRECATED(5.6, "Use the overload taking FPropertyBindingPathIndirection instead")
 	[[nodiscard]] bool ResolvePath(const UStruct* Struct, const FStateTreePropertyPath& Path, FStateTreePropertyIndirection& OutFirstIndirection, FStateTreePropertyPathIndirection& OutLeafIndirection);
-
-#if WITH_EDITORONLY_DATA
-	
-	UE_DEPRECATED_FORGAME(5.3, "Use PropertyPathBindings instead.")
-	UPROPERTY()
-	TArray<FStateTreePropertyBinding> PropertyBindings_DEPRECATED;
-
-	UE_DEPRECATED_FORGAME(5.3, "Use PropertyPathBindings instead.")
-	UPROPERTY()
-	TArray<FStateTreePropertySegment> PropertySegments_DEPRECATED;
-#endif // WITH_EDITORONLY_DATA
 PRAGMA_ENABLE_DEPRECATION_WARNINGS
 
 	friend FStateTreePropertyBindingCompiler;
@@ -783,17 +631,6 @@ struct STATETREEMODULE_API IStateTreeBindingLookup
 
 	/** @return Leaf property based on property path. */
 	virtual const FProperty* GetPropertyPathLeafProperty(const FStateTreePropertyPath& InPath) const PURE_VIRTUAL(IStateTreeBindingLookup::GetPropertyPathLeafProperty, return nullptr; );
-
-PRAGMA_DISABLE_DEPRECATION_WARNINGS
-	UE_DEPRECATED(5.3, "Use version with FStateTreePropertyPath instead.")
-	virtual const FStateTreeEditorPropertyPath* GetPropertyBindingSource(const FStateTreeEditorPropertyPath& InTargetPath) const final { return nullptr; }
-	
-	UE_DEPRECATED(5.3, "Use version with FStateTreePropertyPath instead.")
-	virtual FText GetPropertyPathDisplayName(const FStateTreeEditorPropertyPath& InPath) const final { return FText::GetEmpty(); }
-	
-	UE_DEPRECATED(5.3, "Use version with FStateTreePropertyPath instead.")
-	virtual const FProperty* GetPropertyPathLeafProperty(const FStateTreeEditorPropertyPath& InPath) const final { return nullptr; }
-PRAGMA_ENABLE_DEPRECATION_WARNINGS
 };
 
 
@@ -814,20 +651,3 @@ namespace UE::StateTree
 	STATETREEMODULE_API const FProperty* GetStructSingleOutputProperty(const UStruct& InStruct);
 #endif
 } // UE::StateTree
-
-namespace UE::StateTree::Private
-{
-#if WITH_EDITORONLY_DATA
-// Helper functions to convert between property path types.
-PRAGMA_DISABLE_DEPRECATION_WARNINGS
-	UE_DEPRECATED(5.3, "Will be removed when FStateTreeEditorPropertyPath is removed.")
-	extern STATETREEMODULE_API FStateTreePropertyPath ConvertEditorPath(const FStateTreeEditorPropertyPath& InEditorPath);
-	UE_DEPRECATED(5.3, "Will be removed when FStateTreeEditorPropertyPath is removed.")
-	extern STATETREEMODULE_API FStateTreeEditorPropertyPath ConvertEditorPath(const FStateTreePropertyPath& InPath);
-PRAGMA_ENABLE_DEPRECATION_WARNINGS
-#endif // WITH_EDITORONLY_DATA
-} // UE::StateTree::Private 
-
-#if UE_ENABLE_INCLUDE_ORDER_DEPRECATED_IN_5_2
-#include "CoreMinimal.h"
-#endif
diff --git a/Engine/Plugins/Runtime/StateTree/Source/StateTreeModule/Public/StateTreeReference.h b/Engine/Plugins/Runtime/StateTree/Source/StateTreeModule/Public/StateTreeReference.h
--- a/Engine/Plugins/Runtime/StateTree/Source/StateTreeModule/Public/StateTreeReference.h
+++ b/Engine/Plugins/Runtime/StateTree/Source/StateTreeModule/Public/StateTreeReference.h
@@ -60,12 +60,6 @@ struct STATETREEMODULE_API FStateTreeReference
 	 */
 	void SyncParameters();
 
-	/**
-	 * Sync provided parameters to be compatible with those exposed by the selected StateTree asset.
-	 */
-	UE_DEPRECATED(5.4, "Use SyncParameters() instead.")
-	void SyncParametersToMatchStateTree(FInstancedPropertyBag& ParametersToSync) const;
-
 	/**
 	 * Indicates if current parameters are compatible with those available in the selected StateTree asset.
 	 * @return true when parameters requires to be synced to be compatible with those available in the selected StateTree asset, false otherwise.
diff --git a/Engine/Plugins/Runtime/StateTree/Source/StateTreeModule/Public/StateTreeSchema.h b/Engine/Plugins/Runtime/StateTree/Source/StateTreeModule/Public/StateTreeSchema.h
--- a/Engine/Plugins/Runtime/StateTree/Source/StateTreeModule/Public/StateTreeSchema.h
+++ b/Engine/Plugins/Runtime/StateTree/Source/StateTreeModule/Public/StateTreeSchema.h
@@ -58,7 +58,3 @@ public:
 	
 #endif // WITH_EDITOR
 };
-
-#if UE_ENABLE_INCLUDE_ORDER_DEPRECATED_IN_5_2
-#include "StateTreeTypes.h"
-#endif
diff --git a/Engine/Plugins/Runtime/StateTree/Source/StateTreeModule/Public/StateTreeTypes.h b/Engine/Plugins/Runtime/StateTree/Source/StateTreeModule/Public/StateTreeTypes.h
--- a/Engine/Plugins/Runtime/StateTree/Source/StateTreeModule/Public/StateTreeTypes.h
+++ b/Engine/Plugins/Runtime/StateTree/Source/StateTreeModule/Public/StateTreeTypes.h
@@ -81,7 +81,7 @@ enum class EStateTreeTransitionType : uint8
 	/** Transition to the next selectable sibling state */
 	NextSelectableState,
 
-	NotSet UE_DEPRECATED(5.0, "Use None instead."),
+	NotSet UE_DEPRECATED(all, "Use None instead."),
 };
 
 /** Operand in an expression */
@@ -1011,11 +1011,6 @@ struct STATETREEMODULE_API FStateTreeStateLink
 
 #if WITH_EDITORONLY_DATA
 	FStateTreeStateLink(const EStateTreeTransitionType InType) : LinkType(InType) {}
-
-	UE_DEPRECATED(5.2, "Use UStateTreeState::GetLinkToState() instead.")
-	void Set(const EStateTreeTransitionType InType, const class UStateTreeState* InState = nullptr) {}
-	UE_DEPRECATED(5.2, "Use UStateTreeState::GetLinkToState() instead.")
-	void Set(const class UStateTreeState* InState) {}
 #endif // WITH_EDITORONLY_DATA
 	
 PRAGMA_DISABLE_DEPRECATION_WARNINGS
@@ -1029,9 +1024,6 @@ PRAGMA_ENABLE_DEPRECATION_WARNINGS
 	void PostSerialize(const FArchive& Ar);
 
 #if WITH_EDITORONLY_DATA
-	UE_DEPRECATED(5.2, "Will be removed in later releases.")
-	bool IsValid() const { return ID.IsValid(); }
-	
 	/** Name of the state at the time of linking, used for error reporting. */
 	UPROPERTY(EditDefaultsOnly, Category = "Link")
 	FName Name;
@@ -1044,7 +1036,7 @@ PRAGMA_ENABLE_DEPRECATION_WARNINGS
 	UPROPERTY(EditDefaultsOnly, Category = "Link")
 	EStateTreeTransitionType LinkType = EStateTreeTransitionType::None;
 
-	UE_DEPRECATED(5.2, "Use LinkType instead.")
+	UE_DEPRECATED(all, "Use LinkType instead.")
 	UPROPERTY()
 	EStateTreeTransitionType Type_DEPRECATED = EStateTreeTransitionType::GotoState;
 #endif // WITH_EDITORONLY_DATA
diff --git a/Engine/Plugins/Runtime/StateTree/Source/StateTreeTestSuite/Public/StateTreeTestSuiteModule.h b/Engine/Plugins/Runtime/StateTree/Source/StateTreeTestSuite/Public/StateTreeTestSuiteModule.h
--- a/Engine/Plugins/Runtime/StateTree/Source/StateTreeTestSuite/Public/StateTreeTestSuiteModule.h
+++ b/Engine/Plugins/Runtime/StateTree/Source/StateTreeTestSuite/Public/StateTreeTestSuiteModule.h
@@ -33,8 +33,3 @@ public:
 		return FModuleManager::Get().IsModuleLoaded("StateTreeTestSuiteModule");
 	}
 };
-
-
-#if UE_ENABLE_INCLUDE_ORDER_DEPRECATED_IN_5_2
-#include "CoreMinimal.h"
-#endif
