Fix CIS issue 654421: missing forward declaration

diff --git a/Engine/Plugins/Runtime/StateTree/Source/StateTreeEditorModule/Private/SStateTreeDiff.h b/Engine/Plugins/Runtime/StateTree/Source/StateTreeEditorModule/Private/SStateTreeDiff.h
--- a/Engine/Plugins/Runtime/StateTree/Source/StateTreeEditorModule/Private/SStateTreeDiff.h
+++ b/Engine/Plugins/Runtime/StateTree/Source/StateTreeEditorModule/Private/SStateTreeDiff.h
@@ -10,6 +10,7 @@
 class FSpawnTabArgs;
 class FTabManager;
 class FUICommandList;
+class UStateTreeState;
 
 namespace UE::StateTree::Diff
 {
