#UE Add redirector for EStateTreePropertyAccessType::Offset

diff --git a/Engine/Plugins/Runtime/StateTree/Config/BaseStateTree.ini b/Engine/Plugins/Runtime/StateTree/Config/BaseStateTree.ini
--- a/Engine/Plugins/Runtime/StateTree/Config/BaseStateTree.ini
+++ b/Engine/Plugins/Runtime/StateTree/Config/BaseStateTree.ini
@@ -28,7 +28,7 @@
 
 +EnumRedirects=(OldName="EStateTreePropertyCopyType",ValueChanges=(("Plain","CopyPlain"),("Complex","CopyComplex"),("Bool","CopyBool"),("Struct","CopyStruct"),("Object","CopyObject"),("Name","CopyName"),("FixedArray","CopyFixedArray")))
 +EnumRedirects=(OldName="EStateTreePropertyAccessType",NewName="/Script/PropertyBindingUtilsModule.EPropertyBindingAccessType")
-+EnumRedirects=(OldName="EPropertyBindingAccessType",ValueChanges=(("Array","IndexArray")))
++EnumRedirects=(OldName="EPropertyBindingAccessType",ValueChanges=(("Array","IndexArray"),("EStateTreePropertyAccessType::Offset","EPropertyBindingAccessType::Offset")))
 
 +StructRedirects=(OldName="StateTreeCondition_CompareIntInstanceData",NewName="/script/StateTreeModule.StateTreeCompareIntConditionInstanceData")
 +StructRedirects=(OldName="StateTreeCondition_CompareInt",NewName="/script/StateTreeModule.StateTreeCompareIntCondition")
