[StateTree] restored a deprecation code path that shouldn't have been removed

diff --git a/Engine/Plugins/Runtime/StateTree/Source/StateTreeModule/Private/StateTreePropertyBindings.cpp b/Engine/Plugins/Runtime/StateTree/Source/StateTreeModule/Private/StateTreePropertyBindings.cpp
--- a/Engine/Plugins/Runtime/StateTree/Source/StateTreeModule/Private/StateTreePropertyBindings.cpp
+++ b/Engine/Plugins/Runtime/StateTree/Source/StateTreeModule/Private/StateTreePropertyBindings.cpp
@@ -91,6 +91,50 @@ namespace UE::StateTree
 
 } // UE::StateTree
 
+PRAGMA_DISABLE_DEPRECATION_WARNINGS
+#if WITH_EDITORONLY_DATA
+namespace UE::StateTree::Deprecation
+{
+	FStateTreePropertyPath ConvertEditorPath(const FStateTreeEditorPropertyPath& InEditorPath)
+	{
+		FStateTreePropertyPath Path;
+		Path.SetStructID(InEditorPath.StructID);
+
+		for (const FString& Segment : InEditorPath.Path)
+		{
+			const TCHAR* PropertyNamePtr = nullptr;
+			int32 PropertyNameLength = 0;
+			int32 ArrayIndex = INDEX_NONE;
+			PropertyPathHelpers::FindFieldNameAndArrayIndex(Segment.Len(), *Segment, PropertyNameLength, &PropertyNamePtr, ArrayIndex);
+			FString PropertyNameString(PropertyNameLength, PropertyNamePtr);
+			const FName PropertyName(*PropertyNameString, FNAME_Find);
+			Path.AddPathSegment(PropertyName, ArrayIndex);
+		}
+		return Path;
+	}
+} // UE::StateTree::Deprecation
+
+//----------------------------------------------------------------//
+//  FStateTreePropertyPathBinding
+//----------------------------------------------------------------//
+void FStateTreePropertyPathBinding::PostSerialize(const FArchive& Ar)
+{
+	if (SourcePath_DEPRECATED.IsValid())
+	{
+		SourcePropertyPath = UE::StateTree::Deprecation::ConvertEditorPath(SourcePath_DEPRECATED);
+		SourcePath_DEPRECATED.StructID = FGuid();
+		SourcePath_DEPRECATED.Path.Reset();
+	}
+
+	if (TargetPath_DEPRECATED.IsValid())
+	{
+		TargetPropertyPath = UE::StateTree::Deprecation::ConvertEditorPath(TargetPath_DEPRECATED);
+		TargetPath_DEPRECATED.StructID = FGuid();
+		TargetPath_DEPRECATED.Path.Reset();
+	}
+}
+#endif // WITH_EDITORONLY_DATA
+PRAGMA_ENABLE_DEPRECATION_WARNINGS
 
 //----------------------------------------------------------------//
 //  FStateTreeBindableStructDesc
diff --git a/Engine/Plugins/Runtime/StateTree/Source/StateTreeModule/Public/StateTreePropertyBindings.h b/Engine/Plugins/Runtime/StateTree/Source/StateTreeModule/Public/StateTreePropertyBindings.h
--- a/Engine/Plugins/Runtime/StateTree/Source/StateTreeModule/Public/StateTreePropertyBindings.h
+++ b/Engine/Plugins/Runtime/StateTree/Source/StateTreeModule/Public/StateTreePropertyBindings.h
@@ -200,6 +200,27 @@ struct STATETREEMODULE_API FStateTreePropertyPath : public FPropertyBindingPath
 };
 
 
+USTRUCT()
+struct UE_DEPRECATED(all, "Use FStateTreePropertyPath instead.") STATETREEMODULE_API FStateTreeEditorPropertyPath
+{
+	GENERATED_BODY()
+
+#if WITH_EDITORONLY_DATA
+	/** Handle of the struct this property path is relative to. */
+	UPROPERTY()
+	FGuid StructID;
+
+	/** Property path segments */
+	UPROPERTY()
+	TArray<FString> Path;
+
+	bool IsValid() const
+	{
+		return StructID.IsValid();
+	}
+#endif // WITH_EDITORONLY_DATA
+};
+
 /**
  * Representation of a property binding in StateTree
  */
@@ -235,6 +256,10 @@ FStateTreePropertyPathBinding(FConstStructView InFunctionNodeStruct, const FStat
 
 PRAGMA_ENABLE_DEPRECATION_WARNINGS
 
+#if WITH_EDITORONLY_DATA
+	void PostSerialize(const FArchive& Ar);
+#endif // WITH_EDITORONLY_DATA
+
 	const FStateTreePropertyPath& GetSourcePath() const { return SourcePropertyPath; }
 	const FStateTreePropertyPath& GetTargetPath() const { return TargetPropertyPath; }
 
@@ -266,9 +291,28 @@ private:
 	/**	Instance of bound PropertyFunction. */
 	UPROPERTY()
 	FInstancedStruct PropertyFunctionNode;
+
+PRAGMA_DISABLE_DEPRECATION_WARNINGS
+	UPROPERTY()
+	FStateTreeEditorPropertyPath SourcePath_DEPRECATED;
+
+	UPROPERTY()
+	FStateTreeEditorPropertyPath TargetPath_DEPRECATED;
+PRAGMA_ENABLE_DEPRECATION_WARNINGS
 #endif // WITH_EDITORONLY_DATA
 };
 
+#if WITH_EDITORONLY_DATA
+template<>
+struct TStructOpsTypeTraits<FStateTreePropertyPathBinding> : public TStructOpsTypeTraitsBase2<FStateTreePropertyPathBinding>
+{
+	enum 
+	{
+		WithPostSerialize = true,
+	};
+};
+#endif // WITH_EDITORONLY_DATA
+
 /**
  * Representation of a property reference binding in StateTree.
  */
