StateTree: Add flag to detect if the transitions need ticking.

diff --git a/Engine/Plugins/Runtime/StateTree/Source/StateTreeEditorModule/Private/StateTreeCompiler.cpp b/Engine/Plugins/Runtime/StateTree/Source/StateTreeEditorModule/Private/StateTreeCompiler.cpp
--- a/Engine/Plugins/Runtime/StateTree/Source/StateTreeEditorModule/Private/StateTreeCompiler.cpp
+++ b/Engine/Plugins/Runtime/StateTree/Source/StateTreeEditorModule/Private/StateTreeCompiler.cpp
@@ -1076,6 +1076,15 @@ bool FStateTreeCompiler::CreateStateTransitions()
 				BindingsCompiler.AddSourceStruct(Desc);
 			}
 
+			if (CompactTransition.bTransitionEnabled)
+			{
+				CompactState.bHasTickTriggerTransitions |= Transition.Trigger == EStateTreeTransitionTrigger::OnTick;
+				CompactState.bHasEventTriggerTransitions |= Transition.Trigger == EStateTreeTransitionTrigger::OnEvent;
+				CompactState.bHasSucceededTriggerTransitions |= Transition.Trigger == EStateTreeTransitionTrigger::OnStateCompleted;
+				CompactState.bHasSucceededTriggerTransitions |= Transition.Trigger == EStateTreeTransitionTrigger::OnStateSucceeded;
+				CompactState.bHasFailedTriggerTransitions |= Transition.Trigger == EStateTreeTransitionTrigger::OnStateFailed;
+			}
+
 			const int32 ConditionsBegin = Nodes.Num();
 			if (const auto Validation = UE::StateTree::Compiler::IsValidCount16(ConditionsBegin); Validation.DidFail())
 			{
diff --git a/Engine/Plugins/Runtime/StateTree/Source/StateTreeModule/Private/StateTree.cpp b/Engine/Plugins/Runtime/StateTree/Source/StateTreeModule/Private/StateTree.cpp
--- a/Engine/Plugins/Runtime/StateTree/Source/StateTreeModule/Private/StateTree.cpp
+++ b/Engine/Plugins/Runtime/StateTree/Source/StateTreeModule/Private/StateTree.cpp
@@ -190,6 +190,8 @@ void UStateTree::ResetCompiled()
 	GlobalTasksBegin = 0;
 	GlobalTasksNum = 0;
 	bHasGlobalTransitionTasks = false;
+	bHasGlobalTickTasks = false;
+	bHasGlobalTickTasksOnlyOnEvents = false;
 	
 	ResetLinked();
 
diff --git a/Engine/Plugins/Runtime/StateTree/Source/StateTreeModule/Private/StateTreeExecutionContext.cpp b/Engine/Plugins/Runtime/StateTree/Source/StateTreeModule/Private/StateTreeExecutionContext.cpp
--- a/Engine/Plugins/Runtime/StateTree/Source/StateTreeModule/Private/StateTreeExecutionContext.cpp
+++ b/Engine/Plugins/Runtime/StateTree/Source/StateTreeModule/Private/StateTreeExecutionContext.cpp
@@ -3215,6 +3215,8 @@ bool FStateTreeExecutionContext::TriggerTransitions()
 
 	if (Exec.ActiveFrames.Num() > 0)
 	{
+		// Re-cache bHasEvents, RequestTransition above can create new events.
+		const bool bHasEvents = EventQueue && EventQueue->HasEvents();
 		for (int32 FrameIndex = Exec.ActiveFrames.Num() - 1; FrameIndex >= 0; FrameIndex--)
 		{
 			FStateTreeExecutionFrame& CurrentFrame = Exec.ActiveFrames[FrameIndex];
@@ -3248,7 +3250,8 @@ bool FStateTreeExecutionContext::TriggerTransitions()
 				}
 				
 				// Regular transitions on state
-				if (State.TransitionsNum > 0)
+				//A transition task can trigger an event. We need to add the state if that is a possibility.
+				if (State.ShouldTickTransitions(bHasEvents || State.bHasTransitionTasks))
 				{
 					TransitionHandlers.Emplace((uint8)FrameIndex, StateHandle, EStateTreeTransitionPriority::Normal);
 				}
@@ -3316,10 +3319,10 @@ bool FStateTreeExecutionContext::TriggerTransitions()
 			const FCompactStateTreeState& State = CurrentStateTree->States[Handler.StateHandle.Index];
 
 			// Transitions
-			for (uint8 i = 0; i < State.TransitionsNum; i++)
+			for (uint8 TransitionCounter = 0; TransitionCounter < State.TransitionsNum; ++TransitionCounter)
 			{
 				// All transition conditions must pass
-				const int16 TransitionIndex = State.TransitionsBegin + i;
+				const int16 TransitionIndex = State.TransitionsBegin + TransitionCounter;
 				const FCompactStateTransition& Transition = CurrentStateTree->Transitions[TransitionIndex];
 
 				// Skip disabled transitions
@@ -3392,7 +3395,11 @@ bool FStateTreeExecutionContext::TriggerTransitions()
 					// Dummy event to make sure we iterate to loop below once.
 					TransitionEvents.Emplace(nullptr);
 				}
-				
+				else
+				{
+					ensureMsgf(false, TEXT("The trigger type is not supported."));
+				}
+
 				for (const FStateTreeSharedEvent* TransitionEvent : TransitionEvents)
 				{
 					bool bPassed = false; 
@@ -3504,53 +3511,56 @@ bool FStateTreeExecutionContext::TriggerTransitions()
 			}
 			
 			const EStateTreeTransitionTrigger CompletionTrigger = Exec.LastTickStatus == EStateTreeRunStatus::Succeeded ? EStateTreeTransitionTrigger::OnStateSucceeded : EStateTreeTransitionTrigger::OnStateFailed;
-		
+
 			// Check completion transitions
 			for (int32 StateIndex = StateStartIndex; StateIndex >= 0; StateIndex--)
 			{
 				const FStateTreeStateHandle StateHandle = CurrentFrame.ActiveStates[StateIndex];
 				const FCompactStateTreeState& State = CurrentStateTree->States[StateHandle.Index];
 
-				FCurrentlyProcessedStateScope StateScope(*this, StateHandle);
-				STATETREE_TRACE_SCOPED_STATE_PHASE(StateHandle, EStateTreeUpdatePhase::TriggerTransitions);
-
-				for (uint8 i = 0; i < State.TransitionsNum; i++)
+				if (State.ShouldTickCompletionTransitions(CompletionTrigger == EStateTreeTransitionTrigger::OnStateSucceeded))
 				{
-					// All transition conditions must pass
-					const int16 TransitionIndex = State.TransitionsBegin + i;
-					const FCompactStateTransition& Transition = CurrentStateTree->Transitions[TransitionIndex];
+					FCurrentlyProcessedStateScope StateScope(*this, StateHandle);
+					STATETREE_TRACE_SCOPED_STATE_PHASE(StateHandle, EStateTreeUpdatePhase::TriggerTransitions);
 
-					// Skip disabled transitions
-					if (Transition.bTransitionEnabled == false)
+					for (uint8 TransitionCounter = 0; TransitionCounter < State.TransitionsNum; ++TransitionCounter)
 					{
-						continue;
-					}
+						// All transition conditions must pass
+						const int16 TransitionIndex = State.TransitionsBegin + TransitionCounter;
+						const FCompactStateTransition& Transition = CurrentStateTree->Transitions[TransitionIndex];
 
-					if (EnumHasAnyFlags(Transition.Trigger, CompletionTrigger))
-					{
-						bool bPassed = false;
+						// Skip disabled transitions
+						if (Transition.bTransitionEnabled == false)
 						{
-							STATETREE_TRACE_TRANSITION_EVENT(FStateTreeTransitionSource(FStateTreeIndex16(TransitionIndex), Transition.State, Transition.Priority), EStateTreeTraceEventType::OnEvaluating);
-							STATETREE_TRACE_SCOPED_PHASE(EStateTreeUpdatePhase::TransitionConditions);
-							bPassed = TestAllConditions(CurrentParentFrame, CurrentFrame, Transition.ConditionsBegin, Transition.ConditionsNum);
+							continue;
 						}
 
-						if (bPassed)
+						if (EnumHasAnyFlags(Transition.Trigger, CompletionTrigger))
 						{
-							// No delay allowed on completion conditions.
-							// No priority on completion transitions, use the priority to signal that state is selected.
-							if (RequestTransition(CurrentFrame, Transition.State, EStateTreeTransitionPriority::Normal, /*TransitionEvent*/nullptr, Transition.Fallback))
+							bool bPassed = false;
 							{
-								NextTransitionSource = FStateTreeTransitionSource(FStateTreeIndex16(TransitionIndex), Transition.State, Transition.Priority);
-								break;
+								STATETREE_TRACE_TRANSITION_EVENT(FStateTreeTransitionSource(FStateTreeIndex16(TransitionIndex), Transition.State, Transition.Priority), EStateTreeTraceEventType::OnEvaluating);
+								STATETREE_TRACE_SCOPED_PHASE(EStateTreeUpdatePhase::TransitionConditions);
+								bPassed = TestAllConditions(CurrentParentFrame, CurrentFrame, Transition.ConditionsBegin, Transition.ConditionsNum);
+							}
+
+							if (bPassed)
+							{
+								// No delay allowed on completion conditions.
+								// No priority on completion transitions, use the priority to signal that state is selected.
+								if (RequestTransition(CurrentFrame, Transition.State, EStateTreeTransitionPriority::Normal, /*TransitionEvent*/nullptr, Transition.Fallback))
+								{
+									NextTransitionSource = FStateTreeTransitionSource(FStateTreeIndex16(TransitionIndex), Transition.State, Transition.Priority);
+									break;
+								}
 							}
 						}
 					}
-				}
 
-				if (NextTransition.Priority != EStateTreeTransitionPriority::None) //-V547
-				{
-					break;
+					if (NextTransition.Priority != EStateTreeTransitionPriority::None) //-V547
+					{
+						break;
+					}
 				}
 			}
 		}
diff --git a/Engine/Plugins/Runtime/StateTree/Source/StateTreeModule/Public/StateTreeTypes.h b/Engine/Plugins/Runtime/StateTree/Source/StateTreeModule/Public/StateTreeTypes.h
--- a/Engine/Plugins/Runtime/StateTree/Source/StateTreeModule/Public/StateTreeTypes.h
+++ b/Engine/Plugins/Runtime/StateTree/Source/StateTreeModule/Public/StateTreeTypes.h
@@ -678,6 +678,20 @@ struct STATETREEMODULE_API FCompactStateTreeState
 		return bHasTickTasks || (bHasEvent && bHasTickTasksOnlyOnEvents);
 	}
 
+	/** @return True if the state has any transitions that need ticking. */
+	bool ShouldTickTransitions(bool bHasEvent) const
+	{
+		return bHasTickTriggerTransitions || (bHasEvent && bHasEventTriggerTransitions);
+	}
+
+	/** @return True if the state has any transitions that tick on completion. */
+	bool ShouldTickCompletionTransitions(bool bSucceeded) const
+	{
+		return bHasCompletedTriggerTransitions
+			|| (bHasSucceededTriggerTransitions && bSucceeded)
+			|| (bHasFailedTriggerTransitions && !bSucceeded);
+	}
+
 	/** Description of an event required to enter the state. */
 	UPROPERTY()
 	FCompactEventDesc RequiredEventToEnter;
@@ -782,19 +796,35 @@ struct STATETREEMODULE_API FCompactStateTreeState
 	UPROPERTY()
 	uint8 bHasStateChangeConditions : 1 = false;
 
-	/**
-	 * True if any task has bShouldCallTick.
-	 * Not ticking implies no property copy.
-	 */
+	/** True if any task has bShouldCallTick. */
 	uint8 bHasTickTasks : 1 = false;
 
 	/**
 	 * True if any task has bShouldCallTickOnlyOnEvents.
 	 * No effect if bHasGlobalTickTasks is true.
-	 * Not ticking implies no property copy.
 	 */
 	uint8 bHasTickTasksOnlyOnEvents : 1 = false;
 
+	/** True if the state contains transitions with tick trigger. */
+	UPROPERTY()
+	uint8 bHasTickTriggerTransitions : 1 = false;
+
+	/** True if the state contains transitions with event trigger. */
+	UPROPERTY()	
+	uint8 bHasEventTriggerTransitions : 1 = false;
+
+	/** True if the state contains transitions with completed trigger. */
+	UPROPERTY()
+	uint8 bHasCompletedTriggerTransitions : 1 = false;
+
+	/** True if the state contains transitions with Succeeded trigger. */
+	UPROPERTY()
+	uint8 bHasSucceededTriggerTransitions : 1 = false;
+	
+	/** True if the state contains transitions with Succeeded trigger. */
+	UPROPERTY()
+	uint8 bHasFailedTriggerTransitions : 1 = false;
+
 	/** Should state's required event and enter conditions be evaluated when transition leads directly to it's child. */
 	UPROPERTY()
 	uint8 bCheckPrerequisitesWhenActivatingChildDirectly : 1 = false;
