[State Tree] Fixed a few glitches in transitioning logic
- Fixed when the active path doesn't have any tickable tasks, on-completion transitions in subsequent frames will be ignored
- Fixed when there are only disabled task on the active path, on-completion transitions will be ignored
- Fixed unnecessary looping through parent frames when a on-completion transition has already been found that triggers faulty behaviors in ST Debugger

diff --git a/Engine/Plugins/Runtime/StateTree/Source/StateTreeEditorModule/Private/StateTreeCompiler.cpp b/Engine/Plugins/Runtime/StateTree/Source/StateTreeEditorModule/Private/StateTreeCompiler.cpp
--- a/Engine/Plugins/Runtime/StateTree/Source/StateTreeEditorModule/Private/StateTreeCompiler.cpp
+++ b/Engine/Plugins/Runtime/StateTree/Source/StateTreeEditorModule/Private/StateTreeCompiler.cpp
@@ -752,6 +752,7 @@ bool FStateTreeCompiler::CreateStateTasksAndParameters()
 		}
 		
 		bool bCreateTaskSucceeded = true;
+		int32 EnabledTasksNum = 0;
 		for (FStateTreeEditorNode& TaskNode : Tasks)
 		{
 			// Silently ignore empty nodes.
@@ -760,6 +761,11 @@ bool FStateTreeCompiler::CreateStateTasksAndParameters()
 				continue;
 			}
 
+			if(TaskNode.bIsNodeEnabled)
+			{
+				EnabledTasksNum += 1;
+			}
+
 			const FStateTreeDataHandle TaskDataHandle(EStateTreeDataSourceType::ActiveInstanceData, InstanceDataIndex++, CompactStateHandle);
 			if (!CreateTask(State, TaskNode, TaskDataHandle))
 			{
@@ -789,6 +795,7 @@ bool FStateTreeCompiler::CreateStateTasksAndParameters()
 		}
 
 		CompactState.TasksNum = uint8(TasksNum);
+		CompactState.EnabledTasksNum = uint8(EnabledTasksNum);
 		CompactState.InstanceDataNum = uint8(InstanceDataNum);
 
 		if (!CreateBindingsForNodes(Tasks, FStateTreeIndex16(TasksBegin), InstanceStructs))
diff --git a/Engine/Plugins/Runtime/StateTree/Source/StateTreeModule/Private/StateTreeExecutionContext.cpp b/Engine/Plugins/Runtime/StateTree/Source/StateTreeModule/Private/StateTreeExecutionContext.cpp
--- a/Engine/Plugins/Runtime/StateTree/Source/StateTreeModule/Private/StateTreeExecutionContext.cpp
+++ b/Engine/Plugins/Runtime/StateTree/Source/StateTreeModule/Private/StateTreeExecutionContext.cpp
@@ -2654,7 +2654,7 @@ EStateTreeRunStatus FStateTreeExecutionContext::TickTasks(const float DeltaTime)
 	}
 
 	EStateTreeRunStatus Result = EStateTreeRunStatus::Running;
-	int32 NumTotalTasks = 0;
+	int32 NumTotalEnabledTasks = 0;
 
 	const bool bHasEvents = EventQueue && EventQueue->HasEvents();
 
@@ -2747,7 +2747,8 @@ EStateTreeRunStatus FStateTreeExecutionContext::TickTasks(const float DeltaTime)
 					STATETREE_TRACE_TASK_EVENT(TaskIndex, TaskInstanceView,
 						TaskResult != EStateTreeRunStatus::Running ? EStateTreeTraceEventType::OnTaskCompleted : EStateTreeTraceEventType::OnTicked,
 						TaskResult);
-				
+
+
 					// TODO: Add more control over which states can control the failed/succeeded result.
 					if (TaskResult != EStateTreeRunStatus::Running)
 					{
@@ -2767,16 +2768,17 @@ EStateTreeRunStatus FStateTreeExecutionContext::TickTasks(const float DeltaTime)
 				}
 				ensureMsgf(bTaskTicked, TEXT("ShouldTickTasks is true but no tasks were tick for State: '%s' inside theStateTree %s."), *State.Name.ToString(), *CurrentStateTree->GetPathName());
 			}
-			NumTotalTasks += State.TasksNum;
+
+			NumTotalEnabledTasks += State.EnabledTasksNum;
 		}
 	}
 
-	if (NumTotalTasks == 0)
+	if (NumTotalEnabledTasks == 0)
 	{
-		// No tasks, done ticking.
+		// No enabled tasks, done ticking. Check from the bottom state in the bottom frame to trigger completion transitions
 		Result = EStateTreeRunStatus::Succeeded;
-		Exec.CompletedFrameIndex = FStateTreeIndex16(0);
-		Exec.CompletedStateHandle = Exec.ActiveFrames.Num() > 0 ? Exec.ActiveFrames[0].ActiveStates.Last() : FStateTreeStateHandle::Invalid;
+		Exec.CompletedFrameIndex = Exec.ActiveFrames.Num() > 0 ? FStateTreeIndex16(Exec.ActiveFrames.Num() - 1) : FStateTreeIndex16::Invalid;
+		Exec.CompletedStateHandle = Exec.ActiveFrames.Num() > 0 ? Exec.ActiveFrames.Last().ActiveStates.Last() : FStateTreeStateHandle::Invalid;
 	}
 
 	return Result;
@@ -3563,6 +3565,13 @@ bool FStateTreeExecutionContext::TriggerTransitions()
 					}
 				}
 			}
+
+			// if a valid completion transition has already been found, the remaining transitions in parent frames won't have a higher priority than the found one
+			// so skip the remainder. this also prevented false positive warnings and ensures from STDebugger
+			if (NextTransition.Priority != EStateTreeTransitionPriority::None)
+			{
+				break;
+			}
 		}
 
 		// Handle the case where no transition was found.
diff --git a/Engine/Plugins/Runtime/StateTree/Source/StateTreeModule/Public/StateTreeTypes.h b/Engine/Plugins/Runtime/StateTree/Source/StateTreeModule/Public/StateTreeTypes.h
--- a/Engine/Plugins/Runtime/StateTree/Source/StateTreeModule/Public/StateTreeTypes.h
+++ b/Engine/Plugins/Runtime/StateTree/Source/StateTreeModule/Public/StateTreeTypes.h
@@ -772,6 +772,13 @@ struct STATETREEMODULE_API FCompactStateTreeState
 	UPROPERTY()
 	uint8 TasksNum = 0;
 
+	/**
+	 * Number of enabled tasks
+	 * todo: this should be removed once we finished only compiling enabled elements for StateTree Compiler 
+	 */
+	UPROPERTY()
+	uint8 EnabledTasksNum = 0;
+
 	/** Number of instance data */
 	UPROPERTY()
 	uint8 InstanceDataNum = 0;
