[State Tree] Fix Compiler error for using the non-existent flag

diff --git a/Engine/Plugins/Runtime/StateTree/Source/StateTreeEditorModule/Private/StateTreeCompiler.cpp b/Engine/Plugins/Runtime/StateTree/Source/StateTreeEditorModule/Private/StateTreeCompiler.cpp
--- a/Engine/Plugins/Runtime/StateTree/Source/StateTreeEditorModule/Private/StateTreeCompiler.cpp
+++ b/Engine/Plugins/Runtime/StateTree/Source/StateTreeEditorModule/Private/StateTreeCompiler.cpp
@@ -761,7 +761,8 @@ bool FStateTreeCompiler::CreateStateTasksAndParameters()
 				continue;
 			}
 
-			if(TaskNode.bIsNodeEnabled)
+			FStateTreeTaskBase& Task = TaskNode.Node.GetMutable<FStateTreeTaskBase>();
+			if(Task.bTaskEnabled)
 			{
 				EnabledTasksNum += 1;
 			}
