[Backout] - CL37786887
Original CL Desc
-----------------------------------------------------------------
StateTree: Add flag to detect if the transitions need ticking.

#rnx

diff --git a/Engine/Plugins/Runtime/StateTree/Source/StateTreeEditorModule/Private/StateTreeCompiler.cpp b/Engine/Plugins/Runtime/StateTree/Source/StateTreeEditorModule/Private/StateTreeCompiler.cpp
--- a/Engine/Plugins/Runtime/StateTree/Source/StateTreeEditorModule/Private/StateTreeCompiler.cpp
+++ b/Engine/Plugins/Runtime/StateTree/Source/StateTreeEditorModule/Private/StateTreeCompiler.cpp
@@ -1084,15 +1084,6 @@ bool FStateTreeCompiler::CreateStateTransitions()
 				BindingsCompiler.AddSourceStruct(Desc);
 			}
 
-			if (CompactTransition.bTransitionEnabled)
-			{
-				CompactState.bHasTickTriggerTransitions |= Transition.Trigger == EStateTreeTransitionTrigger::OnTick;
-				CompactState.bHasEventTriggerTransitions |= Transition.Trigger == EStateTreeTransitionTrigger::OnEvent;
-				CompactState.bHasSucceededTriggerTransitions |= Transition.Trigger == EStateTreeTransitionTrigger::OnStateCompleted;
-				CompactState.bHasSucceededTriggerTransitions |= Transition.Trigger == EStateTreeTransitionTrigger::OnStateSucceeded;
-				CompactState.bHasFailedTriggerTransitions |= Transition.Trigger == EStateTreeTransitionTrigger::OnStateFailed;
-			}
-
 			const int32 ConditionsBegin = Nodes.Num();
 			if (const auto Validation = UE::StateTree::Compiler::IsValidCount16(ConditionsBegin); Validation.DidFail())
 			{
diff --git a/Engine/Plugins/Runtime/StateTree/Source/StateTreeModule/Private/StateTree.cpp b/Engine/Plugins/Runtime/StateTree/Source/StateTreeModule/Private/StateTree.cpp
--- a/Engine/Plugins/Runtime/StateTree/Source/StateTreeModule/Private/StateTree.cpp
+++ b/Engine/Plugins/Runtime/StateTree/Source/StateTreeModule/Private/StateTree.cpp
@@ -190,8 +190,6 @@ void UStateTree::ResetCompiled()
 	GlobalTasksBegin = 0;
 	GlobalTasksNum = 0;
 	bHasGlobalTransitionTasks = false;
-	bHasGlobalTickTasks = false;
-	bHasGlobalTickTasksOnlyOnEvents = false;
 	
 	ResetLinked();
 
diff --git a/Engine/Plugins/Runtime/StateTree/Source/StateTreeModule/Private/StateTreeExecutionContext.cpp b/Engine/Plugins/Runtime/StateTree/Source/StateTreeModule/Private/StateTreeExecutionContext.cpp
--- a/Engine/Plugins/Runtime/StateTree/Source/StateTreeModule/Private/StateTreeExecutionContext.cpp
+++ b/Engine/Plugins/Runtime/StateTree/Source/StateTreeModule/Private/StateTreeExecutionContext.cpp
@@ -3217,8 +3217,6 @@ bool FStateTreeExecutionContext::TriggerTransitions()
 
 	if (Exec.ActiveFrames.Num() > 0)
 	{
-		// Re-cache bHasEvents, RequestTransition above can create new events.
-		const bool bHasEvents = EventQueue && EventQueue->HasEvents();
 		for (int32 FrameIndex = Exec.ActiveFrames.Num() - 1; FrameIndex >= 0; FrameIndex--)
 		{
 			FStateTreeExecutionFrame& CurrentFrame = Exec.ActiveFrames[FrameIndex];
@@ -3252,8 +3250,7 @@ bool FStateTreeExecutionContext::TriggerTransitions()
 				}
 				
 				// Regular transitions on state
-				//A transition task can trigger an event. We need to add the state if that is a possibility.
-				if (State.ShouldTickTransitions(bHasEvents || State.bHasTransitionTasks))
+				if (State.TransitionsNum > 0)
 				{
 					TransitionHandlers.Emplace((uint8)FrameIndex, StateHandle, EStateTreeTransitionPriority::Normal);
 				}
@@ -3321,10 +3318,10 @@ bool FStateTreeExecutionContext::TriggerTransitions()
 			const FCompactStateTreeState& State = CurrentStateTree->States[Handler.StateHandle.Index];
 
 			// Transitions
-			for (uint8 TransitionCounter = 0; TransitionCounter < State.TransitionsNum; ++TransitionCounter)
+			for (uint8 i = 0; i < State.TransitionsNum; i++)
 			{
 				// All transition conditions must pass
-				const int16 TransitionIndex = State.TransitionsBegin + TransitionCounter;
+				const int16 TransitionIndex = State.TransitionsBegin + i;
 				const FCompactStateTransition& Transition = CurrentStateTree->Transitions[TransitionIndex];
 
 				// Skip disabled transitions
@@ -3397,11 +3394,7 @@ bool FStateTreeExecutionContext::TriggerTransitions()
 					// Dummy event to make sure we iterate to loop below once.
 					TransitionEvents.Emplace(nullptr);
 				}
-				else
-				{
-					ensureMsgf(false, TEXT("The trigger type is not supported."));
-				}
-
+				
 				for (const FStateTreeSharedEvent* TransitionEvent : TransitionEvents)
 				{
 					bool bPassed = false; 
@@ -3513,56 +3506,53 @@ bool FStateTreeExecutionContext::TriggerTransitions()
 			}
 			
 			const EStateTreeTransitionTrigger CompletionTrigger = Exec.LastTickStatus == EStateTreeRunStatus::Succeeded ? EStateTreeTransitionTrigger::OnStateSucceeded : EStateTreeTransitionTrigger::OnStateFailed;
-
+		
 			// Check completion transitions
 			for (int32 StateIndex = StateStartIndex; StateIndex >= 0; StateIndex--)
 			{
 				const FStateTreeStateHandle StateHandle = CurrentFrame.ActiveStates[StateIndex];
 				const FCompactStateTreeState& State = CurrentStateTree->States[StateHandle.Index];
 
-				if (State.ShouldTickCompletionTransitions(CompletionTrigger == EStateTreeTransitionTrigger::OnStateSucceeded))
+				FCurrentlyProcessedStateScope StateScope(*this, StateHandle);
+				STATETREE_TRACE_SCOPED_STATE_PHASE(StateHandle, EStateTreeUpdatePhase::TriggerTransitions);
+
+				for (uint8 i = 0; i < State.TransitionsNum; i++)
 				{
-					FCurrentlyProcessedStateScope StateScope(*this, StateHandle);
-					STATETREE_TRACE_SCOPED_STATE_PHASE(StateHandle, EStateTreeUpdatePhase::TriggerTransitions);
+					// All transition conditions must pass
+					const int16 TransitionIndex = State.TransitionsBegin + i;
+					const FCompactStateTransition& Transition = CurrentStateTree->Transitions[TransitionIndex];
 
-					for (uint8 TransitionCounter = 0; TransitionCounter < State.TransitionsNum; ++TransitionCounter)
+					// Skip disabled transitions
+					if (Transition.bTransitionEnabled == false)
 					{
-						// All transition conditions must pass
-						const int16 TransitionIndex = State.TransitionsBegin + TransitionCounter;
-						const FCompactStateTransition& Transition = CurrentStateTree->Transitions[TransitionIndex];
+						continue;
+					}
 
-						// Skip disabled transitions
-						if (Transition.bTransitionEnabled == false)
+					if (EnumHasAnyFlags(Transition.Trigger, CompletionTrigger))
+					{
+						bool bPassed = false;
 						{
-							continue;
+							STATETREE_TRACE_TRANSITION_EVENT(FStateTreeTransitionSource(FStateTreeIndex16(TransitionIndex), Transition.State, Transition.Priority), EStateTreeTraceEventType::OnEvaluating);
+							STATETREE_TRACE_SCOPED_PHASE(EStateTreeUpdatePhase::TransitionConditions);
+							bPassed = TestAllConditions(CurrentParentFrame, CurrentFrame, Transition.ConditionsBegin, Transition.ConditionsNum);
 						}
 
-						if (EnumHasAnyFlags(Transition.Trigger, CompletionTrigger))
+						if (bPassed)
 						{
-							bool bPassed = false;
+							// No delay allowed on completion conditions.
+							// No priority on completion transitions, use the priority to signal that state is selected.
+							if (RequestTransition(CurrentFrame, Transition.State, EStateTreeTransitionPriority::Normal, /*TransitionEvent*/nullptr, Transition.Fallback))
 							{
-								STATETREE_TRACE_TRANSITION_EVENT(FStateTreeTransitionSource(FStateTreeIndex16(TransitionIndex), Transition.State, Transition.Priority), EStateTreeTraceEventType::OnEvaluating);
-								STATETREE_TRACE_SCOPED_PHASE(EStateTreeUpdatePhase::TransitionConditions);
-								bPassed = TestAllConditions(CurrentParentFrame, CurrentFrame, Transition.ConditionsBegin, Transition.ConditionsNum);
-							}
-
-							if (bPassed)
-							{
-								// No delay allowed on completion conditions.
-								// No priority on completion transitions, use the priority to signal that state is selected.
-								if (RequestTransition(CurrentFrame, Transition.State, EStateTreeTransitionPriority::Normal, /*TransitionEvent*/nullptr, Transition.Fallback))
-								{
-									NextTransitionSource = FStateTreeTransitionSource(FStateTreeIndex16(TransitionIndex), Transition.State, Transition.Priority);
-									break;
-								}
+								NextTransitionSource = FStateTreeTransitionSource(FStateTreeIndex16(TransitionIndex), Transition.State, Transition.Priority);
+								break;
 							}
 						}
 					}
+				}
 
-					if (NextTransition.Priority != EStateTreeTransitionPriority::None) //-V547
-					{
-						break;
-					}
+				if (NextTransition.Priority != EStateTreeTransitionPriority::None) //-V547
+				{
+					break;
 				}
 			}
 
diff --git a/Engine/Plugins/Runtime/StateTree/Source/StateTreeModule/Public/StateTreeTypes.h b/Engine/Plugins/Runtime/StateTree/Source/StateTreeModule/Public/StateTreeTypes.h
--- a/Engine/Plugins/Runtime/StateTree/Source/StateTreeModule/Public/StateTreeTypes.h
+++ b/Engine/Plugins/Runtime/StateTree/Source/StateTreeModule/Public/StateTreeTypes.h
@@ -678,20 +678,6 @@ struct STATETREEMODULE_API FCompactStateTreeState
 		return bHasTickTasks || (bHasEvent && bHasTickTasksOnlyOnEvents);
 	}
 
-	/** @return True if the state has any transitions that need ticking. */
-	bool ShouldTickTransitions(bool bHasEvent) const
-	{
-		return bHasTickTriggerTransitions || (bHasEvent && bHasEventTriggerTransitions);
-	}
-
-	/** @return True if the state has any transitions that tick on completion. */
-	bool ShouldTickCompletionTransitions(bool bSucceeded) const
-	{
-		return bHasCompletedTriggerTransitions
-			|| (bHasSucceededTriggerTransitions && bSucceeded)
-			|| (bHasFailedTriggerTransitions && !bSucceeded);
-	}
-
 	/** Description of an event required to enter the state. */
 	UPROPERTY()
 	FCompactEventDesc RequiredEventToEnter;
@@ -803,35 +789,19 @@ struct STATETREEMODULE_API FCompactStateTreeState
 	UPROPERTY()
 	uint8 bHasStateChangeConditions : 1 = false;
 
-	/** True if any task has bShouldCallTick. */
+	/**
+	 * True if any task has bShouldCallTick.
+	 * Not ticking implies no property copy.
+	 */
 	uint8 bHasTickTasks : 1 = false;
 
 	/**
 	 * True if any task has bShouldCallTickOnlyOnEvents.
 	 * No effect if bHasGlobalTickTasks is true.
+	 * Not ticking implies no property copy.
 	 */
 	uint8 bHasTickTasksOnlyOnEvents : 1 = false;
 
-	/** True if the state contains transitions with tick trigger. */
-	UPROPERTY()
-	uint8 bHasTickTriggerTransitions : 1 = false;
-
-	/** True if the state contains transitions with event trigger. */
-	UPROPERTY()	
-	uint8 bHasEventTriggerTransitions : 1 = false;
-
-	/** True if the state contains transitions with completed trigger. */
-	UPROPERTY()
-	uint8 bHasCompletedTriggerTransitions : 1 = false;
-
-	/** True if the state contains transitions with Succeeded trigger. */
-	UPROPERTY()
-	uint8 bHasSucceededTriggerTransitions : 1 = false;
-	
-	/** True if the state contains transitions with Succeeded trigger. */
-	UPROPERTY()
-	uint8 bHasFailedTriggerTransitions : 1 = false;
-
 	/** Should state's required event and enter conditions be evaluated when transition leads directly to it's child. */
 	UPROPERTY()
 	uint8 bCheckPrerequisitesWhenActivatingChildDirectly : 1 = false;
