StateTree: Add unittest for linked tree with empty state.

diff --git a/Engine/Plugins/Runtime/StateTree/Source/StateTreeTestSuite/Private/StateTreeLinkedStateTest.cpp b/Engine/Plugins/Runtime/StateTree/Source/StateTreeTestSuite/Private/StateTreeLinkedStateTest.cpp
--- a/Engine/Plugins/Runtime/StateTree/Source/StateTreeTestSuite/Private/StateTreeLinkedStateTest.cpp
+++ b/Engine/Plugins/Runtime/StateTree/Source/StateTreeTestSuite/Private/StateTreeLinkedStateTest.cpp
@@ -154,17 +154,17 @@ struct FStateTreeTest_MultipleSameLinkedAsset : FStateTreeTestBase
 	virtual bool InstantTest() override
 	{
 		//Tree 1
-		//	RootA
+		//	Root
+		//		StateA -> Next
 		//		StateB -> Next
-		//		StateC -> Next
-		//		StateLinkedTreeD1 (Tree2) -> Next
-		//		StateLinkedTreeD2 (Tree2) -> Next
-		//		StateLinkedTreeD3 (Tree2) -> Next
-		//		StateE -> Succeeded
+		//		StateLinkedTreeA (Tree2) -> Next
+		//		StateLinkedTreeB (Tree2) -> Next
+		//		StateLinkedTreeC (Tree2) -> Next
+		//		StateC -> Succeeded
 		//Tree 2
 		//  Global task and parameter
 		//	RootE
-		//		StateF (with transition OnTick to succeeded)
+		//		StateA (with transition OnTick to succeeded)
 
 		FStateTreeCompilerLog Log;
 
@@ -179,19 +179,24 @@ struct FStateTreeTest_MultipleSameLinkedAsset : FStateTreeTestBase
 				EditorData2.RootParameters.Parameters.SetValueInt32("Value", -111);
 				RootParameter_ValueID = EditorData2.RootParameters.Parameters.FindPropertyDescByName("Value")->ID;
 
-				TStateTreeEditorNode<FTestTask_PrintValue>& GlobalTask = EditorData2.AddGlobalTask<FTestTask_PrintValue>("Tree2Global");
+				TStateTreeEditorNode<FTestTask_PrintValue>& GlobalTask = EditorData2.AddGlobalTask<FTestTask_PrintValue>("Tree2GlobalTaskA");
 				GlobalTask.GetInstanceData().Value = -1;
 				EditorData2.AddPropertyBinding(FStateTreePropertyPath(EditorData2.RootParameters.ID, TEXT("Value")), FStateTreePropertyPath(GlobalTask.ID, TEXT("Value")));
 			}
 
-			UStateTreeState& RootC = EditorData2.AddSubTree("RootE");
+			UStateTreeState& Root = EditorData2.AddSubTree("Tree2StateRoot");
 			{
-				UStateTreeState& StateE = RootC.AddChildState("StateF", EStateTreeStateType::State);
-				FStateTreeTransition& Transition = StateE.AddTransition(EStateTreeTransitionTrigger::OnTick, EStateTreeTransitionType::Succeeded);
+				TStateTreeEditorNode<FTestTask_PrintValue>& Task1 = Root.AddTask<FTestTask_PrintValue>("Tree2StateRootTaskA");
+				Task1.GetInstanceData().Value = -2;
+				EditorData2.AddPropertyBinding(FStateTreePropertyPath(EditorData2.RootParameters.ID, TEXT("Value")), FStateTreePropertyPath(Task1.ID, TEXT("Value")));
+			}
+			{
+				UStateTreeState& State = Root.AddChildState("Tree2StateA", EStateTreeStateType::State);
+				FStateTreeTransition& Transition = State.AddTransition(EStateTreeTransitionTrigger::OnTick, EStateTreeTransitionType::Succeeded);
 				Transition.bDelayTransition = true;
 				Transition.DelayDuration = 1.0;
 
-				TStateTreeEditorNode<FTestTask_PrintValue>& Task1 = StateE.AddTask<FTestTask_PrintValue>("Task1");
+				TStateTreeEditorNode<FTestTask_PrintValue>& Task1 = State.AddTask<FTestTask_PrintValue>("Tree2StateATaskA");
 				Task1.GetInstanceData().Value = -2;
 				EditorData2.AddPropertyBinding(FStateTreePropertyPath(EditorData2.RootParameters.ID, TEXT("Value")), FStateTreePropertyPath(Task1.ID, TEXT("Value")));
 			}
@@ -205,52 +210,56 @@ struct FStateTreeTest_MultipleSameLinkedAsset : FStateTreeTestBase
 		UStateTree& StateTree1 = NewStateTree();
 		{
 			UStateTreeEditorData& EditorData = *Cast<UStateTreeEditorData>(StateTree1.EditorData);
-
-			UStateTreeState& Root = EditorData.AddSubTree("RootA");
 			{
-				TStateTreeEditorNode<FTestTask_PrintValue>& GlobalTask = EditorData.AddGlobalTask<FTestTask_PrintValue>("Tree1Global");
+				TStateTreeEditorNode<FTestTask_PrintValue>& GlobalTask = EditorData.AddGlobalTask<FTestTask_PrintValue>("Tree1GlobalTaskA");
 				GlobalTask.GetInstanceData().Value = 99;
 			}
+
+			UStateTreeState& Root = EditorData.AddSubTree("Tree1StateRoot");
+			{
+				TStateTreeEditorNode<FTestTask_PrintValue>& Task1 = Root.AddTask<FTestTask_PrintValue>("Tree1StateRootTaskA");
+				Task1.GetInstanceData().Value = 88;
+			}
 			{
-				UStateTreeState& StateB = Root.AddChildState("StateB", EStateTreeStateType::State);
-				TStateTreeEditorNode<FTestTask_PrintValue>& Task = StateB.AddTask<FTestTask_PrintValue>("TaskB");
+				UStateTreeState& StateB = Root.AddChildState("Tree1StateA", EStateTreeStateType::State);
+				TStateTreeEditorNode<FTestTask_PrintValue>& Task = StateB.AddTask<FTestTask_PrintValue>("Tree1StateATaskA");
 				Task.GetInstanceData().Value = 1;
 				FStateTreeTransition& Transition = StateB.AddTransition(EStateTreeTransitionTrigger::OnTick, EStateTreeTransitionType::NextState);
 				Transition.bDelayTransition = true;
 				Transition.DelayDuration = 1.0;
 			}
 			{
-				UStateTreeState& StateB = Root.AddChildState("StateC", EStateTreeStateType::State);
-				TStateTreeEditorNode<FTestTask_PrintValue>& Task = StateB.AddTask<FTestTask_PrintValue>("TaskC");
+				UStateTreeState& StateB = Root.AddChildState("Tree1StateB", EStateTreeStateType::State);
+				TStateTreeEditorNode<FTestTask_PrintValue>& Task = StateB.AddTask<FTestTask_PrintValue>("Tree1StateBTaskA");
 				Task.GetInstanceData().Value = 2;
 				FStateTreeTransition& Transition = StateB.AddTransition(EStateTreeTransitionTrigger::OnTick, EStateTreeTransitionType::NextState);
 				Transition.bDelayTransition = true;
 				Transition.DelayDuration = 1.0;
 			}
 			{
-				UStateTreeState& C1 = Root.AddChildState("StateLinkedTreeD1", EStateTreeStateType::LinkedAsset);
+				UStateTreeState& C1 = Root.AddChildState("Tree1StateLinkedTreeA", EStateTreeStateType::LinkedAsset);
 				C1.SetLinkedStateAsset(&StateTree2);
 				C1.AddTransition(EStateTreeTransitionTrigger::OnStateSucceeded, EStateTreeTransitionType::NextState);
 				C1.SetParametersPropertyOverridden(RootParameter_ValueID, true);
 				C1.Parameters.Parameters.SetValueInt32("Value", 111);
 			}
 			{
-				UStateTreeState& C2 = Root.AddChildState("StateLinkedTreeD2", EStateTreeStateType::LinkedAsset);
+				UStateTreeState& C2 = Root.AddChildState("Tree1StateLinkedTreeB", EStateTreeStateType::LinkedAsset);
 				C2.SetLinkedStateAsset(&StateTree2);
 				C2.AddTransition(EStateTreeTransitionTrigger::OnStateSucceeded, EStateTreeTransitionType::NextState);
 				C2.SetParametersPropertyOverridden(RootParameter_ValueID, true);
 				C2.Parameters.Parameters.SetValueInt32("Value", 222);
 			}
 			{
-				UStateTreeState& C3 = Root.AddChildState("StateLinkedTreeD3", EStateTreeStateType::LinkedAsset);
+				UStateTreeState& C3 = Root.AddChildState("Tree1StateLinkedTreeC", EStateTreeStateType::LinkedAsset);
 				C3.SetLinkedStateAsset(&StateTree2);
 				C3.AddTransition(EStateTreeTransitionTrigger::OnStateSucceeded, EStateTreeTransitionType::NextState);
 				C3.SetParametersPropertyOverridden(RootParameter_ValueID, true);
 				C3.Parameters.Parameters.SetValueInt32("Value", 333);
 			}
 			{
-				UStateTreeState& StateD = Root.AddChildState("StateE", EStateTreeStateType::State);
-				TStateTreeEditorNode<FTestTask_PrintValue>& Task = StateD.AddTask<FTestTask_PrintValue>("TaskE");
+				UStateTreeState& StateD = Root.AddChildState("Tree1StateC", EStateTreeStateType::State);
+				TStateTreeEditorNode<FTestTask_PrintValue>& Task = StateD.AddTask<FTestTask_PrintValue>("Tree1StateBTaskA");
 				Task.GetInstanceData().Value = 3;
 				FStateTreeTransition& Transition = StateD.AddTransition(EStateTreeTransitionTrigger::OnTick, EStateTreeTransitionType::GotoState, &Root);
 				Transition.bDelayTransition = true;
@@ -274,55 +283,83 @@ struct FStateTreeTest_MultipleSameLinkedAsset : FStateTreeTestBase
 
 			Status = Exec.Start();
 			AITEST_EQUAL("Start should complete with Running", Status, EStateTreeRunStatus::Running);
-			AITEST_TRUE("Start should enter Global tasks", Exec.Expect("Tree1Global", TEXT("EnterState99")));
-			AITEST_TRUE("Start should enter StateB", Exec.Expect("TaskB", TEXT("EnterState1")));
+			AITEST_TRUE("Start should enter Global tasks", Exec.Expect("Tree1GlobalTaskA", TEXT("EnterState99"))
+				.Then("Tree1StateRootTaskA", TEXT("EnterState88"))
+				.Then("Tree1StateATaskA", TEXT("EnterState1")));
 			Exec.LogClear();
 
 			Status = Exec.Tick(1.5f); // over tick, should trigger
 			AITEST_EQUAL("1st Tick should complete with Running", Status, EStateTreeRunStatus::Running);
-			AITEST_TRUE("1st should tick Global tasks", Exec.Expect("Tree1Global", TEXT("Tick99")));
-			AITEST_TRUE("1st Tick tick the StateB", Exec.Expect("TaskB", TEXT("Tick1")));
+			AITEST_TRUE("1st should tick tasks", Exec.Expect("Tree1GlobalTaskA", TEXT("Tick99"))
+				.Then("Tree1StateRootTaskA", TEXT("Tick88"))
+				.Then("Tree1StateATaskA", TEXT("Tick1")));
 			Exec.LogClear();
 
 			Status = Exec.Tick(1.0f);
 			AITEST_EQUAL("2nd Tick should complete with Running", Status, EStateTreeRunStatus::Running);
-			AITEST_TRUE("2nd Tick should tick Global tasks", Exec.Expect("Tree1Global", TEXT("Tick99")));
-			AITEST_TRUE("2nd Tick should tick the StateB", Exec.Expect("TaskB", TEXT("Tick1")));
-			AITEST_TRUE("2nd Tick should exit the StateB", Exec.Expect("TaskB", TEXT("ExitState1")));
-			AITEST_TRUE("2nd Tick should enter the StateC", Exec.Expect("TaskC", TEXT("EnterState2")));
+			AITEST_TRUE("2nd Tick should tick tasks", Exec.Expect("Tree1GlobalTaskA", TEXT("Tick99"))
+				.Then("Tree1StateRootTaskA", TEXT("Tick88"))
+				.Then("Tree1StateATaskA", TEXT("Tick1"))
+				.Then("Tree1StateATaskA", TEXT("ExitState1"))
+				.Then("Tree1StateBTaskA", TEXT("EnterState2")));
 			Exec.LogClear();
 
 			Status = Exec.Tick(1.0f);
 			AITEST_EQUAL("3rd Tick should complete with Running", Status, EStateTreeRunStatus::Running);
-			AITEST_TRUE("3rd should tick Global tasks", Exec.Expect("Tree1Global", TEXT("Tick99")));
-			AITEST_TRUE("3rd Tick tick the StateB", Exec.Expect("TaskC", TEXT("Tick2")));
+			AITEST_TRUE("3rd should tick tasks", Exec.Expect("Tree1GlobalTaskA", TEXT("Tick99"))
+				.Then("Tree1StateRootTaskA", TEXT("Tick88"))
+				.Then("Tree1StateBTaskA", TEXT("Tick2")));
 			Exec.LogClear();
 
 			Status = Exec.Tick(1.0f);
 			AITEST_EQUAL("4th Tick should complete with Running", Status, EStateTreeRunStatus::Running);
-			AITEST_TRUE("4th Tick should tick Global tasks", Exec.Expect("Tree1Global", TEXT("Tick99")));
-			AITEST_TRUE("4th Tick should tick the StateB", Exec.Expect("TaskC", TEXT("Tick2")));
-			//AITEST_TRUE("4th Tick should enter the Global", Exec.Expect("Tree2Global", TEXT("EnterState111")));
-			AITEST_TRUE("4th Tick should exit the StateB", Exec.Expect("TaskC", TEXT("ExitState2")));
-			AITEST_TRUE("4th Tick should enter the StateLinkedTreeC1", Exec.Expect("Task1", TEXT("EnterState111")));
+			AITEST_TRUE("4th Tick should tick tasks", Exec.Expect("Tree1GlobalTaskA", TEXT("Tick99"))
+				.Then("Tree1StateRootTaskA", TEXT("Tick88"))
+				.Then("Tree1StateBTaskA", TEXT("Tick2")));
+			//bug: binding are not executed
+			//AITEST_TRUE("4th Tick should enter the Global", Exec.Expect("Tree2GlobalTaskA", TEXT("EnterState111")));
+			AITEST_TRUE("4th Tick should exit the State", Exec.Expect("Tree1StateBTaskA", TEXT("ExitState2")));
+			AITEST_TRUE("4th Tick should enter the Root", Exec.Expect("Tree2StateRootTaskA", TEXT("EnterState111")));
+			AITEST_TRUE("4th Tick should enter the LinkedTree", Exec.Expect("Tree2StateATaskA", TEXT("EnterState111")));
 			Exec.LogClear();
 
 			Status = Exec.Tick(0.001f);
 			AITEST_EQUAL("5th Tick should complete with Running", Status, EStateTreeRunStatus::Running);
-			AITEST_TRUE("5h Tick should tick the global tasks", Exec.Expect("Tree1Global", TEXT("Tick99")));
-			AITEST_TRUE("5th Tick should tick the global tasks", Exec.Expect("Tree2Global", TEXT("Tick111")));
-			AITEST_TRUE("5th Tick should tick the StateLinkedTreeC1", Exec.Expect("Task1", TEXT("Tick111")));
+			AITEST_TRUE("5h Tick should tick the tasks", Exec.Expect("Tree1GlobalTaskA", TEXT("Tick99")));
+			//bug: the task should be tick in the correct order
+			{
+				//.Then("Tree1StateRootTaskA", TEXT("Tick88"))
+				//.Then("Tree2GlobalTaskA", TEXT("Tick111"))
+				//.Then("Tree2StateRootTaskA", TEXT("Tick111"))
+				//.Then("Tree2StateATaskA", TEXT("Tick111")));
+				AITEST_TRUE("5th Tick should exit the State", Exec.Expect("Tree1StateRootTaskA", TEXT("Tick88")));
+				AITEST_TRUE("5th Tick should enter the Root", Exec.Expect("Tree2GlobalTaskA", TEXT("Tick111")));
+				AITEST_TRUE("5th Tick should enter the State", Exec.Expect("Tree2StateRootTaskA", TEXT("Tick111")));
+				AITEST_TRUE("5th Tick should enter the LinkedTree", Exec.Expect("Tree2StateATaskA", TEXT("Tick111")));
+			}
 			Exec.LogClear();
 
 			Status = Exec.Tick(1.0f);
-			AITEST_EQUAL("5th Tick should complete with Running", Status, EStateTreeRunStatus::Running);
-			AITEST_TRUE("5th Tick should tick the global tasks", Exec.Expect("Tree1Global", TEXT("Tick99")));
-			AITEST_TRUE("5th Tick should tick the global tasks", Exec.Expect("Tree2Global", TEXT("Tick111")));
-			AITEST_TRUE("5th Tick should tick the StateLinkedTreeC1", Exec.Expect("Task1", TEXT("Tick111")));
-			//AITEST_TRUE("5th Tick should exit the global tasks", Exec.Expect("Tree2Global", TEXT("ExitState111")));
-			//AITEST_TRUE("5th Tick should exit the StateLinkedTreeC1", Exec.Expect("Task1", TEXT("ExitState111")));
-			//AITEST_TRUE("5th Tick should enter the global tasks", Exec.Expect("Tree2Global", TEXT("EnterState222")));
-			AITEST_TRUE("5th Tick should enter the StateLinkedTreeC2", Exec.Expect("Task1", TEXT("EnterState222")));
+			AITEST_EQUAL("6th Tick should complete with Running", Status, EStateTreeRunStatus::Running);
+			AITEST_TRUE("6th Tick should tick the tasks", Exec.Expect("Tree1GlobalTaskA", TEXT("Tick99")));
+			//bug: the task should be tick in the correct order
+			{
+				//.Then("Tree1StateRootTaskA", TEXT("Tick88"))
+				//.Then("Tree2GlobalTaskA", TEXT("Tick111"))
+				//.Then("Tree2StateRootTaskA", TEXT("Tick111"))
+				//.Then("Tree2StateATaskA", TEXT("Tick111")));
+				AITEST_TRUE("6th Tick should exit the State", Exec.Expect("Tree1StateRootTaskA", TEXT("Tick88")));
+				AITEST_TRUE("6th Tick should enter the Root", Exec.Expect("Tree2GlobalTaskA", TEXT("Tick111")));
+				AITEST_TRUE("6th Tick should enter the Root", Exec.Expect("Tree2StateRootTaskA", TEXT("Tick111")));
+				AITEST_TRUE("6th Tick should enter the LinkedTree", Exec.Expect("Tree2StateATaskA", TEXT("Tick111")));
+			}
+			//bug: the global task are not exited
+			//AITEST_TRUE("6th Tick should exit the global tasks", Exec.Expect("Tree2GlobalTaskA", TEXT("ExitState111")));
+			AITEST_TRUE("6th Tick should exit the state", Exec.Expect("Tree2StateRootTaskA", TEXT("ExitState111")));
+			AITEST_TRUE("6th Tick should exit the StateLinkedTree", Exec.Expect("Tree2StateATaskA", TEXT("ExitState111")));
+			//AITEST_TRUE("6th Tick should enter the global tasks", Exec.Expect("Tree2GlobalTaskA", TEXT("EnterState222")));
+			AITEST_TRUE("6th Tick should enter the state", Exec.Expect("Tree2StateRootTaskA", TEXT("EnterState222")));
+			AITEST_TRUE("6th Tick should enter the StateLinkedTree", Exec.Expect("Tree2StateATaskA", TEXT("EnterState222")));
 			Exec.LogClear();
 
 			Status = Exec.Tick(1.0f);
@@ -336,6 +373,184 @@ struct FStateTreeTest_MultipleSameLinkedAsset : FStateTreeTestBase
 };
 IMPLEMENT_AI_INSTANT_TEST(FStateTreeTest_MultipleSameLinkedAsset, "System.StateTree.LinkedAsset.MultipleSameTree");
 
+struct FStateTreeTest_EmptyStateWithTickTransitionLinkedAsset : FStateTreeTestBase
+{
+	virtual bool InstantTest() override
+	{
+		//Tree 1
+		//	Root
+		//		StateA -> Next
+		//		StateLinkedTree (Tree2) -> Next
+		//		StateB -> Root
+		//Tree 2
+		//  Global task and parameter
+		//	Root
+		//		FailState (condition false)
+		//		StateA (condition true and with transition OnTick to succeeded)
+
+		FStateTreeCompilerLog Log;
+
+		// Asset 2
+		UStateTree& StateTree2 = NewStateTree();
+		{
+			UStateTreeEditorData& EditorData = *Cast<UStateTreeEditorData>(StateTree2.EditorData);
+			TStateTreeEditorNode<FTestTask_PrintValue>& GlobalTask = EditorData.AddGlobalTask<FTestTask_PrintValue>("Tree2GlobalTaskA");
+			GlobalTask.GetInstanceData().Value = 21;
+
+			UStateTreeState& Root = EditorData.AddSubTree("Tree2StateRoot");
+			{
+				TStateTreeEditorNode<FTestTask_PrintValue>& Task = Root.AddTask<FTestTask_PrintValue>("Tree2StateRootTaskA");
+				Task.GetInstanceData().Value = 22;
+			}
+			{
+				UStateTreeState& State = Root.AddChildState("Tree2StateFail", EStateTreeStateType::State);
+				// Add auto fails condition
+				auto& Condition = State.AddEnterCondition<FStateTreeTestBooleanCondition>();
+				Condition.GetInstanceData().bSuccess = false;
+
+				// Should never see
+				TStateTreeEditorNode<FTestTask_PrintValue>& Task = State.AddTask<FTestTask_PrintValue>("Tree2StateFailTaskA");
+				Task.GetInstanceData().Value = 23;
+			}
+			{
+				UStateTreeState& State = Root.AddChildState("Tree2StateB", EStateTreeStateType::State);
+
+				// Add auto success condition
+				auto& Condition = State.AddEnterCondition<FStateTreeTestBooleanCondition>();
+				Condition.GetInstanceData().bSuccess = true;
+
+				FStateTreeTransition& Transition = State.AddTransition(EStateTreeTransitionTrigger::OnTick, EStateTreeTransitionType::Succeeded);
+				Transition.bDelayTransition = true;
+				Transition.DelayDuration = 1.0;
+			}
+
+
+			FStateTreeCompiler Compiler2(Log);
+			const bool bResult2 = Compiler2.Compile(StateTree2);
+			AITEST_TRUE("StateTree2 should get compiled", bResult2);
+		}
+
+		// Main asset
+		UStateTree& StateTree1 = NewStateTree();
+		{
+			UStateTreeEditorData& EditorData = *Cast<UStateTreeEditorData>(StateTree1.EditorData);
+
+			{
+				TStateTreeEditorNode<FTestTask_PrintValue>& GlobalTask = EditorData.AddGlobalTask<FTestTask_PrintValue>("Tree1GlobalTaskA");
+				GlobalTask.GetInstanceData().Value = 11;
+			}
+			UStateTreeState& Root = EditorData.AddSubTree("Tree1StateRoot");
+			{
+				TStateTreeEditorNode<FTestTask_PrintValue>& Task = Root.AddTask<FTestTask_PrintValue>("Tree1StateRootTaskA");
+				Task.GetInstanceData().Value = 12;
+			}
+			{
+				UStateTreeState& State = Root.AddChildState("Tree1StateA", EStateTreeStateType::State);
+				TStateTreeEditorNode<FTestTask_PrintValue>& Task = State.AddTask<FTestTask_PrintValue>("Tree1StateATaskA");
+				Task.GetInstanceData().Value = 13;
+
+				FStateTreeTransition& Transition = State.AddTransition(EStateTreeTransitionTrigger::OnTick, EStateTreeTransitionType::NextState);
+				Transition.bDelayTransition = true;
+				Transition.DelayDuration = 1.0;
+			}
+			{
+				UStateTreeState& C1 = Root.AddChildState("Tree1StateLinkedTree", EStateTreeStateType::LinkedAsset);
+				C1.SetLinkedStateAsset(&StateTree2);
+				C1.AddTransition(EStateTreeTransitionTrigger::OnStateSucceeded, EStateTreeTransitionType::NextState);
+			}
+			{
+				UStateTreeState& State = Root.AddChildState("Tree1StateB", EStateTreeStateType::State);
+				TStateTreeEditorNode<FTestTask_PrintValue>& Task = State.AddTask<FTestTask_PrintValue>("Tree1StateBTaskA");
+				Task.GetInstanceData().Value = 14;
+				FStateTreeTransition& Transition = State.AddTransition(EStateTreeTransitionTrigger::OnTick, EStateTreeTransitionType::GotoState, &Root);
+				Transition.bDelayTransition = true;
+				Transition.DelayDuration = 1.0;
+			}
+
+			FStateTreeCompiler Compiler(Log);
+			const bool bResult = Compiler.Compile(StateTree1);
+			AITEST_TRUE("StateTree should get compiled", bResult);
+		}
+
+		const FString EnterStateStr(TEXT("EnterState"));
+		const FString ExitStateStr(TEXT("ExitState"));
+
+		{
+			EStateTreeRunStatus Status = EStateTreeRunStatus::Unset;
+			FStateTreeInstanceData InstanceData;
+			FTestStateTreeExecutionContext Exec(StateTree1, StateTree1, InstanceData);
+			const bool bInitSucceeded = Exec.IsValid();
+			AITEST_TRUE("StateTree should init", bInitSucceeded);
+
+			Status = Exec.Start();
+			AITEST_EQUAL("Start should complete with Running", Status, EStateTreeRunStatus::Running);
+			AITEST_TRUE("Start tick in the correct order", Exec.Expect("Tree1GlobalTaskA", TEXT("EnterState11"))
+				.Then("Tree1StateRootTaskA", TEXT("EnterState12"))
+				.Then("Tree1StateATaskA", TEXT("EnterState13")));
+			AITEST_TRUE("Start should be in the correct states", Exec.ExpectInActiveStates("Tree1StateRoot", "Tree1StateA"));
+			Exec.LogClear();
+
+			Status = Exec.Tick(1.5f); // over tick, should trigger
+			AITEST_EQUAL("1st Tick should complete with Running", Status, EStateTreeRunStatus::Running);
+			AITEST_TRUE("1st Tick should tick the Global task", Exec.Expect("Tree1GlobalTaskA", TEXT("Tick11"))
+				.Then("Tree1StateRootTaskA", TEXT("Tick12"))
+				.Then("Tree1StateATaskA", TEXT("Tick13")));
+			AITEST_TRUE("Start should be in the correct states", Exec.ExpectInActiveStates("Tree1StateRoot", "Tree1StateA"));
+			Exec.LogClear();
+
+			Status = Exec.Tick(1.0f);
+			AITEST_EQUAL("2nd Tick should complete with Running", Status, EStateTreeRunStatus::Running);
+			AITEST_TRUE("2nd Tick should tick the Global task", Exec.Expect("Tree1GlobalTaskA", TEXT("Tick11")));
+			AITEST_TRUE("2nd Tick should tick the Root task", Exec.Expect("Tree1StateRootTaskA", TEXT("Tick12")));
+			AITEST_TRUE("2nd Tick should tick the TaskA", Exec.Expect("Tree1StateATaskA", TEXT("Tick13")));
+
+			AITEST_TRUE("2nd Tick should enter the Global task", Exec.Expect("Tree2GlobalTaskA", TEXT("EnterState21")));
+			AITEST_TRUE("2nd Tick should exit the TaskA", Exec.Expect("Tree1StateATaskA", TEXT("ExitState13")));
+			AITEST_TRUE("2nd Tick should enter the Root task", Exec.Expect("Tree2StateRootTaskA", TEXT("EnterState22")));
+			AITEST_FALSE("2nd Tick should NOT enter the FailTask", Exec.Expect("Tree2StateFailTaskA"));
+
+			AITEST_TRUE("2nd Tick should be in the correct states", Exec.ExpectInActiveStates("Tree1StateRoot", "Tree1StateLinkedTree", "Tree2StateRoot", "Tree2StateB"));
+			Exec.LogClear();
+
+			Status = Exec.Tick(1.0f);
+			AITEST_EQUAL("3rd Tick should complete with Running", Status, EStateTreeRunStatus::Running);
+			AITEST_TRUE("3rd Tick should tick Global tasks", Exec.Expect("Tree1GlobalTaskA", TEXT("Tick11")));
+			AITEST_TRUE("3rd Tick should tick the Root task", Exec.Expect("Tree1StateRootTaskA", TEXT("Tick12")));
+			AITEST_TRUE("3rd Tick should tick Global tasks", Exec.Expect("Tree2GlobalTaskA", TEXT("Tick21")));
+			AITEST_TRUE("3rd Tick should tick the Root task", Exec.Expect("Tree2StateRootTaskA", TEXT("Tick22")));
+
+			AITEST_TRUE("3rd Tick should be in the correct states", Exec.ExpectInActiveStates("Tree1StateRoot", "Tree1StateLinkedTree", "Tree2StateRoot", "Tree2StateB"));
+			Exec.LogClear();
+
+			Status = Exec.Tick(1.0f);
+			AITEST_EQUAL("4th Tick should complete with Running", Status, EStateTreeRunStatus::Running);
+			AITEST_TRUE("4th Tick should tick Global tasks", Exec.Expect("Tree1GlobalTaskA", TEXT("Tick11")));
+			AITEST_TRUE("4th Tick should tick the Root task", Exec.Expect("Tree1StateRootTaskA", TEXT("Tick12")));
+			AITEST_TRUE("4th Tick should tick Global tasks", Exec.Expect("Tree2GlobalTaskA", TEXT("Tick21")));
+			AITEST_TRUE("4th Tick should tick the Root task", Exec.Expect("Tree2StateRootTaskA", TEXT("Tick22")));
+			AITEST_TRUE("4th Tick should tick Global tasks", Exec.Expect("Tree2StateRootTaskA", TEXT("ExitState22")));
+			AITEST_TRUE("4th Tick should tick Global tasks", Exec.Expect("Tree2GlobalTaskA", TEXT("ExitState21")));
+			AITEST_TRUE("4th Tick should tick the Root task", Exec.Expect("Tree1StateBTaskA", TEXT("EnterState14")));
+
+			AITEST_TRUE("4th Tick should be in the correct states", Exec.ExpectInActiveStates("Tree1StateRoot", "Tree1StateB"));
+			Exec.LogClear();
+
+			Status = Exec.Tick(0.001f);
+			AITEST_EQUAL("5th Tick should complete with Running", Status, EStateTreeRunStatus::Running);
+			AITEST_TRUE("Start tick in the correct order", Exec.Expect("Tree1GlobalTaskA", TEXT("Tick11"))
+				.Then("Tree1StateRootTaskA", TEXT("Tick12"))
+				.Then("Tree1StateBTaskA", TEXT("Tick14")));
+			AITEST_TRUE("5th Tick should be in the correct states", Exec.ExpectInActiveStates("Tree1StateRoot", "Tree1StateB"));
+			Exec.LogClear();
+
+			Exec.Stop();
+		}
+
+		return true;
+	}
+};
+IMPLEMENT_AI_INSTANT_TEST(FStateTreeTest_EmptyStateWithTickTransitionLinkedAsset, "System.StateTree.LinkedAsset.EmptyStateWithTickTransition");
+
 } // namespace UE::StateTree::Tests
 
 UE_ENABLE_OPTIMIZATION_SHIP
diff --git a/Engine/Plugins/Runtime/StateTree/Source/StateTreeTestSuite/Private/StateTreeTestTypes.h b/Engine/Plugins/Runtime/StateTree/Source/StateTreeTestSuite/Private/StateTreeTestTypes.h
--- a/Engine/Plugins/Runtime/StateTree/Source/StateTreeTestSuite/Private/StateTreeTestTypes.h
+++ b/Engine/Plugins/Runtime/StateTree/Source/StateTreeTestSuite/Private/StateTreeTestTypes.h
@@ -71,6 +71,21 @@ struct FTestStateTreeExecutionContext : public FStateTreeExecutionContext
 	{
 		FLogOrder(const FTestStateTreeExecutionContext& InContext, const int32 InIndex) : Context(InContext), Index(InIndex) {}
 
+		FLogOrder Then(const FName& Name) const
+		{
+			int32 NextIndex = Index;
+			while (NextIndex < Context.Log.Items.Num())
+			{
+				const FStateTreeTestLog::FLogItem& Item = Context.Log.Items[NextIndex];
+				if (Item.Name == Name)
+				{
+					break;
+				}
+				NextIndex++;
+			}
+			return FLogOrder(Context, NextIndex);
+		}
+
 		FLogOrder Then(const FName& Name, const FString& Message) const
 		{
 			int32 NextIndex = Index;
@@ -92,6 +107,11 @@ struct FTestStateTreeExecutionContext : public FStateTreeExecutionContext
 		int32 Index = 0;
 	};
 
+	FLogOrder Expect(const FName& Name) const
+	{
+		return FLogOrder(*this, 0).Then(Name);
+	}
+
 	FLogOrder Expect(const FName& Name, const FString& Message) const
 	{
 		return FLogOrder(*this, 0).Then(Name, Message);
@@ -648,3 +668,33 @@ struct FTestPropertyFunction : public FStateTreePropertyFunctionBase
 		InstanceData.Result = InstanceData.Input + 1;
 	}
 };
+
+USTRUCT()
+struct FStateTreeTestBooleanConditionInstanceData
+{
+	GENERATED_BODY()
+	
+	UPROPERTY(EditAnywhere, Category = Parameters)
+	bool bSuccess = true;
+};
+
+USTRUCT(meta = (Hidden))
+struct FStateTreeTestBooleanCondition : public FStateTreeConditionCommonBase
+{
+	GENERATED_BODY()
+
+	using FInstanceDataType = FStateTreeTestBooleanConditionInstanceData;
+
+	FStateTreeTestBooleanCondition() = default;
+
+	virtual const UStruct* GetInstanceDataType() const override { return FInstanceDataType::StaticStruct(); }
+	virtual bool TestCondition(FStateTreeExecutionContext& Context) const override
+	{
+		FStateTreeTestLog& TestLog = Context.GetExternalData(LogHandle);
+		FInstanceDataType& InstanceData = Context.GetInstanceData(*this);
+		TestLog.Log(Name,  FString::Printf(TEXT("TestCondition=%d"), InstanceData.bSuccess));
+		return InstanceData.bSuccess;
+	}
+
+	TStateTreeExternalDataHandle<FStateTreeTestLog> LogHandle;
+};
\ No newline at end of file
