StateTree: Add FStateTreeEventQueueExecutionContext::IsValid. Only started tree can send an event.

diff --git a/Engine/Plugins/Runtime/StateTree/Source/StateTreeModule/Private/StateTreeExecutionContext.cpp b/Engine/Plugins/Runtime/StateTree/Source/StateTreeModule/Private/StateTreeExecutionContext.cpp
--- a/Engine/Plugins/Runtime/StateTree/Source/StateTreeModule/Private/StateTreeExecutionContext.cpp
+++ b/Engine/Plugins/Runtime/StateTree/Source/StateTreeModule/Private/StateTreeExecutionContext.cpp
@@ -4839,6 +4839,11 @@ FStateTreeEventQueueExecutionContext::FStateTreeEventQueueExecutionContext(const
 {
 }
 
+bool FStateTreeEventQueueExecutionContext::IsValid() const
+{
+	return StateTree.IsReadyToRun() && Storage.GetExecutionState().TreeRunStatus == EStateTreeRunStatus::Running;
+}
+
 void FStateTreeEventQueueExecutionContext::SendEvent(const FGameplayTag Tag, const FConstStructView Payload, const FName Origin) const
 {
 	STATETREE_LOG_AND_TRACE(Verbose, Log, TEXT("Send Event '%s'"), *Tag.ToString());
@@ -4879,6 +4884,13 @@ FStateTreeWeakEventQueueExecutionContext::FStateTreeWeakEventQueueExecutionConte
 {
 }
 
+bool FStateTreeWeakEventQueueExecutionContext::IsValid() const
+{
+	TStrongObjectPtr<const UStateTree> StateTreePinned = StateTree.Pin();
+	TSharedPtr<FStateTreeInstanceStorage> InstanceStoragePinned = Storage.Pin();
+	return Owner.IsValid() && StateTreePinned && InstanceStoragePinned && StateTreePinned->IsReadyToRun() && InstanceStoragePinned->GetExecutionState().TreeRunStatus == EStateTreeRunStatus::Running;
+}
+
 const UObject* FStateTreeWeakEventQueueExecutionContext::GetOwner() const
 {
 	return Owner.Get();
diff --git a/Engine/Plugins/Runtime/StateTree/Source/StateTreeModule/Public/StateTreeExecutionContext.h b/Engine/Plugins/Runtime/StateTree/Source/StateTreeModule/Public/StateTreeExecutionContext.h
--- a/Engine/Plugins/Runtime/StateTree/Source/StateTreeModule/Public/StateTreeExecutionContext.h
+++ b/Engine/Plugins/Runtime/StateTree/Source/StateTreeModule/Public/StateTreeExecutionContext.h
@@ -937,6 +937,9 @@ public:
 	explicit FStateTreeEventQueueExecutionContext(const UObject& Owner, const UStateTree& StateTree, FStateTreeInstanceStorage& Storage);
 
 public:
+	/** Indicates if the context is valid and able to run an instance of the associated StateTree asset. */
+	bool IsValid() const;
+
 	/** @return The owner of the context */
 	gsl::not_null<const UObject*> GetOwner() const
 	{
