StateTree: Fix unit failure introduced with 37839742

diff --git a/Engine/Plugins/Runtime/StateTree/Source/StateTreeTestSuite/Private/StateTreeTestTypes.h b/Engine/Plugins/Runtime/StateTree/Source/StateTreeTestSuite/Private/StateTreeTestTypes.h
--- a/Engine/Plugins/Runtime/StateTree/Source/StateTreeTestSuite/Private/StateTreeTestTypes.h
+++ b/Engine/Plugins/Runtime/StateTree/Source/StateTreeTestSuite/Private/StateTreeTestTypes.h
@@ -696,5 +696,11 @@ struct FStateTreeTestBooleanCondition : public FStateTreeConditionCommonBase
 		return InstanceData.bSuccess;
 	}
 
+	virtual bool Link(FStateTreeLinker& Linker) override
+	{
+		Linker.LinkExternalData(LogHandle);
+		return true;
+	}
+
 	TStateTreeExternalDataHandle<FStateTreeTestLog> LogHandle;
 };
\ No newline at end of file
