[State Tree] Added Fallback Option to TransitionRequest

diff --git a/Engine/Plugins/Runtime/StateTree/Source/StateTreeEditorModule/Private/StateTreeCompiler.cpp b/Engine/Plugins/Runtime/StateTree/Source/StateTreeEditorModule/Private/StateTreeCompiler.cpp
--- a/Engine/Plugins/Runtime/StateTree/Source/StateTreeEditorModule/Private/StateTreeCompiler.cpp
+++ b/Engine/Plugins/Runtime/StateTree/Source/StateTreeEditorModule/Private/StateTreeCompiler.cpp
@@ -986,11 +986,6 @@ bool FStateTreeCompiler::CreateStateTransitions()
 
 			CompactTransition.bTransitionEnabled = Transition.bTransitionEnabled;
 
-			if (Transition.State.LinkType == EStateTreeTransitionType::NextSelectableState)
-			{
-				CompactTransition.Fallback = EStateTreeSelectionFallback::NextSelectableSibling;
-			}
-
 			if (Transition.bDelayTransition)
 			{
 				CompactTransition.Delay.Set(Transition.DelayDuration, Transition.DelayRandomVariance);
@@ -1024,7 +1019,7 @@ bool FStateTreeCompiler::CreateStateTransitions()
 			}
 			
 			CompactTransition.State = FStateTreeStateHandle::Invalid;
-			if (!ResolveTransitionState(SourceState, Transition.State, CompactTransition.State))
+			if (!ResolveTransitionStateAndFallback(SourceState, Transition.State, CompactTransition.State, CompactTransition.Fallback))
 			{
 				bTransitionSucceeded = false;
 				continue;
@@ -1442,7 +1437,7 @@ T* FStateTreeCompiler::CreateNodeWithSharedInstanceData(UStateTreeState* State,
 	return RawNode.GetMutablePtr<T>();
 }
 
-bool FStateTreeCompiler::ResolveTransitionState(const UStateTreeState* SourceState, const FStateTreeStateLink& Link, FStateTreeStateHandle& OutTransitionHandle) const 
+bool FStateTreeCompiler::ResolveTransitionStateAndFallback(const UStateTreeState* SourceState, const FStateTreeStateLink& Link, FStateTreeStateHandle& OutTransitionHandle, EStateTreeSelectionFallback& OutFallback) const 
 {
 	if (Link.LinkType == EStateTreeTransitionType::GotoState)
 	{
@@ -1496,26 +1491,31 @@ bool FStateTreeCompiler::ResolveTransitionState(const UStateTreeState* SourceSta
 	else if(Link.LinkType == EStateTreeTransitionType::Failed)
 	{
 		OutTransitionHandle = FStateTreeStateHandle::Failed;
-		return true;
 	}
 	else if(Link.LinkType == EStateTreeTransitionType::Succeeded)
 	{
 		OutTransitionHandle = FStateTreeStateHandle::Succeeded;
-		return true;
 	}
 	else if(Link.LinkType == EStateTreeTransitionType::None)
 	{
 		OutTransitionHandle = FStateTreeStateHandle::Invalid;
-		return true;
 	}
 PRAGMA_DISABLE_DEPRECATION_WARNINGS
 	else if (Link.LinkType == EStateTreeTransitionType::NotSet)
 	{
 		OutTransitionHandle = FStateTreeStateHandle::Invalid;
-		return false;
 	}
 PRAGMA_ENABLE_DEPRECATION_WARNINGS
-	
+
+	if (Link.LinkType == EStateTreeTransitionType::NextSelectableState)
+	{
+		OutFallback = EStateTreeSelectionFallback::NextSelectableSibling;
+	}
+	else
+	{
+		OutFallback = EStateTreeSelectionFallback::None;
+	}
+
 	return true;
 }
 
@@ -1609,7 +1609,7 @@ bool FStateTreeCompiler::CompileAndValidateNode(const UStateTreeState* SourceSta
 			{
 				FStateTreeStateLink& StateLink = *static_cast<FStateTreeStateLink*>(const_cast<void*>(It->Value));
 
-				if (!ResolveTransitionState(SourceState, StateLink, StateLink.StateHandle))
+				if (!ResolveTransitionStateAndFallback(SourceState, StateLink, StateLink.StateHandle, StateLink.Fallback))
 				{
 					return false;
 				}
diff --git a/Engine/Plugins/Runtime/StateTree/Source/StateTreeEditorModule/Public/StateTreeCompiler.h b/Engine/Plugins/Runtime/StateTree/Source/StateTreeEditorModule/Public/StateTreeCompiler.h
--- a/Engine/Plugins/Runtime/StateTree/Source/StateTreeEditorModule/Public/StateTreeCompiler.h
+++ b/Engine/Plugins/Runtime/StateTree/Source/StateTreeEditorModule/Public/StateTreeCompiler.h
@@ -36,8 +36,8 @@ public:
 	
 private:
 
-	/** Resolves the state a transition points to. SourceState is nullptr for global tasks. */
-	bool ResolveTransitionState(const UStateTreeState* SourceState, const FStateTreeStateLink& Link, FStateTreeStateHandle& OutTransitionHandle) const;
+	/** Resolves the state a transition points to, and the optional fallback for failing to enter the state. SourceState is nullptr for global tasks. */
+	bool ResolveTransitionStateAndFallback(const UStateTreeState* SourceState, const FStateTreeStateLink& Link, FStateTreeStateHandle& OutTransitionHandle, EStateTreeSelectionFallback& OutFallback) const;
 	FStateTreeStateHandle GetStateHandle(const FGuid& StateID) const;
 	UStateTreeState* GetState(const FGuid& StateID) const;
 
diff --git a/Engine/Plugins/Runtime/StateTree/Source/StateTreeModule/Private/StateTreeExecutionContext.cpp b/Engine/Plugins/Runtime/StateTree/Source/StateTreeModule/Private/StateTreeExecutionContext.cpp
--- a/Engine/Plugins/Runtime/StateTree/Source/StateTreeModule/Private/StateTreeExecutionContext.cpp
+++ b/Engine/Plugins/Runtime/StateTree/Source/StateTreeModule/Private/StateTreeExecutionContext.cpp
@@ -725,7 +725,7 @@ void FStateTreeExecutionContext::RequestTransition(const FStateTreeTransitionReq
 		
 		STATETREE_LOG(Verbose, TEXT("Request transition to '%s' at priority %s"), *GetSafeStateName(*CurrentlyProcessedFrame, Request.TargetState), *UEnum::GetDisplayValueAsText(Request.Priority).ToString());
 
-		if (RequestTransition(*CurrentlyProcessedFrame, Request.TargetState, Request.Priority))
+		if (RequestTransition(*CurrentlyProcessedFrame, Request.TargetState, Request.Priority, /*TransitionEvent*/nullptr, Request.Fallback))
 		{
 			NextTransitionSource = FStateTreeTransitionSource(EStateTreeTransitionSourceType::ExternalRequest, Request.TargetState, Request.Priority);
 		}
@@ -3170,7 +3170,7 @@ bool FStateTreeExecutionContext::TriggerTransitions()
 
 		if (CurrentFrame)
 		{
-			if (RequestTransition(*CurrentFrame, Request.TargetState, Request.Priority))
+			if (RequestTransition(*CurrentFrame, Request.TargetState, Request.Priority, /*TransitionEvent*/nullptr, Request.Fallback))
 			{
 				NextTransitionSource = FStateTreeTransitionSource(EStateTreeTransitionSourceType::ExternalRequest, Request.TargetState, Request.Priority);
 			}
diff --git a/Engine/Plugins/Runtime/StateTree/Source/StateTreeModule/Public/StateTreeExecutionTypes.h b/Engine/Plugins/Runtime/StateTree/Source/StateTreeModule/Public/StateTreeExecutionTypes.h
--- a/Engine/Plugins/Runtime/StateTree/Source/StateTreeModule/Public/StateTreeExecutionTypes.h
+++ b/Engine/Plugins/Runtime/StateTree/Source/StateTreeModule/Public/StateTreeExecutionTypes.h
@@ -224,9 +224,19 @@ struct STATETREEMODULE_API FStateTreeTransitionRequest
 
 	FStateTreeTransitionRequest() = default;
 
-	explicit FStateTreeTransitionRequest(const FStateTreeStateHandle InTargetState, const EStateTreeTransitionPriority InPriority = EStateTreeTransitionPriority::Normal)
+	explicit FStateTreeTransitionRequest(const FStateTreeStateLink& InStateLink)
+		: TargetState(InStateLink.StateHandle)
+		, Fallback(InStateLink.Fallback)
+	{
+	}
+
+	explicit FStateTreeTransitionRequest(
+		const FStateTreeStateHandle InTargetState, 
+		const EStateTreeTransitionPriority InPriority = EStateTreeTransitionPriority::Normal,
+		const EStateTreeSelectionFallback InFallback = EStateTreeSelectionFallback::None)
 		: TargetState(InTargetState)
 		, Priority(InPriority)
+		, Fallback(InFallback)
 	{
 	}
 
@@ -249,6 +259,10 @@ struct STATETREEMODULE_API FStateTreeTransitionRequest
 	/** Priority of the transition. */
 	UPROPERTY(EditDefaultsOnly, Category = "Default")
 	EStateTreeTransitionPriority Priority = EStateTreeTransitionPriority::Normal;
+
+	/** Fallback of the transition if it fails to select the target state */
+	UPROPERTY(EditDefaultsOnly, Category = "Default")
+	EStateTreeSelectionFallback Fallback = EStateTreeSelectionFallback::None;
 };
 
 
diff --git a/Engine/Plugins/Runtime/StateTree/Source/StateTreeModule/Public/StateTreeTypes.h b/Engine/Plugins/Runtime/StateTree/Source/StateTreeModule/Public/StateTreeTypes.h
--- a/Engine/Plugins/Runtime/StateTree/Source/StateTreeModule/Public/StateTreeTypes.h
+++ b/Engine/Plugins/Runtime/StateTree/Source/StateTreeModule/Public/StateTreeTypes.h
@@ -1080,7 +1080,11 @@ PRAGMA_ENABLE_DEPRECATION_WARNINGS
 
 	/** Handle of the linked state. */
 	UPROPERTY()
-	FStateTreeStateHandle StateHandle; 
+	FStateTreeStateHandle StateHandle;
+
+	/** Fallback of the transition if it fails to select the target state */
+	UPROPERTY()
+	EStateTreeSelectionFallback Fallback = EStateTreeSelectionFallback::None;
 };
 
 template<>
