StateTree: Auto expands the find tools results.

diff --git a/Engine/Plugins/Runtime/StateTree/Source/StateTreeEditorModule/Private/FindTools/SStateTreeFind.cpp b/Engine/Plugins/Runtime/StateTree/Source/StateTreeEditorModule/Private/FindTools/SStateTreeFind.cpp
--- a/Engine/Plugins/Runtime/StateTree/Source/StateTreeEditorModule/Private/FindTools/SStateTreeFind.cpp
+++ b/Engine/Plugins/Runtime/StateTree/Source/StateTreeEditorModule/Private/FindTools/SStateTreeFind.cpp
@@ -420,6 +420,10 @@ void SFindInAsset::MakeQueryDelayed()
 	{
 		bSearching = false;
 		TreeView->RequestTreeRefresh();
+		for (const TSharedPtr<FFindResult>& Entry : ItemsFound)
+		{
+			ExpandAll(Entry);
+		}
 	}
 	else
 	{
@@ -427,6 +431,16 @@ void SFindInAsset::MakeQueryDelayed()
 	}
 }
 
+void SFindInAsset::ExpandAll(const TSharedPtr<FFindResult>& Entry)
+{
+	TreeView->SetItemExpansion(Entry, true);
+
+	for (const TSharedPtr<FFindResult>& Child : Entry->Children)
+	{
+		ExpandAll(Child);
+	}
+}
+
 void SFindInAsset::TriggerQueryDelayed()
 {
 	FTSTicker::GetCoreTicker().AddTicker(TEXT("SFindInAsset::MakeQuery"), 0.1f,
diff --git a/Engine/Plugins/Runtime/StateTree/Source/StateTreeEditorModule/Private/FindTools/SStateTreeFind.h b/Engine/Plugins/Runtime/StateTree/Source/StateTreeEditorModule/Private/FindTools/SStateTreeFind.h
--- a/Engine/Plugins/Runtime/StateTree/Source/StateTreeEditorModule/Private/FindTools/SStateTreeFind.h
+++ b/Engine/Plugins/Runtime/StateTree/Source/StateTreeEditorModule/Private/FindTools/SStateTreeFind.h
@@ -109,6 +109,9 @@ private:
 	/** Callback to build the context menu when right clicking in the tree */
 	TSharedPtr<SWidget> HandleTreeContextMenuOpening();
 
+	/** Expand all item in the tree. */
+	void ExpandAll(const TSharedPtr<FFindResult>& Entry);
+
 private:
 	TWeakPtr<IStateTreeEditorHost> EditorHost;
 
