StateTree: Fix issue with BlueprintTask that didn't save the flag. The flag were set at compile time but never saved for cooked build.

diff --git a/Engine/Plugins/Runtime/StateTree/Source/StateTreeModule/Private/Blueprint/StateTreeTaskBlueprintBase.cpp b/Engine/Plugins/Runtime/StateTree/Source/StateTreeModule/Private/Blueprint/StateTreeTaskBlueprintBase.cpp
--- a/Engine/Plugins/Runtime/StateTree/Source/StateTreeModule/Private/Blueprint/StateTreeTaskBlueprintBase.cpp
+++ b/Engine/Plugins/Runtime/StateTree/Source/StateTreeModule/Private/Blueprint/StateTreeTaskBlueprintBase.cpp
@@ -12,11 +12,10 @@
 //----------------------------------------------------------------------//
 //  UStateTreeTaskBlueprintBase
 //----------------------------------------------------------------------//
-
 UStateTreeTaskBlueprintBase::UStateTreeTaskBlueprintBase(const FObjectInitializer& ObjectInitializer)
 	: Super(ObjectInitializer)
 	, bShouldStateChangeOnReselect(true)
-	, bShouldCallTick(true)
+	, bShouldCallTick(false) // for when a child overrides the Tick
 	, bShouldCallTickOnlyOnEvents(false)
 	, bShouldCopyBoundPropertiesOnTick(true)
 	, bShouldCopyBoundPropertiesOnExitState(true)
@@ -115,9 +114,27 @@ PRAGMA_ENABLE_DEPRECATION_WARNINGS
 	bShouldCopyBoundPropertiesOnTick = InstanceData.bShouldCopyBoundPropertiesOnTick;
 	bShouldCopyBoundPropertiesOnExitState = InstanceData.bShouldCopyBoundPropertiesOnExitState;
 
+	// The flags on the FStateTreeTaskBase are not saved.
+	TaskFlags = (bShouldStateChangeOnReselect ? 0x01 : 0)
+		| (bShouldCallTick ? 0x02 : 0)
+		| (bShouldCallTickOnlyOnEvents ? 0x04 : 0)
+		| (bShouldCopyBoundPropertiesOnTick ? 0x08 : 0)
+		| (bShouldCopyBoundPropertiesOnExitState ? 0x10 : 0);
+
 	return EDataValidationResult::Valid;
 }
 
+bool FStateTreeBlueprintTaskWrapper::Link(FStateTreeLinker& Linker)
+{
+	bShouldStateChangeOnReselect = (TaskFlags & 0x01) != 0;
+	bShouldCallTick = (TaskFlags & 0x02) != 0;
+	bShouldCallTickOnlyOnEvents = (TaskFlags & 0x04) != 0;
+	bShouldCopyBoundPropertiesOnTick = (TaskFlags & 0x08) != 0;
+	bShouldCopyBoundPropertiesOnExitState = (TaskFlags & 0x10) != 0;
+
+	return Super::Link(Linker);
+}
+
 EStateTreeRunStatus FStateTreeBlueprintTaskWrapper::EnterState(FStateTreeExecutionContext& Context, const FStateTreeTransitionResult& Transition) const
 {
 	UStateTreeTaskBlueprintBase* Instance = Context.GetInstanceDataPtr<UStateTreeTaskBlueprintBase>(*this);
diff --git a/Engine/Plugins/Runtime/StateTree/Source/StateTreeModule/Public/Blueprint/StateTreeTaskBlueprintBase.h b/Engine/Plugins/Runtime/StateTree/Source/StateTreeModule/Public/Blueprint/StateTreeTaskBlueprintBase.h
--- a/Engine/Plugins/Runtime/StateTree/Source/StateTreeModule/Public/Blueprint/StateTreeTaskBlueprintBase.h
+++ b/Engine/Plugins/Runtime/StateTree/Source/StateTreeModule/Public/Blueprint/StateTreeTaskBlueprintBase.h
@@ -134,6 +134,7 @@ struct STATETREEMODULE_API FStateTreeBlueprintTaskWrapper : public FStateTreeTas
 
 	virtual const UStruct* GetInstanceDataType() const override { return TaskClass; };
 	virtual EDataValidationResult Compile(FStateTreeDataView InstanceDataView, TArray<FText>& ValidationMessages) override;
+	virtual bool Link(FStateTreeLinker& Linker) override;
 	virtual EStateTreeRunStatus EnterState(FStateTreeExecutionContext& Context, const FStateTreeTransitionResult& Transition) const override;
 	virtual void ExitState(FStateTreeExecutionContext& Context, const FStateTreeTransitionResult& Transition) const override;
 	virtual void StateCompleted(FStateTreeExecutionContext& Context, const EStateTreeRunStatus CompletionStatus, const FStateTreeActiveStates& CompletedActiveStates) const override;
@@ -146,4 +147,7 @@ struct STATETREEMODULE_API FStateTreeBlueprintTaskWrapper : public FStateTreeTas
 
 	UPROPERTY()
 	TSubclassOf<UStateTreeTaskBlueprintBase> TaskClass = nullptr;
+
+	UPROPERTY()
+	uint8 TaskFlags;
 };
