StateTree: Modify the editor to notify the UI when an asset compiles. Update the state and the node UI to display if a task can trigger a transition. Use the compiled data when posibled because the compiler can modify the values.

diff --git a/Engine/Plugins/Runtime/StateTree/Source/StateTreeEditorModule/Private/Customizations/StateTreeEditorNodeDetails.cpp b/Engine/Plugins/Runtime/StateTree/Source/StateTreeEditorModule/Private/Customizations/StateTreeEditorNodeDetails.cpp
--- a/Engine/Plugins/Runtime/StateTree/Source/StateTreeEditorModule/Private/Customizations/StateTreeEditorNodeDetails.cpp
+++ b/Engine/Plugins/Runtime/StateTree/Source/StateTreeEditorModule/Private/Customizations/StateTreeEditorNodeDetails.cpp
@@ -36,6 +36,7 @@
 #include "StateTreeEditorNodeUtils.h"
 #include "Framework/Notifications/NotificationManager.h"
 #include "Debugger/StateTreeDebuggerUIExtensions.h"
+#include "StateTreeEditingSubsystem.h"
 #include "StateTreePropertyBindings.h"
 #include "Engine/BlueprintGeneratedClass.h"
 #include "Styling/SlateTypes.h"
@@ -335,10 +336,10 @@ public:
 		FGuid ID;
 		UE::StateTree::PropertyHelpers::GetStructValue<FGuid>(IDProperty, ID);
 		
-		UE::StateTreeEditor::Internal::ModifyRow(ChildRow, ID, EditorData);
+		UE::StateTreeEditor::Internal::ModifyRow(ChildRow, ID, EditorData.Get());
 	}
 
-	UStateTreeEditorData* EditorData;
+	TWeakObjectPtr<UStateTreeEditorData> EditorData;
 	TSharedPtr<IPropertyHandle> IDProperty;
 };
 
@@ -352,6 +353,10 @@ TSharedRef<IPropertyTypeCustomization> FStateTreeEditorNodeDetails::MakeInstance
 FStateTreeEditorNodeDetails::~FStateTreeEditorNodeDetails()
 {
 	UE::StateTree::PropertyBinding::OnStateTreePropertyBindingChanged.Remove(OnBindingChangedHandle);
+	if (StateTreeViewModel)
+	{
+		StateTreeViewModel->GetOnAssetChanged().Remove(OnChangedAssetHandle);
+	}
 }
 
 void FStateTreeEditorNodeDetails::CustomizeHeader(TSharedRef<class IPropertyHandle> StructPropertyHandle, class FDetailWidgetRow& HeaderRow, IPropertyTypeCustomizationUtils& StructCustomizationUtils)
@@ -373,11 +378,21 @@ void FStateTreeEditorNodeDetails::CustomizeHeader(TSharedRef<class IPropertyHand
 	check(IndentProperty.IsValid());
 	check(OperandProperty.IsValid());
 
-	UE::StateTreeEditor::EditorNodeUtils::GetNodeBaseScriptStructAndClass(StructProperty, BaseScriptStruct, BaseClass);
+	{
+		UScriptStruct* BaseScriptStructPtr = nullptr;
+		UClass* BaseClassPtr = nullptr;
+		UE::StateTreeEditor::EditorNodeUtils::GetNodeBaseScriptStructAndClass(StructProperty, BaseScriptStructPtr, BaseClassPtr);
+		BaseScriptStruct = BaseScriptStructPtr;
+		BaseClass = BaseClassPtr;
+	}
 
 	UE::StateTree::Delegates::OnIdentifierChanged.AddSP(this, &FStateTreeEditorNodeDetails::OnIdentifierChanged);
 	OnBindingChangedHandle = UE::StateTree::PropertyBinding::OnStateTreePropertyBindingChanged.AddRaw(this, &FStateTreeEditorNodeDetails::OnBindingChanged);
 	FindOuterObjects();
+	if (StateTreeViewModel)
+	{
+		OnChangedAssetHandle = StateTreeViewModel->GetOnAssetChanged().AddSP(this, &FStateTreeEditorNodeDetails::HandleAssetChanged);
+	}
 
 	// Don't draw the header if it's a PropertyFunction.
 	if (UE::StateTreeEditor::Internal::IsOwnedByPropertyFunctionNode(StructProperty))
@@ -558,7 +573,7 @@ void FStateTreeEditorNodeDetails::CustomizeHeader(TSharedRef<class IPropertyHand
 					.VAlign(VAlign_Center)
 					.Padding(FMargin(4.0f, 0.0f))
 					[
-						SAssignNew(FlagBorder, SBorder)
+						SAssignNew(FlagsContainer, SBorder)
 							.BorderImage(FStyleDefaults::GetNoBrush())
 							.Visibility(this, &FStateTreeEditorNodeDetails::AreFlagsVisible)
 					]
@@ -589,7 +604,7 @@ void FStateTreeEditorNodeDetails::CustomizeHeader(TSharedRef<class IPropertyHand
 					.AutoWidth()
 					.VAlign(VAlign_Center)
 					[
-						UE::StateTreeEditor::DebuggerExtensions::CreateEditorNodeWidget(StructPropertyHandle, EditorData)
+						UE::StateTreeEditor::DebuggerExtensions::CreateEditorNodeWidget(StructPropertyHandle, EditorData.Get())
 					]
 
 					// Browse To BP Button
@@ -661,60 +676,85 @@ void FStateTreeEditorNodeDetails::CustomizeHeader(TSharedRef<class IPropertyHand
 		.CopyAction(FUIAction(FExecuteAction::CreateSP(this, &FStateTreeEditorNodeDetails::OnCopyNode)))
 		.PasteAction(FUIAction(FExecuteAction::CreateSP(this, &FStateTreeEditorNodeDetails::OnPasteNode)));
 
-	if (FlagBorder)
+	MakeFlagsWidget();
+}
+
+void FStateTreeEditorNodeDetails::MakeFlagsWidget()
+{
+	FlagsContainer->SetPadding(FMargin(4.0f));
+	FlagsContainer->SetContent(SNullWidget::NullWidget);
+
+	const UStateTree* StateTreePtr = StateTree.Get();
+
+	TArray<void*> RawNodeDatas;
+	StructProperty->AccessRawData(RawNodeDatas);
+	bool bShowCallTick = false;
+	bool bShouldCallTickOnlyOnEvents = false;
+	bool bHasTransitionTick = false;
+	for (void* RawNodeData : RawNodeDatas)
 	{
-		TArray<void*> RawNodeDatas;
-		StructProperty->AccessRawData(RawNodeDatas);
-		bool bShowCallTick = false;
-		bool bShouldCallTickOnlyOnEvents = false;
-		for (void* RawNodeData : RawNodeDatas)
+		if (const FStateTreeEditorNode* EditorNode = reinterpret_cast<const FStateTreeEditorNode*>(RawNodeData))
 		{
-			if (const FStateTreeEditorNode* EditorNode = reinterpret_cast<const FStateTreeEditorNode*>(RawNodeData))
+			bool bUseEditorData = true;
+			// Use the compiled version if it exists. It is more accurate (like with BP tasks) but less interactive (the user needs to compile) :(
+			if (StateTreePtr)
+			{
+				if (const FStateTreeTaskBase* CompiledTask = StateTreePtr->GetNode(StateTreePtr->GetNodeIndexFromId(EditorNode->ID).AsInt32()).GetPtr<const FStateTreeTaskBase>())
+				{
+					bShowCallTick = bShowCallTick || CompiledTask->bShouldCallTick;
+					bShouldCallTickOnlyOnEvents = bShouldCallTickOnlyOnEvents || CompiledTask->bShouldCallTickOnlyOnEvents;
+					bHasTransitionTick = bHasTransitionTick || CompiledTask->bShouldAffectTransitions;
+					bUseEditorData = false;
+				}
+			}
+
+			if (bUseEditorData)
 			{
 				if (const FStateTreeTaskBase* TreeTaskNodePtr = EditorNode->Node.GetPtr<FStateTreeTaskBase>())
 				{
 					bShowCallTick = bShowCallTick || TreeTaskNodePtr->bShouldCallTick;
 					bShouldCallTickOnlyOnEvents = bShouldCallTickOnlyOnEvents || TreeTaskNodePtr->bShouldCallTickOnlyOnEvents;
+					bHasTransitionTick = bHasTransitionTick || TreeTaskNodePtr->bShouldAffectTransitions;
 				}
 			}
 		}
+	}
 
-		if (bShowCallTick || bShouldCallTickOnlyOnEvents)
+	if (bShowCallTick || bShouldCallTickOnlyOnEvents || bHasTransitionTick)
+	{
+		TSharedRef<SHorizontalBox> Box = SNew(SHorizontalBox);
+		const bool bShowTickIcon = bShowCallTick || bHasTransitionTick;
+		if (bShowTickIcon)
 		{
-			TSharedRef<SHorizontalBox> Box = SNew(SHorizontalBox);
-			if (bShowCallTick)
-			{
-				Box->AddSlot()
-				.AutoWidth()
-				.VAlign(VAlign_Center)
-				[
-					SNew(SImage)
-					.Image(FStateTreeEditorStyle::Get().GetBrush("StateTreeEditor.Flags.Tick"))
-					.ColorAndOpacity(FSlateColor::UseSubduedForeground())
-					.ToolTipText(LOCTEXT("TaskTick", "The task ticks at runtime."))
-				];
-			}
+			Box->AddSlot()
+			.AutoWidth()
+			.VAlign(VAlign_Center)
+			[
+				SNew(SImage)
+				.Image(FStateTreeEditorStyle::Get().GetBrush("StateTreeEditor.Flags.Tick"))
+				.ColorAndOpacity(FSlateColor::UseSubduedForeground())
+				.ToolTipText(LOCTEXT("TaskTick", "The task ticks at runtime."))
+			];
+		}
 
-			if (!bShowCallTick && bShouldCallTickOnlyOnEvents)
-			{
-				Box->AddSlot()
-				.AutoWidth()
-				.VAlign(VAlign_Center)
-				[
-					SNew(SImage)
+		if (!bShowTickIcon && bShouldCallTickOnlyOnEvents)
+		{
+			Box->AddSlot()
+			.AutoWidth()
+			.VAlign(VAlign_Center)
+			[
+				SNew(SImage)
 					.Image(FStateTreeEditorStyle::Get().GetBrush("StateTreeEditor.Flags.TickOnEvent"))
 					.ColorAndOpacity(FSlateColor::UseSubduedForeground())
 					.ToolTipText(LOCTEXT("TaskTickEvent", "The task ticks on event at runtime."))
-				];
-			}
-
-			FlagBorder->SetPadding(FMargin(4.f, 4.f, 4.f, 4.f));
-			FlagBorder->SetContent(Box);
+			];
 		}
+
+		FlagsContainer->SetPadding(FMargin(4.0f));
+		FlagsContainer->SetContent(Box);
 	}
 }
 
-
 FReply FStateTreeEditorNodeDetails::OnRowMouseDown(const FGeometry& MyGeometry, const FPointerEvent& MouseEvent)
 {
 	if (MouseEvent.GetEffectingButton() == EKeys::RightMouseButton)
@@ -778,7 +818,7 @@ void FStateTreeEditorNodeDetails::OnPasteNode()
 	
 	const UStruct* NodeTypeStruct = TempNode.Node.GetScriptStruct(); 
 	// Only allow valid node types for this property (e.g. do not mix task with conditions).
-	if (ErrorPipe.NumErrors > 0 || !NodeTypeStruct || !NodeTypeStruct->IsChildOf(BaseScriptStruct))
+	if (ErrorPipe.NumErrors > 0 || !NodeTypeStruct || !NodeTypeStruct->IsChildOf(BaseScriptStruct.Get()))
 	{
 		FNotificationInfo NotificationInfo(FText::GetEmpty());
 		NotificationInfo.Text = FText::Format(LOCTEXT("NotSupportedByType", "This property only accepts nodes of type {0}."), BaseScriptStruct->GetDisplayNameText());
@@ -788,7 +828,8 @@ void FStateTreeEditorNodeDetails::OnPasteNode()
 	}
 	
 	// Reject nodes that are not allowed by the schema.
-	const UStateTreeSchema* Schema = EditorData ? EditorData->Schema : nullptr;
+	UStateTreeEditorData* EditorDataPtr = EditorData.Get();
+	const UStateTreeSchema* Schema = EditorDataPtr ? EditorDataPtr->Schema : nullptr;
 	if (Schema)
 	{
 		bool bNodeIsAllowed = false;
@@ -858,9 +899,9 @@ void FStateTreeEditorNodeDetails::OnPasteNode()
 				EditorNode->ID = FGuid::NewGuid();
 
 				// Copy bindings from the copied node.
-				if (OldStructID.IsValid() && EditorData)
+				if (OldStructID.IsValid() && EditorDataPtr)
 				{
-					if (FStateTreeEditorPropertyBindings* Bindings = EditorData->GetPropertyEditorBindings())
+					if (FStateTreeEditorPropertyBindings* Bindings = EditorDataPtr->GetPropertyEditorBindings())
 					{
 						Bindings->CopyBindings(OldStructID, EditorNode->ID);
 					}
@@ -939,12 +980,14 @@ void FStateTreeEditorNodeDetails::CustomizeChildren(TSharedRef<class IPropertyHa
 		StructBuilder.AddProperty(IDProperty.ToSharedRef());
 	}
 	
+	UStateTreeEditorData* EditorDataPtr = EditorData.Get();
+
 	// Node
-	TSharedRef<FBindableNodeInstanceDetails> NodeDetails = MakeShareable(new FBindableNodeInstanceDetails(NodeProperty, {}, EditorData));
+	TSharedRef<FBindableNodeInstanceDetails> NodeDetails = MakeShareable(new FBindableNodeInstanceDetails(NodeProperty, {}, EditorDataPtr));
 	StructBuilder.AddCustomBuilder(NodeDetails);
 
 	// Instance
-	TSharedRef<FBindableNodeInstanceDetails> InstanceDetails = MakeShareable(new FBindableNodeInstanceDetails(InstanceProperty, IDProperty, EditorData));
+	TSharedRef<FBindableNodeInstanceDetails> InstanceDetails = MakeShareable(new FBindableNodeInstanceDetails(InstanceProperty, IDProperty, EditorDataPtr));
 	StructBuilder.AddCustomBuilder(InstanceDetails);
 
 	// InstanceObject
@@ -988,7 +1031,7 @@ void FStateTreeEditorNodeDetails::CustomizeChildren(TSharedRef<class IPropertyHa
 		for (FSortedChild& Child : SortedChildren)
 		{
 			IDetailPropertyRow& ChildRow = StructBuilder.AddProperty(Child.PropertyHandle.ToSharedRef());
-			UE::StateTreeEditor::Internal::ModifyRow(ChildRow, ID, EditorData);
+			UE::StateTreeEditor::Internal::ModifyRow(ChildRow, ID, EditorDataPtr);
 		}
 	}
 }
@@ -1023,7 +1066,8 @@ void FStateTreeEditorNodeDetails::OnBindingChanged(const FStateTreePropertyPath&
 {
 	check(StructProperty);
 
-	if (!EditorData)
+	UStateTreeEditorData* EditorDataPtr = EditorData.Get();
+	if (!EditorDataPtr)
 	{
 		return;
 	}
@@ -1039,7 +1083,7 @@ void FStateTreeEditorNodeDetails::OnBindingChanged(const FStateTreePropertyPath&
 		return;
 	}
 
-	const FStateTreeBindingLookup BindingLookup(EditorData);
+	const FStateTreeBindingLookup BindingLookup(EditorDataPtr);
 
 	for (int32 i = 0; i < OuterObjects.Num(); i++)
 	{
@@ -1063,8 +1107,9 @@ void FStateTreeEditorNodeDetails::FindOuterObjects()
 {
 	check(StructProperty);
 	
-	EditorData = nullptr;
-	StateTree = nullptr;
+	EditorData.Reset();
+	StateTree.Reset();
+	StateTreeViewModel.Reset();
 
 	TArray<UObject*> OuterObjects;
 	StructProperty->GetOuterObjects(OuterObjects);
@@ -1081,6 +1126,10 @@ void FStateTreeEditorNodeDetails::FindOuterObjects()
 		{
 			StateTree = OuterStateTree;
 			EditorData = OuterEditorData;
+			if (UStateTreeEditingSubsystem* StateTreeEditingSubsystem = GEditor ? GEditor->GetEditorSubsystem<UStateTreeEditingSubsystem>() : nullptr)
+			{
+				StateTreeViewModel = StateTreeEditingSubsystem->FindOrAddViewModel(OuterStateTree);
+			}
 			break;
 		}
 	}
@@ -1422,7 +1471,8 @@ FReply FStateTreeEditorNodeDetails::OnDescriptionClicked(const FGeometry& MyGeom
 FText FStateTreeEditorNodeDetails::GetNodeDescription() const
 {
 	check(StructProperty);
-	if (!EditorData)
+	const UStateTreeEditorData* EditorDataPtr = EditorData.Get();
+	if (!EditorDataPtr)
 	{
 		return FText::GetEmpty();
 	}
@@ -1435,7 +1485,7 @@ FText FStateTreeEditorNodeDetails::GetNodeDescription() const
 		FText Description = LOCTEXT("EmptyNodeRich", "<s>None</>");
 		if (const FStateTreeEditorNode* Node = static_cast<FStateTreeEditorNode*>(RawNodeData[0]))
 		{
-			return EditorData->GetNodeDescription(*Node, EStateTreeNodeFormatting::RichText);
+			return EditorDataPtr->GetNodeDescription(*Node, EStateTreeNodeFormatting::RichText);
 		}
 		return Description;
 	}
@@ -1453,7 +1503,8 @@ EVisibility FStateTreeEditorNodeDetails::IsNodeDescriptionVisible() const
 
 	if (ScriptStruct != nullptr && ScriptStruct->IsChildOf(FStateTreeTaskBase::StaticStruct()))
 	{
-		const UStateTreeSchema* Schema = EditorData ? EditorData->Schema : nullptr;
+		const UStateTreeEditorData* EditorDataPtr = EditorData.Get();
+		const UStateTreeSchema* Schema = EditorDataPtr ? EditorDataPtr->Schema : nullptr;
 		if (Schema && Schema->AllowMultipleTasks() == false)
 		{
 			// Single task states use the state name as task name.
@@ -1467,7 +1518,9 @@ EVisibility FStateTreeEditorNodeDetails::IsNodeDescriptionVisible() const
 FText FStateTreeEditorNodeDetails::GetNodeTooltip() const
 {
 	check(StructProperty);
-	if (!EditorData)
+
+	const UStateTreeEditorData* EditorDataPtr = EditorData.Get();
+	if (!EditorDataPtr)
 	{
 		return FText::GetEmpty();
 	}
@@ -1584,9 +1637,9 @@ void FStateTreeEditorNodeDetails::HandleNameCommitted(const FText& NewText, ETex
 
 			StructProperty->NotifyPostChange(EPropertyChangeType::ValueSet);
 
-			if (StateTree)
+			if (const UStateTree* StateTreePtr = StateTree.Get())
 			{
-				UE::StateTree::Delegates::OnIdentifierChanged.Broadcast(*StateTree);
+				UE::StateTree::Delegates::OnIdentifierChanged.Broadcast(*StateTreePtr);
 			}
 
 			if (GEditor)
@@ -1605,7 +1658,9 @@ void FStateTreeEditorNodeDetails::HandleNameCommitted(const FText& NewText, ETex
 FText FStateTreeEditorNodeDetails::GetNodePickerTooltip() const
 {
 	check(StructProperty);
-	if (!EditorData)
+
+	const UStateTreeEditorData* EditorDataPtr = EditorData.Get();
+	if (!EditorDataPtr)
 	{
 		return FText::GetEmpty();
 	}
@@ -1620,7 +1675,7 @@ FText FStateTreeEditorNodeDetails::GetNodePickerTooltip() const
 		FText Description = LOCTEXT("EmptyNodeStyled", "<s>None</>");
 		if (const FStateTreeEditorNode* Node = static_cast<FStateTreeEditorNode*>(RawNodeData[0]))
 		{
-			TextBuilder.AppendLine(EditorData->GetNodeDescription(*Node));
+			TextBuilder.AppendLine(EditorDataPtr->GetNodeDescription(*Node));
 		}
 	}
 
@@ -1756,8 +1811,8 @@ void FStateTreeEditorNodeDetails::GeneratePickerMenu(class FMenuBuilder& InMenuB
 
 	TSharedRef<SStateTreeNodeTypePicker> Picker = SNew(SStateTreeNodeTypePicker)
 		.Schema(EditorData->Schema)
-		.BaseScriptStruct(BaseScriptStruct)
-		.BaseClass(BaseClass)
+		.BaseScriptStruct(BaseScriptStruct.Get())
+		.BaseClass(BaseClass.Get())
 		.CurrentStruct(CommonStruct)
 		.OnNodeTypePicked(SStateTreeNodeTypePicker::FOnNodeStructPicked::CreateSP(this, &FStateTreeEditorNodeDetails::OnNodePicked));
 	
@@ -1851,7 +1906,7 @@ TSharedRef<SWidget> FStateTreeEditorNodeDetails::GenerateOptionsMenu()
 	MenuBuilder.EndSection();
 
 	// Append debugger items.
-	UE::StateTreeEditor::DebuggerExtensions::AppendEditorNodeMenuItems(MenuBuilder, StructProperty, EditorData);
+	UE::StateTreeEditor::DebuggerExtensions::AppendEditorNodeMenuItems(MenuBuilder, StructProperty, EditorData.Get());
 
 	return MenuBuilder.MakeWidget();
 }
@@ -1928,4 +1983,9 @@ void FStateTreeEditorNodeDetails::OnNodePicked(const UStruct* InStruct) const
 	}
 }
 
+void FStateTreeEditorNodeDetails::HandleAssetChanged()
+{
+	MakeFlagsWidget();
+}
+
 #undef LOCTEXT_NAMESPACE
diff --git a/Engine/Plugins/Runtime/StateTree/Source/StateTreeEditorModule/Private/Customizations/StateTreeEditorNodeDetails.h b/Engine/Plugins/Runtime/StateTree/Source/StateTreeEditorModule/Private/Customizations/StateTreeEditorNodeDetails.h
--- a/Engine/Plugins/Runtime/StateTree/Source/StateTreeEditorModule/Private/Customizations/StateTreeEditorNodeDetails.h
+++ b/Engine/Plugins/Runtime/StateTree/Source/StateTreeEditorModule/Private/Customizations/StateTreeEditorNodeDetails.h
@@ -10,6 +10,7 @@ struct FObjectKey;
 struct FOptionalSize;
 
 class FDetailWidgetRow;
+class FStateTreeViewModel;
 class IDetailChildrenBuilder;
 class IPropertyHandle;
 class IPropertyUtilities;
@@ -44,6 +45,8 @@ public:
 	virtual void CustomizeChildren(TSharedRef<IPropertyHandle> StructPropertyHandle, IDetailChildrenBuilder& StructBuilder, IPropertyTypeCustomizationUtils& StructCustomizationUtils) override;
 
 private:
+	void MakeFlagsWidget();
+
 	bool ShouldResetToDefault(TSharedPtr<IPropertyHandle> PropertyHandle) const;
 	void ResetToDefault(TSharedPtr<IPropertyHandle> PropertyHandle);
 	void OnCopyNode();
@@ -115,15 +118,18 @@ private:
 	void OnDeleteAllNodes() const;
 	void OnDuplicateNode() const;
 	void OnRenameNode() const;
+	void HandleAssetChanged();
 
-	UScriptStruct* BaseScriptStruct = nullptr;
-	UClass* BaseClass = nullptr;
+	TWeakObjectPtr<UScriptStruct> BaseScriptStruct;
+	TWeakObjectPtr<UClass> BaseClass;
 	TSharedPtr<SWidgetSwitcher> NameSwitcher;
 	TSharedPtr<SInlineEditableTextBlock> NameEdit;
 	TSharedPtr<SBorder> RowBorder; 
+	TSharedPtr<SBorder> FlagsContainer; 
 
-	UStateTreeEditorData* EditorData = nullptr;
-	UStateTree* StateTree = nullptr;
+	TWeakObjectPtr<UStateTreeEditorData> EditorData;
+	TWeakObjectPtr<UStateTree> StateTree;
+	TSharedPtr<FStateTreeViewModel> StateTreeViewModel;
 	
 	TSharedPtr<IPropertyUtilities> PropUtils;
 	TSharedPtr<IPropertyHandle> StructProperty;
@@ -136,4 +142,5 @@ private:
 	TSharedPtr<IPropertyHandle> OperandProperty;
 
 	FDelegateHandle OnBindingChangedHandle;
+	FDelegateHandle OnChangedAssetHandle;
 };
diff --git a/Engine/Plugins/Runtime/StateTree/Source/StateTreeEditorModule/Private/SStateTreeViewRow.cpp b/Engine/Plugins/Runtime/StateTree/Source/StateTreeEditorModule/Private/SStateTreeViewRow.cpp
--- a/Engine/Plugins/Runtime/StateTree/Source/StateTreeEditorModule/Private/SStateTreeViewRow.cpp
+++ b/Engine/Plugins/Runtime/StateTree/Source/StateTreeEditorModule/Private/SStateTreeViewRow.cpp
@@ -49,6 +49,8 @@ void SStateTreeViewRow::Construct(const FArguments& InArgs, const TSharedRef<STa
 	const UStateTreeState* State = InState.Get();
 	WeakEditorData = State != nullptr ? State->GetTypedOuter<UStateTreeEditorData>() : nullptr;
 
+	AssetChangedHandle = StateTreeViewModel->GetOnAssetChanged().AddSP(this, &SStateTreeViewRow::HandleAssetChanged);
+
 	ConstructInternal(STableRow::FArguments()
 		.OnDragDetected(this, &SStateTreeViewRow::HandleDragDetected)
 		.OnDragLeave(this, &SStateTreeViewRow::HandleDragLeave)
@@ -305,7 +307,7 @@ void SStateTreeViewRow::Construct(const FArguments& InArgs, const TSharedRef<STa
 											.AutoWidth()
 											.Padding(FMargin(0.0f))
 											[
-												SAssignNew(FlagBorder, SBorder)
+												SAssignNew(FlagsContainer, SBorder)
 												.BorderImage(FStyleDefaults::GetNoBrush())
 											]
 
@@ -466,60 +468,16 @@ void SStateTreeViewRow::Construct(const FArguments& InArgs, const TSharedRef<STa
 		];
 	}
 
-	if (EnumHasAllFlags(GetDefault<UStateTreeEditorUserSettings>()->GetStatesViewDisplayNodeType(), EStateTreeEditorUserSettingsNodeType::Flag))
-	{
-		const UStateTree* StateTree = StateTreeViewModel->GetStateTree();
-		if (State && StateTree)
-		{
-			if (const FCompactStateTreeState* RuntimeState = StateTree->GetStateFromHandle(StateTree->GetStateHandleFromId(State->ID)))
-			{
-				if (RuntimeState->bHasTickTasks || RuntimeState->bHasTickTasksOnlyOnEvents || RuntimeState->bHasTransitionTasks)
-				{
-					TSharedRef<SHorizontalBox> Box = SNew(SHorizontalBox);
-					if (RuntimeState->bHasTickTasks)
-					{
-						Box->AddSlot()
-						.AutoWidth()
-						.VAlign(VAlign_Center)
-						[
-							SNew(SImage)
-							.Image(FStateTreeEditorStyle::Get().GetBrush("StateTreeEditor.Flags.Tick"))
-							.ColorAndOpacity(IconTint)
-							.ToolTipText(LOCTEXT("NodeTick", "The node ticks at runtime."))
-						];
-					}
 
-					if (!RuntimeState->bHasTickTasks && RuntimeState->bHasTickTasksOnlyOnEvents)
-					{
-						Box->AddSlot()
-						.AutoWidth()
-						.VAlign(VAlign_Center)
-						[
-							SNew(SImage)
-							.Image(FStateTreeEditorStyle::Get().GetBrush("StateTreeEditor.Flags.TickOnEvent"))
-							.ColorAndOpacity(IconTint)
-							.ToolTipText(LOCTEXT("NodeTickEvent", "The node ticks on event at runtime."))
-						];
-					}
+	MakeFlagsWidget();
+}
 
-					if (RuntimeState->bHasTransitionTasks)
-					{
-						Box->AddSlot()
-						.AutoWidth()
-						.VAlign(VAlign_Center)
-						[
-							SNew(SImage)
-							.Image(FStateTreeEditorStyle::Get().GetBrush("StateTreeEditor.Transitions"))
-							.ColorAndOpacity(IconTint)
-							.ToolTipText(LOCTEXT("NodeTickTransition", "The node ticks on to trigger transition."))
-						];
-					}
 
-					FlagBorder->SetPadding(FMargin(4.f, 4.f, 4.f, 4.f));
-					FlagBorder->SetContent(Box);
-				}
-			}
-		}
+SStateTreeViewRow::~SStateTreeViewRow()
+{
+	if (StateTreeViewModel)
+	{
+		StateTreeViewModel->GetOnAssetChanged().Remove(AssetChangedHandle);
 	}
 }
 
@@ -926,6 +884,69 @@ TSharedRef<SWidget> SStateTreeViewRow::MakeConditionsWidget(const TSharedPtr<SSc
 	return ConditionsBox;
 }
 
+void SStateTreeViewRow::MakeFlagsWidget()
+{
+	FlagsContainer->SetPadding(FMargin(0.0f));
+	FlagsContainer->SetContent(SNullWidget::NullWidget);
+
+	const UStateTree* StateTree = StateTreeViewModel ? StateTreeViewModel->GetStateTree() : nullptr;
+	const UStateTreeState* State = WeakState.Get();
+	const bool bDisplayFlags = EnumHasAllFlags(GetDefault<UStateTreeEditorUserSettings>()->GetStatesViewDisplayNodeType(), EStateTreeEditorUserSettingsNodeType::Flag);
+	static constexpr FLinearColor IconTint = FLinearColor(1.0f, 1.0f, 1.0f, 0.5f);
+
+	if (bDisplayFlags && State && StateTree)
+	{
+		if (const FCompactStateTreeState* RuntimeState = StateTree->GetStateFromHandle(StateTree->GetStateHandleFromId(State->ID)))
+		{
+			if (RuntimeState->bHasTickTasks || RuntimeState->bHasTickTasksOnlyOnEvents || RuntimeState->bHasTransitionTasks)
+			{
+				TSharedRef<SHorizontalBox> Box = SNew(SHorizontalBox);
+				if (RuntimeState->bHasTickTasks)
+				{
+					Box->AddSlot()
+					.AutoWidth()
+					.VAlign(VAlign_Center)
+					[
+						SNew(SImage)
+						.Image(FStateTreeEditorStyle::Get().GetBrush("StateTreeEditor.Flags.Tick"))
+						.ColorAndOpacity(IconTint)
+						.ToolTipText(LOCTEXT("NodeTick", "The node ticks at runtime."))
+					];
+				}
+
+				if (!RuntimeState->bHasTickTasks && RuntimeState->bHasTickTasksOnlyOnEvents)
+				{
+					Box->AddSlot()
+					.AutoWidth()
+					.VAlign(VAlign_Center)
+					[
+						SNew(SImage)
+						.Image(FStateTreeEditorStyle::Get().GetBrush("StateTreeEditor.Flags.TickOnEvent"))
+						.ColorAndOpacity(IconTint)
+						.ToolTipText(LOCTEXT("NodeTickEvent", "The node ticks on event at runtime."))
+					];
+				}
+
+				if (RuntimeState->bHasTransitionTasks)
+				{
+					Box->AddSlot()
+					.AutoWidth()
+					.VAlign(VAlign_Center)
+					[
+						SNew(SImage)
+						.Image(FStateTreeEditorStyle::Get().GetBrush("StateTreeEditor.Transitions"))
+						.ColorAndOpacity(IconTint)
+						.ToolTipText(LOCTEXT("NodeTickTransition", "The node ticks on to trigger transition."))
+					];
+				}
+
+				FlagsContainer->SetPadding(FMargin(4.0f));
+				FlagsContainer->SetContent(Box);
+			}
+		}
+	}
+}
+
 void SStateTreeViewRow::RequestRename() const
 {
 	if (NameTextBlock)
@@ -2136,4 +2157,9 @@ FReply SStateTreeViewRow::HandleAcceptDrop(const FDragDropEvent& DragDropEvent,
 	return FReply::Unhandled();
 }
 
+void SStateTreeViewRow::HandleAssetChanged()
+{
+	MakeFlagsWidget();
+}
+
 #undef LOCTEXT_NAMESPACE
diff --git a/Engine/Plugins/Runtime/StateTree/Source/StateTreeEditorModule/Private/SStateTreeViewRow.h b/Engine/Plugins/Runtime/StateTree/Source/StateTreeEditorModule/Private/SStateTreeViewRow.h
--- a/Engine/Plugins/Runtime/StateTree/Source/StateTreeEditorModule/Private/SStateTreeViewRow.h
+++ b/Engine/Plugins/Runtime/StateTree/Source/StateTreeEditorModule/Private/SStateTreeViewRow.h
@@ -14,6 +14,7 @@ struct FStateTreeEditorNode;
 
 class UStateTreeEditorData;
 class SStateTreeView;
+class SBorder;
 class SInlineEditableTextBlock;
 class SScrollBox;
 class FStateTreeViewModel;
@@ -28,11 +29,14 @@ public:
 	SLATE_END_ARGS()
 
 	void Construct(const FArguments& InArgs, const TSharedRef<STableViewBase>& InOwnerTableView, TWeakObjectPtr<UStateTreeState> InState, const TSharedPtr<SScrollBox>& ViewBox, TSharedPtr<FStateTreeViewModel> InStateTreeViewModel);
+	~SStateTreeViewRow();
+
 	void RequestRename() const;
 
 private:
 	TSharedRef<SWidget> MakeTasksWidget(const TSharedPtr<SScrollBox>& ViewBox);
 	TSharedRef<SWidget> MakeConditionsWidget(const TSharedPtr<SScrollBox>& ViewBox);
+	void MakeFlagsWidget();
 
 	FSlateColor GetTitleColor(const float Alpha = 1.0f, const float Lighten = 0.0f) const;
 	FSlateColor GetActiveStateColor() const;
@@ -117,8 +121,13 @@ private:
 	TOptional<EItemDropZone> HandleCanAcceptDrop(const FDragDropEvent& DragDropEvent, EItemDropZone DropZone, TWeakObjectPtr<UStateTreeState> TargetState) const;
 	FReply HandleAcceptDrop(const FDragDropEvent& DragDropEvent, EItemDropZone DropZone, TWeakObjectPtr<UStateTreeState> TargetState) const;
 
+	void HandleAssetChanged();
+
 	TSharedPtr<FStateTreeViewModel> StateTreeViewModel;
 	TWeakObjectPtr<UStateTreeState> WeakState;
 	TWeakObjectPtr<UStateTreeEditorData> WeakEditorData;
 	TSharedPtr<SInlineEditableTextBlock> NameTextBlock;
+	TSharedPtr<SBorder> FlagsContainer;
+
+	FDelegateHandle AssetChangedHandle;
 };
diff --git a/Engine/Plugins/Runtime/StateTree/Source/StateTreeEditorModule/Private/StateTreeEditingSubsystem.cpp b/Engine/Plugins/Runtime/StateTree/Source/StateTreeEditorModule/Private/StateTreeEditingSubsystem.cpp
--- a/Engine/Plugins/Runtime/StateTree/Source/StateTreeEditorModule/Private/StateTreeEditingSubsystem.cpp
+++ b/Engine/Plugins/Runtime/StateTree/Source/StateTreeEditorModule/Private/StateTreeEditingSubsystem.cpp
@@ -47,7 +47,16 @@ bool UStateTreeEditingSubsystem::CompileStateTree(const gsl::not_null<UStateTree
 		
 		UE_LOG(LogStateTreeEditor, Error, TEXT("Failed to compile '%s', errors follow."), *InStateTree->GetFullName());
 		InOutLog.DumpToLog(LogStateTreeEditor);
-	}	
+	}
+
+	// Notify the UI that something changed. Make sure to not request a new viewmodel. That way, we are not creating new viewmodel when cooking/PIE.
+	if (UStateTreeEditingSubsystem* StateTreeEditingSubsystem = GEditor ? GEditor->GetEditorSubsystem<UStateTreeEditingSubsystem>() : nullptr)
+	{
+		if (TSharedPtr<FStateTreeViewModel> ViewModel = StateTreeEditingSubsystem->StateTreeViewModels.FindRef(FObjectKey(InStateTree)))
+		{
+			ViewModel->NotifyAssetChangedExternally();
+		}
+	}
 
 	return bCompilationResult;
 }
