StateTree: Initialize the TaskFalgs variable.

#rnx

diff --git a/Engine/Plugins/Runtime/StateTree/Source/StateTreeModule/Public/Blueprint/StateTreeTaskBlueprintBase.h b/Engine/Plugins/Runtime/StateTree/Source/StateTreeModule/Public/Blueprint/StateTreeTaskBlueprintBase.h
--- a/Engine/Plugins/Runtime/StateTree/Source/StateTreeModule/Public/Blueprint/StateTreeTaskBlueprintBase.h
+++ b/Engine/Plugins/Runtime/StateTree/Source/StateTreeModule/Public/Blueprint/StateTreeTaskBlueprintBase.h
@@ -149,5 +149,5 @@ struct STATETREEMODULE_API FStateTreeBlueprintTaskWrapper : public FStateTreeTas
 	TSubclassOf<UStateTreeTaskBlueprintBase> TaskClass = nullptr;
 
 	UPROPERTY()
-	uint8 TaskFlags;
+	uint8 TaskFlags = 0;
 };
