StateTree: Display the other states that are linked to the selected state.

diff --git a/Engine/Plugins/Runtime/StateTree/Source/StateTreeEditorModule/Private/Customizations/Widgets/SCompactStateTreeView.cpp b/Engine/Plugins/Runtime/StateTree/Source/StateTreeEditorModule/Private/Customizations/Widgets/SCompactStateTreeView.cpp
--- a/Engine/Plugins/Runtime/StateTree/Source/StateTreeEditorModule/Private/Customizations/Widgets/SCompactStateTreeView.cpp
+++ b/Engine/Plugins/Runtime/StateTree/Source/StateTreeEditorModule/Private/Customizations/Widgets/SCompactStateTreeView.cpp
@@ -1,4 +1,4 @@
-﻿// Copyright Epic Games, Inc. All Rights Reserved.
+// Copyright Epic Games, Inc. All Rights Reserved.
 
 #include "SCompactStateTreeView.h"
 
@@ -16,12 +16,17 @@
 
 TMap<FObjectKey, SCompactStateTreeView::FStateExpansionState> SCompactStateTreeView::StateExpansionStates;
 
-SCompactStateTreeView::SCompactStateTreeView()
+FSlateColor SCompactStateTreeView::FStateTreeStateItem::GetBorderColor() const
 {
-}
+	if (LinkState == ELinkState::LinkState_None)
+	{
+		return FLinearColor(0.0f, 0.0f, 0.0f, 0.0f);
+	}
 
-SCompactStateTreeView::~SCompactStateTreeView()
-{
+	static const FName NAME_LinkingIn = "Colors.StateLinkingIn";
+	static const FName NAME_StateLinkedOut = "Colors.StateLinkedOut";
+	const FName ColorName = (LinkState & ELinkState::LinkState_LinkingIn) != 0 ? NAME_LinkingIn : NAME_StateLinkedOut;
+	return FStateTreeEditorStyle::Get().GetColor(ColorName);
 }
 
 void SCompactStateTreeView::Construct(const SCompactStateTreeView::FArguments& InArgs, TSharedPtr<FStateTreeViewModel> InViewModel)
@@ -30,6 +35,7 @@ void SCompactStateTreeView::Construct(const SCompactStateTreeView::FArguments& I
 	WeakStateTreeEditorData = InArgs._StateTreeEditorData;
 	bSubtreesOnly = InArgs._SubtreesOnly;
 	bSelectableStatesOnly = InArgs._SelectableStatesOnly;
+	bShowLinkedStates = InArgs._ShowLinkedStates;
 	OnSelectionChanged = InArgs._OnSelectionChanged;
 	OnContextMenuOpening = InArgs._OnContextMenuOpening;
 	
@@ -97,12 +103,12 @@ void SCompactStateTreeView::SetSelection(TConstArrayView<FGuid> Selection)
 	{
 		return;
 	}
-	
+
 	if (bIsSettingSelection)
 	{
 		return;
 	}
-	
+
 	TArray<TSharedPtr<FStateTreeStateItem>> SelectedStates;
 	FindStatesByIDRecursive(FilteredRootItem, Selection, SelectedStates);
 
@@ -111,6 +117,43 @@ void SCompactStateTreeView::SetSelection(TConstArrayView<FGuid> Selection)
 	StateItemTree->ClearSelection();
 	StateItemTree->SetItemSelection(SelectedStates, true);
 
+	if (bShowLinkedStates)
+	{
+		ResetLinkedStates();
+
+		// Find the Linked items
+		TArray<FGuid> LinkingIn;
+		TArray<FGuid> LinkedOut;
+		for (TSharedPtr<FStateTreeStateItem>& StateItem : SelectedStates)
+		{
+			StateTreeViewModel->GetLinkStates(StateItem->StateID, LinkingIn, LinkedOut);
+		}
+
+		// Set the outline
+		{
+			TArray<TSharedPtr<FStateTreeStateItem>> FoundStates;
+			FoundStates.Reserve(LinkingIn.Num());
+			FindStatesByIDRecursive(FilteredRootItem, LinkingIn, FoundStates);
+
+			for (TSharedPtr<FStateTreeStateItem> Item : FoundStates)
+			{
+				PreviousLinkedStates.Add(Item);
+				Item->LinkState |= FStateTreeStateItem::LinkState_LinkingIn;
+			}
+		}
+		{
+			TArray<TSharedPtr<FStateTreeStateItem>> FoundStates;
+			FoundStates.Reserve(LinkedOut.Num());
+			FindStatesByIDRecursive(FilteredRootItem, LinkedOut, FoundStates);
+
+			for (TSharedPtr<FStateTreeStateItem> Item : FoundStates)
+			{
+				PreviousLinkedStates.AddUnique(Item);
+				Item->LinkState |= FStateTreeStateItem::LinkState_LinkedOut;
+			}
+		}
+	}
+
 	if (SelectedStates.Num() == 1)
 	{
 		StateItemTree->RequestScrollIntoView(SelectedStates[0]);	
@@ -119,6 +162,18 @@ void SCompactStateTreeView::SetSelection(TConstArrayView<FGuid> Selection)
 	bIsSettingSelection = false;
 }
 
+void SCompactStateTreeView::ResetLinkedStates()
+{
+	for (TWeakPtr<FStateTreeStateItem>& PreviousLinkedState : PreviousLinkedStates)
+	{
+		if (TSharedPtr<FStateTreeStateItem> Pin = PreviousLinkedState.Pin())
+		{
+			Pin->LinkState = FStateTreeStateItem::LinkState_None;
+		}
+	}
+	PreviousLinkedStates.Reset();
+}
+
 TArray<FGuid> SCompactStateTreeView::GetSelection() const
 {
 	TArray<FGuid> SelectedItemIDs;
@@ -291,7 +346,11 @@ TSharedRef<ITableRow> SCompactStateTreeView::GenerateStateItemRow(TSharedPtr<FSt
 		.OnCanAcceptDrop(this, &SCompactStateTreeView::HandleCanAcceptDrop)
 		.OnAcceptDrop(this, &SCompactStateTreeView::HandleAcceptDrop)
 		[
-			Container
+			SNew(SBorder)
+			.BorderBackgroundColor(Item.Get(), &FStateTreeStateItem::GetBorderColor)
+			[
+				Container
+			]
 		];
 }
 
@@ -379,6 +438,8 @@ void SCompactStateTreeView::UpdateFilteredRoot(bool bRestoreSelection)
 		Selection = GetSelection();
 	}
 
+	ResetLinkedStates();
+
 	if (FilterStrings.IsEmpty())
 	{
 		// Show all when there's no filter string.
diff --git a/Engine/Plugins/Runtime/StateTree/Source/StateTreeEditorModule/Private/Customizations/Widgets/SCompactStateTreeView.h b/Engine/Plugins/Runtime/StateTree/Source/StateTreeEditorModule/Private/Customizations/Widgets/SCompactStateTreeView.h
--- a/Engine/Plugins/Runtime/StateTree/Source/StateTreeEditorModule/Private/Customizations/Widgets/SCompactStateTreeView.h
+++ b/Engine/Plugins/Runtime/StateTree/Source/StateTreeEditorModule/Private/Customizations/Widgets/SCompactStateTreeView.h
@@ -1,4 +1,4 @@
-﻿// Copyright Epic Games, Inc. All Rights Reserved.
+// Copyright Epic Games, Inc. All Rights Reserved.
 
 #pragma once
 
@@ -35,13 +35,11 @@ public:
 		SLATE_ARGUMENT(ESelectionMode::Type, SelectionMode)
 		SLATE_ARGUMENT(bool, SelectableStatesOnly)
 		SLATE_ARGUMENT(bool, SubtreesOnly)
+		SLATE_ARGUMENT(bool, ShowLinkedStates)
 		SLATE_EVENT(FOnSelectionChanged, OnSelectionChanged)
 		SLATE_EVENT(FOnContextMenuOpening, OnContextMenuOpening)
 	SLATE_END_ARGS()
 
-	SCompactStateTreeView();
-	virtual ~SCompactStateTreeView() override;
-
 	void Construct(const FArguments& InArgs, TSharedPtr<FStateTreeViewModel> InViewModel = nullptr);
 
 	/** @returns widget to focus (search box) when the picker is opened. */
@@ -65,11 +63,21 @@ private:
 		{
 		}
 
+		enum ELinkState
+		{
+			LinkState_None = 0x00,
+			LinkState_LinkingIn = 0x01,
+			LinkState_LinkedOut = 0x02,
+		};
+
+		FSlateColor GetBorderColor() const;
+
 		FText Desc;
 		FText TooltipText;
 		FGuid StateID = {};
 		FSlateColor Color = FSlateColor(FLinearColor::White);
 		const FSlateBrush* Icon = nullptr;
+		uint8 LinkState = LinkState_None;
 		bool bIsSubTree = false;
 		bool bIsLinked = false;
 		bool bIsEnabled = true;
@@ -85,6 +93,7 @@ private:
 
 	void CacheStates();
 	void CacheState(TSharedPtr<FStateTreeStateItem> ParentNode, const UStateTreeState* State);
+	void ResetLinkedStates();
 
 	TSharedRef<ITableRow> GenerateStateItemRow(TSharedPtr<FStateTreeStateItem> Item, const TSharedRef<STableViewBase>& OwnerTable);
 	void GetStateItemChildren(TSharedPtr<FStateTreeStateItem> Item, TArray<TSharedPtr<FStateTreeStateItem>>& OutItems) const;
@@ -115,12 +124,15 @@ private:
 	TWeakObjectPtr<const UStateTreeEditorData> WeakStateTreeEditorData = nullptr;
 
 	TArray<FString> FilterStrings;
-	
+	TArray<TWeakPtr<FStateTreeStateItem>> PreviousLinkedStates;
+
 	// If set, allow to select only states marked as subtrees.
 	bool bSubtreesOnly = false;
 
 	bool bSelectableStatesOnly = false;
 
+	bool bShowLinkedStates = false;
+
 	bool bIsSettingSelection = false;
 
 	TSharedPtr<FStateTreeViewModel> StateTreeViewModel;
diff --git a/Engine/Plugins/Runtime/StateTree/Source/StateTreeEditorModule/Private/SStateTreeOutliner.cpp b/Engine/Plugins/Runtime/StateTree/Source/StateTreeEditorModule/Private/SStateTreeOutliner.cpp
--- a/Engine/Plugins/Runtime/StateTree/Source/StateTreeEditorModule/Private/SStateTreeOutliner.cpp
+++ b/Engine/Plugins/Runtime/StateTree/Source/StateTreeEditorModule/Private/SStateTreeOutliner.cpp
@@ -1,4 +1,4 @@
-﻿// Copyright Epic Games, Inc. All Rights Reserved.
+// Copyright Epic Games, Inc. All Rights Reserved.
 
 #include "SStateTreeOutliner.h"
 
@@ -54,6 +54,7 @@ void SStateTreeOutliner::Construct(const FArguments& InArgs, TSharedRef<FStateTr
 		.StateTreeEditorData(StateTreeViewModel->GetStateTreeEditorData())
 		.OnSelectionChanged(this, &SStateTreeOutliner::HandleTreeViewSelectionChanged)
 		.OnContextMenuOpening(this, &SStateTreeOutliner::HandleContextMenuOpening)
+		.ShowLinkedStates(true)
 	];
 	
 	CommandList = InCommandList;
diff --git a/Engine/Plugins/Runtime/StateTree/Source/StateTreeEditorModule/Private/StateTreeEditorStyle.cpp b/Engine/Plugins/Runtime/StateTree/Source/StateTreeEditorModule/Private/StateTreeEditorStyle.cpp
--- a/Engine/Plugins/Runtime/StateTree/Source/StateTreeEditorModule/Private/StateTreeEditorStyle.cpp
+++ b/Engine/Plugins/Runtime/StateTree/Source/StateTreeEditorModule/Private/StateTreeEditorStyle.cpp
@@ -419,6 +419,10 @@ FStateTreeEditorStyle::FStateTreeEditorStyle()
 		Set("StateTreeEditor.Flags.Tick", new IMAGE_BRUSH_SVG("Icons/Tick", CoreStyleConstants::Icon16x16));
 		Set("StateTreeEditor.Flags.TickOnEvent", new IMAGE_BRUSH_SVG("Icons/TickEvent", CoreStyleConstants::Icon16x16));
 	}
+	{
+		Set("Colors.StateLinkingIn", FLinearColor::Yellow);
+		Set("Colors.StateLinkedOut", FLinearColor::Green);
+	}
 }
 
 void FStateTreeEditorStyle::Register()
diff --git a/Engine/Plugins/Runtime/StateTree/Source/StateTreeEditorModule/Private/StateTreeViewModel.cpp b/Engine/Plugins/Runtime/StateTree/Source/StateTreeEditorModule/Private/StateTreeViewModel.cpp
--- a/Engine/Plugins/Runtime/StateTree/Source/StateTreeEditorModule/Private/StateTreeViewModel.cpp
+++ b/Engine/Plugins/Runtime/StateTree/Source/StateTreeEditorModule/Private/StateTreeViewModel.cpp
@@ -578,6 +578,51 @@ void FStateTreeViewModel::GetSubTrees(TArray<TWeakObjectPtr<UStateTreeState>>& O
 	}
 }
 
+void FStateTreeViewModel::GetLinkStates(FGuid StateID, TArray<FGuid>& LinkingIn, TArray<FGuid>& LinkedOut) const
+{
+	const UStateTreeState* State = GetStateByID(StateID);
+	if (State == nullptr)
+	{
+		return;
+	}
+
+	if (UStateTreeEditorData* TreeData = TreeDataWeak.Get())
+	{
+		TreeData->VisitHierarchy([&LinkingIn, &LinkedOut, StateID = State->ID](UStateTreeState& State, UStateTreeState* ParentState)
+			{
+				if (State.ID == StateID)
+				{
+					return EStateTreeVisitor::Continue;
+				}
+				if (State.Type == EStateTreeStateType::Linked && StateID == State.LinkedSubtree.ID)
+				{
+					LinkingIn.AddUnique(State.ID);
+				}
+				else
+				{
+					for (const FStateTreeTransition& Transition : State.Transitions)
+					{
+						if (Transition.State.ID == StateID)
+						{
+							LinkingIn.AddUnique(State.ID);
+						}
+					}
+				}
+				return EStateTreeVisitor::Continue;
+			});
+
+		if (State->Type == EStateTreeStateType::Linked)
+		{
+			LinkedOut.AddUnique(State->LinkedSubtree.ID);
+		}
+
+		for (const FStateTreeTransition& Transition : State->Transitions)
+		{
+			LinkedOut.AddUnique(Transition.State.ID);
+		}
+	}
+}
+
 void FStateTreeViewModel::PostUndo(bool bSuccess)
 {
 	// TODO: see if we can narrow this down.
diff --git a/Engine/Plugins/Runtime/StateTree/Source/StateTreeEditorModule/Public/StateTreeViewModel.h b/Engine/Plugins/Runtime/StateTree/Source/StateTreeEditorModule/Public/StateTreeViewModel.h
--- a/Engine/Plugins/Runtime/StateTree/Source/StateTreeEditorModule/Public/StateTreeViewModel.h
+++ b/Engine/Plugins/Runtime/StateTree/Source/StateTreeEditorModule/Public/StateTreeViewModel.h
@@ -76,6 +76,9 @@ public:
 	int32 GetSubTreeCount() const;
 	void GetSubTrees(TArray<TWeakObjectPtr<UStateTreeState>>& OutSubtrees) const;
 
+	/** Find the state that link to the StateTree and the state the state can link to. */
+	void GetLinkStates(FGuid StateID, TArray<FGuid>& LinkingIn, TArray<FGuid>& LinkedOut) const;
+
 	// Gets and sets StateTree view expansion state store in the asset.
 	void SetPersistentExpandedStates(TSet<TWeakObjectPtr<UStateTreeState>>& InExpandedStates);
 	void GetPersistentExpandedStates(TSet<TWeakObjectPtr<UStateTreeState>>& OutExpandedStates);
