[StateTree] added new debug delegates that can  be used them for debug purposes.
- added new dedicated update phases for trace/debug

diff --git a/Engine/Plugins/Runtime/StateTree/Source/StateTreeModule/Private/Debugger/StateTreeDebug.cpp b/Engine/Plugins/Runtime/StateTree/Source/StateTreeModule/Private/Debugger/StateTreeDebug.cpp
--- a/Engine/Plugins/Runtime/StateTree/Source/StateTreeModule/Private/Debugger/StateTreeDebug.cpp
+++ b/Engine/Plugins/Runtime/StateTree/Source/StateTreeModule/Private/Debugger/StateTreeDebug.cpp
@@ -9,19 +9,25 @@
 
 namespace UE::StateTree::Debug
 {
+FPhaseDelegate OnBeginUpdatePhase_AnyThread;
+FPhaseDelegate OnEndUpdatePhase_AnyThread;
 
-FStateTreeDebug::FNodeDelegate FStateTreeDebug::OnConditionEnterState_AnyThread;
-FStateTreeDebug::FNodeDelegate FStateTreeDebug::OnTestCondition_AnyThread;
-FStateTreeDebug::FNodeDelegate FStateTreeDebug::OnConditionExitState_AnyThread;
-FStateTreeDebug::FNodeDelegate FStateTreeDebug::OnEvaluatorEnterTree_AnyThread;
-FStateTreeDebug::FNodeDelegate FStateTreeDebug::OnTickEvaluator_AnyThread;
-FStateTreeDebug::FNodeDelegate FStateTreeDebug::OnEvaluatorExitTree_AnyThread;
-FStateTreeDebug::FNodeDelegate FStateTreeDebug::OnTaskEnterState_AnyThread;
-FStateTreeDebug::FNodeDelegate FStateTreeDebug::OnTickTask_AnyThread;
-FStateTreeDebug::FNodeDelegate FStateTreeDebug::OnTaskExitState_AnyThread;
-FStateTreeDebug::FEventSentDelegate FStateTreeDebug::OnEventSent_AnyThread;
-
-FStateTreeDebug::FNodeReference::FNodeReference(gsl::not_null<const UStateTree*> InStateTree, FStateTreeIndex16 InNodeIndex)
+FStateDelegate OnStateEvent_AnyThread;
+
+FTransitionDelegate OnTransitionEvent_AnyThread;
+
+FNodeDelegate OnConditionEnterState_AnyThread;
+FNodeDelegate OnTestCondition_AnyThread;
+FNodeDelegate OnConditionExitState_AnyThread;
+FNodeDelegate OnEvaluatorEnterTree_AnyThread;
+FNodeDelegate OnTickEvaluator_AnyThread;
+FNodeDelegate OnEvaluatorExitTree_AnyThread;
+FNodeDelegate OnTaskEnterState_AnyThread;
+FNodeDelegate OnTickTask_AnyThread;
+FNodeDelegate OnTaskExitState_AnyThread;
+FEventSentDelegate OnEventSent_AnyThread;
+
+FNodeReference::FNodeReference(const gsl::not_null<const UStateTree*> InStateTree, const FStateTreeIndex16 InNodeIndex)
 	: StateTree(InStateTree)
 	, Index(InNodeIndex)
 {
@@ -29,7 +35,17 @@ FStateTreeDebug::FNodeReference::FNodeReference(gsl::not_null<const UStateTree*>
 
 namespace Private
 {
-void NodeEnter(const FStateTreeExecutionContext& ExecutionContext, FStateTreeDebug::FNodeReference Node, FStateTreeDebug::FNodeDelegate& Delegate)
+void UpdatePhaseEnter(const FStateTreeExecutionContext& ExecutionContext, const EStateTreeUpdatePhase Phase, const FStateTreeStateHandle StateHandle, const FPhaseDelegate& Delegate)
+{
+	Delegate.Broadcast(ExecutionContext, Phase, StateHandle);
+}
+
+void UpdatePhaseExit(const FStateTreeExecutionContext& ExecutionContext, const EStateTreeUpdatePhase Phase, const FStateTreeStateHandle StateHandle, const FPhaseDelegate& Delegate)
+{
+	Delegate.Broadcast(ExecutionContext, Phase, StateHandle);
+}
+
+void NodeEnter(const FStateTreeExecutionContext& ExecutionContext, const FNodeReference Node, const FNodeDelegate& Delegate)
 {
 	FStateTreeInstanceData* InstanceData = ExecutionContext.GetMutableInstanceData();
 	FGuid NodeId = Node.StateTree->GetNodeIdFromIndex(Node.Index);
@@ -39,11 +55,11 @@ void NodeEnter(const FStateTreeExecutionContext& ExecutionContext, FStateTreeDeb
 		{
 			RuntimeValidation->NodeEnterState(NodeId);
 		}
-		Delegate.Broadcast(ExecutionContext, FStateTreeDebug::FNodeDelegateArgs{ .Node = Node, .NodeId = NodeId });
+		Delegate.Broadcast(ExecutionContext, FNodeDelegateArgs{ .Node = Node, .NodeId = NodeId });
 	}
 }
 
-void NodeExit(const FStateTreeExecutionContext& ExecutionContext, FStateTreeDebug::FNodeReference Node, FStateTreeDebug::FNodeDelegate& Delegate)
+void NodeExit(const FStateTreeExecutionContext& ExecutionContext, const FNodeReference Node, const FNodeDelegate& Delegate)
 {
 	FStateTreeInstanceData* InstanceData = ExecutionContext.GetMutableInstanceData();
 	FGuid NodeId = Node.StateTree->GetNodeIdFromIndex(Node.Index);
@@ -53,69 +69,88 @@ void NodeExit(const FStateTreeExecutionContext& ExecutionContext, FStateTreeDebu
 		{
 			RuntimeValidation->NodeExitState(NodeId);
 		}
-		Delegate.Broadcast(ExecutionContext, FStateTreeDebug::FNodeDelegateArgs{ .Node = Node, .NodeId = NodeId });
+		Delegate.Broadcast(ExecutionContext, FNodeDelegateArgs{ .Node = Node, .NodeId = NodeId });
 	}
 }
 
-void NodeTick(const FStateTreeExecutionContext& ExecutionContext, FStateTreeDebug::FNodeReference Node, FStateTreeDebug::FNodeDelegate& Delegate)
+void NodeTick(const FStateTreeExecutionContext& ExecutionContext, const FNodeReference Node, const FNodeDelegate& Delegate)
 {
-	FStateTreeInstanceData* InstanceData = ExecutionContext.GetMutableInstanceData();
-	FGuid NodeId = Node.StateTree->GetNodeIdFromIndex(Node.Index);
+	const FGuid NodeId = Node.StateTree->GetNodeIdFromIndex(Node.Index);
 	if (ensure(NodeId.IsValid()))
 	{
-		Delegate.Broadcast(ExecutionContext, FStateTreeDebug::FNodeDelegateArgs{ .Node = Node, .NodeId = NodeId });
+		Delegate.Broadcast(ExecutionContext, FNodeDelegateArgs{ .Node = Node, .NodeId = NodeId });
 	}
 }
 } //namespace Private
 
-void FStateTreeDebug::ConditionEnterState(const FStateTreeExecutionContext& ExecutionContext, FNodeReference Node)
+void EnterPhase(const FStateTreeExecutionContext& ExecutionContext, const EStateTreeUpdatePhase Phase, const FStateTreeStateHandle StateHandle)
+{
+	Private::UpdatePhaseEnter(ExecutionContext, Phase, StateHandle, OnBeginUpdatePhase_AnyThread);
+}
+
+void ExitPhase(const FStateTreeExecutionContext& ExecutionContext, const EStateTreeUpdatePhase Phase, const FStateTreeStateHandle StateHandle)
+{
+	Private::UpdatePhaseExit(ExecutionContext, Phase, StateHandle, OnEndUpdatePhase_AnyThread);
+}
+
+void StateEvent(const FStateTreeExecutionContext& ExecutionContext, const FStateTreeStateHandle StateHandle, const EStateTreeTraceEventType EventType)
+{
+	OnStateEvent_AnyThread.Broadcast(ExecutionContext, StateHandle, EventType);
+}
+
+void TransitionEvent(const FStateTreeExecutionContext& ExecutionContext, const FStateTreeTransitionSource& TransitionSource, const EStateTreeTraceEventType EventType)
+{
+	OnTransitionEvent_AnyThread.Broadcast(ExecutionContext, TransitionSource, EventType);
+}
+
+void ConditionEnterState(const FStateTreeExecutionContext& ExecutionContext, const FNodeReference Node)
 {
 	Private::NodeEnter(ExecutionContext, Node, OnConditionEnterState_AnyThread);
 }
 
-void FStateTreeDebug::ConditionTest(const FStateTreeExecutionContext& ExecutionContext, FNodeReference Node)
+void ConditionTest(const FStateTreeExecutionContext& ExecutionContext, const FNodeReference Node)
 {
 	Private::NodeTick(ExecutionContext, Node, OnTestCondition_AnyThread);
 }
 
-void FStateTreeDebug::ConditionExitState(const FStateTreeExecutionContext& ExecutionContext, FNodeReference Node)
+void ConditionExitState(const FStateTreeExecutionContext& ExecutionContext, const FNodeReference Node)
 {
 	Private::NodeExit(ExecutionContext, Node, OnConditionExitState_AnyThread);
 }
 
-void FStateTreeDebug::EvaluatorEnterTree(const FStateTreeExecutionContext& ExecutionContext, FNodeReference Node)
+void EvaluatorEnterTree(const FStateTreeExecutionContext& ExecutionContext, const FNodeReference Node)
 {
 	Private::NodeEnter(ExecutionContext, Node, OnEvaluatorEnterTree_AnyThread);
 }
 
-void FStateTreeDebug::EvaluatorTick(const FStateTreeExecutionContext& ExecutionContext, FNodeReference Node)
+void EvaluatorTick(const FStateTreeExecutionContext& ExecutionContext, const FNodeReference Node)
 {
 	Private::NodeTick(ExecutionContext, Node, OnTickEvaluator_AnyThread);
 }
 
-void FStateTreeDebug::EvaluatorExitTree(const FStateTreeExecutionContext& ExecutionContext, FNodeReference Node)
+void EvaluatorExitTree(const FStateTreeExecutionContext& ExecutionContext, const FNodeReference Node)
 {
 	Private::NodeExit(ExecutionContext, Node, OnEvaluatorExitTree_AnyThread);
 }
 
-void FStateTreeDebug::TaskEnterState(const FStateTreeExecutionContext& ExecutionContext, FNodeReference Node)
+void TaskEnterState(const FStateTreeExecutionContext& ExecutionContext, const FNodeReference Node)
 {
 	Private::NodeEnter(ExecutionContext, Node, OnTaskEnterState_AnyThread);
 }
 
-void FStateTreeDebug::TaskTick(const FStateTreeExecutionContext& ExecutionContext, FNodeReference Node)
+void TaskTick(const FStateTreeExecutionContext& ExecutionContext, const FNodeReference Node)
 {
 	Private::NodeTick(ExecutionContext, Node, OnTickTask_AnyThread);
 }
 
-void FStateTreeDebug::TaskExitState(const FStateTreeExecutionContext& ExecutionContext, FNodeReference Node)
+void TaskExitState(const FStateTreeExecutionContext& ExecutionContext, const FNodeReference Node)
 {
 	Private::NodeExit(ExecutionContext, Node, OnTaskExitState_AnyThread);
 }
 
-void FStateTreeDebug::EventSent(const FStateTreeMinimalExecutionContext& ExecutionContext, gsl::not_null<const UStateTree*> StateTree, FGameplayTag Tag, FConstStructView Payload, FName Origin)
+void EventSent(const FStateTreeMinimalExecutionContext& ExecutionContext, const gsl::not_null<const UStateTree*> StateTree, const FGameplayTag Tag, const FConstStructView Payload, const FName Origin)
 {
-	OnEventSent_AnyThread.Broadcast(ExecutionContext, FStateTreeDebug::FEventSentDelegateArgs{ .StateTree = StateTree, .Tag = Tag, .Payload = Payload, .Origin = Origin });
+	OnEventSent_AnyThread.Broadcast(ExecutionContext, FEventSentDelegateArgs{ .StateTree = StateTree, .Tag = Tag, .Payload = Payload, .Origin = Origin });
 }
 
 }//namespace UE::StateTree::Debug
diff --git a/Engine/Plugins/Runtime/StateTree/Source/StateTreeModule/Private/Debugger/StateTreeRuntimeValidationInstanceData.h b/Engine/Plugins/Runtime/StateTree/Source/StateTreeModule/Private/Debugger/StateTreeRuntimeValidationInstanceData.h
--- a/Engine/Plugins/Runtime/StateTree/Source/StateTreeModule/Private/Debugger/StateTreeRuntimeValidationInstanceData.h
+++ b/Engine/Plugins/Runtime/StateTree/Source/StateTreeModule/Private/Debugger/StateTreeRuntimeValidationInstanceData.h
@@ -3,7 +3,6 @@
 #pragma once
 
 #include "StateTreeTypes.h"
-#include "Debugger/StateTreeDebug.h"
 
 class UStateTree;
 class UObject;
diff --git a/Engine/Plugins/Runtime/StateTree/Source/StateTreeModule/Private/StateTreeExecutionContext.cpp b/Engine/Plugins/Runtime/StateTree/Source/StateTreeModule/Private/StateTreeExecutionContext.cpp
--- a/Engine/Plugins/Runtime/StateTree/Source/StateTreeModule/Private/StateTreeExecutionContext.cpp
+++ b/Engine/Plugins/Runtime/StateTree/Source/StateTreeModule/Private/StateTreeExecutionContext.cpp
@@ -21,61 +21,12 @@
 #define STATETREE_LOG(Verbosity, Format, ...) UE_VLOG_ALWAYS_UELOG(GetOwner(), LogStateTree, Verbosity, TEXT("%s: ") Format, *GetInstanceDescriptionInternal(), ##__VA_ARGS__)
 #define STATETREE_CLOG(Condition, Verbosity, Format, ...) UE_CVLOG_ALWAYS_UELOG((Condition), GetOwner(), LogStateTree, Verbosity, TEXT("%s: ") Format, *GetInstanceDescriptionInternal(), ##__VA_ARGS__)
 
-#define STATETREE_LOG_AND_TRACE(LogVerbosity, TraceVerbosity, Format, ...) \
-	UE_VLOG_ALWAYS_UELOG(GetOwner(), LogStateTree, LogVerbosity, TEXT("%s: ") Format, *GetInstanceDescriptionInternal(), ##__VA_ARGS__); \
-	STATETREE_TRACE_LOG_EVENT(TraceVerbosity, Format, ##__VA_ARGS__)
-
-#if WITH_STATETREE_TRACE
-	#define STATETREE_TRACE_ID_NAME PREPROCESSOR_JOIN(InstanceId,__LINE__) \
-
-	#define STATETREE_TRACE_SCOPED_PHASE(Phase) \
-		FStateTreeInstanceDebugId STATETREE_TRACE_ID_NAME = GetInstanceDebugId(); \
-		TRACE_STATETREE_PHASE_EVENT(STATETREE_TRACE_ID_NAME, Phase, EStateTreeTraceEventType::Push, FStateTreeStateHandle::Invalid) \
-		ON_SCOPE_EXIT { TRACE_STATETREE_PHASE_EVENT(STATETREE_TRACE_ID_NAME, Phase, EStateTreeTraceEventType::Pop, FStateTreeStateHandle::Invalid) }
-
-	#define STATETREE_TRACE_SCOPED_STATE(StateHandle) \
-		FStateTreeInstanceDebugId STATETREE_TRACE_ID_NAME = GetInstanceDebugId(); \
-		TRACE_STATETREE_PHASE_EVENT(STATETREE_TRACE_ID_NAME, EStateTreeUpdatePhase::Unset, EStateTreeTraceEventType::Push, StateHandle) \
-		ON_SCOPE_EXIT { TRACE_STATETREE_PHASE_EVENT(STATETREE_TRACE_ID_NAME, EStateTreeUpdatePhase::Unset, EStateTreeTraceEventType::Pop, StateHandle) }
-
-	#define STATETREE_TRACE_SCOPED_STATE_PHASE(StateHandle, Phase) \
-		FStateTreeInstanceDebugId STATETREE_TRACE_ID_NAME = GetInstanceDebugId(); \
-		TRACE_STATETREE_PHASE_EVENT(STATETREE_TRACE_ID_NAME, Phase, EStateTreeTraceEventType::Push, StateHandle) \
-		ON_SCOPE_EXIT { TRACE_STATETREE_PHASE_EVENT(STATETREE_TRACE_ID_NAME, Phase, EStateTreeTraceEventType::Pop, StateHandle) }
-
-	#define STATETREE_TRACE_INSTANCE_EVENT(EventType)						TRACE_STATETREE_INSTANCE_EVENT(GetInstanceDebugId(), GetStateTree(), *GetInstanceDescriptionInternal(), EventType);
-	#define STATETREE_TRACE_INSTANCE_FRAME_EVENT(InstanceDebugId, Frame)	TRACE_STATETREE_INSTANCE_FRAME_EVENT(InstanceDebugId, Frame);
-	#define STATETREE_TRACE_PHASE_BEGIN(Phase)								TRACE_STATETREE_PHASE_EVENT(GetInstanceDebugId(), Phase, EStateTreeTraceEventType::Push, FStateTreeStateHandle::Invalid)
-	#define STATETREE_TRACE_PHASE_END(Phase)								TRACE_STATETREE_PHASE_EVENT(GetInstanceDebugId(), Phase, EStateTreeTraceEventType::Pop, FStateTreeStateHandle::Invalid)
-	#define STATETREE_TRACE_ACTIVE_STATES_EVENT(ActiveFrames)				TRACE_STATETREE_ACTIVE_STATES_EVENT(GetInstanceDebugId(), ActiveFrames);
-	#define STATETREE_TRACE_LOG_EVENT(TraceVerbosity, Format, ...)			TRACE_STATETREE_LOG_EVENT(GetInstanceDebugId(), TraceVerbosity, Format, ##__VA_ARGS__)
-	#define STATETREE_TRACE_STATE_EVENT(StateHandle, EventType)				TRACE_STATETREE_STATE_EVENT(GetInstanceDebugId(), StateHandle, EventType);
-	#define STATETREE_TRACE_TASK_EVENT(Index, DataView, EventType, Status)	TRACE_STATETREE_TASK_EVENT(GetInstanceDebugId(), FStateTreeIndex16(Index), DataView, EventType, Status);
-	#define STATETREE_TRACE_EVALUATOR_EVENT(Index, DataView, EventType)		TRACE_STATETREE_EVALUATOR_EVENT(GetInstanceDebugId(), FStateTreeIndex16(Index), DataView, EventType);
-	#define STATETREE_TRACE_CONDITION_EVENT(Index, DataView, EventType)		TRACE_STATETREE_CONDITION_EVENT(GetInstanceDebugId(), FStateTreeIndex16(Index), DataView, EventType);	
-	#define STATETREE_TRACE_TRANSITION_EVENT(Source, EventType)				TRACE_STATETREE_TRANSITION_EVENT(GetInstanceDebugId(), Source, EventType);
-#else
-	#define STATETREE_TRACE_SCOPED_PHASE(Phase)
-	#define STATETREE_TRACE_SCOPED_STATE(StateHandle)
-	#define STATETREE_TRACE_SCOPED_STATE_PHASE(StateHandle, Phase)
-	#define STATETREE_TRACE_INSTANCE_EVENT(EventType)
-	#define STATETREE_TRACE_INSTANCE_FRAME_EVENT(InstanceDebugId, Frame)
-	#define STATETREE_TRACE_PHASE_BEGIN(Phase)
-	#define STATETREE_TRACE_PHASE_END(Phase)
-	#define STATETREE_TRACE_ACTIVE_STATES_EVENT(ActiveFrames)
-	#define STATETREE_TRACE_LOG_EVENT(TraceVerbosity, Format, ...)
-	#define STATETREE_TRACE_STATE_EVENT(StateHandle, EventType)
-	#define STATETREE_TRACE_TASK_EVENT(Index, DataView, EventType, Status)
-	#define STATETREE_TRACE_EVALUATOR_EVENT(Index, DataView, EventType)
-	#define STATETREE_TRACE_CONDITION_EVENT(Index, DataView, EventType)
-	#define STATETREE_TRACE_TRANSITION_EVENT(Source, EventType)
-#endif // WITH_STATETREE_TRACE
-
-namespace UE::StateTree
+namespace UE::StateTree::Debug
 {
-	constexpr int32 DebugIndentSize = 2;	// Debug printing indent for hierarchical data.
-	static std::atomic<uint32> SerialNumber = 0;
-} // UE::StateTree
+	constexpr int32 IndentSize = 2;	// Debug printing indent for hierarchical data.
+	static std::atomic<uint32> InstanceSerialNumber = 0;
+}; // UE::StateTree
+
 
 FStateTreeMinimalExecutionContext::FStateTreeMinimalExecutionContext(UObject& InOwner, const UStateTree& InStateTree, FStateTreeInstanceData& InInstanceData)
 	: Owner(InOwner)
@@ -226,7 +177,8 @@ void FStateTreeMinimalExecutionContext::SendEvent(const FGameplayTag Tag, const
 		return;
 	}
 
-	STATETREE_LOG_AND_TRACE(Verbose, Log, TEXT("Send Event '%s'"), *Tag.ToString());
+	STATETREE_LOG(VeryVerbose, TEXT("Send Event '%s'"), *Tag.ToString());
+	UE_STATETREE_DEBUG_LOG_EVENT(this, Log, TEXT("Send Event '%s'"), *Tag.ToString());
 
 	FStateTreeEventQueue& LocalEventQueue = Storage.GetMutableEventQueue();
 	LocalEventQueue.SendEvent(&Owner, Tag, Payload, Origin);
@@ -250,7 +202,7 @@ FStateTreeInstanceDebugId FStateTreeMinimalExecutionContext::GetInstanceDebugId(
 	FStateTreeInstanceDebugId& InstanceDebugId = Storage.GetMutableExecutionState().InstanceDebugId;
 	if (!InstanceDebugId.IsValid())
 	{
-		InstanceDebugId = FStateTreeInstanceDebugId(GetTypeHash(GetInstanceDescriptionInternal()), ++UE::StateTree::SerialNumber);
+		InstanceDebugId = FStateTreeInstanceDebugId(GetTypeHash(GetInstanceDescriptionInternal()), ++UE::StateTree::Debug::InstanceSerialNumber);
 	}
 	return InstanceDebugId;
 }
@@ -268,7 +220,7 @@ FStateTreeExecutionContext::FCurrentlyProcessedFrameScope::FCurrentlyProcessedFr
 	Context.CurrentlyProcessedParentFrame = CurrentParentFrame;
 	Context.CurrentlyProcessedSharedInstanceStorage = SharedInstanceDataStorage;
 	
-	STATETREE_TRACE_INSTANCE_FRAME_EVENT(Context.GetInstanceDebugId(), Context.CurrentlyProcessedFrame);
+	UE_STATETREE_DEBUG_INSTANCE_FRAME_EVENT(&Context, Context.CurrentlyProcessedFrame);
 }
 
 FStateTreeExecutionContext::FCurrentlyProcessedFrameScope::~FCurrentlyProcessedFrameScope()
@@ -279,7 +231,7 @@ FStateTreeExecutionContext::FCurrentlyProcessedFrameScope::~FCurrentlyProcessedF
 
 	if (Context.CurrentlyProcessedFrame)
 	{
-		STATETREE_TRACE_INSTANCE_FRAME_EVENT(Context.GetInstanceDebugId(), Context.CurrentlyProcessedFrame);
+		UE_STATETREE_DEBUG_INSTANCE_FRAME_EVENT(&Context, Context.CurrentlyProcessedFrame);
 	}
 }
 
@@ -491,11 +443,11 @@ EStateTreeRunStatus FStateTreeExecutionContext::Start(FStartParameters Parameter
 	}
 
 	// Must sent instance creation event first 
-	STATETREE_TRACE_INSTANCE_EVENT(EStateTreeTraceEventType::Push);
+	UE_STATETREE_DEBUG_INSTANCE_EVENT(this, EStateTreeTraceEventType::Push);
 	
 	// Set scoped phase only for properly initialized context with valid Instance data
 	// since we need it to output the InstanceId
-	STATETREE_TRACE_SCOPED_PHASE(EStateTreeUpdatePhase::StartTree);
+	UE_STATETREE_DEBUG_SCOPED_PHASE(this, EStateTreeUpdatePhase::StartTree);
 
 	STATETREE_LOG(VeryVerbose, TEXT("%hs: Starting State Tree %s on owner '%s'."),
 		__FUNCTION__, *GetFullNameSafe(&RootStateTree), *GetNameSafe(&Owner));
@@ -579,7 +531,8 @@ EStateTreeRunStatus FStateTreeExecutionContext::Start(FStartParameters Parameter
 	
 	if (Exec.RequestedStop != EStateTreeRunStatus::Unset)
 	{
-		STATETREE_LOG_AND_TRACE(VeryVerbose, Log, TEXT("Processing Deferred Stop"));
+		STATETREE_LOG(VeryVerbose, TEXT("Processing Deferred Stop"));
+		UE_STATETREE_DEBUG_LOG_EVENT(this, Log, TEXT("Processing Deferred Stop"));
 		Result = Stop(Exec.RequestedStop);
 	}
 	
@@ -606,7 +559,7 @@ EStateTreeRunStatus FStateTreeExecutionContext::Stop(EStateTreeRunStatus Complet
 
 	// Set scoped phase only for properly initialized context with valid Instance data
 	// since we need it to output the InstanceId 
-	STATETREE_TRACE_SCOPED_PHASE(EStateTreeUpdatePhase::StopTree);
+	UE_STATETREE_DEBUG_SCOPED_PHASE(this, EStateTreeUpdatePhase::StopTree);
 
 	// Make sure that we return a valid completion status (i.e. Succeeded, Failed or Stopped)
 	if (CompletionStatus == EStateTreeRunStatus::Unset
@@ -620,7 +573,8 @@ EStateTreeRunStatus FStateTreeExecutionContext::Stop(EStateTreeRunStatus Complet
 	// A reentrant call to Stop or a call from Start or Tick must be deferred.
 	if (Exec.CurrentPhase != EStateTreeUpdatePhase::Unset)
 	{
-		STATETREE_LOG_AND_TRACE(VeryVerbose, Log, TEXT("Deferring Stop at end of %s"), *UEnum::GetDisplayValueAsText(Exec.CurrentPhase).ToString());
+		STATETREE_LOG(VeryVerbose, TEXT("Deferring Stop at end of %s"), *UEnum::GetDisplayValueAsText(Exec.CurrentPhase).ToString());
+		UE_STATETREE_DEBUG_LOG_EVENT(this, Log, TEXT("Deferring Stop at end of %s"), *UEnum::GetDisplayValueAsText(Exec.CurrentPhase).ToString());
 
 		Exec.RequestedStop = CompletionStatus;
 		return EStateTreeRunStatus::Running;
@@ -647,8 +601,8 @@ EStateTreeRunStatus FStateTreeExecutionContext::Stop(EStateTreeRunStatus Complet
 	}
 
 	// Trace before resetting the instance data since it is required to provide all the event information
-	STATETREE_TRACE_ACTIVE_STATES_EVENT({});
-	STATETREE_TRACE_INSTANCE_EVENT(EStateTreeTraceEventType::Pop);
+	UE_STATETREE_DEBUG_ACTIVE_STATES_EVENT(this, {});
+	UE_STATETREE_DEBUG_INSTANCE_EVENT(this, EStateTreeTraceEventType::Pop);
 
 	// Destruct all allocated instance data (does not shrink the buffer). This will invalidate Exec too.
 	InstanceData.Reset();
@@ -708,7 +662,9 @@ EStateTreeRunStatus FStateTreeExecutionContext::TickPostlude()
 	
 	if (Exec.RequestedStop != EStateTreeRunStatus::Unset)
 	{
-		STATETREE_LOG_AND_TRACE(VeryVerbose, Log, TEXT("Processing Deferred Stop"));
+		STATETREE_LOG(VeryVerbose, TEXT("Processing Deferred Stop"));
+		UE_STATETREE_DEBUG_LOG_EVENT(this, Log, TEXT("Processing Deferred Stop"));
+
 		Result = Stop(Exec.RequestedStop);
 	}
 
@@ -789,13 +745,16 @@ void FStateTreeExecutionContext::TickUpdateTasksInternal(const float DeltaTime)
 	}
 	else
 	{
-		STATETREE_TRACE_LOG_EVENT(Log, TEXT("Global tasks completed (%s), stopping the tree"), *UEnum::GetDisplayValueAsText(EvalAndGlobalTaskStatus).ToString());
+		UE_STATETREE_DEBUG_LOG_EVENT(this, Log, TEXT("Global tasks completed (%s), stopping the tree"), *UEnum::GetDisplayValueAsText(EvalAndGlobalTaskStatus).ToString());
+		STATETREE_LOG(Log, TEXT("Global tasks completed (%s), stopping the tree"), *UEnum::GetDisplayValueAsText(EvalAndGlobalTaskStatus).ToString());
 		Exec.RequestedStop = EvalAndGlobalTaskStatus;
 	}
 }
 
 void FStateTreeExecutionContext::TickTriggerTransitionsInternal()
 {
+	UE_STATETREE_DEBUG_SCOPED_PHASE(this, EStateTreeUpdatePhase::TickTransitions);
+
 	FStateTreeExecutionState& Exec = GetExecState();
 
 	// If stop is requested, do not try to trigger transitions.
@@ -812,8 +771,8 @@ void FStateTreeExecutionContext::TickTriggerTransitionsInternal()
 		// Trigger conditional transitions or state succeed/failed transitions. First tick transition is handled here too.
 		if (TriggerTransitions())
 		{
-			STATETREE_TRACE_SCOPED_PHASE(EStateTreeUpdatePhase::ApplyTransitions);
-			STATETREE_TRACE_TRANSITION_EVENT(NextTransitionSource, EStateTreeTraceEventType::OnTransition);
+			UE_STATETREE_DEBUG_SCOPED_PHASE(this, EStateTreeUpdatePhase::ApplyTransitions);
+			UE_STATETREE_DEBUG_TRANSITION_EVENT(this, NextTransitionSource, EStateTreeTraceEventType::OnTransition);
 			NextTransitionSource.Reset();
 
 			ExitState(NextTransition);
@@ -1744,7 +1703,7 @@ EStateTreeRunStatus FStateTreeExecutionContext::EnterState(FStateTreeTransitionR
 	EStateTreeRunStatus Result = EStateTreeRunStatus::Running;
 
 	STATETREE_LOG(Log, TEXT("Enter state '%s' (%d)"), *DebugGetStatePath(Transition.NextActiveFrames), Exec.StateChangeCount);
-	STATETREE_TRACE_PHASE_BEGIN(EStateTreeUpdatePhase::EnterStates);
+	UE_STATETREE_DEBUG_ENTER_PHASE(this, EStateTreeUpdatePhase::EnterStates);
 
 	// The previous active frames are needed for state enter logic.
 	TArray<FStateTreeExecutionFrame, TConcurrentLinearArrayAllocator<FDefaultBlockAllocationTag>> PreviousActiveFrames;
@@ -1829,8 +1788,8 @@ EStateTreeRunStatus FStateTreeExecutionContext::EnterState(FStateTreeTransitionR
 
 			if (bIsEnteringState)
 			{
-				STATETREE_TRACE_STATE_EVENT(CurrentHandle, EStateTreeTraceEventType::OnEntering);
-				STATETREE_LOG(Log, TEXT("%*sState '%s' (%s)"), (FrameIndex + StateIndex + 1)*UE::StateTree::DebugIndentSize, TEXT("")
+				UE_STATETREE_DEBUG_STATE_EVENT(this, CurrentHandle, EStateTreeTraceEventType::OnEntering);
+				STATETREE_LOG(Log, TEXT("%*sState '%s' (%s)"), (FrameIndex + StateIndex + 1)*UE::StateTree::Debug::IndentSize, TEXT("")
 					, *GetSafeStateName(NextFrame, CurrentHandle)
 					, *UEnum::GetDisplayValueAsText(CurrentTransition.ChangeType).ToString());
 			}
@@ -1887,7 +1846,7 @@ EStateTreeRunStatus FStateTreeExecutionContext::EnterState(FStateTreeTransitionR
 				// Ignore disabled task
 				if (Task.bTaskEnabled == false)
 				{
-					STATETREE_LOG(VeryVerbose, TEXT("%*sSkipped 'EnterState' for disabled Task: '%s'"), UE::StateTree::DebugIndentSize, TEXT(""), *Task.Name.ToString());
+					STATETREE_LOG(VeryVerbose, TEXT("%*sSkipped 'EnterState' for disabled Task: '%s'"), UE::StateTree::Debug::IndentSize, TEXT(""), *Task.Name.ToString());
 					continue;
 				}
 
@@ -1896,7 +1855,7 @@ EStateTreeRunStatus FStateTreeExecutionContext::EnterState(FStateTreeTransitionR
 
 				if (bIsEnteringState && bShouldCallEnterState)
 				{
-					STATETREE_LOG(Verbose, TEXT("%*sTask '%s'.EnterState()"), (FrameIndex + StateIndex + 1)*UE::StateTree::DebugIndentSize, TEXT(""), *Task.Name.ToString());
+					STATETREE_LOG(Verbose, TEXT("%*sTask '%s'.EnterState()"), (FrameIndex + StateIndex + 1)*UE::StateTree::Debug::IndentSize, TEXT(""), *Task.Name.ToString());
 					UE_STATETREE_DEBUG_TASK_ENTER_STATE(this, NextFrame.StateTree, FStateTreeIndex16(TaskIndex));
 
 					EStateTreeRunStatus Status = EStateTreeRunStatus::Unset;
@@ -1907,7 +1866,7 @@ EStateTreeRunStatus FStateTreeExecutionContext::EnterState(FStateTreeTransitionR
 						Status = Task.EnterState(*this, CurrentTransition);
 					}
 
-					STATETREE_TRACE_TASK_EVENT(TaskIndex, TaskInstanceView, EStateTreeTraceEventType::OnEntered, Status);
+					UE_STATETREE_DEBUG_TASK_EVENT(this, TaskIndex, TaskInstanceView, EStateTreeTraceEventType::OnEntered, Status);
 
 					if (Status != EStateTreeRunStatus::Running)
 					{
@@ -1932,14 +1891,14 @@ EStateTreeRunStatus FStateTreeExecutionContext::EnterState(FStateTreeTransitionR
 
 			if (bIsEnteringState)
 			{
-				STATETREE_TRACE_STATE_EVENT(CurrentHandle, EStateTreeTraceEventType::OnEntered);
+				UE_STATETREE_DEBUG_STATE_EVENT(this, CurrentHandle, EStateTreeTraceEventType::OnEntered);
 			}
 		}
 	}
 
-	STATETREE_TRACE_PHASE_END(EStateTreeUpdatePhase::EnterStates);
+	UE_STATETREE_DEBUG_EXIT_PHASE(this, EStateTreeUpdatePhase::EnterStates);
 
-	STATETREE_TRACE_ACTIVE_STATES_EVENT(Exec.ActiveFrames);
+	UE_STATETREE_DEBUG_ACTIVE_STATES_EVENT(this, Exec.ActiveFrames);
 
 	return Result;
 }
@@ -2080,7 +2039,7 @@ void FStateTreeExecutionContext::ExitState(const FStateTreeTransitionResult& Tra
 
 	// Call in reverse order.
 	STATETREE_LOG(Log, TEXT("Exit state '%s' (%d)"), *DebugGetStatePath(Exec.ActiveFrames), Exec.StateChangeCount);
-	STATETREE_TRACE_SCOPED_PHASE(EStateTreeUpdatePhase::ExitStates);
+	UE_STATETREE_DEBUG_SCOPED_PHASE(this, EStateTreeUpdatePhase::ExitStates);
 
 	FStateTreeTransitionResult CurrentTransition = Transition;
 	int32 CallIndex = ExitStateCalls.Num() - 1;
@@ -2101,11 +2060,11 @@ void FStateTreeExecutionContext::ExitState(const FStateTreeTransitionResult& Tra
 			const FExitStateCall& ExitCall = ExitStateCalls[CallIndex--];
 			CurrentTransition.ChangeType = ExitCall.ChangeType;
 
-			STATETREE_LOG(Log, TEXT("%*sState '%s' (%s)"), (FrameIndex + StateIndex + 1)*UE::StateTree::DebugIndentSize, TEXT("")
+			STATETREE_LOG(Log, TEXT("%*sState '%s' (%s)"), (FrameIndex + StateIndex + 1)*UE::StateTree::Debug::IndentSize, TEXT("")
 				, *GetSafeStateName(CurrentFrame, CurrentHandle)
 				, *UEnum::GetDisplayValueAsText(CurrentTransition.ChangeType).ToString());
 
-			STATETREE_TRACE_STATE_EVENT(CurrentHandle, EStateTreeTraceEventType::OnExiting);
+			UE_STATETREE_DEBUG_STATE_EVENT(this, CurrentHandle, EStateTreeTraceEventType::OnExiting);
 
 			if (ExitCall.bShouldCall)
 			{
@@ -2137,7 +2096,7 @@ void FStateTreeExecutionContext::ExitState(const FStateTreeTransitionResult& Tra
 						// Ignore disabled task
 						if (Task.bTaskEnabled == false)
 						{
-							STATETREE_LOG(VeryVerbose, TEXT("%*sSkipped 'ExitState' for disabled Task: '%s'"), UE::StateTree::DebugIndentSize, TEXT(""), *Task.Name.ToString());
+							STATETREE_LOG(VeryVerbose, TEXT("%*sSkipped 'ExitState' for disabled Task: '%s'"), UE::StateTree::Debug::IndentSize, TEXT(""), *Task.Name.ToString());
 							continue;
 						}
 
@@ -2146,14 +2105,14 @@ void FStateTreeExecutionContext::ExitState(const FStateTreeTransitionResult& Tra
 
 						if (bShouldCallStateChange)
 						{
-							STATETREE_LOG(Verbose, TEXT("%*sTask '%s'.ExitState()"), (FrameIndex + StateIndex + 1)*UE::StateTree::DebugIndentSize, TEXT(""), *Task.Name.ToString());
+							STATETREE_LOG(Verbose, TEXT("%*sTask '%s'.ExitState()"), (FrameIndex + StateIndex + 1)*UE::StateTree::Debug::IndentSize, TEXT(""), *Task.Name.ToString());
 							UE_STATETREE_DEBUG_TASK_EXIT_STATE(this, CurrentStateTree, FStateTreeIndex16(TaskIndex));
 							{
 								QUICK_SCOPE_CYCLE_COUNTER(StateTree_Task_ExitState);
 								CSV_SCOPED_TIMING_STAT_EXCLUSIVE(StateTree_Task_ExitState);
 								Task.ExitState(*this, CurrentTransition);
 							}
-							STATETREE_TRACE_TASK_EVENT(TaskIndex, TaskInstanceView, EStateTreeTraceEventType::OnExited, Transition.CurrentRunStatus);
+							UE_STATETREE_DEBUG_TASK_EVENT(this, TaskIndex, TaskInstanceView, EStateTreeTraceEventType::OnExited, Transition.CurrentRunStatus);
 						}
 					}
 				}
@@ -2194,7 +2153,7 @@ void FStateTreeExecutionContext::ExitState(const FStateTreeTransitionResult& Tra
 				}
 			}
 
-			STATETREE_TRACE_STATE_EVENT(CurrentHandle, EStateTreeTraceEventType::OnExited);
+			UE_STATETREE_DEBUG_STATE_EVENT(this, CurrentHandle, EStateTreeTraceEventType::OnExited);
 		}
 
 		// Frame exit call
@@ -2222,7 +2181,7 @@ void FStateTreeExecutionContext::StateCompleted()
 	}
 
 	STATETREE_LOG(Verbose, TEXT("State Completed %s (%d)"), *UEnum::GetDisplayValueAsText(Exec.LastTickStatus).ToString(), Exec.StateChangeCount);
-	STATETREE_TRACE_SCOPED_PHASE(EStateTreeUpdatePhase::StateCompleted);
+	UE_STATETREE_DEBUG_SCOPED_PHASE(this, EStateTreeUpdatePhase::StateCompleted);
 
 	// Call from child towards root to allow to pass results back.
 	// Note: Completed is assumed to be called immediately after tick or enter state, so there's no property copying.
@@ -2244,8 +2203,8 @@ void FStateTreeExecutionContext::StateCompleted()
 
 				FCurrentlyProcessedStateScope StateScope(*this, CurrentHandle);
 				
-				STATETREE_LOG(Verbose, TEXT("%*sState '%s'"), (FrameIndex + StateIndex + 1)*UE::StateTree::DebugIndentSize, TEXT(""), *GetSafeStateName(CurrentFrame, CurrentHandle));
-				STATETREE_TRACE_STATE_EVENT(CurrentHandle, EStateTreeTraceEventType::OnStateCompleted);
+				STATETREE_LOG(Verbose, TEXT("%*sState '%s'"), (FrameIndex + StateIndex + 1)*UE::StateTree::Debug::IndentSize, TEXT(""), *GetSafeStateName(CurrentFrame, CurrentHandle));
+				UE_STATETREE_DEBUG_STATE_EVENT(this, CurrentHandle, EStateTreeTraceEventType::OnStateCompleted);
 
 				// Notify Tasks
 				for (int32 TaskIndex = (State.TasksBegin + State.TasksNum) - 1; TaskIndex >= State.TasksBegin; TaskIndex--)
@@ -2262,11 +2221,11 @@ void FStateTreeExecutionContext::StateCompleted()
 						// Ignore disabled task
 						if (Task.bTaskEnabled == false)
 						{
-							STATETREE_LOG(VeryVerbose, TEXT("%*sSkipped 'StateCompleted' for disabled Task: '%s'"), UE::StateTree::DebugIndentSize, TEXT(""), *Task.Name.ToString());
+							STATETREE_LOG(VeryVerbose, TEXT("%*sSkipped 'StateCompleted' for disabled Task: '%s'"), UE::StateTree::Debug::IndentSize, TEXT(""), *Task.Name.ToString());
 							continue;
 						}
 
-						STATETREE_LOG(Verbose, TEXT("%*sTask '%s'.StateCompleted()"), (FrameIndex + StateIndex + 1)*UE::StateTree::DebugIndentSize, TEXT(""), *Task.Name.ToString());
+						STATETREE_LOG(Verbose, TEXT("%*sTask '%s'.StateCompleted()"), (FrameIndex + StateIndex + 1)*UE::StateTree::Debug::IndentSize, TEXT(""), *Task.Name.ToString());
 						Task.StateCompleted(*this, Exec.LastTickStatus, CurrentFrame.ActiveStates);
 					}
 				}
@@ -2306,7 +2265,7 @@ void FStateTreeExecutionContext::StateCompleted()
 EStateTreeRunStatus FStateTreeExecutionContext::TickEvaluatorsAndGlobalTasks(const float DeltaTime, const bool bTickGlobalTasks)
 {
 	CSV_SCOPED_TIMING_STAT_EXCLUSIVE(StateTree_TickEvaluators);
-	STATETREE_TRACE_SCOPED_PHASE(EStateTreeUpdatePhase::TickingGlobalTasks);
+	UE_STATETREE_DEBUG_SCOPED_PHASE(this, EStateTreeUpdatePhase::TickingGlobalTasks);
 
 	STATETREE_LOG(VeryVerbose, TEXT("Ticking Evaluators & Global Tasks"));
 
@@ -2342,7 +2301,7 @@ EStateTreeRunStatus FStateTreeExecutionContext::TickEvaluatorsAndGlobalTasks(con
 					QUICK_SCOPE_CYCLE_COUNTER(StateTree_Eval_Tick);
 					Eval.Tick(*this, DeltaTime);
 
-					STATETREE_TRACE_EVALUATOR_EVENT(EvalIndex, EvalInstanceView, EStateTreeTraceEventType::OnTicked);
+					UE_STATETREE_DEBUG_EVALUATOR_EVENT(this, EvalIndex, EvalInstanceView, EStateTreeTraceEventType::OnTicked);
 				}
 			}
 
@@ -2364,7 +2323,7 @@ EStateTreeRunStatus FStateTreeExecutionContext::TickEvaluatorsAndGlobalTasks(con
 						// Ignore disabled task
 						if (Task.bTaskEnabled == false)
 						{
-							STATETREE_LOG(VeryVerbose, TEXT("%*sSkipped 'Tick' for disabled Task: '%s'"), UE::StateTree::DebugIndentSize, TEXT(""), *Task.Name.ToString());
+							STATETREE_LOG(VeryVerbose, TEXT("%*sSkipped 'Tick' for disabled Task: '%s'"), UE::StateTree::Debug::IndentSize, TEXT(""), *Task.Name.ToString());
 							continue;
 						}
 
@@ -2382,7 +2341,7 @@ EStateTreeRunStatus FStateTreeExecutionContext::TickEvaluatorsAndGlobalTasks(con
 							CopyBatchOnActiveInstances(CurrentParentFrame, CurrentFrame, TaskInstanceView, Task.BindingsBatch);
 						}
 
-						//STATETREE_TRACE_TASK_EVENT(TaskIndex, TaskDataView, EStateTreeTraceEventType::OnTickingTask, EStateTreeRunStatus::Running);
+						//UE_STATETREE_DEBUG_TASK_EVENT(this, TaskIndex, TaskDataView, EStateTreeTraceEventType::OnTickingTask, EStateTreeRunStatus::Running);
 						UE_STATETREE_DEBUG_TASK_TICK(this, CurrentStateTree, TaskIndex);
 
 						EStateTreeRunStatus TaskResult = EStateTreeRunStatus::Unset;
@@ -2394,7 +2353,7 @@ EStateTreeRunStatus FStateTreeExecutionContext::TickEvaluatorsAndGlobalTasks(con
 							bTaskTicked = true;
 						}
 
-						STATETREE_TRACE_TASK_EVENT(TaskIndex, TaskInstanceView,
+						UE_STATETREE_DEBUG_TASK_EVENT(this, TaskIndex, TaskInstanceView,
 							TaskResult != EStateTreeRunStatus::Running ? EStateTreeTraceEventType::OnTaskCompleted : EStateTreeTraceEventType::OnTicked,
 							TaskResult);
 
@@ -2422,7 +2381,7 @@ EStateTreeRunStatus FStateTreeExecutionContext::TickEvaluatorsAndGlobalTasks(con
 EStateTreeRunStatus FStateTreeExecutionContext::StartEvaluatorsAndGlobalTasks(FStateTreeIndex16& OutLastInitializedTaskIndex)
 {
 	CSV_SCOPED_TIMING_STAT_EXCLUSIVE(StateTree_StartEvaluators);
-	STATETREE_TRACE_SCOPED_PHASE(EStateTreeUpdatePhase::StartGlobalTasks);
+	UE_STATETREE_DEBUG_SCOPED_PHASE(this, EStateTreeUpdatePhase::StartGlobalTasks);
 
 	STATETREE_LOG(Verbose, TEXT("Start Evaluators & Global tasks"));
 
@@ -2459,7 +2418,7 @@ EStateTreeRunStatus FStateTreeExecutionContext::StartEvaluatorsAndGlobalTasks(FS
 					QUICK_SCOPE_CYCLE_COUNTER(StateTree_Eval_TreeStart);
 					Eval.TreeStart(*this);
 
-					STATETREE_TRACE_EVALUATOR_EVENT(EvalIndex, EvalInstanceView, EStateTreeTraceEventType::OnTreeStarted);
+					UE_STATETREE_DEBUG_EVALUATOR_EVENT(this, EvalIndex, EvalInstanceView, EStateTreeTraceEventType::OnTreeStarted);
 				}
 			}
 
@@ -2482,7 +2441,7 @@ EStateTreeRunStatus FStateTreeExecutionContext::StartEvaluatorsAndGlobalTasks(FS
 				// Ignore disabled task
 				if (Task.bTaskEnabled == false)
 				{
-					STATETREE_LOG(VeryVerbose, TEXT("%*sSkipped 'EnterState' for disabled Task: '%s'"), UE::StateTree::DebugIndentSize, TEXT(""), *Task.Name.ToString());
+					STATETREE_LOG(VeryVerbose, TEXT("%*sSkipped 'EnterState' for disabled Task: '%s'"), UE::StateTree::Debug::IndentSize, TEXT(""), *Task.Name.ToString());
 					continue;
 				}
 
@@ -2492,7 +2451,7 @@ EStateTreeRunStatus FStateTreeExecutionContext::StartEvaluatorsAndGlobalTasks(FS
 					QUICK_SCOPE_CYCLE_COUNTER(StateTree_Task_TreeStart);
 					const EStateTreeRunStatus TaskStatus = Task.EnterState(*this, Transition);
 
-					STATETREE_TRACE_TASK_EVENT(TaskIndex, TaskInstanceView, EStateTreeTraceEventType::OnEntered, TaskStatus);
+					UE_STATETREE_DEBUG_TASK_EVENT(this, TaskIndex, TaskInstanceView, EStateTreeTraceEventType::OnEntered, TaskStatus);
 
 					if (TaskStatus != EStateTreeRunStatus::Running)
 					{
@@ -2511,7 +2470,7 @@ EStateTreeRunStatus FStateTreeExecutionContext::StartEvaluatorsAndGlobalTasks(FS
 void FStateTreeExecutionContext::StopEvaluatorsAndGlobalTasks(const EStateTreeRunStatus CompletionStatus, const FStateTreeIndex16 LastInitializedTaskIndex)
 {
 	CSV_SCOPED_TIMING_STAT_EXCLUSIVE(StateTree_StopEvaluators);
-	STATETREE_TRACE_SCOPED_PHASE(EStateTreeUpdatePhase::StopGlobalTasks);
+	UE_STATETREE_DEBUG_SCOPED_PHASE(this, EStateTreeUpdatePhase::StopGlobalTasks);
 
 	STATETREE_LOG(Verbose, TEXT("Stop Evaluators & Global Tasks"));
 
@@ -2592,7 +2551,7 @@ void FStateTreeExecutionContext::CallStopOnEvaluatorsAndGlobalTasks(const FState
 		// Ignore disabled task
 		if (Task.bTaskEnabled == false)
 		{
-			STATETREE_LOG(VeryVerbose, TEXT("%*sSkipped 'ExitState' for disabled Task: '%s'"), UE::StateTree::DebugIndentSize, TEXT(""), *Task.Name.ToString());
+			STATETREE_LOG(VeryVerbose, TEXT("%*sSkipped 'ExitState' for disabled Task: '%s'"), UE::StateTree::Debug::IndentSize, TEXT(""), *Task.Name.ToString());
 			continue;
 		}
 
@@ -2605,7 +2564,7 @@ void FStateTreeExecutionContext::CallStopOnEvaluatorsAndGlobalTasks(const FState
 				QUICK_SCOPE_CYCLE_COUNTER(StateTree_Task_TreeStop);
 				Task.ExitState(*this, Transition);
 			}
-			STATETREE_TRACE_TASK_EVENT(TaskIndex, TaskInstanceView, EStateTreeTraceEventType::OnExited, Transition.CurrentRunStatus);
+			UE_STATETREE_DEBUG_TASK_EVENT(this, TaskIndex, TaskInstanceView, EStateTreeTraceEventType::OnExited, Transition.CurrentRunStatus);
 		}
 	}
 
@@ -2621,7 +2580,7 @@ void FStateTreeExecutionContext::CallStopOnEvaluatorsAndGlobalTasks(const FState
 			QUICK_SCOPE_CYCLE_COUNTER(StateTree_Eval_TreeStop);
 			Eval.TreeStop(*this);
 
-			STATETREE_TRACE_EVALUATOR_EVENT(EvalIndex, EvalInstanceView, EStateTreeTraceEventType::OnTreeStopped);
+			UE_STATETREE_DEBUG_EVALUATOR_EVENT(this, EvalIndex, EvalInstanceView, EStateTreeTraceEventType::OnTreeStopped);
 		}
 	}
 }
@@ -2632,16 +2591,14 @@ EStateTreeRunStatus FStateTreeExecutionContext::StartTemporaryEvaluatorsAndGloba
 	{
 		return EStateTreeRunStatus::Failed;
 	}
-	
-	CSV_SCOPED_TIMING_STAT_EXCLUSIVE(StateTree_StartEvaluators);
-	// @todo: figure out debugger phase for temporary start.
-	//	STATETREE_TRACE_SCOPED_PHASE(EStateTreeUpdatePhase::StartGlobalTasks);
 
+	CSV_SCOPED_TIMING_STAT_EXCLUSIVE(StateTree_StartEvaluators);
 	STATETREE_LOG(Verbose, TEXT("Start Temporary Evaluators & Global tasks while trying to select linked asset: %s"), *GetNameSafe(CurrentFrame.StateTree));
 
-	FStateTreeExecutionState& Exec = GetExecState();
-	EStateTreeRunStatus Result = EStateTreeRunStatus::Running;
 	FCurrentlyProcessedFrameScope FrameScope(*this, CurrentParentFrame, CurrentFrame);
+	UE_STATETREE_DEBUG_SCOPED_PHASE(this, EStateTreeUpdatePhase::StartGlobalTasksForSelection);
+
+	EStateTreeRunStatus Result = EStateTreeRunStatus::Running;
 	const UStateTree* CurrentStateTree = CurrentFrame.StateTree;
 
 	// Start evaluators
@@ -2668,7 +2625,7 @@ EStateTreeRunStatus FStateTreeExecutionContext::StartTemporaryEvaluatorsAndGloba
 			QUICK_SCOPE_CYCLE_COUNTER(StateTree_Eval_TreeStart);
 			Eval.TreeStart(*this);
 
-			STATETREE_TRACE_EVALUATOR_EVENT(EvalIndex, EvalInstanceView, EStateTreeTraceEventType::OnTreeStarted);
+			UE_STATETREE_DEBUG_EVALUATOR_EVENT(this, EvalIndex, EvalInstanceView, EStateTreeTraceEventType::OnTreeStarted);
 		}
 	}
 
@@ -2682,7 +2639,7 @@ EStateTreeRunStatus FStateTreeExecutionContext::StartTemporaryEvaluatorsAndGloba
 		// Ignore disabled task
 		if (Task.bTaskEnabled == false)
 		{
-			STATETREE_LOG(VeryVerbose, TEXT("%*sSkipped 'EnterState' for disabled Task: '%s'"), UE::StateTree::DebugIndentSize, TEXT(""), *Task.Name.ToString());
+			STATETREE_LOG(VeryVerbose, TEXT("%*sSkipped 'EnterState' for disabled Task: '%s'"), UE::StateTree::Debug::IndentSize, TEXT(""), *Task.Name.ToString());
 			continue;
 		}
 
@@ -2706,7 +2663,7 @@ EStateTreeRunStatus FStateTreeExecutionContext::StartTemporaryEvaluatorsAndGloba
 		QUICK_SCOPE_CYCLE_COUNTER(StateTree_Task_TreeStart);
 		const EStateTreeRunStatus TaskStatus = Task.EnterState(*this, Transition);
 
-		STATETREE_TRACE_TASK_EVENT(TaskIndex, TaskDataView, EStateTreeTraceEventType::OnEntered, TaskStatus);
+		UE_STATETREE_DEBUG_TASK_EVENT(this, TaskIndex, TaskDataView, EStateTreeTraceEventType::OnEntered, TaskStatus);
 
 		if (TaskStatus != EStateTreeRunStatus::Running)
 		{
@@ -2720,9 +2677,11 @@ EStateTreeRunStatus FStateTreeExecutionContext::StartTemporaryEvaluatorsAndGloba
 
 void FStateTreeExecutionContext::StopTemporaryEvaluatorsAndGlobalTasks(const FStateTreeExecutionFrame* CurrentParentFrame, const FStateTreeExecutionFrame& CurrentFrame)
 {
-	// @todo: figure out debugger phase for temporary stop.
 	STATETREE_LOG(Verbose, TEXT("Stop Temporary Evaluators & Global tasks"));
 
+	FCurrentlyProcessedFrameScope FrameScope(*this, CurrentParentFrame, CurrentFrame);
+	UE_STATETREE_DEBUG_SCOPED_PHASE(this, EStateTreeUpdatePhase::StopGlobalTasksForSelection);
+
 	// Create temporary transition to stop the unused global tasks and evaluators.
 	constexpr  EStateTreeRunStatus CompletionStatus = EStateTreeRunStatus::Stopped; 
 	FStateTreeTransitionResult Transition;
@@ -2750,8 +2709,7 @@ void FStateTreeExecutionContext::StopTemporaryEvaluatorsAndGlobalTasks(const FSt
 			{
 				NodeInstanceView = FStateTreeDataView(TempInstance.Instance);
 			}
-			
-			FCurrentlyProcessedFrameScope FrameScope(*this, CurrentParentFrame, CurrentFrame);
+
 			FNodeInstanceDataScope DataScope(*this, TempInstance.DataHandle, NodeInstanceView);
 
 			FConstStructView NodeView = CurrentFrame.StateTree->Nodes[TempInstance.OwnerNodeIndex.Get()];
@@ -2763,7 +2721,7 @@ void FStateTreeExecutionContext::StopTemporaryEvaluatorsAndGlobalTasks(const FSt
 					QUICK_SCOPE_CYCLE_COUNTER(StateTree_Task_TreeStop);
 					Task->ExitState(*this, Transition);
 				}
-				STATETREE_TRACE_TASK_EVENT(TempInstance.OwnerNodeIndex.Get(), NodeInstanceView, EStateTreeTraceEventType::OnExited, Transition.CurrentRunStatus);
+				UE_STATETREE_DEBUG_TASK_EVENT(this, TempInstance.OwnerNodeIndex.Get(), NodeInstanceView, EStateTreeTraceEventType::OnExited, Transition.CurrentRunStatus);
 			}
 			else if (const FStateTreeEvaluatorBase* Eval = NodeView.GetPtr<const FStateTreeEvaluatorBase>())
 			{
@@ -2773,7 +2731,7 @@ void FStateTreeExecutionContext::StopTemporaryEvaluatorsAndGlobalTasks(const FSt
 					QUICK_SCOPE_CYCLE_COUNTER(StateTree_Eval_TreeStop);
 					Eval->TreeStop(*this);
 
-					STATETREE_TRACE_EVALUATOR_EVENT(TempInstance.OwnerNodeIndex.Get(), NodeInstanceView, EStateTreeTraceEventType::OnTreeStopped);
+					UE_STATETREE_DEBUG_EVALUATOR_EVENT(this, TempInstance.OwnerNodeIndex.Get(), NodeInstanceView, EStateTreeTraceEventType::OnTreeStopped);
 				}
 			}
 		}
@@ -2783,7 +2741,7 @@ void FStateTreeExecutionContext::StopTemporaryEvaluatorsAndGlobalTasks(const FSt
 EStateTreeRunStatus FStateTreeExecutionContext::TickTasks(const float DeltaTime)
 {
 	CSV_SCOPED_TIMING_STAT_EXCLUSIVE(StateTree_TickTasks);
-	STATETREE_TRACE_SCOPED_PHASE(EStateTreeUpdatePhase::TickingTasks);
+	UE_STATETREE_DEBUG_SCOPED_PHASE(this, EStateTreeUpdatePhase::TickingTasks);
 
 	FStateTreeExecutionState& Exec = GetExecState();
 
@@ -2819,9 +2777,9 @@ EStateTreeRunStatus FStateTreeExecutionContext::TickTasks(const float DeltaTime)
 			const FCompactStateTreeState& State = CurrentStateTree->States[CurrentHandle.Index];
 
 			FCurrentlyProcessedStateScope StateScope(*this, CurrentHandle);
-			STATETREE_TRACE_SCOPED_STATE(CurrentHandle);
+			UE_STATETREE_DEBUG_SCOPED_STATE(this, CurrentHandle);
 
-			STATETREE_CLOG(State.TasksNum > 0, VeryVerbose, TEXT("%*sState '%s'"), (FrameIndex + StateIndex + 1)*UE::StateTree::DebugIndentSize, TEXT("")
+			STATETREE_CLOG(State.TasksNum > 0, VeryVerbose, TEXT("%*sState '%s'"), (FrameIndex + StateIndex + 1)*UE::StateTree::Debug::IndentSize, TEXT("")
 				, *DebugGetStatePath(Exec.ActiveFrames, &CurrentFrame, StateIndex));
 
 			if (State.ShouldTickTasks(bHasEvents))
@@ -2850,12 +2808,12 @@ EStateTreeRunStatus FStateTreeExecutionContext::TickTasks(const float DeltaTime)
 					// Ignore disabled task
 					if (Task.bTaskEnabled == false)
 					{
-						STATETREE_LOG(VeryVerbose, TEXT("%*sSkipped 'Tick' for disabled Task: '%s'"), UE::StateTree::DebugIndentSize, TEXT(""), *Task.Name.ToString());
+						STATETREE_LOG(VeryVerbose, TEXT("%*sSkipped 'Tick' for disabled Task: '%s'"), UE::StateTree::Debug::IndentSize, TEXT(""), *Task.Name.ToString());
 						continue;
 					}
 
 					const bool bNeedsTick = bShouldTickTasks && (Task.bShouldCallTick || (bHasEvents && Task.bShouldCallTickOnlyOnEvents));
-					STATETREE_LOG(VeryVerbose, TEXT("%*s  Tick: '%s' %s"), (FrameIndex + StateIndex + 1)*UE::StateTree::DebugIndentSize, TEXT("")
+					STATETREE_LOG(VeryVerbose, TEXT("%*s  Tick: '%s' %s"), (FrameIndex + StateIndex + 1)*UE::StateTree::Debug::IndentSize, TEXT("")
 						, *Task.Name.ToString()
 						, !bNeedsTick ? TEXT("[not ticked]") : TEXT(""));
 
@@ -2871,7 +2829,7 @@ EStateTreeRunStatus FStateTreeExecutionContext::TickTasks(const float DeltaTime)
 						CopyBatchOnActiveInstances(CurrentParentFrame, CurrentFrame, TaskInstanceView, Task.BindingsBatch);
 					}
 
-					//STATETREE_TRACE_TASK_EVENT(TaskIndex, TaskDataView, EStateTreeTraceEventType::OnTickingTask, EStateTreeRunStatus::Running);
+					//UE_STATETREE_DEBUG_TASK_EVENT(this, TaskIndex, TaskDataView, EStateTreeTraceEventType::OnTickingTask, EStateTreeRunStatus::Running);
 					UE_STATETREE_DEBUG_TASK_TICK(this, CurrentStateTree, TaskIndex);
 
 					EStateTreeRunStatus TaskResult = EStateTreeRunStatus::Unset;
@@ -2883,7 +2841,7 @@ EStateTreeRunStatus FStateTreeExecutionContext::TickTasks(const float DeltaTime)
 						bTaskTicked = true;
 					}
 
-					STATETREE_TRACE_TASK_EVENT(TaskIndex, TaskInstanceView,
+					UE_STATETREE_DEBUG_TASK_EVENT(this, TaskIndex, TaskInstanceView,
 						TaskResult != EStateTreeRunStatus::Running ? EStateTreeTraceEventType::OnTaskCompleted : EStateTreeTraceEventType::OnTicked,
 						TaskResult);
 
@@ -2954,8 +2912,8 @@ bool FStateTreeExecutionContext::TestAllConditions(const FStateTreeExecutionFram
 				if (!CopyBatchWithValidation(CurrentParentFrame, CurrentFrame, ConditionInstanceView, Cond.BindingsBatch))
 				{
 					// If the source data cannot be accessed, the whole expression evaluates to false.
-					STATETREE_TRACE_CONDITION_EVENT(ConditionIndex, ConditionInstanceView, EStateTreeTraceEventType::InternalForcedFailure);
-					STATETREE_TRACE_LOG_EVENT(Warning, TEXT("Evaluation forced to false: source data cannot be accessed (e.g. enter conditions trying to access inactive parent state)"));
+					UE_STATETREE_DEBUG_CONDITION_EVENT(this, ConditionIndex, ConditionInstanceView, EStateTreeTraceEventType::InternalForcedFailure);
+					UE_STATETREE_DEBUG_LOG_EVENT(this, Warning, TEXT("Evaluation forced to false: source data cannot be accessed (e.g. enter conditions trying to access inactive parent state)"));
 					Values[0] = false;
 					break;
 				}
@@ -2964,7 +2922,7 @@ bool FStateTreeExecutionContext::TestAllConditions(const FStateTreeExecutionFram
 			UE_STATETREE_DEBUG_CONDITION_TEST_CONDITION(this, CurrentFrame.StateTree, Index);
 
 			bValue = Cond.TestCondition(*this);
-			STATETREE_TRACE_CONDITION_EVENT(ConditionIndex, ConditionInstanceView, bValue ? EStateTreeTraceEventType::Passed : EStateTreeTraceEventType::Failed);
+			UE_STATETREE_DEBUG_CONDITION_EVENT(this, ConditionIndex, ConditionInstanceView, bValue ? EStateTreeTraceEventType::Passed : EStateTreeTraceEventType::Failed);
 			
 			// Reset copied properties that might contain object references.
 			if (Cond.BindingsBatch.IsValid())
@@ -2975,7 +2933,7 @@ bool FStateTreeExecutionContext::TestAllConditions(const FStateTreeExecutionFram
 		else
 		{
 			bValue = Cond.EvaluationMode == EStateTreeConditionEvaluationMode::ForcedTrue;
-			STATETREE_TRACE_CONDITION_EVENT(ConditionIndex, FStateTreeDataView{}, bValue ? EStateTreeTraceEventType::ForcedSuccess : EStateTreeTraceEventType::ForcedFailure);
+			UE_STATETREE_DEBUG_CONDITION_EVENT(this, ConditionIndex, FStateTreeDataView{}, bValue ? EStateTreeTraceEventType::ForcedSuccess : EStateTreeTraceEventType::ForcedFailure);
 		}
 
 		const int32 DeltaIndent = Cond.DeltaIndent;
@@ -3286,14 +3244,15 @@ void FStateTreeExecutionContext::SetupNextTransition(const FStateTreeExecutionFr
 bool FStateTreeExecutionContext::TriggerTransitions()
 {
 	CSV_SCOPED_TIMING_STAT_EXCLUSIVE(StateTree_TriggerTransition);
-	STATETREE_TRACE_SCOPED_PHASE(EStateTreeUpdatePhase::TriggerTransitions);
+	UE_STATETREE_DEBUG_SCOPED_PHASE(this, EStateTreeUpdatePhase::TriggerTransitions);
 
 	FAllowDirectTransitionsScope AllowDirectTransitionsScope(*this); // Set flag for the scope of this function to allow direct transitions without buffering.
 	FStateTreeExecutionState& Exec = GetExecState();
 
 	if (EventQueue && EventQueue->HasEvents())
 	{
-		STATETREE_LOG_AND_TRACE(Verbose, Log, TEXT("Trigger transitions with events: %s"), *DebugGetEventsAsString());
+		STATETREE_LOG(Verbose, TEXT("Trigger transitions with events: %s"), *DebugGetEventsAsString());
+		UE_STATETREE_DEBUG_LOG_EVENT(this, Log, TEXT("Trigger transitions with events: %s"), *DebugGetEventsAsString());
 	}
 
 	NextTransition.Reset();
@@ -3461,7 +3420,7 @@ bool FStateTreeExecutionContext::TriggerTransitions()
 
 		FCurrentlyProcessedFrameScope FrameScope(*this, CurrentParentFrame, CurrentFrame);
 		FCurrentlyProcessedStateScope StateScope(*this, Handler.StateHandle);
-		STATETREE_TRACE_SCOPED_STATE(Handler.StateHandle);
+		UE_STATETREE_DEBUG_SCOPED_STATE(this, Handler.StateHandle);
 
 		if (Handler.TaskIndex.IsValid())
 		{
@@ -3470,15 +3429,15 @@ bool FStateTreeExecutionContext::TriggerTransitions()
 			// Ignore disabled task
 			if (Task.bTaskEnabled == false)
 			{
-				STATETREE_LOG(VeryVerbose, TEXT("%*sSkipped 'TriggerTransitions' for disabled Task: '%s'"), UE::StateTree::DebugIndentSize, TEXT(""), *Task.Name.ToString());
+				STATETREE_LOG(VeryVerbose, TEXT("%*sSkipped 'TriggerTransitions' for disabled Task: '%s'"), UE::StateTree::Debug::IndentSize, TEXT(""), *Task.Name.ToString());
 				continue;
 			}
 
 			const FStateTreeDataView TaskInstanceView = GetDataView(CurrentParentFrame, CurrentFrame, Task.InstanceDataHandle);
 			FNodeInstanceDataScope DataScope(*this, Task.InstanceDataHandle, TaskInstanceView);
 
-			STATETREE_LOG(VeryVerbose, TEXT("%*sTriggerTransitions: '%s'"), UE::StateTree::DebugIndentSize, TEXT(""), *Task.Name.ToString());
-			STATETREE_TRACE_TASK_EVENT(Handler.TaskIndex.Get(), TaskInstanceView, EStateTreeTraceEventType::OnEvaluating, EStateTreeRunStatus::Running);
+			STATETREE_LOG(VeryVerbose, TEXT("%*sTriggerTransitions: '%s'"), UE::StateTree::Debug::IndentSize, TEXT(""), *Task.Name.ToString());
+			UE_STATETREE_DEBUG_TASK_EVENT(this, Handler.TaskIndex.Get(), TaskInstanceView, EStateTreeTraceEventType::OnEvaluating, EStateTreeRunStatus::Running);
 			check(TaskInstanceView.IsValid());
 			// TODO: bug the batch data might not be copied if !State.ShouldTickTasks() && State.bHasTransitionTasks and the same for bellow, need to copy at least once
 			Task.TriggerTransitions(*this);
@@ -3576,8 +3535,8 @@ bool FStateTreeExecutionContext::TriggerTransitions()
 					bool bPassed = false; 
 					{
 						FCurrentlyProcessedTransitionEventScope TransitionEventScope(*this, TransitionEvent ? TransitionEvent->Get() : nullptr);
-						STATETREE_TRACE_TRANSITION_EVENT(FStateTreeTransitionSource(FStateTreeIndex16(TransitionIndex), Transition.State, Transition.Priority), EStateTreeTraceEventType::OnEvaluating);
-						STATETREE_TRACE_SCOPED_PHASE(EStateTreeUpdatePhase::TransitionConditions);
+						UE_STATETREE_DEBUG_TRANSITION_EVENT(this, FStateTreeTransitionSource(FStateTreeIndex16(TransitionIndex), Transition.State, Transition.Priority), EStateTreeTraceEventType::OnEvaluating);
+						UE_STATETREE_DEBUG_SCOPED_PHASE(this, EStateTreeUpdatePhase::TransitionConditions);
 						bPassed = TestAllConditions(CurrentParentFrame, CurrentFrame, Transition.ConditionsBegin, Transition.ConditionsNum);
 					}
 
@@ -3630,9 +3589,10 @@ bool FStateTreeExecutionContext::TriggerTransitions()
 							}
 						}
 
+						UE_STATETREE_DEBUG_TRANSITION_EVENT(this, FStateTreeTransitionSource(FStateTreeIndex16(TransitionIndex), Transition.State, Transition.Priority), EStateTreeTraceEventType::OnRequesting);
 						if (RequestTransition(CurrentFrame, Transition.State, Transition.Priority, TransitionEvent, Transition.Fallback))
 						{
-							// If the transition was successfully requested with a specific event, consume and remove the event, it's been used.  
+							// If the transition was successfully requested with a specific event, consume and remove the event, it's been used.
 							if (TransitionEvent && Transition.bConsumeEventOnSelect)
 							{
 								ConsumeEvent(*TransitionEvent);
@@ -3693,7 +3653,7 @@ bool FStateTreeExecutionContext::TriggerTransitions()
 				if (State.ShouldTickCompletionTransitions(CompletionTrigger == EStateTreeTransitionTrigger::OnStateSucceeded))
 				{
 					FCurrentlyProcessedStateScope StateScope(*this, StateHandle);
-					STATETREE_TRACE_SCOPED_STATE_PHASE(StateHandle, EStateTreeUpdatePhase::TriggerTransitions);
+					UE_STATETREE_DEBUG_SCOPED_STATE_PHASE(this, StateHandle, EStateTreeUpdatePhase::TriggerTransitions);
 
 					for (uint8 TransitionCounter = 0; TransitionCounter < State.TransitionsNum; ++TransitionCounter)
 					{
@@ -3711,8 +3671,8 @@ bool FStateTreeExecutionContext::TriggerTransitions()
 						{
 							bool bPassed = false;
 							{
-								STATETREE_TRACE_TRANSITION_EVENT(FStateTreeTransitionSource(FStateTreeIndex16(TransitionIndex), Transition.State, Transition.Priority), EStateTreeTraceEventType::OnEvaluating);
-								STATETREE_TRACE_SCOPED_PHASE(EStateTreeUpdatePhase::TransitionConditions);
+								UE_STATETREE_DEBUG_TRANSITION_EVENT(this, FStateTreeTransitionSource(FStateTreeIndex16(TransitionIndex), Transition.State, Transition.Priority), EStateTreeTraceEventType::OnEvaluating);
+								UE_STATETREE_DEBUG_SCOPED_PHASE(this, EStateTreeUpdatePhase::TransitionConditions);
 								bPassed = TestAllConditions(CurrentParentFrame, CurrentFrame, Transition.ConditionsBegin, Transition.ConditionsNum);
 							}
 
@@ -3720,6 +3680,7 @@ bool FStateTreeExecutionContext::TriggerTransitions()
 							{
 								// No delay allowed on completion conditions.
 								// No priority on completion transitions, use the priority to signal that state is selected.
+								UE_STATETREE_DEBUG_TRANSITION_EVENT(this, FStateTreeTransitionSource(FStateTreeIndex16(TransitionIndex), Transition.State, Transition.Priority), EStateTreeTraceEventType::OnRequesting);
 								if (RequestTransition(CurrentFrame, Transition.State, EStateTreeTransitionPriority::Normal, /*TransitionEvent*/nullptr, Transition.Fallback))
 								{
 									NextTransitionSource = FStateTreeTransitionSource(FStateTreeIndex16(TransitionIndex), Transition.State, Transition.Priority);
@@ -3747,7 +3708,8 @@ bool FStateTreeExecutionContext::TriggerTransitions()
 		// Handle the case where no transition was found.
 		if (NextTransition.Priority == EStateTreeTransitionPriority::None) //-V547
 		{
-			STATETREE_LOG_AND_TRACE(Verbose, Warning, TEXT("Could not trigger completion transition, jump back to root state."));
+			STATETREE_LOG(Verbose, TEXT("Could not trigger completion transition, jump back to root state."));
+			UE_STATETREE_DEBUG_LOG_EVENT(this, Warning, TEXT("Could not trigger completion transition, jump back to root state."));
 
 			check(!Exec.ActiveFrames.IsEmpty());
 			FStateTreeExecutionFrame& RootFrame = Exec.ActiveFrames[0];
@@ -3760,7 +3722,8 @@ bool FStateTreeExecutionContext::TriggerTransitions()
 			}
 			else
 			{
-				STATETREE_LOG_AND_TRACE(Warning, Error, TEXT("Failed to select root state. Stopping the tree with failure."));
+				STATETREE_LOG(Warning, TEXT("Failed to select root state. Stopping the tree with failure."));
+				UE_STATETREE_DEBUG_LOG_EVENT(this, Error, TEXT("Failed to select root state. Stopping the tree with failure."));
 
 				SetupNextTransition(RootFrame, FStateTreeStateHandle::Failed, EStateTreeTransitionPriority::Critical);
 
@@ -4005,7 +3968,7 @@ bool FStateTreeExecutionContext::SelectStateInternal(
 		return false;
 	}
 
-	STATETREE_TRACE_SCOPED_STATE_PHASE(NextStateHandle, EStateTreeUpdatePhase::StateSelection);
+	UE_STATETREE_DEBUG_SCOPED_STATE_PHASE(this, NextStateHandle, EStateTreeUpdatePhase::StateSelection);
 
 	// The state cannot be directly selected.
 	if (NextState.SelectionBehavior == EStateTreeStateSelectionBehavior::None)
@@ -4148,9 +4111,9 @@ bool FStateTreeExecutionContext::SelectStateInternal(
 		if (bShouldPrerequisitesBeChecked)
 		{
 			// Check that the state can be entered
-			STATETREE_TRACE_PHASE_BEGIN(EStateTreeUpdatePhase::EnterConditions);
+			UE_STATETREE_DEBUG_ENTER_PHASE(this, EStateTreeUpdatePhase::EnterConditions);
 			const bool bEnterConditionsPassed = TestAllConditions(CurrentParentFrame, CurrentFrame, NextState.EnterConditionsBegin, NextState.EnterConditionsNum);
-			STATETREE_TRACE_PHASE_END(EStateTreeUpdatePhase::EnterConditions);
+			UE_STATETREE_DEBUG_EXIT_PHASE(this, EStateTreeUpdatePhase::EnterConditions);
 
 			if (!bEnterConditionsPassed)
 			{
@@ -4338,13 +4301,13 @@ bool FStateTreeExecutionContext::SelectStateInternal(
 		else if (NextState.SelectionBehavior == EStateTreeStateSelectionBehavior::TryEnterState)
 		{
 			// Select this state.
-			STATETREE_TRACE_STATE_EVENT(NextStateHandle, EStateTreeTraceEventType::OnStateSelected);
+			UE_STATETREE_DEBUG_STATE_EVENT(this, NextStateHandle, EStateTreeTraceEventType::OnStateSelected);
 			bSucceededToSelectState = true;
 			break;
 		}
 		else if (NextState.SelectionBehavior == EStateTreeStateSelectionBehavior::TryFollowTransitions)
 		{
-			STATETREE_TRACE_SCOPED_STATE_PHASE(NextStateHandle, EStateTreeUpdatePhase::TrySelectBehavior);
+			UE_STATETREE_DEBUG_SCOPED_STATE_PHASE(this, NextStateHandle, EStateTreeUpdatePhase::TrySelectBehavior);
 
 			EStateTreeTransitionPriority CurrentPriority = EStateTreeTransitionPriority::None;
 
@@ -4423,8 +4386,8 @@ bool FStateTreeExecutionContext::SelectStateInternal(
 					{
 						FCurrentlyProcessedTransitionEventScope TransitionEventScope(*this, SelectedStateTransitionEvent ? SelectedStateTransitionEvent->Get() : nullptr);
 
-						STATETREE_TRACE_TRANSITION_EVENT(FStateTreeTransitionSource(FStateTreeIndex16(TransitionIndex), Transition.State, Transition.Priority), EStateTreeTraceEventType::OnEvaluating);
-						STATETREE_TRACE_SCOPED_PHASE(EStateTreeUpdatePhase::TransitionConditions);
+						UE_STATETREE_DEBUG_TRANSITION_EVENT(this, FStateTreeTransitionSource(FStateTreeIndex16(TransitionIndex), Transition.State, Transition.Priority), EStateTreeTraceEventType::OnEvaluating);
+						UE_STATETREE_DEBUG_SCOPED_PHASE(this, EStateTreeUpdatePhase::TransitionConditions);
 
 						bTransitionConditionsPassed = TestAllConditions(CurrentParentFrame, CurrentFrame, Transition.ConditionsBegin, Transition.ConditionsNum);
 					}
@@ -4456,7 +4419,7 @@ bool FStateTreeExecutionContext::SelectStateInternal(
 		{
 			if (NextState.HasChildren())
 			{
-				STATETREE_TRACE_SCOPED_STATE_PHASE(NextStateHandle, EStateTreeUpdatePhase::TrySelectBehavior);
+				UE_STATETREE_DEBUG_SCOPED_STATE_PHASE(this, NextStateHandle, EStateTreeUpdatePhase::TrySelectBehavior);
 
 				// If the state has children, proceed to select children.
 				for (uint16 ChildState = NextState.ChildrenBegin; ChildState < NextState.ChildrenEnd; ChildState = CurrentStateTree->States[ChildState].GetNextSibling())
@@ -4477,7 +4440,7 @@ bool FStateTreeExecutionContext::SelectStateInternal(
 			else
 			{
 				// Select this state (For backwards compatibility)
-				STATETREE_TRACE_STATE_EVENT(NextStateHandle, EStateTreeTraceEventType::OnStateSelected);
+				UE_STATETREE_DEBUG_STATE_EVENT(this, NextStateHandle, EStateTreeTraceEventType::OnStateSelected);
 				bSucceededToSelectState = true;
 				break;
 			}
@@ -4486,7 +4449,7 @@ bool FStateTreeExecutionContext::SelectStateInternal(
 		{
 			if (NextState.HasChildren())
 			{
-				STATETREE_TRACE_SCOPED_STATE_PHASE(NextStateHandle, EStateTreeUpdatePhase::TrySelectBehavior);
+				UE_STATETREE_DEBUG_SCOPED_STATE_PHASE(this, NextStateHandle, EStateTreeUpdatePhase::TrySelectBehavior);
 				
 				TArray<uint16, TInlineAllocator<8>> NextLevelChildStates;
 				for (uint16 ChildState = NextState.ChildrenBegin; ChildState < NextState.ChildrenEnd; ChildState = CurrentStateTree->States[ChildState].GetNextSibling())
@@ -4517,7 +4480,7 @@ bool FStateTreeExecutionContext::SelectStateInternal(
 			else
 			{
 				// Select this state (For backwards compatibility)
-				STATETREE_TRACE_STATE_EVENT(NextStateHandle, EStateTreeTraceEventType::OnStateSelected);
+				UE_STATETREE_DEBUG_STATE_EVENT(this, NextStateHandle, EStateTreeTraceEventType::OnStateSelected);
 				bSucceededToSelectState = true;
 				break;
 			}
@@ -4526,7 +4489,7 @@ bool FStateTreeExecutionContext::SelectStateInternal(
 		{
 			if (NextState.HasChildren())
 			{
-				STATETREE_TRACE_SCOPED_STATE_PHASE(NextStateHandle, EStateTreeUpdatePhase::TrySelectBehavior);
+				UE_STATETREE_DEBUG_SCOPED_STATE_PHASE(this, NextStateHandle, EStateTreeUpdatePhase::TrySelectBehavior);
 
 				TArray<uint16, TInlineAllocator<8>> NextLevelChildStates;
 				for (uint16 ChildState = NextState.ChildrenBegin; ChildState < NextState.ChildrenEnd; ChildState = CurrentStateTree->States[ChildState].GetNextSibling())
@@ -4580,7 +4543,7 @@ bool FStateTreeExecutionContext::SelectStateInternal(
 			else
 			{
 				// Select this state (For backwards compatibility)
-				STATETREE_TRACE_STATE_EVENT(NextStateHandle, EStateTreeTraceEventType::OnStateSelected);
+				UE_STATETREE_DEBUG_STATE_EVENT(this, NextStateHandle, EStateTreeTraceEventType::OnStateSelected);
 				bSucceededToSelectState = true;
 				break;
 			}
@@ -4643,7 +4606,7 @@ bool FStateTreeExecutionContext::SelectStateInternal(
 			else
 			{
 				// Select this state (For backwards compatibility)
-				STATETREE_TRACE_STATE_EVENT(NextStateHandle, EStateTreeTraceEventType::OnStateSelected);
+				UE_STATETREE_DEBUG_STATE_EVENT(this, NextStateHandle, EStateTreeTraceEventType::OnStateSelected);
 				bSucceededToSelectState = true;
 				break;
 			}
@@ -5002,26 +4965,3 @@ bool FStateTreeWeakMinimalExecutionContext::SendEvent(const FGameplayTag Tag, co
 
 #undef STATETREE_LOG
 #undef STATETREE_CLOG
-#undef STATETREE_LOG_AND_TRACE
-
-#undef STATETREE_TRACE_SCOPED_PHASE
-#undef STATETREE_TRACE_SCOPED_STATE
-#undef STATETREE_TRACE_SCOPED_STATE_PHASE
-#undef STATETREE_TRACE_INSTANCE_EVENT
-#undef STATETREE_TRACE_INSTANCE_FRAME_EVENT
-#undef STATETREE_TRACE_PHASE_BEGIN
-#undef STATETREE_TRACE_PHASE_END
-#undef STATETREE_TRACE_ACTIVE_STATES_EVENT
-#undef STATETREE_TRACE_LOG_EVENT
-#undef STATETREE_TRACE_STATE_EVENT
-#undef STATETREE_TRACE_TASK_EVENT
-#undef STATETREE_TRACE_EVALUATOR_EVENT
-#undef STATETREE_TRACE_CONDITION_EVENT
-#undef STATETREE_TRACE_TRANSITION_EVENT
-
-#if WITH_STATETREE_TRACE
-#undef STATETREE_TRACE_ID_NAME
-#undef STATETREE_TRACE_SCOPED_PHASE
-#undef STATETREE_TRACE_SCOPED_STATE
-#undef STATETREE_TRACE_SCOPED_STATE_PHASE
-#endif //WITH_STATETREE_TRACE
diff --git a/Engine/Plugins/Runtime/StateTree/Source/StateTreeModule/Public/Debugger/StateTreeDebug.h b/Engine/Plugins/Runtime/StateTree/Source/StateTreeModule/Public/Debugger/StateTreeDebug.h
--- a/Engine/Plugins/Runtime/StateTree/Source/StateTreeModule/Public/Debugger/StateTreeDebug.h
+++ b/Engine/Plugins/Runtime/StateTree/Source/StateTreeModule/Public/Debugger/StateTreeDebug.h
@@ -2,35 +2,92 @@
 
 #pragma once
 
-#include "CoreMinimal.h"
 #include "StateTreeTypes.h"
 #include "Templates/NotNull.h"
 
 class UStateTree;
 struct FStateTreeExecutionContext;
 struct FStateTreeMinimalExecutionContext;
+struct FStateTreeTransitionSource;
+enum class EStateTreeUpdatePhase : uint8;
+enum class EStateTreeTraceEventType : uint8;
 
 #if WITH_STATETREE_DEBUG
 
-#define UE_STATETREE_DEBUG_CONDITION_ENTER_STATE(InExecutionContextPtr, InStateTree, InNodeIndex) ::UE::StateTree::Debug::FStateTreeDebug::ConditionEnterState(*(InExecutionContextPtr), ::UE::StateTree::Debug::FStateTreeDebug::FNodeReference((InStateTree), FStateTreeIndex16(InNodeIndex)))
-#define UE_STATETREE_DEBUG_CONDITION_TEST_CONDITION(InExecutionContextPtr, InStateTree, InNodeIndex) ::UE::StateTree::Debug::FStateTreeDebug::ConditionTest(*(InExecutionContextPtr), ::UE::StateTree::Debug::FStateTreeDebug::FNodeReference((InStateTree), FStateTreeIndex16(InNodeIndex)))
-#define UE_STATETREE_DEBUG_CONDITION_EXIT_STATE(InExecutionContextPtr, InStateTree, InNodeIndex) ::UE::StateTree::Debug::FStateTreeDebug::ConditionExitState(*(InExecutionContextPtr), ::UE::StateTree::Debug::FStateTreeDebug::FNodeReference((InStateTree), FStateTreeIndex16(InNodeIndex)))
-
-#define UE_STATETREE_DEBUG_EVALUATOR_ENTER_TREE(InExecutionContextPtr, InStateTree, InNodeIndex) ::UE::StateTree::Debug::FStateTreeDebug::EvaluatorEnterTree(*(InExecutionContextPtr), ::UE::StateTree::Debug::FStateTreeDebug::FNodeReference((InStateTree), FStateTreeIndex16(InNodeIndex)))
-#define UE_STATETREE_DEBUG_EVALUATOR_TICK(InExecutionContextPtr, InStateTree, InNodeIndex) ::UE::StateTree::Debug::FStateTreeDebug::EvaluatorTick(*(InExecutionContextPtr), ::UE::StateTree::Debug::FStateTreeDebug::FNodeReference((InStateTree), FStateTreeIndex16(InNodeIndex)))
-#define UE_STATETREE_DEBUG_EVALUATOR_EXIT_TREE(InExecutionContextPtr, InStateTree, InNodeIndex) ::UE::StateTree::Debug::FStateTreeDebug::EvaluatorExitTree(*(InExecutionContextPtr), ::UE::StateTree::Debug::FStateTreeDebug::FNodeReference((InStateTree), FStateTreeIndex16(InNodeIndex)))
-
-#define UE_STATETREE_DEBUG_TASK_ENTER_STATE(InExecutionContextPtr, InStateTree, InNodeIndex) ::UE::StateTree::Debug::FStateTreeDebug::TaskEnterState(*(InExecutionContextPtr), ::UE::StateTree::Debug::FStateTreeDebug::FNodeReference((InStateTree), FStateTreeIndex16(InNodeIndex)))
-#define UE_STATETREE_DEBUG_TASK_TICK(InExecutionContextPtr, InStateTree, InNodeIndex) ::UE::StateTree::Debug::FStateTreeDebug::TaskTick(*(InExecutionContextPtr), ::UE::StateTree::Debug::FStateTreeDebug::FNodeReference((InStateTree), FStateTreeIndex16(InNodeIndex)))
-#define UE_STATETREE_DEBUG_TASK_EXIT_STATE(InExecutionContextPtr, InStateTree, InNodeIndex) ::UE::StateTree::Debug::FStateTreeDebug::TaskExitState(*(InExecutionContextPtr), ::UE::StateTree::Debug::FStateTreeDebug::FNodeReference((InStateTree), FStateTreeIndex16(InNodeIndex)))
-
-#define UE_STATETREE_DEBUG_SEND_EVENT(InExecutionContextPtr, InStateTree, InTag, InPayload, InOrigin) ::UE::StateTree::Debug::FStateTreeDebug::EventSent(*(InExecutionContextPtr), InStateTree, InTag, InPayload, InOrigin)
+#define UE_STATETREE_DEBUG_CONDITION_ENTER_STATE(InExecutionContextPtr, InStateTree, InNodeIndex) ::UE::StateTree::Debug::ConditionEnterState(*(InExecutionContextPtr), ::UE::StateTree::Debug::FNodeReference((InStateTree), FStateTreeIndex16(InNodeIndex)))
+#define UE_STATETREE_DEBUG_CONDITION_TEST_CONDITION(InExecutionContextPtr, InStateTree, InNodeIndex) ::UE::StateTree::Debug::ConditionTest(*(InExecutionContextPtr), ::UE::StateTree::Debug::FNodeReference((InStateTree), FStateTreeIndex16(InNodeIndex)))
+#define UE_STATETREE_DEBUG_CONDITION_EXIT_STATE(InExecutionContextPtr, InStateTree, InNodeIndex) ::UE::StateTree::Debug::ConditionExitState(*(InExecutionContextPtr), ::UE::StateTree::Debug::FNodeReference((InStateTree), FStateTreeIndex16(InNodeIndex)))
+
+#define UE_STATETREE_DEBUG_EVALUATOR_ENTER_TREE(InExecutionContextPtr, InStateTree, InNodeIndex) ::UE::StateTree::Debug::EvaluatorEnterTree(*(InExecutionContextPtr), ::UE::StateTree::Debug::FNodeReference((InStateTree), FStateTreeIndex16(InNodeIndex)))
+#define UE_STATETREE_DEBUG_EVALUATOR_TICK(InExecutionContextPtr, InStateTree, InNodeIndex) ::UE::StateTree::Debug::EvaluatorTick(*(InExecutionContextPtr), ::UE::StateTree::Debug::FNodeReference((InStateTree), FStateTreeIndex16(InNodeIndex)))
+#define UE_STATETREE_DEBUG_EVALUATOR_EXIT_TREE(InExecutionContextPtr, InStateTree, InNodeIndex) ::UE::StateTree::Debug::EvaluatorExitTree(*(InExecutionContextPtr), ::UE::StateTree::Debug::FNodeReference((InStateTree), FStateTreeIndex16(InNodeIndex)))
+
+#define UE_STATETREE_DEBUG_TASK_ENTER_STATE(InExecutionContextPtr, InStateTree, InNodeIndex) ::UE::StateTree::Debug::TaskEnterState(*(InExecutionContextPtr), ::UE::StateTree::Debug::FNodeReference((InStateTree), FStateTreeIndex16(InNodeIndex)))
+#define UE_STATETREE_DEBUG_TASK_TICK(InExecutionContextPtr, InStateTree, InNodeIndex) ::UE::StateTree::Debug::TaskTick(*(InExecutionContextPtr), ::UE::StateTree::Debug::FNodeReference((InStateTree), FStateTreeIndex16(InNodeIndex)))
+#define UE_STATETREE_DEBUG_TASK_EXIT_STATE(InExecutionContextPtr, InStateTree, InNodeIndex) ::UE::StateTree::Debug::TaskExitState(*(InExecutionContextPtr), ::UE::StateTree::Debug::FNodeReference((InStateTree), FStateTreeIndex16(InNodeIndex)))
+
+#define UE_STATETREE_DEBUG_SEND_EVENT(InExecutionContextPtr, InStateTree, InTag, InPayload, InOrigin) ::UE::StateTree::Debug::EventSent(*(InExecutionContextPtr), InStateTree, InTag, InPayload, InOrigin)
+
+#define UE_STATETREE_DEBUG_ENTER_PHASE(InExecutionContextPtr, InPhase) \
+		::UE::StateTree::Debug::EnterPhase(*(InExecutionContextPtr), (InPhase), FStateTreeStateHandle::Invalid); \
+		TRACE_STATETREE_PHASE_EVENT((InExecutionContextPtr)->GetInstanceDebugId(), (InPhase), EStateTreeTraceEventType::Push, FStateTreeStateHandle::Invalid)
+
+#define UE_STATETREE_DEBUG_EXIT_PHASE(InExecutionContextPtr, InPhase) \
+		::UE::StateTree::Debug::ExitPhase(*(InExecutionContextPtr), (InPhase), FStateTreeStateHandle::Invalid); \
+		TRACE_STATETREE_PHASE_EVENT((InExecutionContextPtr)->GetInstanceDebugId(), (InPhase), EStateTreeTraceEventType::Pop, FStateTreeStateHandle::Invalid)
+
+#define UE_STATETREE_DEBUG_STATE_EVENT(InExecutionContextPtr, InStateHandle, InEventType) \
+		::UE::StateTree::Debug::StateEvent(*(InExecutionContextPtr), (InStateHandle), (InEventType)); \
+		TRACE_STATETREE_STATE_EVENT((InExecutionContextPtr)->GetInstanceDebugId(), (InStateHandle), (InEventType))
+
+#define UE_STATETREE_DEBUG_TRANSITION_EVENT(InExecutionContextPtr, InTransitionSource, InEventType) \
+		::UE::StateTree::Debug::TransitionEvent(*(InExecutionContextPtr), (InTransitionSource), (InEventType)); \
+		TRACE_STATETREE_TRANSITION_EVENT((InExecutionContextPtr)->GetInstanceDebugId(), (InTransitionSource), (InEventType))
+
+#define UE_STATETREE_DEBUG_LOG_EVENT(InExecutionContextPtr, InLogVerbosity, InFormat, ...) TRACE_STATETREE_LOG_EVENT((InExecutionContextPtr)->GetInstanceDebugId(), InLogVerbosity, InFormat, ##__VA_ARGS__)
+#define UE_STATETREE_DEBUG_CONDITION_EVENT(InExecutionContextPtr, InIndex, InDataView, InEventType) TRACE_STATETREE_CONDITION_EVENT((InExecutionContextPtr)->GetInstanceDebugId(), FStateTreeIndex16(InIndex), (InDataView), (InEventType));
+#define UE_STATETREE_DEBUG_INSTANCE_EVENT(InExecutionContextPtr, InEventType) TRACE_STATETREE_INSTANCE_EVENT((InExecutionContextPtr)->GetInstanceDebugId(), (InExecutionContextPtr)->GetStateTree(), *((InExecutionContextPtr)->GetInstanceDescriptionInternal()), (InEventType));
+#define UE_STATETREE_DEBUG_INSTANCE_FRAME_EVENT(InExecutionContextPtr, InFrame) TRACE_STATETREE_INSTANCE_FRAME_EVENT((InExecutionContextPtr)->GetInstanceDebugId(), (InFrame));
+#define UE_STATETREE_DEBUG_ACTIVE_STATES_EVENT(InExecutionContextPtr, InActiveFrames) TRACE_STATETREE_ACTIVE_STATES_EVENT((InExecutionContextPtr)->GetInstanceDebugId(), InActiveFrames);
+#define UE_STATETREE_DEBUG_TASK_EVENT(InExecutionContextPtr, InIndex, InDataView, InEventType, InStatus) TRACE_STATETREE_TASK_EVENT((InExecutionContextPtr)->GetInstanceDebugId(), FStateTreeIndex16(InIndex), (InDataView), (InEventType), (InStatus));
+#define UE_STATETREE_DEBUG_EVALUATOR_EVENT(InExecutionContextPtr, InIndex, InDataView, InEventType) TRACE_STATETREE_EVALUATOR_EVENT((InExecutionContextPtr)->GetInstanceDebugId(), FStateTreeIndex16(InIndex), (InDataView), (InEventType));
+#define UE_STATETREE_DEBUG_CONDITION_EVENT(InExecutionContextPtr, InIndex, InDataView, InEventType) TRACE_STATETREE_CONDITION_EVENT((InExecutionContextPtr)->GetInstanceDebugId(), FStateTreeIndex16(InIndex), (InDataView), (InEventType));
+
+#define UE_STATETREE_ID_NAME PREPROCESSOR_JOIN(InstanceId,__LINE__) \
+
+/** Scope based macros captures the instance ID since it might not be accessible when exiting the scope */
+#define UE_STATETREE_DEBUG_SCOPED_PHASE(InExecutionContextPtr, InPhase) \
+		FStateTreeInstanceDebugId UE_STATETREE_ID_NAME = (InExecutionContextPtr)->GetInstanceDebugId(); \
+		UE_STATETREE_DEBUG_ENTER_PHASE((InExecutionContextPtr), (InPhase)); \
+		ON_SCOPE_EXIT \
+		{ \
+			::UE::StateTree::Debug::ExitPhase(*(InExecutionContextPtr), (InPhase), FStateTreeStateHandle::Invalid); \
+			TRACE_STATETREE_PHASE_EVENT(UE_STATETREE_ID_NAME, (InPhase), EStateTreeTraceEventType::Pop, FStateTreeStateHandle::Invalid); \
+		}
+
+/** Scope based macros captures the instance ID since it might not be accessible when exiting the scope */
+#define UE_STATETREE_DEBUG_SCOPED_STATE(InExecutionContextPtr, InStateHandle) \
+		FStateTreeInstanceDebugId UE_STATETREE_ID_NAME = (InExecutionContextPtr)->GetInstanceDebugId(); \
+		UE_STATETREE_DEBUG_STATE_EVENT((InExecutionContextPtr), (InStateHandle), EStateTreeTraceEventType::Push); \
+		ON_SCOPE_EXIT \
+		{ \
+			::UE::StateTree::Debug::StateEvent(*(InExecutionContextPtr), (InStateHandle), EStateTreeTraceEventType::Pop); \
+			TRACE_STATETREE_STATE_EVENT(UE_STATETREE_ID_NAME, (InStateHandle), EStateTreeTraceEventType::Pop); \
+		}
+
+/** Scope based macros captures the instance ID since it might not be accessible when exiting the scope */
+#define UE_STATETREE_DEBUG_SCOPED_STATE_PHASE(InExecutionContextPtr, InStateHandle, InPhase) \
+		FStateTreeInstanceDebugId UE_STATETREE_ID_NAME = (InExecutionContextPtr)->GetInstanceDebugId(); \
+		::UE::StateTree::Debug::EnterPhase(*(InExecutionContextPtr), (InPhase), (InStateHandle)); \
+		TRACE_STATETREE_PHASE_EVENT(UE_STATETREE_ID_NAME, (InPhase), EStateTreeTraceEventType::Push, (InStateHandle)); \
+		ON_SCOPE_EXIT \
+		{ \
+			::UE::StateTree::Debug::ExitPhase(*(InExecutionContextPtr), (InPhase), (InStateHandle)); \
+			TRACE_STATETREE_PHASE_EVENT(UE_STATETREE_ID_NAME, (InPhase), EStateTreeTraceEventType::Pop, (InStateHandle)); \
+		}
 
 namespace UE::StateTree::Debug
 {
-class STATETREEMODULE_API FStateTreeDebug final
-{
-public:
 	struct FNodeReference
 	{
 		explicit FNodeReference(gsl::not_null<const UStateTree*> InStateTree, FStateTreeIndex16 InNodeIndex);
@@ -38,114 +95,144 @@ public:
 		FStateTreeIndex16 Index;
 	};
 
-	static void ConditionEnterState(const FStateTreeExecutionContext& ExecutionContext, FNodeReference Node);
-	static void ConditionTest(const FStateTreeExecutionContext& ExecutionContext, FNodeReference Node);
-	static void ConditionExitState(const FStateTreeExecutionContext& ExecutionContext, FNodeReference Node);
-
-	static void EvaluatorEnterTree(const FStateTreeExecutionContext& ExecutionContext, FNodeReference Node);
-	static void EvaluatorTick(const FStateTreeExecutionContext& ExecutionContext, FNodeReference Node);
-	static void EvaluatorExitTree(const FStateTreeExecutionContext& ExecutionContext, FNodeReference Node);
-
-	static void TaskEnterState(const FStateTreeExecutionContext& ExecutionContext, FNodeReference Node);
-	static void TaskTick(const FStateTreeExecutionContext& ExecutionContext, FNodeReference Node);
-	static void TaskExitState(const FStateTreeExecutionContext& ExecutionContext, FNodeReference Node);
-
-	static void EventSent(const FStateTreeMinimalExecutionContext& ExecutionContext, gsl::not_null<const UStateTree*> StateTree, FGameplayTag Tag, FConstStructView Payload, FName Origin);
-
-public:
 	struct FNodeDelegateArgs
 	{
 		FNodeReference Node;
 		FGuid NodeId;
 	};
 
-	DECLARE_TS_MULTICAST_DELEGATE_TwoParams(FNodeDelegate, const FStateTreeExecutionContext& ExecutionContext, FNodeDelegateArgs Node);
+	struct FEventSentDelegateArgs
+	{
+		gsl::not_null<const UStateTree*> StateTree;
+		FGameplayTag Tag;
+		FConstStructView Payload;
+		FName Origin;
+	};
+
+	DECLARE_TS_MULTICAST_DELEGATE_TwoParams(FEventSentDelegate, const FStateTreeMinimalExecutionContext& ExecutionContext, const FEventSentDelegateArgs& EventSentArgs);
+	DECLARE_TS_MULTICAST_DELEGATE_TwoParams(FNodeDelegate, const FStateTreeExecutionContext&, FNodeDelegateArgs);
+	DECLARE_TS_MULTICAST_DELEGATE_ThreeParams(FPhaseDelegate, const FStateTreeExecutionContext&, EStateTreeUpdatePhase, FStateTreeStateHandle);
+	DECLARE_TS_MULTICAST_DELEGATE_ThreeParams(FStateDelegate, const FStateTreeExecutionContext&, FStateTreeStateHandle, EStateTreeTraceEventType);
+	DECLARE_TS_MULTICAST_DELEGATE_ThreeParams(FTransitionDelegate, const FStateTreeExecutionContext&, const FStateTreeTransitionSource&, EStateTreeTraceEventType);
+
+	void ConditionEnterState(const FStateTreeExecutionContext& ExecutionContext, FNodeReference Node);
+	void ConditionTest(const FStateTreeExecutionContext& ExecutionContext, FNodeReference Node);
+	void ConditionExitState(const FStateTreeExecutionContext& ExecutionContext, FNodeReference Node);
+
+	void EvaluatorEnterTree(const FStateTreeExecutionContext& ExecutionContext, FNodeReference Node);
+	void EvaluatorTick(const FStateTreeExecutionContext& ExecutionContext, FNodeReference Node);
+	void EvaluatorExitTree(const FStateTreeExecutionContext& ExecutionContext, FNodeReference Node);
+
+	void TaskEnterState(const FStateTreeExecutionContext& ExecutionContext, FNodeReference Node);
+	void TaskTick(const FStateTreeExecutionContext& ExecutionContext, FNodeReference Node);
+	void TaskExitState(const FStateTreeExecutionContext& ExecutionContext, FNodeReference Node);
+
+	void EventSent(const FStateTreeMinimalExecutionContext& ExecutionContext, gsl::not_null<const UStateTree*> StateTree, FGameplayTag Tag, FConstStructView Payload, FName Origin);
+
+	void EnterPhase(const FStateTreeExecutionContext& ExecutionContext, EStateTreeUpdatePhase Phase, FStateTreeStateHandle StateHandle);
+	void ExitPhase(const FStateTreeExecutionContext& ExecutionContext, EStateTreeUpdatePhase Phase, FStateTreeStateHandle StateHandle);
+
+	void StateEvent(const FStateTreeExecutionContext& ExecutionContext, FStateTreeStateHandle StateHandle, EStateTreeTraceEventType EventType);
+
+	void TransitionEvent(const FStateTreeExecutionContext& ExecutionContext, const FStateTreeTransitionSource& TransitionSource, EStateTreeTraceEventType EventType);
 
 	/**
 	 * Debugging callback for when a condition activates.
 	 * @note The callback executes inside the StateTree logic.
 	 * @note The StateTree can execute on any thread.
 	 */
-	static FNodeDelegate OnConditionEnterState_AnyThread;
+	STATETREEMODULE_API extern FNodeDelegate OnConditionEnterState_AnyThread;
 
 	/**
 	 * Debugging callback for before a condition is tested.
 	 * @note The callback executes inside the StateTree logic.
 	 * @note The StateTree can execute on any thread.
 	 */
-	static FNodeDelegate OnTestCondition_AnyThread;
+	STATETREEMODULE_API extern FNodeDelegate OnTestCondition_AnyThread;
 
 	/**
 	 * Debugging callback for when a condition deactivates.
 	 * @note The callback executes inside the StateTree logic.
 	 * @note The StateTree can execute on any thread.
 	 */
-	static FNodeDelegate OnConditionExitState_AnyThread;
+	STATETREEMODULE_API extern FNodeDelegate OnConditionExitState_AnyThread;
 
 	/**
 	 * Debugging callback for when a evaluator activates.
 	 * @note The callback executes inside the StateTree logic.
 	 * @note The StateTree can execute on any thread.
 	 */
-	static FNodeDelegate OnEvaluatorEnterTree_AnyThread;
+	STATETREEMODULE_API extern FNodeDelegate OnEvaluatorEnterTree_AnyThread;
 
 	/**
 	 * Debugging callback for before an evaluator ticks.
 	 * @note The callback executes inside the StateTree logic.
 	 * @note The StateTree can execute on any thread.
 	 */
-	static FNodeDelegate OnTickEvaluator_AnyThread;
+	STATETREEMODULE_API extern FNodeDelegate OnTickEvaluator_AnyThread;
 
 	/**
 	 * Debugging callback for when a evaluator deactivates.
 	 * @note The callback executes inside the StateTree logic.
 	 * @note The StateTree can execute on any thread.
 	 */
-	static FNodeDelegate OnEvaluatorExitTree_AnyThread;
+	STATETREEMODULE_API extern FNodeDelegate OnEvaluatorExitTree_AnyThread;
 
 	/**
 	 * Debugging callback for when a task activates.
 	 * @note The callback executes inside the StateTree logic.
 	 * @note The StateTree can execute on any thread.
 	 */
-	static FNodeDelegate OnTaskEnterState_AnyThread;
+	STATETREEMODULE_API extern FNodeDelegate OnTaskEnterState_AnyThread;
 
 	/**
 	 * Debugging callback executed before a task ticks.
 	 * @note The callback executes inside the StateTree logic.
 	 * @note The StateTree can execute on any thread.
 	 */
-	static FNodeDelegate OnTickTask_AnyThread;
+	STATETREEMODULE_API extern FNodeDelegate OnTickTask_AnyThread;
 
 	/**
 	 * Debugging callback for when a task deactivates.
 	 * @note The callback executes inside the StateTree logic.
 	 * @note The StateTree can execute on any thread.
 	 */
-	static FNodeDelegate OnTaskExitState_AnyThread;
+	STATETREEMODULE_API extern FNodeDelegate OnTaskExitState_AnyThread;
 
-public:
-	struct FEventSentDelegateArgs
-	{
-		gsl::not_null<const UStateTree*> StateTree;
-		FGameplayTag Tag;
-		FConstStructView Payload;
-		FName Origin;
-	};
+	/**
+	 * Debugging callback for when entering an update phase (global to the tree or specific to a state).
+	 * @note The callback executes inside the StateTree logic.
+	 * @note The StateTree can execute on any thread.
+	 */
+	STATETREEMODULE_API extern FPhaseDelegate OnBeginUpdatePhase_AnyThread;
 
-	DECLARE_TS_MULTICAST_DELEGATE_TwoParams(FEventSentDelegate, const FStateTreeMinimalExecutionContext& ExecutionContext, const FEventSentDelegateArgs& EventSentArgs);
+	/**
+	 * Debugging callback for when exiting an update phase (global to the tree or specific to a state).
+	 * @note The callback executes inside the StateTree logic.
+	 * @note The StateTree can execute on any thread.
+	 */
+	STATETREEMODULE_API extern FPhaseDelegate OnEndUpdatePhase_AnyThread;
 
 	/**
-	 * Debugging callback for when an event is sent.
+	 * Debugging callback when an action related to a state is executing (e.g., entering, exiting, selecting, etc.).
+	 * @note The callback executes inside the StateTree logic.
+	 * @note The StateTree can execute on any thread.
+	 */
+	STATETREEMODULE_API extern FStateDelegate OnStateEvent_AnyThread;
+
+	/**
+	 * Debugging callback when an action related to a transition is executing (e.g., requesting, evaluating, etc.).
 	 * @note The callback executes inside the StateTree logic.
 	 * @note The StateTree can execute on any thread.
 	 */
-	static FEventSentDelegate OnEventSent_AnyThread;
-	
-private:
-	FStateTreeDebug() = default;
-};
+	STATETREEMODULE_API extern FTransitionDelegate OnTransitionEvent_AnyThread;
 
+	/**
+	 * Debugging callback for when an event is sent.
+	 * @note The callback executes inside the StateTree logic.
+	 * @note The StateTree can execute on any thread.
+	 */
+	STATETREEMODULE_API extern FEventSentDelegate OnEventSent_AnyThread;
 } //UE::StateTree::Debug
 
 #else
@@ -164,4 +251,20 @@ private:
 
 #define UE_STATETREE_DEBUG_SEND_EVENT(...)
 
+#define UE_STATETREE_DEBUG_ENTER_PHASE(...)
+#define UE_STATETREE_DEBUG_EXIT_PHASE(...)
+#define UE_STATETREE_DEBUG_STATE_EVENT(...)
+#define UE_STATETREE_DEBUG_TRANSITION_EVENT(...)
+#define UE_STATETREE_DEBUG_LOG_EVENT(...)
+#define UE_STATETREE_DEBUG_CONDITION_EVENT(...)
+#define UE_STATETREE_DEBUG_INSTANCE_EVENT(...)
+#define UE_STATETREE_DEBUG_INSTANCE_FRAME_EVENT(...)
+#define UE_STATETREE_DEBUG_ACTIVE_STATES_EVENT(...)
+#define UE_STATETREE_DEBUG_TASK_EVENT(...)
+#define UE_STATETREE_DEBUG_EVALUATOR_EVENT(...)
+#define UE_STATETREE_DEBUG_CONDITION_EVENT(...)
+#define UE_STATETREE_DEBUG_SCOPED_PHASE(...)
+#define UE_STATETREE_DEBUG_SCOPED_STATE(...)
+#define UE_STATETREE_DEBUG_SCOPED_STATE_PHASE(...)
+
 #endif // WITH_STATETREE_DEBUG
diff --git a/Engine/Plugins/Runtime/StateTree/Source/StateTreeModule/Public/Debugger/StateTreeRuntimeValidation.h b/Engine/Plugins/Runtime/StateTree/Source/StateTreeModule/Public/Debugger/StateTreeRuntimeValidation.h
--- a/Engine/Plugins/Runtime/StateTree/Source/StateTreeModule/Public/Debugger/StateTreeRuntimeValidation.h
+++ b/Engine/Plugins/Runtime/StateTree/Source/StateTreeModule/Public/Debugger/StateTreeRuntimeValidation.h
@@ -2,7 +2,11 @@
 
 #pragma once
 
+#include "StateTreeTypes.h"
+
+#if UE_ENABLE_INCLUDE_ORDER_DEPRECATED_IN_5_6
 #include "Debugger/StateTreeDebug.h"
+#endif
 
 class UStateTree;
 class UObject;
diff --git a/Engine/Plugins/Runtime/StateTree/Source/StateTreeModule/Public/Debugger/StateTreeTraceTypes.h b/Engine/Plugins/Runtime/StateTree/Source/StateTreeModule/Public/Debugger/StateTreeTraceTypes.h
--- a/Engine/Plugins/Runtime/StateTree/Source/StateTreeModule/Public/Debugger/StateTreeTraceTypes.h
+++ b/Engine/Plugins/Runtime/StateTree/Source/StateTreeModule/Public/Debugger/StateTreeTraceTypes.h
@@ -40,6 +40,7 @@ enum class EStateTreeTraceEventType : uint8
 	ForcedSuccess			UMETA(DisplayName = "Forced Success"),
 	ForcedFailure			UMETA(DisplayName = "Forced Failure"),
 	InternalForcedFailure	UMETA(DisplayName = "Internal Forced Failure"),
+	OnRequesting			UMETA(DisplayName = "Requesting"),
 	OnEvaluating			UMETA(DisplayName = "Evaluating"),
 	OnTransition			UMETA(DisplayName = "Transition"),
 	OnTreeStarted			UMETA(DisplayName = "Tree Started"),
diff --git a/Engine/Plugins/Runtime/StateTree/Source/StateTreeModule/Public/StateTreeExecutionTypes.h b/Engine/Plugins/Runtime/StateTree/Source/StateTreeModule/Public/StateTreeExecutionTypes.h
--- a/Engine/Plugins/Runtime/StateTree/Source/StateTreeModule/Public/StateTreeExecutionTypes.h
+++ b/Engine/Plugins/Runtime/StateTree/Source/StateTreeModule/Public/StateTreeExecutionTypes.h
@@ -16,23 +16,26 @@ class UStateTree;
 UENUM()
 enum class EStateTreeUpdatePhase : uint8
 {
-	Unset					= 0,
-	StartTree				UMETA(DisplayName = "Start Tree"),
-	StopTree				UMETA(DisplayName = "Stop Tree"),
-	StartGlobalTasks		UMETA(DisplayName = "Start Global Tasks & Evaluators"),
-	StopGlobalTasks			UMETA(DisplayName = "Stop Global Tasks & Evaluators"),
-	TickStateTree			UMETA(DisplayName = "Tick State Tree"),
-	ApplyTransitions		UMETA(DisplayName = "Transition"),
-	TriggerTransitions		UMETA(DisplayName = "Trigger Transitions"),
-	TickingGlobalTasks		UMETA(DisplayName = "Tick Global Tasks & Evaluators"),
-	TickingTasks			UMETA(DisplayName = "Tick Tasks"),
-	TransitionConditions	UMETA(DisplayName = "Transition conditions"),
-	StateSelection			UMETA(DisplayName = "Try Enter"),
-	TrySelectBehavior		UMETA(DisplayName = "Try Select Behavior"),
-	EnterConditions			UMETA(DisplayName = "Enter conditions"),
-	EnterStates				UMETA(DisplayName = "Enter States"),
-	ExitStates				UMETA(DisplayName = "Exit States"),
-	StateCompleted			UMETA(DisplayName = "State(s) Completed")
+	Unset							= 0,
+	StartTree						UMETA(DisplayName = "Start Tree"),
+	StopTree						UMETA(DisplayName = "Stop Tree"),
+	StartGlobalTasks				UMETA(DisplayName = "Start Global Tasks & Evaluators"),
+	StartGlobalTasksForSelection	UMETA(DisplayName = "Start Global Tasks & Evaluators for selection"),
+	StopGlobalTasks					UMETA(DisplayName = "Stop Global Tasks & Evaluators"),
+	StopGlobalTasksForSelection		UMETA(DisplayName = "Stop Global Tasks & Evaluators for selection"),
+	TickStateTree					UMETA(DisplayName = "Tick State Tree"),
+	ApplyTransitions				UMETA(DisplayName = "Transition"),
+	TickTransitions					UMETA(DisplayName = "Tick Transitions"),
+	TriggerTransitions				UMETA(DisplayName = "Trigger Transitions"),
+	TickingGlobalTasks				UMETA(DisplayName = "Tick Global Tasks & Evaluators"),
+	TickingTasks					UMETA(DisplayName = "Tick Tasks"),
+	TransitionConditions			UMETA(DisplayName = "Transition conditions"),
+	StateSelection					UMETA(DisplayName = "Try Enter"),
+	TrySelectBehavior				UMETA(DisplayName = "Try Select Behavior"),
+	EnterConditions					UMETA(DisplayName = "Enter conditions"),
+	EnterStates						UMETA(DisplayName = "Enter States"),
+	ExitStates						UMETA(DisplayName = "Exit States"),
+	StateCompleted					UMETA(DisplayName = "State(s) Completed")
 };
 
 
