StateTree: Remove the bTaskTicked ensure that creates false positives if a previous state fails. It highlighted a bug with binding that will not copy if the task does not tick. Another task might bind to that data. (A to B to C and B do not tick).

diff --git a/Engine/Plugins/Runtime/StateTree/Source/StateTreeModule/Private/StateTreeExecutionContext.cpp b/Engine/Plugins/Runtime/StateTree/Source/StateTreeModule/Private/StateTreeExecutionContext.cpp
--- a/Engine/Plugins/Runtime/StateTree/Source/StateTreeModule/Private/StateTreeExecutionContext.cpp
+++ b/Engine/Plugins/Runtime/StateTree/Source/StateTreeModule/Private/StateTreeExecutionContext.cpp
@@ -2863,7 +2863,6 @@ EStateTreeRunStatus FStateTreeExecutionContext::TickTasks(const float DeltaTime)
 						bShouldTickTasks = false;
 					}
 				}
-				ensureMsgf(bTaskTicked, TEXT("ShouldTickTasks is true but no tasks were tick for State: '%s' inside theStateTree %s."), *State.Name.ToString(), *CurrentStateTree->GetPathName());
 			}
 
 			NumTotalEnabledTasks += State.EnabledTasksNum;
