StateTree: Revert check before sending the event introduced by 38076983. The event queue is reset when you start the tree but a EnterState callback may trigger an event before the tree is mark as "Started".

diff --git a/Engine/Plugins/Runtime/StateTree/Source/StateTreeModule/Private/StateTreeExecutionContext.cpp b/Engine/Plugins/Runtime/StateTree/Source/StateTreeModule/Private/StateTreeExecutionContext.cpp
--- a/Engine/Plugins/Runtime/StateTree/Source/StateTreeModule/Private/StateTreeExecutionContext.cpp
+++ b/Engine/Plugins/Runtime/StateTree/Source/StateTreeModule/Private/StateTreeExecutionContext.cpp
@@ -169,14 +169,6 @@ void FStateTreeMinimalExecutionContext::SendEvent(const FGameplayTag Tag, const
 		return;
 	}
 
-	// Events are reset when a state tree execution starts.
-	if (Storage.GetExecutionState().TreeRunStatus != EStateTreeRunStatus::Running)
-	{
-		STATETREE_LOG(Warning, TEXT("%hs: StateTree execution context is not started ('%s' using StateTree '%s')"),
-			__FUNCTION__, *GetNameSafe(&Owner), *GetFullNameSafe(&RootStateTree));
-		return;
-	}
-
 	STATETREE_LOG(VeryVerbose, TEXT("Send Event '%s'"), *Tag.ToString());
 	UE_STATETREE_DEBUG_LOG_EVENT(this, Log, TEXT("Send Event '%s'"), *Tag.ToString());
 
