[State Tree Unittests]
- Added a subtree with no active tasks test case
- Renamed/Regrouped Subtree tasks that one test was always hidden
- tweak comment

diff --git a/Engine/Plugins/Runtime/StateTree/Source/StateTreeTestSuite/Private/StateTreeTest.cpp b/Engine/Plugins/Runtime/StateTree/Source/StateTreeTestSuite/Private/StateTreeTest.cpp
--- a/Engine/Plugins/Runtime/StateTree/Source/StateTreeTestSuite/Private/StateTreeTest.cpp
+++ b/Engine/Plugins/Runtime/StateTree/Source/StateTreeTestSuite/Private/StateTreeTest.cpp
@@ -366,7 +366,7 @@ struct FStateTreeTest_Restart : FStateTreeTestBase
 };
 IMPLEMENT_AI_INSTANT_TEST(FStateTreeTest_Restart, "System.StateTree.Restart");
 
-struct FStateTreeTest_SubTree : FStateTreeTestBase
+struct FStateTreeTest_SubTree_ActiveTasks : FStateTreeTestBase
 {
 	virtual bool InstantTest() override
 	{
@@ -442,18 +442,99 @@ struct FStateTreeTest_SubTree : FStateTreeTestBase
 		return true;
 	}
 };
-IMPLEMENT_AI_INSTANT_TEST(FStateTreeTest_SubTree, "System.StateTree.SubTree");
+IMPLEMENT_AI_INSTANT_TEST(FStateTreeTest_SubTree_ActiveTasks, "System.StateTree.SubTree.ActiveTasks");
 
-struct FStateTreeTest_SubTreeCondition : FStateTreeTestBase
+struct FStateTreeTest_SubTree_NoActiveTasks : FStateTreeTestBase
+{
+	virtual bool InstantTest() override
+	{
+		UStateTree& StateTree = NewStateTree();
+		UStateTreeEditorData& EditorData = *Cast<UStateTreeEditorData>(StateTree.EditorData);
+
+		/*
+		 - RootA
+			- StateA : SubTree -> StateB
+			- StateB
+		 - RootB -> StateB
+		 - SubTree[DisabledTask] -> StateB
+			- StateC -> RootB
+		 */
+		UStateTreeState& RootA = EditorData.AddSubTree(FName(TEXT("RootA")));
+		UStateTreeState& StateA = RootA.AddChildState(FName(TEXT("StateA")));
+		UStateTreeState& StateB = RootA.AddChildState(FName(TEXT("StateB")));
+
+		UStateTreeState& RootB = EditorData.AddSubTree(FName(TEXT("RootB")));
+
+		UStateTreeState& SubTree = EditorData.AddSubTree(FName(TEXT("SubTree")));
+		UStateTreeState& StateC = SubTree.AddChildState(FName(TEXT("StateC")));
+
+		SubTree.Type = EStateTreeStateType::Subtree;
+		StateA.Type = EStateTreeStateType::Linked;
+		StateA.LinkedSubtree = SubTree.GetLinkToState();
+
+		StateA.AddTransition(EStateTreeTransitionTrigger::OnStateCompleted, EStateTreeTransitionType::GotoState, &StateB);
+		RootB.AddTransition(EStateTreeTransitionTrigger::OnStateCompleted, EStateTreeTransitionType::GotoState, &StateB);
+		SubTree.AddTransition(EStateTreeTransitionTrigger::OnStateCompleted, EStateTreeTransitionType::GotoState, &StateB);
+		StateC.AddTransition(EStateTreeTransitionTrigger::OnStateCompleted, EStateTreeTransitionType::GotoState, &RootB);
+
+		auto& TaskNode = SubTree.AddTask<FTestTask_Stand>(FName(TEXT("DisabledTask")));
+		TaskNode.GetNode().bTaskEnabled = false;
+
+		FStateTreeCompilerLog Log;
+		FStateTreeCompiler Compiler(Log);
+		const bool bResult = Compiler.Compile(StateTree);
+		AITEST_TRUE("StateTree should get compiled", bResult);
+
+		EStateTreeRunStatus Status = EStateTreeRunStatus::Unset;
+		FStateTreeInstanceData InstanceData;
+		FTestStateTreeExecutionContext Exec(StateTree, StateTree, InstanceData);
+		const bool bInitSucceeded = Exec.IsValid();
+		AITEST_TRUE("StateTree should init", bInitSucceeded);
+
+		// Start and enter state
+		Status = Exec.Start();
+
+		AITEST_TRUE("StateTree Active States should be in RootA/StateA/SubTree/StateC", 
+					Exec.ExpectInActiveStates(RootA.Name, StateA.Name, SubTree.Name, StateC.Name));
+		AITEST_TRUE("StateTree should be running", Status == EStateTreeRunStatus::Running);
+		Exec.LogClear();
+
+		// Transition from the subtree frame. The parent frame and the disabled task should be ignored.
+		Status = Exec.Tick(0.1f);
+		AITEST_TRUE("StateTree Active States should be in RootB", 
+					Exec.ExpectInActiveStates(RootB.Name));
+		AITEST_TRUE("StateTree should be running", Status == EStateTreeRunStatus::Running);
+		Exec.LogClear();
+
+		Status = Exec.Tick(0.1f);
+		AITEST_TRUE("StateTree Active States should be in RootA/StateB",
+					Exec.ExpectInActiveStates(RootA.Name, StateB.Name));
+		AITEST_TRUE("StateTree should be running", Status == EStateTreeRunStatus::Running);
+		Exec.LogClear();
+
+		Status = Exec.Tick(0.1f);
+		AITEST_TRUE("StateTree Active States should be in RootA/StateA/SubTree/StateC",
+					Exec.ExpectInActiveStates(RootA.Name, StateA.Name, SubTree.Name, StateC.Name));
+		AITEST_TRUE("StateTree should be running", Status == EStateTreeRunStatus::Running);
+		Exec.LogClear();
+
+		Exec.Stop();
+
+		return true;
+	}
+};
+IMPLEMENT_AI_INSTANT_TEST(FStateTreeTest_SubTree_NoActiveTasks, "System.StateTree.SubTree.NoActiveTasks");
+
+struct FStateTreeTest_SubTree_Condition : FStateTreeTestBase
 {
 	virtual bool InstantTest() override
 	{
 		/*
 		- Root
 			- Linked : Subtree -> Root
-		- SubTree : Task1
-			- ? State1 : Task2 -> Succeeded // condition linked to Task1
-			- State2 : Task3
+			- SubTree : Task1
+				- ? State1 : Task2 -> Succeeded // condition linked to Task1
+				- State2 : Task3
 		*/
 		
 		UStateTree& StateTree = NewStateTree();
@@ -524,7 +605,7 @@ struct FStateTreeTest_SubTreeCondition : FStateTreeTestBase
 		return true;
 	}
 };
-IMPLEMENT_AI_INSTANT_TEST(FStateTreeTest_SubTreeCondition, "System.StateTree.SubTreeCondition");
+IMPLEMENT_AI_INSTANT_TEST(FStateTreeTest_SubTree_Condition, "System.StateTree.SubTree.Condition");
 
 struct FStateTreeTest_SubTree_CascadedSucceeded : FStateTreeTestBase
 {
