StateTree: WeakContext uses the current context for transition if it exists. ensure if a user creates a new context while we are ticking the current context. That checks if 2 contexts do operation on the same data.

diff --git a/Engine/Plugins/Runtime/StateTree/Source/StateTreeModule/Private/StateTreeExecutionContext.cpp b/Engine/Plugins/Runtime/StateTree/Source/StateTreeModule/Private/StateTreeExecutionContext.cpp
--- a/Engine/Plugins/Runtime/StateTree/Source/StateTreeModule/Private/StateTreeExecutionContext.cpp
+++ b/Engine/Plugins/Runtime/StateTree/Source/StateTreeModule/Private/StateTreeExecutionContext.cpp
@@ -200,6 +200,25 @@ FStateTreeInstanceDebugId FStateTreeMinimalExecutionContext::GetInstanceDebugId(
 }
 #endif // WITH_STATETREE_TRACE
 
+FStateTreeExecutionContext::FCurrentlyExecutionContextScope::FCurrentlyExecutionContextScope(FStateTreeExecutionContext& InContext)
+	: CurrentContext(InContext)
+{
+	FStateTreeExecutionState& Exec = InContext.GetExecState();
+	ensureMsgf(Exec.CurrentExecutionContext == nullptr || Exec.CurrentExecutionContext == &InContext, TEXT("Only one execution context can operate on the State Tree Instance Data."));
+	PreviousExecutionContext = Exec.CurrentExecutionContext;
+	Exec.CurrentExecutionContext = &InContext;
+}
+
+FStateTreeExecutionContext::FCurrentlyExecutionContextScope::~FCurrentlyExecutionContextScope()
+{
+	FStateTreeExecutionState& Exec = CurrentContext.GetExecState();
+	// The tree can stop (it will reset the execution instance data) while we are ticking the state tree.
+	if (Exec.CurrentExecutionContext == &CurrentContext)
+	{
+		Exec.CurrentExecutionContext = PreviousExecutionContext;
+	}
+}
+
 FStateTreeExecutionContext::FCurrentlyProcessedFrameScope::FCurrentlyProcessedFrameScope(FStateTreeExecutionContext& InContext, const FStateTreeExecutionFrame* CurrentParentFrame, const FStateTreeExecutionFrame& CurrentFrame): Context(InContext)
 {
 	check(CurrentFrame.StateTree);
@@ -416,6 +435,8 @@ EStateTreeRunStatus FStateTreeExecutionContext::Start(FStartParameters Parameter
 		SetGlobalParameters(RootStateTree.GetDefaultParameters());
 	}
 
+	FCurrentlyExecutionContextScope ContextScope(*this);
+
 	// Initialize for the init frame.
 	ensure(Exec.ActiveFrames.Num() == 0);
 	FStateTreeExecutionFrame& InitFrame = Exec.ActiveFrames.AddDefaulted_GetRef();
@@ -549,6 +570,8 @@ EStateTreeRunStatus FStateTreeExecutionContext::Stop(EStateTreeRunStatus Complet
 		return EStateTreeRunStatus::Failed;
 	}
 
+	FCurrentlyExecutionContextScope ContextScope(*this);
+
 	// Set scoped phase only for properly initialized context with valid Instance data
 	// since we need it to output the InstanceId 
 	UE_STATETREE_DEBUG_SCOPED_PHASE(this, EStateTreeUpdatePhase::StopTree);
@@ -667,6 +690,8 @@ EStateTreeRunStatus FStateTreeExecutionContext::Tick(const float DeltaTime)
 {
 	CSV_SCOPED_TIMING_STAT_EXCLUSIVE(StateTree_Tick);
 
+	FCurrentlyExecutionContextScope ContextScope(*this);
+
 	const EStateTreeRunStatus PreludeResult = TickPrelude();
 	if (PreludeResult != EStateTreeRunStatus::Running)
 	{
@@ -683,6 +708,8 @@ EStateTreeRunStatus FStateTreeExecutionContext::TickUpdateTasks(const float Delt
 {
 	CSV_SCOPED_TIMING_STAT_EXCLUSIVE(StateTree_Tick);
 
+	FCurrentlyExecutionContextScope ContextScope(*this);
+
 	const EStateTreeRunStatus PreludeResult = TickPrelude();
 	if (PreludeResult != EStateTreeRunStatus::Running)
 	{
@@ -698,6 +725,8 @@ EStateTreeRunStatus FStateTreeExecutionContext::TickTriggerTransitions()
 {
 	CSV_SCOPED_TIMING_STAT_EXCLUSIVE(StateTree_Tick);
 
+	FCurrentlyExecutionContextScope ContextScope(*this);
+
 	const EStateTreeRunStatus PreludeResult = TickPrelude();
 	if (PreludeResult != EStateTreeRunStatus::Running)
 	{
@@ -858,6 +887,11 @@ void FStateTreeExecutionContext::RequestTransition(const FStateTreeTransitionReq
 	ScheduleNextTick();
 }
 
+void FStateTreeExecutionContext::RequestTransition(FStateTreeStateHandle InTargetState, EStateTreeTransitionPriority InPriority, EStateTreeSelectionFallback InFallback)
+{
+	RequestTransition(FStateTreeTransitionRequest(InTargetState, InPriority, InFallback));
+}
+
 void FStateTreeExecutionContext::UpdateInstanceData(TConstArrayView<FStateTreeExecutionFrame> CurrentActiveFrames, TArrayView<FStateTreeExecutionFrame> NextActiveFrames)
 {
 	// Estimate how many new instance data items we might have.
@@ -4955,7 +4989,11 @@ bool FStateTreeWeakExecutionContext::RequestTransition(FStateTreeStateHandle Tar
 	if (OwnerPinned && StateTreePinned && StoragePinned)
 	{
 		FStateTreeExecutionState& Exec = StoragePinned->GetMutableExecutionState();
-		if (Exec.ActiveFrames.Num() > 0)
+		if (Exec.CurrentExecutionContext)
+		{
+			Exec.CurrentExecutionContext->RequestTransition(TargetState, Priority, Fallback);
+		}
+		else if (Exec.ActiveFrames.Num() > 0)
 		{
 			//todo: weak should check if there is an active FStateTreeExecutionContext and use that to trigger transition directly instead of caching
 			const FStateTreeExecutionFrame* RootFrame = &Exec.ActiveFrames[0];
diff --git a/Engine/Plugins/Runtime/StateTree/Source/StateTreeModule/Public/StateTreeExecutionContext.h b/Engine/Plugins/Runtime/StateTree/Source/StateTreeModule/Public/StateTreeExecutionContext.h
--- a/Engine/Plugins/Runtime/StateTree/Source/StateTreeModule/Public/StateTreeExecutionContext.h
+++ b/Engine/Plugins/Runtime/StateTree/Source/StateTreeModule/Public/StateTreeExecutionContext.h
@@ -478,6 +478,17 @@ public:
 	 */
 	void RequestTransition(const FStateTreeTransitionRequest& Request);
 
+	/**
+	 * Requests transition to a state.
+	 * If called during during transition processing (e.g. from FStateTreeTaskBase::TriggerTransitions()) the transition
+	 * is attempted to be activate immediately (it can fail e.g. because of preconditions on a target state).
+	 * If called outside the transition handling, the request is buffered and handled at the beginning of next transition processing.
+	 * @param TargetState The state to transition to.
+	 * @param Priority The priority of the transition.
+	 * @param Fallback of the transition if it fails to select the target state.
+	 */
+	void RequestTransition(FStateTreeStateHandle TargetState, EStateTreeTransitionPriority Priority = EStateTreeTransitionPriority::Normal, EStateTreeSelectionFallback Fallback = EStateTreeSelectionFallback::None);
+
 	/** @return data view of the specified handle relative to given frame. */
 	static FStateTreeDataView GetDataViewFromInstanceStorage(FStateTreeInstanceStorage& InstanceDataStorage, FStateTreeInstanceStorage* CurrentlyProcessedSharedInstanceStorage, const FStateTreeExecutionFrame* ParentFrame, const FStateTreeExecutionFrame& CurrentFrame, const FStateTreeDataHandle Handle);
 
@@ -818,6 +829,17 @@ protected:
 	/** Pointer to the shared instance data of the current frame we're processing.  */
 	FStateTreeInstanceStorage* CurrentlyProcessedSharedInstanceStorage = nullptr;
 
+	/** Helper struct to track currently execution context. */
+	struct FCurrentlyExecutionContextScope
+	{
+		FCurrentlyExecutionContextScope(FStateTreeExecutionContext& InContext);
+		~FCurrentlyExecutionContextScope();
+
+	private:
+		FStateTreeExecutionContext& CurrentContext;
+		FStateTreeExecutionContext* PreviousExecutionContext;
+	};
+
 	/** Helper struct to track currently processed frame. */
 	struct FCurrentlyProcessedFrameScope
 	{
@@ -1035,7 +1057,7 @@ public:
 	 * @param Priority The priority of the transition.
 	 * @param Fallback of the transition if it fails to select the target state.
 	 */
-	bool RequestTransition(FStateTreeStateHandle TargetState, EStateTreeTransitionPriority Priority = EStateTreeTransitionPriority::Normal, const EStateTreeSelectionFallback InFallback = EStateTreeSelectionFallback::None) const;
+	bool RequestTransition(FStateTreeStateHandle TargetState, EStateTreeTransitionPriority Priority = EStateTreeTransitionPriority::Normal, EStateTreeSelectionFallback Fallback = EStateTreeSelectionFallback::None) const;
 
 private:
 	TWeakObjectPtr<UObject> Owner;
diff --git a/Engine/Plugins/Runtime/StateTree/Source/StateTreeModule/Public/StateTreeExecutionTypes.h b/Engine/Plugins/Runtime/StateTree/Source/StateTreeModule/Public/StateTreeExecutionTypes.h
--- a/Engine/Plugins/Runtime/StateTree/Source/StateTreeModule/Public/StateTreeExecutionTypes.h
+++ b/Engine/Plugins/Runtime/StateTree/Source/StateTreeModule/Public/StateTreeExecutionTypes.h
@@ -7,6 +7,7 @@
 
 #include "StateTreeExecutionTypes.generated.h"
 
+struct FStateTreeExecutionContext;
 class UStateTree;
 
 /**
@@ -673,6 +674,9 @@ PRAGMA_ENABLE_DEPRECATION_WARNINGS
 	UPROPERTY(Transient)
 	TInstancedStruct<FStateTreeExecutionExtension> ExecutionExtension;
 
+	/** The current context used to operate on the instance. */
+	FStateTreeExecutionContext* CurrentExecutionContext = nullptr;
+
 	/** The index of the task that failed during enter state. Exit state uses it to call ExitState() symmetrically. */
 	UPROPERTY()
 	FStateTreeIndex16 EnterStateFailedFrameIndex = FStateTreeIndex16::Invalid;
