StateTree: Multithread access detection. Detect if 2 threads are accessing the same instance data. The validation can be deactivate with the cvar StateTree.RuntimeValidation.MultithreadAccessDetector

diff --git a/Engine/Plugins/Runtime/StateTree/Source/StateTreeModule/Private/Debugger/StateTreeRuntimeValidation.cpp b/Engine/Plugins/Runtime/StateTree/Source/StateTreeModule/Private/Debugger/StateTreeRuntimeValidation.cpp
--- a/Engine/Plugins/Runtime/StateTree/Source/StateTreeModule/Private/Debugger/StateTreeRuntimeValidation.cpp
+++ b/Engine/Plugins/Runtime/StateTree/Source/StateTreeModule/Private/Debugger/StateTreeRuntimeValidation.cpp
@@ -14,6 +14,13 @@ namespace UE::StateTree::Debug
 
 namespace Private
 {
+bool RuntimeValidationThreadingAccessDetector = true;
+static FAutoConsoleVariableRef CVarRuntimeValidationThreadingAccessDetector(
+	TEXT("StateTree.RuntimeValidation.MultithreadAccessDetector"),
+	RuntimeValidationThreadingAccessDetector,
+	TEXT("Test if an execution context is accessed from different threads at the same time.")
+);
+
 bool RuntimeValidationContext = true;
 static FAutoConsoleVariableRef CVarRuntimeValidationContext(
 	TEXT("StateTree.RuntimeValidation.Context"),
@@ -84,6 +91,22 @@ void FRuntimeValidationInstanceData::SetContext(const UObject* InNewOwner, const
 	Owner = NewOwner;
 }
 
+void FRuntimeValidationInstanceData::AcquireAccess()
+{
+	if (Private::RuntimeValidationThreadingAccessDetector)
+	{
+		AccessDetector.AcquireWriteAccess();
+	}
+}
+
+void FRuntimeValidationInstanceData::ReleaseAccess()
+{
+	if (Private::RuntimeValidationThreadingAccessDetector)
+	{
+		AccessDetector.ReleaseWriteAccess();
+	}
+}
+
 void FRuntimeValidationInstanceData::NodeEnterState(FGuid Id)
 {
 	FNodeStatePair* Found = NodeStates.FindByPredicate([Id](const FNodeStatePair& Other)
@@ -158,12 +181,36 @@ void FRuntimeValidation::SetContext(const UObject* Owner, const UStateTree* Stat
 	}
 }
 
+void FRuntimeValidation::AcquireAccess() const
+{
+	if (RuntimeValidationData)
+	{
+		RuntimeValidationData->AcquireAccess();
+	}
+}
+
+void FRuntimeValidation::ReleaseAccess() const
+{
+	if (RuntimeValidationData)
+	{
+		RuntimeValidationData->ReleaseAccess();
+	}
+}
+
 #else
 
 void FRuntimeValidation::SetContext(const UObject*, const UStateTree*) const
 {
 }
 
+void FRuntimeValidation::AcquireAccess() const
+{
+}
+
+void FRuntimeValidation::ReleaseAccess() const
+{
+}
+
 #endif //WITH_STATETREE_DEBUG
 
 } // UE::StateTree::Debug
diff --git a/Engine/Plugins/Runtime/StateTree/Source/StateTreeModule/Private/Debugger/StateTreeRuntimeValidationInstanceData.h b/Engine/Plugins/Runtime/StateTree/Source/StateTreeModule/Private/Debugger/StateTreeRuntimeValidationInstanceData.h
--- a/Engine/Plugins/Runtime/StateTree/Source/StateTreeModule/Private/Debugger/StateTreeRuntimeValidationInstanceData.h
+++ b/Engine/Plugins/Runtime/StateTree/Source/StateTreeModule/Private/Debugger/StateTreeRuntimeValidationInstanceData.h
@@ -9,6 +9,8 @@ class UObject;
 
 #if WITH_STATETREE_DEBUG
 
+#include "Misc/MTAccessDetector.h"
+
 namespace UE::StateTree::Debug
 {
 /**
@@ -21,6 +23,10 @@ public:
 	~FRuntimeValidationInstanceData();
 
 	void SetContext(const UObject* Owner, const UStateTree* StateTree);
+
+	void AcquireAccess();
+	void ReleaseAccess();
+
 	void NodeEnterState(FGuid Id);
 	void NodeExitState(FGuid Id);
 
@@ -40,6 +46,9 @@ private:
 	TArray<FNodeStatePair> NodeStates;
 	TWeakObjectPtr<const UStateTree> StateTree;
 	TWeakObjectPtr<const UObject> Owner;
+
+	/** Multithread access detector. */
+	FRWRecursiveAccessDetector AccessDetector;
 };
 
 ENUM_CLASS_FLAGS(FRuntimeValidationInstanceData::EState);
diff --git a/Engine/Plugins/Runtime/StateTree/Source/StateTreeModule/Private/StateTreeExecutionContext.cpp b/Engine/Plugins/Runtime/StateTree/Source/StateTreeModule/Private/StateTreeExecutionContext.cpp
--- a/Engine/Plugins/Runtime/StateTree/Source/StateTreeModule/Private/StateTreeExecutionContext.cpp
+++ b/Engine/Plugins/Runtime/StateTree/Source/StateTreeModule/Private/StateTreeExecutionContext.cpp
@@ -33,6 +33,7 @@ FStateTreeMinimalExecutionContext::FStateTreeMinimalExecutionContext(UObject& In
 	, RootStateTree(InStateTree)
 	, Storage(InInstanceData.GetMutableStorage())
 {
+	Storage.GetRuntimeValidation().AcquireAccess();
 	if (InStateTree.IsReadyToRun())
 	{
 		Storage.GetRuntimeValidation().SetContext(&Owner, &RootStateTree);
@@ -44,12 +45,18 @@ FStateTreeMinimalExecutionContext::FStateTreeMinimalExecutionContext(UObject& In
 	, RootStateTree(InStateTree)
 	, Storage(InStorage)
 {
+	Storage.GetRuntimeValidation().AcquireAccess();
 	if (InStateTree.IsReadyToRun())
 	{
 		Storage.GetRuntimeValidation().SetContext(&Owner, &RootStateTree);
 	}
 }
 
+FStateTreeMinimalExecutionContext::~FStateTreeMinimalExecutionContext()
+{
+	Storage.GetRuntimeValidation().ReleaseAccess();
+}
+
 FString FStateTreeMinimalExecutionContext::GetInstanceDescriptionInternal() const
 {
 	const TInstancedStruct<FStateTreeExecutionExtension>& ExecutionExtension = Storage.GetExecutionState().ExecutionExtension;
@@ -415,9 +422,13 @@ EStateTreeRunStatus FStateTreeExecutionContext::Start(FStartParameters Parameter
 		Stop();
 	}
 
+	// Release the access created in the constructor.
+	InstanceDataStorage->GetRuntimeValidation().ReleaseAccess();
+
 	// Initialize instance data. No active states yet, so we'll initialize the evals and global tasks.
 	InstanceData.Reset();
 
+	InstanceDataStorage->GetRuntimeValidation().AcquireAccess();
 	InstanceDataStorage->GetRuntimeValidation().SetContext(&Owner, &RootStateTree);
 	Exec.ExecutionExtension = MoveTemp(Parameters.ExecutionExtension);
 	if (Parameters.SharedEventQueue)
@@ -619,9 +630,14 @@ EStateTreeRunStatus FStateTreeExecutionContext::Stop(EStateTreeRunStatus Complet
 	UE_STATETREE_DEBUG_ACTIVE_STATES_EVENT(this, {});
 	UE_STATETREE_DEBUG_INSTANCE_EVENT(this, EStateTreeTraceEventType::Pop);
 
+	// Release the access created in the constructor.
+	InstanceDataStorage->GetRuntimeValidation().ReleaseAccess();
+
 	// Destruct all allocated instance data (does not shrink the buffer). This will invalidate Exec too.
 	InstanceData.Reset();
 
+	InstanceDataStorage->GetRuntimeValidation().AcquireAccess();
+
 	// External data needs to be recollected if this exec context is reused.
 	bActiveExternalDataCollected = false;
 
@@ -4988,6 +5004,8 @@ bool FStateTreeWeakExecutionContext::RequestTransition(FStateTreeStateHandle Tar
 	TSharedPtr<FStateTreeInstanceStorage> StoragePinned = Storage.Pin();
 	if (OwnerPinned && StateTreePinned && StoragePinned)
 	{
+		StoragePinned->GetRuntimeValidation().AcquireAccess();
+
 		FStateTreeExecutionState& Exec = StoragePinned->GetMutableExecutionState();
 		if (Exec.CurrentExecutionContext)
 		{
@@ -4995,7 +5013,6 @@ bool FStateTreeWeakExecutionContext::RequestTransition(FStateTreeStateHandle Tar
 		}
 		else if (Exec.ActiveFrames.Num() > 0)
 		{
-			//todo: weak should check if there is an active FStateTreeExecutionContext and use that to trigger transition directly instead of caching
 			const FStateTreeExecutionFrame* RootFrame = &Exec.ActiveFrames[0];
 
 			FStateTreeTransitionRequest Request;
@@ -5016,6 +5033,8 @@ bool FStateTreeWeakExecutionContext::RequestTransition(FStateTreeStateHandle Tar
 			}
 			return true;
 		}
+
+		StoragePinned->GetRuntimeValidation().ReleaseAccess();
 	}
 	return false;
 }
diff --git a/Engine/Plugins/Runtime/StateTree/Source/StateTreeModule/Public/Debugger/StateTreeRuntimeValidation.h b/Engine/Plugins/Runtime/StateTree/Source/StateTreeModule/Public/Debugger/StateTreeRuntimeValidation.h
--- a/Engine/Plugins/Runtime/StateTree/Source/StateTreeModule/Public/Debugger/StateTreeRuntimeValidation.h
+++ b/Engine/Plugins/Runtime/StateTree/Source/StateTreeModule/Public/Debugger/StateTreeRuntimeValidation.h
@@ -37,6 +37,12 @@ public:
 	/** Set the Owner of the data */
 	void SetContext(const UObject* Owner, const UStateTree* StateTree) const;
 
+	/** Detect invalid multithreading access of the instance data. */
+	void AcquireAccess() const;
+
+	/** The current thread stop using the instance data. */
+	void ReleaseAccess() const;
+
 private:
 #if WITH_STATETREE_DEBUG
 	FRuntimeValidationInstanceData* RuntimeValidationData = nullptr;
diff --git a/Engine/Plugins/Runtime/StateTree/Source/StateTreeModule/Public/StateTreeExecutionContext.h b/Engine/Plugins/Runtime/StateTree/Source/StateTreeModule/Public/StateTreeExecutionContext.h
--- a/Engine/Plugins/Runtime/StateTree/Source/StateTreeModule/Public/StateTreeExecutionContext.h
+++ b/Engine/Plugins/Runtime/StateTree/Source/StateTreeModule/Public/StateTreeExecutionContext.h
@@ -47,6 +47,7 @@ struct STATETREEMODULE_API FStateTreeMinimalExecutionContext
 public:
 	explicit FStateTreeMinimalExecutionContext(UObject& Owner, const UStateTree& StateTree, FStateTreeInstanceData& InInstanceData);
 	explicit FStateTreeMinimalExecutionContext(UObject& Owner, const UStateTree& StateTree, FStateTreeInstanceStorage& Storage);
+	~FStateTreeMinimalExecutionContext();
 
 public:
 	/**
