StateTree: Review feedback from CL 38189094

diff --git a/Engine/Plugins/Runtime/StateTree/Source/StateTreeModule/Private/Debugger/StateTreeRuntimeValidationInstanceData.h b/Engine/Plugins/Runtime/StateTree/Source/StateTreeModule/Private/Debugger/StateTreeRuntimeValidationInstanceData.h
--- a/Engine/Plugins/Runtime/StateTree/Source/StateTreeModule/Private/Debugger/StateTreeRuntimeValidationInstanceData.h
+++ b/Engine/Plugins/Runtime/StateTree/Source/StateTreeModule/Private/Debugger/StateTreeRuntimeValidationInstanceData.h
@@ -48,7 +48,7 @@ private:
 	TWeakObjectPtr<const UObject> Owner;
 
 	/** Multithread access detector. */
-	FRWRecursiveAccessDetector AccessDetector;
+	UE_MT_DECLARE_RW_RECURSIVE_ACCESS_DETECTOR(AccessDetector);
 };
 
 ENUM_CLASS_FLAGS(FRuntimeValidationInstanceData::EState);
diff --git a/Engine/Plugins/Runtime/StateTree/Source/StateTreeModule/Private/StateTreeExecutionContext.cpp b/Engine/Plugins/Runtime/StateTree/Source/StateTreeModule/Private/StateTreeExecutionContext.cpp
--- a/Engine/Plugins/Runtime/StateTree/Source/StateTreeModule/Private/StateTreeExecutionContext.cpp
+++ b/Engine/Plugins/Runtime/StateTree/Source/StateTreeModule/Private/StateTreeExecutionContext.cpp
@@ -29,15 +29,8 @@ namespace UE::StateTree::Debug
 
 
 FStateTreeMinimalExecutionContext::FStateTreeMinimalExecutionContext(UObject& InOwner, const UStateTree& InStateTree, FStateTreeInstanceData& InInstanceData)
-	: Owner(InOwner)
-	, RootStateTree(InStateTree)
-	, Storage(InInstanceData.GetMutableStorage())
+	: FStateTreeMinimalExecutionContext(InOwner, InStateTree, InInstanceData.GetMutableStorage())
 {
-	Storage.GetRuntimeValidation().AcquireAccess();
-	if (InStateTree.IsReadyToRun())
-	{
-		Storage.GetRuntimeValidation().SetContext(&Owner, &RootStateTree);
-	}
 }
 
 FStateTreeMinimalExecutionContext::FStateTreeMinimalExecutionContext(UObject& InOwner, const UStateTree& InStateTree, FStateTreeInstanceStorage& InStorage)
diff --git a/Engine/Plugins/Runtime/StateTree/Source/StateTreeModule/Public/Debugger/StateTreeRuntimeValidation.h b/Engine/Plugins/Runtime/StateTree/Source/StateTreeModule/Public/Debugger/StateTreeRuntimeValidation.h
--- a/Engine/Plugins/Runtime/StateTree/Source/StateTreeModule/Public/Debugger/StateTreeRuntimeValidation.h
+++ b/Engine/Plugins/Runtime/StateTree/Source/StateTreeModule/Public/Debugger/StateTreeRuntimeValidation.h
@@ -40,7 +40,7 @@ public:
 	/** Detect invalid multithreading access of the instance data. */
 	void AcquireAccess() const;
 
-	/** The current thread stop using the instance data. */
+	/** Stop the multitheading access detected. */
 	void ReleaseAccess() const;
 
 private:
