[StateTree] added missing debug event for EStateTreeUpdatePhase::TickStateTree

diff --git a/Engine/Plugins/Runtime/StateTree/Source/StateTreeModule/Private/StateTreeExecutionContext.cpp b/Engine/Plugins/Runtime/StateTree/Source/StateTreeModule/Private/StateTreeExecutionContext.cpp
--- a/Engine/Plugins/Runtime/StateTree/Source/StateTreeModule/Private/StateTreeExecutionContext.cpp
+++ b/Engine/Plugins/Runtime/StateTree/Source/StateTreeModule/Private/StateTreeExecutionContext.cpp
@@ -669,6 +669,7 @@ EStateTreeRunStatus FStateTreeExecutionContext::TickPrelude()
 
 	// From this point any calls to Stop should be deferred.
 	Exec.CurrentPhase = EStateTreeUpdatePhase::TickStateTree;
+	UE_STATETREE_DEBUG_ENTER_PHASE(this, EStateTreeUpdatePhase::TickStateTree);
 
 	return EStateTreeRunStatus::Running;
 }
@@ -679,6 +680,7 @@ EStateTreeRunStatus FStateTreeExecutionContext::TickPostlude()
 	FStateTreeExecutionState& Exec = GetExecState();
 
 	// Reset phase since we are now safe to stop.
+	UE_STATETREE_DEBUG_EXIT_PHASE(this, EStateTreeUpdatePhase::TickStateTree);
 	Exec.CurrentPhase = EStateTreeUpdatePhase::Unset;
 
 	// Use local for resulting run state since Stop will reset the instance data.
