StateTree: Move the multithread detector from the runtime validation to the instance data. The detector was reporting false positives when 2 execution contexts were created inside the same thread.

diff --git a/Engine/Plugins/Runtime/GameplayStateTree/Source/GameplayStateTreeModule/Private/BehaviorTree/Tasks/BTTask_RunDynamicStateTree.cpp b/Engine/Plugins/Runtime/GameplayStateTree/Source/GameplayStateTreeModule/Private/BehaviorTree/Tasks/BTTask_RunDynamicStateTree.cpp
--- a/Engine/Plugins/Runtime/GameplayStateTree/Source/GameplayStateTreeModule/Private/BehaviorTree/Tasks/BTTask_RunDynamicStateTree.cpp
+++ b/Engine/Plugins/Runtime/GameplayStateTree/Source/GameplayStateTreeModule/Private/BehaviorTree/Tasks/BTTask_RunDynamicStateTree.cpp
@@ -78,9 +78,13 @@ void UBTTask_RunDynamicStateTree::TickTask(UBehaviorTreeComponent& OwnerComp, ui
 		FinishLatentTask(OwnerComp, EBTNodeResult::Failed);
 	}
 
-	FStateTreeExecutionContext Context(*OwnerComp.GetOwner(), *StateTreeRef.GetStateTree(), InstanceData);
-	SetContextDataDelegate.ExecuteIfBound(Context, OwnerComp, InjectionTag);
-	const EStateTreeRunStatus TickStatus = Context.Tick(DeltaSeconds);
+	EStateTreeRunStatus TickStatus = EStateTreeRunStatus::Unset;
+	{
+		FStateTreeExecutionContext Context(*OwnerComp.GetOwner(), *StateTreeRef.GetStateTree(), InstanceData);
+		SetContextDataDelegate.ExecuteIfBound(Context, OwnerComp, InjectionTag);
+		TickStatus = Context.Tick(DeltaSeconds);
+	}
+
 	if (TickStatus != EStateTreeRunStatus::Running)
 	{
 		FinishLatentTask(OwnerComp, GameplayStateTreeBTUtils::StateTreeRunStatusToBTNodeResult(TickStatus));
diff --git a/Engine/Plugins/Runtime/StateTree/Source/StateTreeModule/Private/Debugger/StateTreeRuntimeValidation.cpp b/Engine/Plugins/Runtime/StateTree/Source/StateTreeModule/Private/Debugger/StateTreeRuntimeValidation.cpp
--- a/Engine/Plugins/Runtime/StateTree/Source/StateTreeModule/Private/Debugger/StateTreeRuntimeValidation.cpp
+++ b/Engine/Plugins/Runtime/StateTree/Source/StateTreeModule/Private/Debugger/StateTreeRuntimeValidation.cpp
@@ -14,13 +14,6 @@ namespace UE::StateTree::Debug
 
 namespace Private
 {
-bool RuntimeValidationThreadingAccessDetector = false;
-static FAutoConsoleVariableRef CVarRuntimeValidationThreadingAccessDetector(
-	TEXT("StateTree.RuntimeValidation.MultithreadAccessDetector"),
-	RuntimeValidationThreadingAccessDetector,
-	TEXT("Test if an execution context is accessed from different threads at the same time.")
-);
-
 bool RuntimeValidationContext = true;
 static FAutoConsoleVariableRef CVarRuntimeValidationContext(
 	TEXT("StateTree.RuntimeValidation.Context"),
@@ -91,22 +84,6 @@ void FRuntimeValidationInstanceData::SetContext(const UObject* InNewOwner, const
 	Owner = NewOwner;
 }
 
-void FRuntimeValidationInstanceData::AcquireAccess()
-{
-	if (Private::RuntimeValidationThreadingAccessDetector)
-	{
-		AccessDetector.AcquireWriteAccess();
-	}
-}
-
-void FRuntimeValidationInstanceData::ReleaseAccess()
-{
-	if (Private::RuntimeValidationThreadingAccessDetector)
-	{
-		AccessDetector.ReleaseWriteAccess();
-	}
-}
-
 void FRuntimeValidationInstanceData::NodeEnterState(FGuid Id)
 {
 	FNodeStatePair* Found = NodeStates.FindByPredicate([Id](const FNodeStatePair& Other)
@@ -181,36 +158,12 @@ void FRuntimeValidation::SetContext(const UObject* Owner, const UStateTree* Stat
 	}
 }
 
-void FRuntimeValidation::AcquireAccess() const
-{
-	if (RuntimeValidationData)
-	{
-		RuntimeValidationData->AcquireAccess();
-	}
-}
-
-void FRuntimeValidation::ReleaseAccess() const
-{
-	if (RuntimeValidationData)
-	{
-		RuntimeValidationData->ReleaseAccess();
-	}
-}
-
 #else
 
 void FRuntimeValidation::SetContext(const UObject*, const UStateTree*) const
 {
 }
 
-void FRuntimeValidation::AcquireAccess() const
-{
-}
-
-void FRuntimeValidation::ReleaseAccess() const
-{
-}
-
 #endif //WITH_STATETREE_DEBUG
 
 } // UE::StateTree::Debug
diff --git a/Engine/Plugins/Runtime/StateTree/Source/StateTreeModule/Private/Debugger/StateTreeRuntimeValidationInstanceData.h b/Engine/Plugins/Runtime/StateTree/Source/StateTreeModule/Private/Debugger/StateTreeRuntimeValidationInstanceData.h
--- a/Engine/Plugins/Runtime/StateTree/Source/StateTreeModule/Private/Debugger/StateTreeRuntimeValidationInstanceData.h
+++ b/Engine/Plugins/Runtime/StateTree/Source/StateTreeModule/Private/Debugger/StateTreeRuntimeValidationInstanceData.h
@@ -9,8 +9,6 @@ class UObject;
 
 #if WITH_STATETREE_DEBUG
 
-#include "Misc/MTAccessDetector.h"
-
 namespace UE::StateTree::Debug
 {
 /**
@@ -24,9 +22,6 @@ public:
 
 	void SetContext(const UObject* Owner, const UStateTree* StateTree);
 
-	void AcquireAccess();
-	void ReleaseAccess();
-
 	void NodeEnterState(FGuid Id);
 	void NodeExitState(FGuid Id);
 
@@ -46,9 +41,6 @@ private:
 	TArray<FNodeStatePair> NodeStates;
 	TWeakObjectPtr<const UStateTree> StateTree;
 	TWeakObjectPtr<const UObject> Owner;
-
-	/** Multithread access detector. */
-	UE_MT_DECLARE_RW_RECURSIVE_ACCESS_DETECTOR(AccessDetector);
 };
 
 ENUM_CLASS_FLAGS(FRuntimeValidationInstanceData::EState);
diff --git a/Engine/Plugins/Runtime/StateTree/Source/StateTreeModule/Private/StateTreeExecutionContext.cpp b/Engine/Plugins/Runtime/StateTree/Source/StateTreeModule/Private/StateTreeExecutionContext.cpp
--- a/Engine/Plugins/Runtime/StateTree/Source/StateTreeModule/Private/StateTreeExecutionContext.cpp
+++ b/Engine/Plugins/Runtime/StateTree/Source/StateTreeModule/Private/StateTreeExecutionContext.cpp
@@ -38,7 +38,7 @@ FStateTreeMinimalExecutionContext::FStateTreeMinimalExecutionContext(UObject& In
 	, RootStateTree(InStateTree)
 	, Storage(InStorage)
 {
-	Storage.GetRuntimeValidation().AcquireAccess();
+	Storage.AcquireWriteAccess();
 	if (InStateTree.IsReadyToRun())
 	{
 		Storage.GetRuntimeValidation().SetContext(&Owner, &RootStateTree);
@@ -47,7 +47,7 @@ FStateTreeMinimalExecutionContext::FStateTreeMinimalExecutionContext(UObject& In
 
 FStateTreeMinimalExecutionContext::~FStateTreeMinimalExecutionContext()
 {
-	Storage.GetRuntimeValidation().ReleaseAccess();
+	Storage.ReleaseWriteAccess();
 }
 
 FString FStateTreeMinimalExecutionContext::GetInstanceDescriptionInternal() const
@@ -415,13 +415,9 @@ EStateTreeRunStatus FStateTreeExecutionContext::Start(FStartParameters Parameter
 		Stop();
 	}
 
-	// Release the access created in the constructor.
-	InstanceDataStorage->GetRuntimeValidation().ReleaseAccess();
-
 	// Initialize instance data. No active states yet, so we'll initialize the evals and global tasks.
 	InstanceData.Reset();
 
-	InstanceDataStorage->GetRuntimeValidation().AcquireAccess();
 	InstanceDataStorage->GetRuntimeValidation().SetContext(&Owner, &RootStateTree);
 	Exec.ExecutionExtension = MoveTemp(Parameters.ExecutionExtension);
 	if (Parameters.SharedEventQueue)
@@ -623,14 +619,9 @@ EStateTreeRunStatus FStateTreeExecutionContext::Stop(EStateTreeRunStatus Complet
 	UE_STATETREE_DEBUG_ACTIVE_STATES_EVENT(this, {});
 	UE_STATETREE_DEBUG_INSTANCE_EVENT(this, EStateTreeTraceEventType::Pop);
 
-	// Release the access created in the constructor.
-	InstanceDataStorage->GetRuntimeValidation().ReleaseAccess();
-
 	// Destruct all allocated instance data (does not shrink the buffer). This will invalidate Exec too.
 	InstanceData.Reset();
 
-	InstanceDataStorage->GetRuntimeValidation().AcquireAccess();
-
 	// External data needs to be recollected if this exec context is reused.
 	bActiveExternalDataCollected = false;
 
@@ -4999,7 +4990,9 @@ bool FStateTreeWeakExecutionContext::RequestTransition(FStateTreeStateHandle Tar
 	TSharedPtr<FStateTreeInstanceStorage> StoragePinned = Storage.Pin();
 	if (OwnerPinned && StateTreePinned && StoragePinned)
 	{
-		StoragePinned->GetRuntimeValidation().AcquireAccess();
+		// Manually acquire and release the write access.
+		//The context may already exist and we want to confirm that we are on the same thread.
+		StoragePinned->AcquireWriteAccess();
 
 		FStateTreeExecutionState& Exec = StoragePinned->GetMutableExecutionState();
 		if (Exec.CurrentExecutionContext)
@@ -5029,7 +5022,7 @@ bool FStateTreeWeakExecutionContext::RequestTransition(FStateTreeStateHandle Tar
 			return true;
 		}
 
-		StoragePinned->GetRuntimeValidation().ReleaseAccess();
+		StoragePinned->ReleaseWriteAccess();
 	}
 	return false;
 }
diff --git a/Engine/Plugins/Runtime/StateTree/Source/StateTreeModule/Private/StateTreeInstanceData.cpp b/Engine/Plugins/Runtime/StateTree/Source/StateTreeModule/Private/StateTreeInstanceData.cpp
--- a/Engine/Plugins/Runtime/StateTree/Source/StateTreeModule/Private/StateTreeInstanceData.cpp
+++ b/Engine/Plugins/Runtime/StateTree/Source/StateTreeModule/Private/StateTreeInstanceData.cpp
@@ -16,9 +16,18 @@
 const FGuid FStateTreeInstanceStorageCustomVersion::GUID(0x60C4F0DE, 0x8B264C34, 0xAA937201, 0x5DFF09CC);
 FCustomVersionRegistration GRegisterStateTreeInstanceStorageCustomVersion(FStateTreeInstanceStorageCustomVersion::GUID, FStateTreeInstanceStorageCustomVersion::LatestVersion, TEXT("StateTreeInstanceStorage"));
 
-namespace UE::StateTree
+namespace UE::StateTree::Debug::Private
 {
+	bool RuntimeValidationThreadingAccessDetector = true;
+	static FAutoConsoleVariableRef CVarRuntimeValidationThreadingAccessDetector(
+		TEXT("StateTree.RuntimeValidation.MultithreadAccessDetector"),
+		RuntimeValidationThreadingAccessDetector,
+		TEXT("Test if an execution context is accessed from different threads at the same time.")
+	);
+} //UE::StateTree::Debug::Private
 
+namespace UE::StateTree
+{
 	/**
 	 * Duplicates object, and tries to covert old BP classes (REINST_*) to their newer version.
 	 */
@@ -100,6 +109,10 @@ FStateTreeInstanceStorage::FStateTreeInstanceStorage(const FStateTreeInstanceSto
 	, EventQueue(MakeShared<FStateTreeEventQueue>(*Other.EventQueue))
 	, TransitionRequests(Other.TransitionRequests)
 	, GlobalParameters(Other.GlobalParameters)
+#if ENABLE_MT_DETECTOR
+	, AccessDetector(Other.AccessDetector)
+#endif
+	, bIsOwningEventQueue(true)
 #if WITH_STATETREE_DEBUG
 	, RuntimeValidationData(MakePimpl<UE::StateTree::Debug::FRuntimeValidationInstanceData>(*Other.RuntimeValidationData))
 #endif
@@ -113,6 +126,10 @@ FStateTreeInstanceStorage::FStateTreeInstanceStorage(FStateTreeInstanceStorage&&
 	, EventQueue(Other.EventQueue)
 	, TransitionRequests(MoveTemp(Other.TransitionRequests))
 	, GlobalParameters(MoveTemp(Other.GlobalParameters))
+#if ENABLE_MT_DETECTOR
+	, AccessDetector(MoveTemp(Other.AccessDetector))
+#endif
+	, bIsOwningEventQueue(Other.bIsOwningEventQueue)
 #if WITH_STATETREE_DEBUG
 	, RuntimeValidationData(MoveTemp(Other.RuntimeValidationData))
 #endif
@@ -132,6 +149,16 @@ FStateTreeInstanceStorage& FStateTreeInstanceStorage::operator=(const FStateTree
 	TransitionRequests = Other.TransitionRequests;
 	GlobalParameters = Other.GlobalParameters;
 
+#if ENABLE_MT_DETECTOR
+	AccessDetector = Other.AccessDetector;
+#endif
+
+	bIsOwningEventQueue = true;
+
+#if WITH_STATETREE_DEBUG
+	RuntimeValidationData = MakePimpl<UE::StateTree::Debug::FRuntimeValidationInstanceData>(*Other.RuntimeValidationData);
+#endif
+
 	return *this;
 }
 
@@ -145,6 +172,18 @@ FStateTreeInstanceStorage& FStateTreeInstanceStorage::operator=(FStateTreeInstan
 	TransitionRequests = MoveTemp(Other.TransitionRequests);
 	GlobalParameters = MoveTemp(Other.GlobalParameters);
 
+#if ENABLE_MT_DETECTOR
+	AccessDetector = MoveTemp(Other.AccessDetector);
+#endif
+
+	bIsOwningEventQueue = Other.bIsOwningEventQueue;
+	Other.bIsOwningEventQueue = true;
+
+#if WITH_STATETREE_DEBUG
+	RuntimeValidationData = MoveTemp(Other.RuntimeValidationData);
+	Other.RuntimeValidationData = MakePimpl<UE::StateTree::Debug::FRuntimeValidationInstanceData>();
+#endif
+
 	return *this;
 }
 
@@ -291,6 +330,22 @@ void FStateTreeInstanceStorage::Reset()
 #endif
 }
 
+void FStateTreeInstanceStorage::AcquireWriteAccess()
+{
+	if (UE::StateTree::Debug::Private::RuntimeValidationThreadingAccessDetector)
+	{
+		UE_MT_ACQUIRE_WRITE_ACCESS(AccessDetector)
+	}
+}
+
+void FStateTreeInstanceStorage::ReleaseWriteAccess()
+{
+	if (UE::StateTree::Debug::Private::RuntimeValidationThreadingAccessDetector)
+	{
+		UE_MT_RELEASE_WRITE_ACCESS(AccessDetector)
+	}
+}
+
 UE::StateTree::Debug::FRuntimeValidation FStateTreeInstanceStorage::GetRuntimeValidation() const
 {
 #if WITH_STATETREE_DEBUG
diff --git a/Engine/Plugins/Runtime/StateTree/Source/StateTreeModule/Public/Debugger/StateTreeRuntimeValidation.h b/Engine/Plugins/Runtime/StateTree/Source/StateTreeModule/Public/Debugger/StateTreeRuntimeValidation.h
--- a/Engine/Plugins/Runtime/StateTree/Source/StateTreeModule/Public/Debugger/StateTreeRuntimeValidation.h
+++ b/Engine/Plugins/Runtime/StateTree/Source/StateTreeModule/Public/Debugger/StateTreeRuntimeValidation.h
@@ -37,12 +37,6 @@ public:
 	/** Set the Owner of the data */
 	void SetContext(const UObject* Owner, const UStateTree* StateTree) const;
 
-	/** Detect invalid multithreading access of the instance data. */
-	void AcquireAccess() const;
-
-	/** Stop the multitheading access detected. */
-	void ReleaseAccess() const;
-
 private:
 #if WITH_STATETREE_DEBUG
 	FRuntimeValidationInstanceData* RuntimeValidationData = nullptr;
diff --git a/Engine/Plugins/Runtime/StateTree/Source/StateTreeModule/Public/StateTreeInstanceData.h b/Engine/Plugins/Runtime/StateTree/Source/StateTreeModule/Public/StateTreeInstanceData.h
--- a/Engine/Plugins/Runtime/StateTree/Source/StateTreeModule/Public/StateTreeInstanceData.h
+++ b/Engine/Plugins/Runtime/StateTree/Source/StateTreeModule/Public/StateTreeInstanceData.h
@@ -7,6 +7,7 @@
 #include "StateTreeEvents.h"
 #include "StateTreeTypes.h"
 #include "StateTreeExecutionTypes.h"
+#include "Misc/MTAccessDetector.h"
 #include "Templates/PimplPtr.h"
 #include "Templates/SharedPointer.h"
 #include "StateTreeInstanceData.generated.h"
@@ -248,6 +249,12 @@ struct STATETREEMODULE_API FStateTreeInstanceStorage
 	/** Resets the storage to initial state. */
 	void Reset();
 
+	/** Start the invalid multithreading access detection. */
+	void AcquireWriteAccess();
+
+	/** Stop the multitheading access detection. */
+	void ReleaseWriteAccess();
+
 	/** Get the data used at runtime to confirm the inner working of the StateTree. */
 	UE::StateTree::Debug::FRuntimeValidation GetRuntimeValidation() const;
 
@@ -275,6 +282,9 @@ protected:
 	UPROPERTY(Transient)
 	FInstancedPropertyBag GlobalParameters;
 
+	/** Used to detect if we are using the instance data on multiple threads at the same time. */
+	UE_MT_DECLARE_RW_RECURSIVE_ACCESS_DETECTOR(AccessDetector)
+
 	/* True if the storage owns the event queue. */
 	bool bIsOwningEventQueue = true;
 
