[State Tree Compiler]
- Moved BindStruct out of CreateSharedInstanceData for ease of error reporting with Node Info
- Removed duplicated validity check
- Added warning report for Condition Evaluation Mode and Disabled Tasks

diff --git a/Engine/Plugins/Runtime/StateTree/Source/StateTreeEditorModule/Private/StateTreeCompiler.cpp b/Engine/Plugins/Runtime/StateTree/Source/StateTreeEditorModule/Private/StateTreeCompiler.cpp
--- a/Engine/Plugins/Runtime/StateTree/Source/StateTreeEditorModule/Private/StateTreeCompiler.cpp
+++ b/Engine/Plugins/Runtime/StateTree/Source/StateTreeEditorModule/Private/StateTreeCompiler.cpp
@@ -1319,7 +1319,13 @@ bool FStateTreeCompiler::CreatePropertyFunction(const FStateTreeEditorNode& Func
 		return false;
 	}
 
-	FStateTreeNodeBase* Node = CreateNodeWithSharedInstanceData(nullptr, UE::StateTree::Editor::PropertyFunctionStateName, FuncEditorNode, EStateTreeBindableStructSource::PropertyFunction);
+	FStateTreeBindableStructDesc StructDesc;
+	StructDesc.StatePath = UE::StateTree::Editor::PropertyFunctionStateName;
+	StructDesc.ID = FuncEditorNode.ID;
+	StructDesc.Name = FuncEditorNode.GetName();
+	StructDesc.DataSource = EStateTreeBindableStructSource::PropertyFunction;
+
+	FStateTreeNodeBase* Node = CreateNodeWithSharedInstanceData(nullptr, FuncEditorNode, StructDesc);
 	if(Node == nullptr)
 	{
 		return false;
@@ -1347,7 +1353,7 @@ bool FStateTreeCompiler::CreatePropertyFunction(const FStateTreeEditorNode& Func
 }
 
 template<class T>
-T* FStateTreeCompiler::CreateNodeWithSharedInstanceData(UStateTreeState* State, const FString& StatePath, const FStateTreeEditorNode& EditorNode, EStateTreeBindableStructSource StructSource)
+T* FStateTreeCompiler::CreateNodeWithSharedInstanceData(UStateTreeState* State, const FStateTreeEditorNode& EditorNode, FStateTreeBindableStructDesc& StructDesc)
 {
 	if (!EditorNode.Node.IsValid())
 	{
@@ -1356,12 +1362,6 @@ T* FStateTreeCompiler::CreateNodeWithSharedInstanceData(UStateTreeState* State,
 
 	check(EditorNode.Node.GetScriptStruct()->IsChildOf<T>());
 
-	FStateTreeBindableStructDesc StructDesc;
-	StructDesc.StatePath = StatePath;
-	StructDesc.ID = EditorNode.ID;
-	StructDesc.Name = EditorNode.GetName();
-	StructDesc.DataSource = StructSource;
-
 	// Check that item has valid instance initialized.
 	if (!EditorNode.Instance.IsValid() && EditorNode.InstanceObject == nullptr)
 	{
@@ -1521,14 +1521,22 @@ PRAGMA_ENABLE_DEPRECATION_WARNINGS
 
 bool FStateTreeCompiler::CreateCondition(UStateTreeState& State, const FString& StatePath, const FStateTreeEditorNode& CondNode, const EStateTreeExpressionOperand Operand, const int8 DeltaIndent)
 {
-	if (!CondNode.Node.IsValid())
-	{
-		// Empty line in conditions array, just silently ignore.
-		return true;
-	}
+	FStateTreeBindableStructDesc StructDesc;
+	StructDesc.StatePath = StatePath;
+	StructDesc.ID = CondNode.ID;
+	StructDesc.Name = CondNode.GetName();
+	StructDesc.DataSource = EStateTreeBindableStructSource::Condition;
 
-	if (FStateTreeConditionBase* Cond = CreateNodeWithSharedInstanceData<FStateTreeConditionBase>(&State, StatePath, CondNode, EStateTreeBindableStructSource::Condition))
+	if (FStateTreeConditionBase* Cond = CreateNodeWithSharedInstanceData<FStateTreeConditionBase>(&State, CondNode, StructDesc))
 	{
+		if (Cond->EvaluationMode == EStateTreeConditionEvaluationMode::ForcedFalse
+			|| Cond->EvaluationMode == EStateTreeConditionEvaluationMode::ForcedTrue)
+		{
+			Log.Reportf(EMessageSeverity::Info, StructDesc, 
+						TEXT("The condition result will always be %s."),
+						Cond->EvaluationMode == EStateTreeConditionEvaluationMode::ForcedTrue ? TEXT("True") : TEXT("False"));
+		}
+
 		Cond->Operand = Operand;
 		Cond->DeltaIndent = DeltaIndent;
 		return true;
@@ -1575,13 +1583,13 @@ bool FStateTreeCompiler::CreateConsiderations(UStateTreeState& State, const FStr
 
 bool FStateTreeCompiler::CreateConsideration(UStateTreeState& State, const FString& StatePath, const FStateTreeEditorNode& ConsiderationNode, const EStateTreeExpressionOperand Operand, const int8 DeltaIndent)
 {
-	if (!ConsiderationNode.Node.IsValid())
-	{
-		// Empty line in conditions array, just silently ignore.
-		return true;
-	}
+	FStateTreeBindableStructDesc StructDesc;
+	StructDesc.StatePath = StatePath;
+	StructDesc.ID = ConsiderationNode.ID;
+	StructDesc.Name = ConsiderationNode.GetName();
+	StructDesc.DataSource = EStateTreeBindableStructSource::Consideration;
 
-	if (FStateTreeConsiderationBase* Consideration = CreateNodeWithSharedInstanceData<FStateTreeConsiderationBase>(&State, StatePath, ConsiderationNode, EStateTreeBindableStructSource::Consideration))
+	if (FStateTreeConsiderationBase* Consideration = CreateNodeWithSharedInstanceData<FStateTreeConsiderationBase>(&State, ConsiderationNode, StructDesc))
 	{
 		Consideration->Operand = Operand;
 		Consideration->DeltaIndent = DeltaIndent;
@@ -1730,6 +1738,11 @@ bool FStateTreeCompiler::CreateTask(UStateTreeState* State, const FStateTreeEdit
 
 	StructDesc.DataHandle = Task.InstanceDataHandle;
 
+	if (!Task.bTaskEnabled)
+	{
+		Log.Reportf(EMessageSeverity::Info, StructDesc, TEXT("Task is disabled and will have no effect."));
+	}
+
 	if (!CompileAndValidateNode(State, StructDesc, Node,  InstanceDataView))
 	{
 		return false;
diff --git a/Engine/Plugins/Runtime/StateTree/Source/StateTreeEditorModule/Public/StateTreeCompiler.h b/Engine/Plugins/Runtime/StateTree/Source/StateTreeEditorModule/Public/StateTreeCompiler.h
--- a/Engine/Plugins/Runtime/StateTree/Source/StateTreeEditorModule/Public/StateTreeCompiler.h
+++ b/Engine/Plugins/Runtime/StateTree/Source/StateTreeEditorModule/Public/StateTreeCompiler.h
@@ -62,7 +62,7 @@ private:
 	bool CreateConsideration(UStateTreeState& State, const FString& StatePath, const FStateTreeEditorNode& ConsiderationNode, const EStateTreeExpressionOperand Operand, const int8 DeltaIndent);
 
 	template<class T = FStateTreeNodeBase>
-	T* CreateNodeWithSharedInstanceData(UStateTreeState* State, const FString& StatePath, const FStateTreeEditorNode& EditorNode, EStateTreeBindableStructSource StructSource);
+	T* CreateNodeWithSharedInstanceData(UStateTreeState* State, const FStateTreeEditorNode& EditorNode, FStateTreeBindableStructDesc& StructDesc);
 
 	bool CreateTask(UStateTreeState* State, const FStateTreeEditorNode& TaskNode, const FStateTreeDataHandle TaskDataHandle);
 	bool CreateEvaluator(const FStateTreeEditorNode& EvalNode, const FStateTreeDataHandle EvalDataHandle);
