[StateTreeDebugger]
- added missing information in transition source struct to refer to the proper asset in case of linked trees (this fixes an ensure for mismatching transition information when reading the trace)
- fixed scoped state phase display (missing case for Push/Pop)
- fixed broken icons in frame details view for log events (warning, error and fatal)

diff --git a/Engine/Plugins/Runtime/StateTree/Source/StateTreeEditorModule/Private/Debugger/SStateTreeDebuggerView.cpp b/Engine/Plugins/Runtime/StateTree/Source/StateTreeEditorModule/Private/Debugger/SStateTreeDebuggerView.cpp
--- a/Engine/Plugins/Runtime/StateTree/Source/StateTreeEditorModule/Private/Debugger/SStateTreeDebuggerView.cpp
+++ b/Engine/Plugins/Runtime/StateTree/Source/StateTreeEditorModule/Private/Debugger/SStateTreeDebuggerView.cpp
@@ -902,12 +902,14 @@ void SStateTreeDebuggerView::OnDebuggerScrubStateChanged(const UE::StateTreeDebu
 		if (const FStateTreeTraceStateEvent* StateEvent = Event.TryGet<FStateTreeTraceStateEvent>())
 		{
 			if (StateEvent->EventType == EStateTreeTraceEventType::OnEntering
-			 	|| StateEvent->EventType == EStateTreeTraceEventType::OnExiting)
+			 	|| StateEvent->EventType == EStateTreeTraceEventType::OnExiting
+				|| StateEvent->EventType == EStateTreeTraceEventType::Push)
 			{
 				bShouldAddToScopeStack = true;
 			}
 			else if (StateEvent->EventType == EStateTreeTraceEventType::OnEntered
-				|| StateEvent->EventType == EStateTreeTraceEventType::OnExited)
+				|| StateEvent->EventType == EStateTreeTraceEventType::OnExited
+				|| StateEvent->EventType == EStateTreeTraceEventType::Pop)
 			{
 				bShouldPopScopeStack = true;
 			}
diff --git a/Engine/Plugins/Runtime/StateTree/Source/StateTreeEditorModule/Private/StateTreeEditorStyle.cpp b/Engine/Plugins/Runtime/StateTree/Source/StateTreeEditorModule/Private/StateTreeEditorStyle.cpp
--- a/Engine/Plugins/Runtime/StateTree/Source/StateTreeEditorModule/Private/StateTreeEditorStyle.cpp
+++ b/Engine/Plugins/Runtime/StateTree/Source/StateTreeEditorModule/Private/StateTreeEditorStyle.cpp
@@ -322,6 +322,9 @@ FStateTreeEditorStyle::FStateTreeEditorStyle()
 
 		Set("StateTreeEditor.Debugger.Unset", new CORE_IMAGE_BRUSH_SVG("Starship/Common/help", CoreStyleConstants::Icon16x16, FStyleColors::AccentBlack));
 
+		Set("StateTreeEditor.Debugger.Log.Warning", new CORE_IMAGE_BRUSH_SVG("Starship/Common/alert-circle", CoreStyleConstants::Icon16x16, FStyleColors::AccentYellow));
+		Set("StateTreeEditor.Debugger.Log.Error", new CORE_IMAGE_BRUSH_SVG("Starship/Common/x-circle", CoreStyleConstants::Icon16x16, FStyleColors::AccentRed));
+
 		// Common Node Icons
 		Set("Node.EnableDisable", new CORE_IMAGE_BRUSH_SVG("Starship/Common/check-circle", CoreStyleConstants::Icon16x16));
 		Set("Node.Time", new CORE_IMAGE_BRUSH_SVG("Starship/Common/Recent", CoreStyleConstants::Icon16x16));
@@ -403,9 +406,6 @@ FStateTreeEditorStyle::FStateTreeEditorStyle()
 		Set("StateTreeEditor.Transition.Condition", new IMAGE_BRUSH_SVG("Icons/State_Conditions", CoreStyleConstants::Icon16x16, FStyleColors::AccentGray));
 		Set("StateTreeEditor.Debugger.Condition.OnTransition", new IMAGE_BRUSH_SVG("Icons/State_Conditions", CoreStyleConstants::Icon16x16, FStyleColors::AccentGray));
 
-		Set("StateTreeEditor.Debugger.Log.Warning", new IMAGE_BRUSH_SVG("Icons/Alert", CoreStyleConstants::Icon16x16, FStyleColors::AccentYellow));
-		Set("StateTreeEditor.Debugger.Log.Error", new IMAGE_BRUSH_SVG("Icons/Failure", CoreStyleConstants::Icon16x16, FStyleColors::AccentRed));
-
 		// Common Node Icons
 		Set("Node.Movement", new IMAGE_BRUSH_SVG("Icons/Movement", CoreStyleConstants::Icon16x16));
 		Set("Node.Tag", new IMAGE_BRUSH_SVG("Icons/Tag", CoreStyleConstants::Icon16x16));
diff --git a/Engine/Plugins/Runtime/StateTree/Source/StateTreeModule/Private/Debugger/StateTreeTrace.cpp b/Engine/Plugins/Runtime/StateTree/Source/StateTreeModule/Private/Debugger/StateTreeTrace.cpp
--- a/Engine/Plugins/Runtime/StateTree/Source/StateTreeModule/Private/Debugger/StateTreeTrace.cpp
+++ b/Engine/Plugins/Runtime/StateTree/Source/StateTreeModule/Private/Debugger/StateTreeTrace.cpp
@@ -603,18 +603,18 @@ void OutputInstanceLifetimeEvent(
 	}
 }
 
-void OutputInstanceFrameEvent(
+void OutputInstanceAssetEvent(
 	const FStateTreeInstanceDebugId InstanceId,
-	const FStateTreeExecutionFrame* Frame
+	const UStateTree* StateTree
 	)
 {
-	check(Frame != nullptr);
+	check(StateTree != nullptr);
 
 	if (UE_TRACE_CHANNELEXPR_IS_ENABLED(StateTreeDebugChannel))
 	{
 		TraceBufferedEvents(InstanceId);
 
-		const FStateTreeIndex16 AssetDebugId = FindOrAddDebugIdForAsset(Frame->StateTree.Get());
+		const FStateTreeIndex16 AssetDebugId = FindOrAddDebugIdForAsset(StateTree);
 
 		UE_TRACE_LOG(StateTreeDebugger, InstanceFrameEvent, StateTreeDebugChannel)
 			<< InstanceFrameEvent.Cycle(FPlatformTime::Cycles64())
@@ -626,6 +626,15 @@ void OutputInstanceFrameEvent(
 	// and we don't expect the trace channel to be enabled/disabled during a single execution context update.
 }
 
+void OutputInstanceFrameEvent(
+	const FStateTreeInstanceDebugId InstanceId,
+	const FStateTreeExecutionFrame* Frame
+	)
+{
+	check(Frame != nullptr);
+	OutputInstanceAssetEvent(InstanceId, Frame->StateTree.Get());
+}
+
 void OutputLogEventTrace(
 	const FStateTreeInstanceDebugId InstanceId,
 	ELogVerbosity::Type Verbosity,
@@ -711,6 +720,12 @@ void OutputTransitionEventTrace(
 	const EStateTreeTraceEventType EventType
 	)
 {
+	// Output an instance Frame/Asset event to put the transition data in the proper context when it gets read.
+	if (const UStateTree* StateTree = Source.Asset.Get())
+	{
+		OutputInstanceAssetEvent(InstanceId, StateTree);
+	}
+
 	FBufferArchive Archive;
 	Archive << EventType;
 
diff --git a/Engine/Plugins/Runtime/StateTree/Source/StateTreeModule/Private/Debugger/StateTreeTraceAnalyzer.cpp b/Engine/Plugins/Runtime/StateTree/Source/StateTreeModule/Private/Debugger/StateTreeTraceAnalyzer.cpp
--- a/Engine/Plugins/Runtime/StateTree/Source/StateTreeModule/Private/Debugger/StateTreeTraceAnalyzer.cpp
+++ b/Engine/Plugins/Runtime/StateTree/Source/StateTreeModule/Private/Debugger/StateTreeTraceAnalyzer.cpp
@@ -201,6 +201,7 @@ bool FStateTreeTraceAnalyzer::OnEvent(const uint16 RouteId, EStyle Style, const
 		{
 			const FStateTreeTraceTransitionEvent Event(WorldTime,
 				FStateTreeTransitionSource(
+					nullptr, // we simply reuse the struct to hold the received data, we don't need to provide a StateTree
 					EventData.GetValue<EStateTreeTransitionSourceType>("SourceType"),
 					FStateTreeIndex16(EventData.GetValue<uint16>("TransitionIndex")),
 					FStateTreeStateHandle(EventData.GetValue<uint16>("TargetStateIndex")),
diff --git a/Engine/Plugins/Runtime/StateTree/Source/StateTreeModule/Private/StateTreeExecutionContext.cpp b/Engine/Plugins/Runtime/StateTree/Source/StateTreeModule/Private/StateTreeExecutionContext.cpp
--- a/Engine/Plugins/Runtime/StateTree/Source/StateTreeModule/Private/StateTreeExecutionContext.cpp
+++ b/Engine/Plugins/Runtime/StateTree/Source/StateTreeModule/Private/StateTreeExecutionContext.cpp
@@ -858,7 +858,7 @@ void FStateTreeExecutionContext::RequestTransition(const FStateTreeTransitionReq
 
 		if (RequestTransition(*CurrentlyProcessedFrame, Request.TargetState, Request.Priority, /*TransitionEvent*/nullptr, Request.Fallback))
 		{
-			NextTransitionSource = FStateTreeTransitionSource(EStateTreeTransitionSourceType::ExternalRequest, Request.TargetState, Request.Priority);
+			NextTransitionSource = FStateTreeTransitionSource(CurrentlyProcessedFrame->StateTree, EStateTreeTransitionSourceType::ExternalRequest, Request.TargetState, Request.Priority);
 		}
 	}
 	else
@@ -3297,7 +3297,7 @@ bool FStateTreeExecutionContext::TriggerTransitions()
 		{
 			if (RequestTransition(*CurrentFrame, Request.TargetState, Request.Priority, /*TransitionEvent*/nullptr, Request.Fallback))
 			{
-				NextTransitionSource = FStateTreeTransitionSource(EStateTreeTransitionSourceType::ExternalRequest, Request.TargetState, Request.Priority);
+				NextTransitionSource = FStateTreeTransitionSource(CurrentFrame->StateTree, EStateTreeTransitionSourceType::ExternalRequest, Request.TargetState, Request.Priority);
 			}
 		}
 	}
@@ -3512,13 +3512,13 @@ bool FStateTreeExecutionContext::TriggerTransitions()
 							// Trigger Delayed Transition when the delay has passed.
 							if (RequestTransition(CurrentFrame, Transition.State, Transition.Priority, &DelayedTransition.Transition.CapturedEvent, Transition.Fallback))
 							{
-								// If the transition was successfully requested with a specific event, consume and remove the event, it's been used.  
+								// If the transition was successfully requested with a specific event, consume and remove the event, it's been used.
 								if (DelayedTransition.Transition.CapturedEvent.IsValid() && Transition.bConsumeEventOnSelect)
 								{
 									ConsumeEvent(DelayedTransition.Transition.CapturedEvent);
 								}
 
-								NextTransitionSource = FStateTreeTransitionSource(FStateTreeIndex16(TransitionIndex), Transition.State, Transition.Priority);
+								NextTransitionSource = FStateTreeTransitionSource(CurrentFrame.StateTree, FStateTreeIndex16(TransitionIndex), Transition.State, Transition.Priority);
 								bTriggeredDelayedTransition = true;
 								break;
 							}
@@ -3562,7 +3562,7 @@ bool FStateTreeExecutionContext::TriggerTransitions()
 					bool bPassed = false; 
 					{
 						FCurrentlyProcessedTransitionEventScope TransitionEventScope(*this, TransitionEvent ? TransitionEvent->Get() : nullptr);
-						UE_STATETREE_DEBUG_TRANSITION_EVENT(this, FStateTreeTransitionSource(FStateTreeIndex16(TransitionIndex), Transition.State, Transition.Priority), EStateTreeTraceEventType::OnEvaluating);
+						UE_STATETREE_DEBUG_TRANSITION_EVENT(this, FStateTreeTransitionSource(CurrentFrame.StateTree, FStateTreeIndex16(TransitionIndex), Transition.State, Transition.Priority), EStateTreeTraceEventType::OnEvaluating);
 						UE_STATETREE_DEBUG_SCOPED_PHASE(this, EStateTreeUpdatePhase::TransitionConditions);
 						bPassed = TestAllConditions(CurrentParentFrame, CurrentFrame, Transition.ConditionsBegin, Transition.ConditionsNum);
 					}
@@ -3616,7 +3616,7 @@ bool FStateTreeExecutionContext::TriggerTransitions()
 							}
 						}
 
-						UE_STATETREE_DEBUG_TRANSITION_EVENT(this, FStateTreeTransitionSource(FStateTreeIndex16(TransitionIndex), Transition.State, Transition.Priority), EStateTreeTraceEventType::OnRequesting);
+						UE_STATETREE_DEBUG_TRANSITION_EVENT(this, FStateTreeTransitionSource(CurrentFrame.StateTree, FStateTreeIndex16(TransitionIndex), Transition.State, Transition.Priority), EStateTreeTraceEventType::OnRequesting);
 						if (RequestTransition(CurrentFrame, Transition.State, Transition.Priority, TransitionEvent, Transition.Fallback))
 						{
 							// If the transition was successfully requested with a specific event, consume and remove the event, it's been used.
@@ -3625,7 +3625,7 @@ bool FStateTreeExecutionContext::TriggerTransitions()
 								ConsumeEvent(*TransitionEvent);
 							}
 							
-							NextTransitionSource = FStateTreeTransitionSource(FStateTreeIndex16(TransitionIndex), Transition.State, Transition.Priority);
+							NextTransitionSource = FStateTreeTransitionSource(CurrentFrame.StateTree, FStateTreeIndex16(TransitionIndex), Transition.State, Transition.Priority);
 							break;
 						}
 					}
@@ -3698,7 +3698,7 @@ bool FStateTreeExecutionContext::TriggerTransitions()
 						{
 							bool bPassed = false;
 							{
-								UE_STATETREE_DEBUG_TRANSITION_EVENT(this, FStateTreeTransitionSource(FStateTreeIndex16(TransitionIndex), Transition.State, Transition.Priority), EStateTreeTraceEventType::OnEvaluating);
+								UE_STATETREE_DEBUG_TRANSITION_EVENT(this, FStateTreeTransitionSource(CurrentFrame.StateTree, FStateTreeIndex16(TransitionIndex), Transition.State, Transition.Priority), EStateTreeTraceEventType::OnEvaluating);
 								UE_STATETREE_DEBUG_SCOPED_PHASE(this, EStateTreeUpdatePhase::TransitionConditions);
 								bPassed = TestAllConditions(CurrentParentFrame, CurrentFrame, Transition.ConditionsBegin, Transition.ConditionsNum);
 							}
@@ -3707,10 +3707,10 @@ bool FStateTreeExecutionContext::TriggerTransitions()
 							{
 								// No delay allowed on completion conditions.
 								// No priority on completion transitions, use the priority to signal that state is selected.
-								UE_STATETREE_DEBUG_TRANSITION_EVENT(this, FStateTreeTransitionSource(FStateTreeIndex16(TransitionIndex), Transition.State, Transition.Priority), EStateTreeTraceEventType::OnRequesting);
+								UE_STATETREE_DEBUG_TRANSITION_EVENT(this, FStateTreeTransitionSource(CurrentFrame.StateTree, FStateTreeIndex16(TransitionIndex), Transition.State, Transition.Priority), EStateTreeTraceEventType::OnRequesting);
 								if (RequestTransition(CurrentFrame, Transition.State, EStateTreeTransitionPriority::Normal, /*TransitionEvent*/nullptr, Transition.Fallback))
 								{
-									NextTransitionSource = FStateTreeTransitionSource(FStateTreeIndex16(TransitionIndex), Transition.State, Transition.Priority);
+									NextTransitionSource = FStateTreeTransitionSource(CurrentFrame.StateTree, FStateTreeIndex16(TransitionIndex), Transition.State, Transition.Priority);
 									break;
 								}
 							}
@@ -3745,7 +3745,7 @@ bool FStateTreeExecutionContext::TriggerTransitions()
 				
 			if (RequestTransition(RootFrame, FStateTreeStateHandle::Root, EStateTreeTransitionPriority::Normal))
 			{
-				NextTransitionSource = FStateTreeTransitionSource(EStateTreeTransitionSourceType::Internal, FStateTreeStateHandle::Root, EStateTreeTransitionPriority::Normal);
+				NextTransitionSource = FStateTreeTransitionSource(RootFrame.StateTree, EStateTreeTransitionSourceType::Internal, FStateTreeStateHandle::Root, EStateTreeTransitionPriority::Normal);
 			}
 			else
 			{
@@ -4413,7 +4413,8 @@ bool FStateTreeExecutionContext::SelectStateInternal(
 					{
 						FCurrentlyProcessedTransitionEventScope TransitionEventScope(*this, SelectedStateTransitionEvent ? SelectedStateTransitionEvent->Get() : nullptr);
 
-						UE_STATETREE_DEBUG_TRANSITION_EVENT(this, FStateTreeTransitionSource(FStateTreeIndex16(TransitionIndex), Transition.State, Transition.Priority), EStateTreeTraceEventType::OnEvaluating);
+						UE_STATETREE_DEBUG_TRANSITION_EVENT(this, FStateTreeTransitionSource(CurrentFrame.StateTree, FStateTreeIndex16(TransitionIndex), Transition.State, Transition.Priority), EStateTreeTraceEventType::OnEvaluating);
+
 						UE_STATETREE_DEBUG_SCOPED_PHASE(this, EStateTreeUpdatePhase::TransitionConditions);
 
 						bTransitionConditionsPassed = TestAllConditions(CurrentParentFrame, CurrentFrame, Transition.ConditionsBegin, Transition.ConditionsNum);
diff --git a/Engine/Plugins/Runtime/StateTree/Source/StateTreeModule/Private/StateTreeExecutionTypes.cpp b/Engine/Plugins/Runtime/StateTree/Source/StateTreeModule/Private/StateTreeExecutionTypes.cpp
--- a/Engine/Plugins/Runtime/StateTree/Source/StateTreeModule/Private/StateTreeExecutionTypes.cpp
+++ b/Engine/Plugins/Runtime/StateTree/Source/StateTreeModule/Private/StateTreeExecutionTypes.cpp
@@ -1,6 +1,7 @@
 // Copyright Epic Games, Inc. All Rights Reserved.
 
 #include "StateTreeExecutionTypes.h"
+#include "StateTree.h"
 
 const FStateTreeExternalDataHandle FStateTreeExternalDataHandle::Invalid = FStateTreeExternalDataHandle();
 
@@ -8,6 +9,27 @@ const FStateTreeExternalDataHandle FStateTreeExternalDataHandle::Invalid = FStat
 const FStateTreeInstanceDebugId FStateTreeInstanceDebugId::Invalid = FStateTreeInstanceDebugId();
 #endif // WITH_STATETREE_TRACE
 
+//----------------------------------------------------------------------//
+// FStateTreeTransitionSource
+//----------------------------------------------------------------------//
+FStateTreeTransitionSource::FStateTreeTransitionSource(
+	const UStateTree* StateTree,
+	const EStateTreeTransitionSourceType SourceType,
+	const FStateTreeIndex16 TransitionIndex,
+	const FStateTreeStateHandle TargetState,
+	const EStateTreeTransitionPriority Priority
+	)
+	: Asset(StateTree)
+	, SourceType(SourceType)
+	, TransitionIndex(TransitionIndex)
+	, TargetState(TargetState)
+	, Priority(Priority)
+{
+}
+
+//----------------------------------------------------------------------//
+// FStateTreeTransitionResult
+//----------------------------------------------------------------------//
 FStateTreeTransitionResult::FStateTreeTransitionResult(const FRecordedStateTreeTransitionResult& RecordedTransition)
 {
 	for (int32 RecordedFrameIndex = 0; RecordedFrameIndex < RecordedTransition.NextActiveFrames.Num(); RecordedFrameIndex++)
@@ -36,6 +58,9 @@ FStateTreeTransitionResult::FStateTreeTransitionResult(const FRecordedStateTreeT
 	SourceRootState = RecordedTransition.SourceRootState;
 }
 
+//----------------------------------------------------------------------//
+// FRecordedStateTreeTransitionResult
+//----------------------------------------------------------------------//
 FRecordedStateTreeTransitionResult::FRecordedStateTreeTransitionResult(const FStateTreeTransitionResult& Transition)
 {
 	check(Transition.NextActiveFrames.Num() == Transition.NextActiveFrameEvents.Num());
@@ -65,6 +90,9 @@ FRecordedStateTreeTransitionResult::FRecordedStateTreeTransitionResult(const FSt
 	SourceRootState = Transition.SourceRootState;
 }
 
+//----------------------------------------------------------------------//
+// FStateTreeExecutionFrame
+//----------------------------------------------------------------------//
 FStateTreeExecutionFrame::FStateTreeExecutionFrame(const FRecordedStateTreeExecutionFrame& RecordedExecutionFrame)
 {
 	StateTree = RecordedExecutionFrame.StateTree;
@@ -73,6 +101,9 @@ FStateTreeExecutionFrame::FStateTreeExecutionFrame(const FRecordedStateTreeExecu
 	bIsGlobalFrame = RecordedExecutionFrame.bIsGlobalFrame;
 }
 
+//----------------------------------------------------------------------//
+// FRecordedStateTreeExecutionFrame
+//----------------------------------------------------------------------//
 FRecordedStateTreeExecutionFrame::FRecordedStateTreeExecutionFrame(const FStateTreeExecutionFrame& ExecutionFrame)
 {
 	StateTree = ExecutionFrame.StateTree;
@@ -81,6 +112,9 @@ FRecordedStateTreeExecutionFrame::FRecordedStateTreeExecutionFrame(const FStateT
 	bIsGlobalFrame = ExecutionFrame.bIsGlobalFrame;
 }
 
+//----------------------------------------------------------------------//
+// FStateTreeScheduledTick
+//----------------------------------------------------------------------//
 FStateTreeScheduledTick FStateTreeScheduledTick::MakeSleep()
 {
 	return FStateTreeScheduledTick(UE_FLOAT_NON_FRACTIONAL);
diff --git a/Engine/Plugins/Runtime/StateTree/Source/StateTreeModule/Public/Debugger/StateTreeTrace.h b/Engine/Plugins/Runtime/StateTree/Source/StateTreeModule/Public/Debugger/StateTreeTrace.h
--- a/Engine/Plugins/Runtime/StateTree/Source/StateTreeModule/Public/Debugger/StateTreeTrace.h
+++ b/Engine/Plugins/Runtime/StateTree/Source/StateTreeModule/Public/Debugger/StateTreeTrace.h
@@ -31,6 +31,7 @@ namespace UE::StateTreeTrace
 	void OutputPhaseScopeEvent(FStateTreeInstanceDebugId InstanceId, EStateTreeUpdatePhase Phase, EStateTreeTraceEventType EventType, FStateTreeStateHandle StateHandle);
 	void OutputAssetDebugIdEvent(const UStateTree* StateTree, FStateTreeIndex16 AssetDebugId);
 	void OutputInstanceLifetimeEvent(FStateTreeInstanceDebugId InstanceId, const UStateTree* StateTree, const TCHAR* InstanceName, EStateTreeTraceEventType EventType);
+	void OutputInstanceAssetEvent(FStateTreeInstanceDebugId InstanceId, const UStateTree* StateTree);
 	void OutputInstanceFrameEvent(FStateTreeInstanceDebugId InstanceId, const FStateTreeExecutionFrame* Frame);
 	void OutputLogEventTrace(FStateTreeInstanceDebugId InstanceId, ELogVerbosity::Type Verbosity, const TCHAR* Fmt, ...);
 	void OutputStateEventTrace(FStateTreeInstanceDebugId InstanceId, FStateTreeStateHandle StateHandle, EStateTreeTraceEventType EventType);
@@ -91,15 +92,15 @@ namespace UE::StateTreeTrace
 
 #else //WITH_STATETREE_TRACE
 
-#define TRACE_STATETREE_INSTANCE_EVENT(InstanceID, StateTree, InstanceName, EventType)
-#define TRACE_STATETREE_INSTANCE_FRAME_EVENT(InstanceID, Frame)
-#define TRACE_STATETREE_PHASE_EVENT(InstanceID, Phase, EventType, StateHandle)
-#define TRACE_STATETREE_LOG_EVENT(InstanceId, TraceVerbosity, Format, ...)
-#define TRACE_STATETREE_STATE_EVENT(InstanceId, StateHandle, EventType)
-#define TRACE_STATETREE_TASK_EVENT(InstanceId, TaskIdx, DataView, EventType, Status)
-#define TRACE_STATETREE_EVALUATOR_EVENT(InstanceId, EvaluatorIdx, DataView, EventType)
-#define TRACE_STATETREE_CONDITION_EVENT(InstanceId, ConditionIdx, DataView, EventType)
-#define TRACE_STATETREE_TRANSITION_EVENT(InstanceId, TransitionIdx, EventType)
-#define TRACE_STATETREE_ACTIVE_STATES_EVENT(InstanceId, ActiveFrames)
+#define TRACE_STATETREE_INSTANCE_EVENT(...)
+#define TRACE_STATETREE_INSTANCE_FRAME_EVENT(...)
+#define TRACE_STATETREE_PHASE_EVENT(...)
+#define TRACE_STATETREE_LOG_EVENT(...)
+#define TRACE_STATETREE_STATE_EVENT(...)
+#define TRACE_STATETREE_TASK_EVENT(...)
+#define TRACE_STATETREE_EVALUATOR_EVENT(...)
+#define TRACE_STATETREE_CONDITION_EVENT(...)
+#define TRACE_STATETREE_TRANSITION_EVENT(...)
+#define TRACE_STATETREE_ACTIVE_STATES_EVENT(...)
 
 #endif // WITH_STATETREE_TRACE
diff --git a/Engine/Plugins/Runtime/StateTree/Source/StateTreeModule/Public/StateTreeExecutionTypes.h b/Engine/Plugins/Runtime/StateTree/Source/StateTreeModule/Public/StateTreeExecutionTypes.h
--- a/Engine/Plugins/Runtime/StateTree/Source/StateTreeModule/Public/StateTreeExecutionTypes.h
+++ b/Engine/Plugins/Runtime/StateTree/Source/StateTreeModule/Public/StateTreeExecutionTypes.h
@@ -459,21 +459,33 @@ struct STATETREEMODULE_API FStateTreeTransitionSource
 
 	FStateTreeTransitionSource() = default;
 
+	UE_DEPRECATED(5.6, "Use the version that requires a pointer to the StateTree asset.")
 	explicit FStateTreeTransitionSource(const EStateTreeTransitionSourceType SourceType, const FStateTreeIndex16 TransitionIndex, const FStateTreeStateHandle TargetState, const EStateTreeTransitionPriority Priority)
-	: SourceType(SourceType)
-	, TransitionIndex(TransitionIndex)
-	, TargetState(TargetState)
-	, Priority(Priority)
+	: FStateTreeTransitionSource(nullptr, SourceType, TransitionIndex, TargetState, Priority)
 	{
 	}
 
+	UE_DEPRECATED(5.6, "Use the version that requires a pointer to the StateTree asset.")
 	explicit FStateTreeTransitionSource(const FStateTreeIndex16 TransitionIndex, const FStateTreeStateHandle TargetState, const EStateTreeTransitionPriority Priority)
-	: FStateTreeTransitionSource(EStateTreeTransitionSourceType::Asset, TransitionIndex, TargetState, Priority)
+	: FStateTreeTransitionSource(nullptr, EStateTreeTransitionSourceType::Asset, TransitionIndex, TargetState, Priority)
 	{
 	}
 
+	UE_DEPRECATED(5.6, "Use the version that requires a pointer to the StateTree asset.")
 	explicit FStateTreeTransitionSource(const EStateTreeTransitionSourceType SourceType, const FStateTreeStateHandle TargetState, const EStateTreeTransitionPriority Priority)
-	: FStateTreeTransitionSource(SourceType, FStateTreeIndex16::Invalid, TargetState, Priority)
+	: FStateTreeTransitionSource(nullptr, SourceType, FStateTreeIndex16::Invalid, TargetState, Priority)
+	{
+	}
+
+	explicit FStateTreeTransitionSource(const UStateTree* StateTree, const EStateTreeTransitionSourceType SourceType, const FStateTreeIndex16 TransitionIndex, const FStateTreeStateHandle TargetState, const EStateTreeTransitionPriority Priority);
+
+	explicit FStateTreeTransitionSource(const UStateTree* StateTree, const FStateTreeIndex16 TransitionIndex, const FStateTreeStateHandle TargetState, const EStateTreeTransitionPriority Priority)
+	: FStateTreeTransitionSource(StateTree, EStateTreeTransitionSourceType::Asset, TransitionIndex, TargetState, Priority)
+	{
+	}
+
+	explicit FStateTreeTransitionSource(const UStateTree* StateTree, const EStateTreeTransitionSourceType SourceType, const FStateTreeStateHandle TargetState, const EStateTreeTransitionPriority Priority)
+	: FStateTreeTransitionSource(StateTree, SourceType, FStateTreeIndex16::Invalid, TargetState, Priority)
 	{
 	}
 
@@ -482,6 +494,9 @@ struct STATETREEMODULE_API FStateTreeTransitionSource
 		*this = {};
 	}
 
+	/** The StateTree asset owning the transition and state. */
+	TWeakObjectPtr<const UStateTree> Asset;
+
 	/** Describes where the transition originated. */
 	EStateTreeTransitionSourceType SourceType = EStateTreeTransitionSourceType::Unset;
 
