StateTree: Add custom tick rate in seconds to states. If an active state has a custom tick rate then  the state cannot sleep, all the other states (children or parents) will also tick at that rate. If more than one active states has a custom tick rate then the smallest custom tick rate wins.

diff --git a/Engine/Plugins/Runtime/StateTree/Source/StateTreeEditorModule/Private/Customizations/StateTreeStateDetails.cpp b/Engine/Plugins/Runtime/StateTree/Source/StateTreeEditorModule/Private/Customizations/StateTreeStateDetails.cpp
--- a/Engine/Plugins/Runtime/StateTree/Source/StateTreeEditorModule/Private/Customizations/StateTreeStateDetails.cpp
+++ b/Engine/Plugins/Runtime/StateTree/Source/StateTreeEditorModule/Private/Customizations/StateTreeStateDetails.cpp
@@ -80,6 +80,7 @@ void FStateTreeStateDetails::CustomizeDetails(IDetailLayoutBuilder& DetailBuilde
 	const TSharedPtr<IPropertyHandle> TypeProperty = DetailBuilder.GetProperty(GET_MEMBER_NAME_CHECKED(UStateTreeState, Type));
 	const TSharedPtr<IPropertyHandle> LinkedSubtreeProperty = DetailBuilder.GetProperty(GET_MEMBER_NAME_CHECKED(UStateTreeState, LinkedSubtree));
 	const TSharedPtr<IPropertyHandle> LinkedAssetProperty = DetailBuilder.GetProperty(GET_MEMBER_NAME_CHECKED(UStateTreeState, LinkedAsset));
+	const TSharedPtr<IPropertyHandle> CustomTickRateProperty = DetailBuilder.GetProperty(GET_MEMBER_NAME_CHECKED(UStateTreeState, CustomTickRate));
 	const TSharedPtr<IPropertyHandle> ParametersProperty = DetailBuilder.GetProperty(GET_MEMBER_NAME_CHECKED(UStateTreeState, Parameters));
 	const TSharedPtr<IPropertyHandle> SelectionBehaviorProperty = DetailBuilder.GetProperty(GET_MEMBER_NAME_CHECKED(UStateTreeState, SelectionBehavior));
 	const TSharedPtr<IPropertyHandle> RequiredEventToEnterProperty = DetailBuilder.GetProperty(GET_MEMBER_NAME_CHECKED(UStateTreeState, RequiredEventToEnter));
@@ -159,6 +160,13 @@ void FStateTreeStateDetails::CustomizeDetails(IDetailLayoutBuilder& DetailBuilde
 	// Color
 	ColorRefProperty->MarkHiddenByCustomization();
 	StateCategory.AddProperty(ColorRefProperty);
+	
+	// Custom Tick Rate
+	CustomTickRateProperty->MarkHiddenByCustomization();
+	if (Schema && Schema->IsScheduledTickAllowed())
+	{
+		StateCategory.AddProperty(CustomTickRateProperty);
+	}
 
 	// Type
 	TypeProperty->MarkHiddenByCustomization();
diff --git a/Engine/Plugins/Runtime/StateTree/Source/StateTreeEditorModule/Private/SStateTreeViewRow.cpp b/Engine/Plugins/Runtime/StateTree/Source/StateTreeEditorModule/Private/SStateTreeViewRow.cpp
--- a/Engine/Plugins/Runtime/StateTree/Source/StateTreeEditorModule/Private/SStateTreeViewRow.cpp
+++ b/Engine/Plugins/Runtime/StateTree/Source/StateTreeEditorModule/Private/SStateTreeViewRow.cpp
@@ -898,10 +898,10 @@ void SStateTreeViewRow::MakeFlagsWidget()
 	{
 		if (const FCompactStateTreeState* RuntimeState = StateTree->GetStateFromHandle(StateTree->GetStateHandleFromId(State->ID)))
 		{
-			if (RuntimeState->bHasTickTasks || RuntimeState->bHasTickTasksOnlyOnEvents || RuntimeState->bHasTransitionTasks)
+			if (RuntimeState->bHasTickTasks || RuntimeState->bHasTickTasksOnlyOnEvents || RuntimeState->bHasTransitionTasks || RuntimeState->bHasCustomTickRate)
 			{
 				TSharedRef<SHorizontalBox> Box = SNew(SHorizontalBox);
-				if (RuntimeState->bHasTickTasks)
+				if (RuntimeState->bHasCustomTickRate)
 				{
 					Box->AddSlot()
 					.AutoWidth()
@@ -910,11 +910,22 @@ void SStateTreeViewRow::MakeFlagsWidget()
 						SNew(SImage)
 						.Image(FStateTreeEditorStyle::Get().GetBrush("StateTreeEditor.Flags.Tick"))
 						.ColorAndOpacity(IconTint)
-						.ToolTipText(LOCTEXT("NodeTick", "The node ticks at runtime."))
+						.ToolTipText(LOCTEXT("StateCustomTick", "The state has a custom tick rate."))
 					];
 				}
-
-				if (!RuntimeState->bHasTickTasks && RuntimeState->bHasTickTasksOnlyOnEvents)
+				else if (RuntimeState->bHasTickTasks)
+				{
+					Box->AddSlot()
+					.AutoWidth()
+					.VAlign(VAlign_Center)
+					[
+						SNew(SImage)
+						.Image(FStateTreeEditorStyle::Get().GetBrush("StateTreeEditor.Flags.Tick"))
+						.ColorAndOpacity(IconTint)
+						.ToolTipText(LOCTEXT("StateNodeTick", "The state contains at least one task that ticks at runtime."))
+					];
+				}
+				else if (RuntimeState->bHasTickTasksOnlyOnEvents)
 				{
 					Box->AddSlot()
 					.AutoWidth()
@@ -923,10 +934,10 @@ void SStateTreeViewRow::MakeFlagsWidget()
 						SNew(SImage)
 						.Image(FStateTreeEditorStyle::Get().GetBrush("StateTreeEditor.Flags.TickOnEvent"))
 						.ColorAndOpacity(IconTint)
-						.ToolTipText(LOCTEXT("NodeTickEvent", "The node ticks on event at runtime."))
+						.ToolTipText(LOCTEXT("StateNodeTickEvent", "The state contains at least one task that ticks at runtime when there's an event."))
 					];
 				}
-
+				
 				if (RuntimeState->bHasTransitionTasks)
 				{
 					Box->AddSlot()
@@ -936,7 +947,7 @@ void SStateTreeViewRow::MakeFlagsWidget()
 						SNew(SImage)
 						.Image(FStateTreeEditorStyle::Get().GetBrush("StateTreeEditor.Transitions"))
 						.ColorAndOpacity(IconTint)
-						.ToolTipText(LOCTEXT("NodeTickTransition", "The node ticks on to trigger transition."))
+						.ToolTipText(LOCTEXT("StateNodeTickTransition", "The state contains at least one task that ticks at runtime when evaluating transitions."))
 					];
 				}
 
diff --git a/Engine/Plugins/Runtime/StateTree/Source/StateTreeEditorModule/Private/StateTreeCompiler.cpp b/Engine/Plugins/Runtime/StateTree/Source/StateTreeEditorModule/Private/StateTreeCompiler.cpp
--- a/Engine/Plugins/Runtime/StateTree/Source/StateTreeEditorModule/Private/StateTreeCompiler.cpp
+++ b/Engine/Plugins/Runtime/StateTree/Source/StateTreeEditorModule/Private/StateTreeCompiler.cpp
@@ -449,6 +449,13 @@ bool FStateTreeCompiler::CreateStateRecursive(UStateTreeState& State, const FSta
 	CompactState.bCheckPrerequisitesWhenActivatingChildDirectly = State.bCheckPrerequisitesWhenActivatingChildDirectly;
 	CompactState.Weight = State.Weight;
 
+	CompactState.bHasCustomTickRate = State.bHasCustomTickRate && Schema->IsScheduledTickAllowed();
+	CompactState.CustomTickRate = FMath::Max(State.CustomTickRate, 0.0f);
+	if (CompactState.bHasCustomTickRate && State.CustomTickRate < 0.0f)
+	{
+		Log.Reportf(EMessageSeverity::Warning, TEXT("The custom tick rate has to be greater than or equal to 0."));
+	}
+
 	CompactState.Type = State.Type;
 	CompactState.SelectionBehavior = State.SelectionBehavior;
 
diff --git a/Engine/Plugins/Runtime/StateTree/Source/StateTreeEditorModule/Public/StateTreeState.h b/Engine/Plugins/Runtime/StateTree/Source/StateTreeEditorModule/Public/StateTreeState.h
--- a/Engine/Plugins/Runtime/StateTree/Source/StateTreeEditorModule/Public/StateTreeState.h
+++ b/Engine/Plugins/Runtime/StateTree/Source/StateTreeEditorModule/Public/StateTreeState.h
@@ -352,6 +352,19 @@ public:
 	UPROPERTY(EditDefaultsOnly, Category = "State")
 	TObjectPtr<UStateTree> LinkedAsset = nullptr;
 
+	/** 
+	 * Tick rate in seconds the state tasks and transitions should tick.
+	 * If set the state cannot sleep.
+	 * If set all the other states (children or parents) will also tick at that rate.
+	 * If more than one active states has a custom tick rate then the smallest custom tick rate wins.
+	 */
+	UPROPERTY(EditDefaultsOnly, Category = "State", meta = (EditCondition = "bHasCustomTickRate", ClampMin = 0.0f))
+	float CustomTickRate = 0.0f;
+
+	/** Activate the CustomTickRate. */
+	UPROPERTY(EditDefaultsOnly, Category = "Enter Conditions", meta=(InlineEditConditionToggle))
+	bool bHasCustomTickRate = false;
+
 	/** Parameters of this state. If the state is linked to another state or asset, the parameters are for the linked state. */
 	UPROPERTY(EditDefaultsOnly, Category = "State")
 	FStateTreeStateParameters Parameters;
diff --git a/Engine/Plugins/Runtime/StateTree/Source/StateTreeModule/Private/StateTreeExecutionContext.cpp b/Engine/Plugins/Runtime/StateTree/Source/StateTreeModule/Private/StateTreeExecutionContext.cpp
--- a/Engine/Plugins/Runtime/StateTree/Source/StateTreeModule/Private/StateTreeExecutionContext.cpp
+++ b/Engine/Plugins/Runtime/StateTree/Source/StateTreeModule/Private/StateTreeExecutionContext.cpp
@@ -99,37 +99,55 @@ FStateTreeScheduledTick FStateTreeMinimalExecutionContext::GetNextScheduledTick(
 	const FStateTreeEventQueue& EventQueue = Storage.GetEventQueue();
 	const bool bHasEvents = EventQueue.HasEvents();
 
-	// Do we have a state that require a tick or is waiting for an event.
-	for (int32 FrameIndex = 0; FrameIndex < Exec.ActiveFrames.Num(); ++FrameIndex)
+	// Do we have a state that requires a tick or is waiting for an event.
 	{
-		const FStateTreeExecutionFrame& CurrentFrame = Exec.ActiveFrames[FrameIndex];
-		const UStateTree* CurrentStateTree = CurrentFrame.StateTree;
-
-		// Test global tasks.
-		if (CurrentStateTree->ShouldTickGlobalTasks(bHasEvents))
+		TOptional<float> StateCustomTickRate;
+		bool bHasTaskWithEveryFramesTick = false;
+		for (int32 FrameIndex = 0; FrameIndex < Exec.ActiveFrames.Num(); ++FrameIndex)
 		{
-			return FStateTreeScheduledTick::MakeEveryFrames();
-		}
+			const FStateTreeExecutionFrame& CurrentFrame = Exec.ActiveFrames[FrameIndex];
+			const UStateTree* CurrentStateTree = CurrentFrame.StateTree;
 
-		// Test active states tasks.
-		for (int32 StateIndex = 0; StateIndex < CurrentFrame.ActiveStates.Num(); ++StateIndex)
-		{
-			const FStateTreeStateHandle CurrentHandle = CurrentFrame.ActiveStates[StateIndex];
-			const FCompactStateTreeState& State = CurrentStateTree->States[CurrentHandle.Index];
-			if (State.bEnabled)
+			// Test global tasks.
+			if (CurrentStateTree->ShouldTickGlobalTasks(bHasEvents))
+			{
+				bHasTaskWithEveryFramesTick = true;
+			}
+
+			// Test active states tasks.
+			for (int32 StateIndex = 0; StateIndex < CurrentFrame.ActiveStates.Num(); ++StateIndex)
 			{
-				if (State.ShouldTickTasks(bHasEvents)
-					|| State.bHasTransitionTasks
-					|| State.ShouldTickTransitions(bHasEvents)
-					)
+				const FStateTreeStateHandle CurrentHandle = CurrentFrame.ActiveStates[StateIndex];
+				const FCompactStateTreeState& State = CurrentStateTree->States[CurrentHandle.Index];
+				if (State.bEnabled)
 				{
-					// todo: ShouldTickTransitions has onevent or ontick. both can be already triggered and we are waiting for the delay
-					// todo: find out if we need EveryFrames or NextFrame
-					// todo: ask the task->GetScheduledTick
-					return FStateTreeScheduledTick::MakeEveryFrames();
+					if (State.bHasCustomTickRate)
+					{
+						StateCustomTickRate = StateCustomTickRate.IsSet() ? FMath::Min(StateCustomTickRate.GetValue(), State.CustomTickRate) : State.CustomTickRate;
+					}
+					else if (!StateCustomTickRate.IsSet())
+					{
+						if (State.ShouldTickTasks(bHasEvents)
+							|| State.bHasTransitionTasks
+							|| State.ShouldTickTransitions(bHasEvents)
+							)
+						{
+							// todo: ShouldTickTransitions has onevent or ontick. both can be already triggered and we are waiting for the delay
+							// todo: find out if we need EveryFrames or NextFrame
+							bHasTaskWithEveryFramesTick = true;
+						}
+					}
 				}
 			}
 		}
+		if (StateCustomTickRate.IsSet())
+		{
+			return FStateTreeScheduledTick::MakeCustomTickRate(StateCustomTickRate.GetValue());
+		}
+		if (bHasTaskWithEveryFramesTick)
+		{
+			return FStateTreeScheduledTick::MakeEveryFrames();
+		}
 	}
 
 	// Transitions
diff --git a/Engine/Plugins/Runtime/StateTree/Source/StateTreeModule/Public/StateTreeTypes.h b/Engine/Plugins/Runtime/StateTree/Source/StateTreeModule/Public/StateTreeTypes.h
--- a/Engine/Plugins/Runtime/StateTree/Source/StateTreeModule/Public/StateTreeTypes.h
+++ b/Engine/Plugins/Runtime/StateTree/Source/StateTreeModule/Public/StateTreeTypes.h
@@ -755,6 +755,14 @@ struct STATETREEMODULE_API FCompactStateTreeState
 	/** Weight used to scale the normalized final utility score for this state */
 	UPROPERTY()
 	float Weight = 1.f;
+	
+	/**
+	 * Tick rate in seconds the state tasks and transitions should tick
+	 * If set the state cannot sleep.
+	 * Must be greater than or equal to 0.
+	 */
+	UPROPERTY()
+	float CustomTickRate = 0.f;
 
 	/** Number of enter conditions */
 	UPROPERTY()
@@ -843,6 +851,10 @@ struct STATETREEMODULE_API FCompactStateTreeState
 	/** If set to true, the required event is consumed (later state selection cannot react to it) if state selection can be made. */
 	UPROPERTY()
 	uint8 bConsumeEventOnSelect : 1 = true;
+
+	/** If set to true, the state has a custom tick rate. */
+	UPROPERTY()
+	uint8 bHasCustomTickRate : 1 = false;
 };
 
 USTRUCT()
