StateTree: Do not force a tick when there's a linked tree in the active states. The copy binding will occur when the state needs ticking.
Remove TaskTicked test runtime test because a task may fail and the following task would not tick.

diff --git a/Engine/Plugins/Runtime/StateTree/Source/StateTreeModule/Private/StateTree.cpp b/Engine/Plugins/Runtime/StateTree/Source/StateTreeModule/Private/StateTree.cpp
--- a/Engine/Plugins/Runtime/StateTree/Source/StateTreeModule/Private/StateTree.cpp
+++ b/Engine/Plugins/Runtime/StateTree/Source/StateTreeModule/Private/StateTree.cpp
@@ -678,13 +678,9 @@ void UStateTree::UpdateRuntimeFlags()
 	for (FCompactStateTreeState& State : States)
 	{
 		// Update the state task flags.
-		
-		// Tick Linked and LinkedAsset because they CopyBatchOnActiveInstances when they tick. That info might be important when doing transaction.
-		const bool bShouldTickByDefault = (State.Type == EStateTreeStateType::Linked || State.Type == EStateTreeStateType::LinkedAsset);
-		State.bHasTickTasks = bShouldTickByDefault;
-		State.bHasTickTasksOnlyOnEvents = bShouldTickByDefault;
+		State.bHasTickTasks = false;
+		State.bHasTickTasksOnlyOnEvents = false;
 		State.bHasTransitionTasks = false;
-
 		for (int32 TaskIndex = State.TasksBegin; TaskIndex < (State.TasksBegin + State.TasksNum); ++TaskIndex)
 		{
 			const FStateTreeTaskBase& Task = Nodes[TaskIndex].Get<const FStateTreeTaskBase>();
diff --git a/Engine/Plugins/Runtime/StateTree/Source/StateTreeModule/Private/StateTreeExecutionContext.cpp b/Engine/Plugins/Runtime/StateTree/Source/StateTreeModule/Private/StateTreeExecutionContext.cpp
--- a/Engine/Plugins/Runtime/StateTree/Source/StateTreeModule/Private/StateTreeExecutionContext.cpp
+++ b/Engine/Plugins/Runtime/StateTree/Source/StateTreeModule/Private/StateTreeExecutionContext.cpp
@@ -2359,7 +2359,6 @@ EStateTreeRunStatus FStateTreeExecutionContext::TickEvaluatorsAndGlobalTasks(con
 
 				if (CurrentStateTree->ShouldTickGlobalTasks(bHasEvents))
 				{
-					bool bTaskTicked = false;
 					for (int32 TaskIndex = CurrentStateTree->GlobalTasksBegin; TaskIndex < (CurrentStateTree->GlobalTasksBegin + CurrentStateTree->GlobalTasksNum); TaskIndex++)
 					{
 						const FStateTreeTaskBase& Task = CurrentStateTree->Nodes[TaskIndex].Get<const FStateTreeTaskBase>();
@@ -2396,7 +2395,6 @@ EStateTreeRunStatus FStateTreeExecutionContext::TickEvaluatorsAndGlobalTasks(con
 							CSV_SCOPED_TIMING_STAT_EXCLUSIVE(StateTree_Task_Tick);
 
 							TaskResult = Task.Tick(*this, DeltaTime);
-							bTaskTicked = true;
 						}
 
 						UE_STATETREE_DEBUG_TASK_EVENT(this, TaskIndex, TaskInstanceView,
@@ -2414,8 +2412,6 @@ EStateTreeRunStatus FStateTreeExecutionContext::TickEvaluatorsAndGlobalTasks(con
 							bShouldTickTasks = false;
 						}
 					}
-
-					ensureMsgf(bTaskTicked, TEXT("ShouldTickGlobalTasks is true but no tasks were tick for the StateTree `%s`."), *CurrentStateTree->GetPathName());
 				}
 			}
 		}
@@ -2828,22 +2824,17 @@ EStateTreeRunStatus FStateTreeExecutionContext::TickTasks(const float DeltaTime)
 			STATETREE_CLOG(State.TasksNum > 0, VeryVerbose, TEXT("%*sState '%s'"), (FrameIndex + StateIndex + 1)*UE::StateTree::Debug::IndentSize, TEXT("")
 				, *DebugGetStatePath(Exec.ActiveFrames, &CurrentFrame, StateIndex));
 
-			if (State.ShouldTickTasks(bHasEvents))
+			if (State.Type == EStateTreeStateType::Linked || State.Type == EStateTreeStateType::LinkedAsset)
 			{
-				bool bTaskTicked = false;
-
-				if (State.Type == EStateTreeStateType::Linked || State.Type == EStateTreeStateType::LinkedAsset)
+				if (State.ParameterDataHandle.IsValid() && State.ParameterBindingsBatch.IsValid())
 				{
-					// Linked and LinkedAsset are considered to always tick because the set data that can be used in transitions.
-					bTaskTicked = true;
-
-					if (State.ParameterDataHandle.IsValid() && State.ParameterBindingsBatch.IsValid())
-					{
-						const FStateTreeDataView StateParamsDataView = GetDataView(CurrentParentFrame, CurrentFrame, State.ParameterDataHandle);
-						CopyBatchOnActiveInstances(CurrentParentFrame, CurrentFrame, StateParamsDataView, State.ParameterBindingsBatch);
-					}
+					const FStateTreeDataView StateParamsDataView = GetDataView(CurrentParentFrame, CurrentFrame, State.ParameterDataHandle);
+					CopyBatchOnActiveInstances(CurrentParentFrame, CurrentFrame, StateParamsDataView, State.ParameterBindingsBatch);
 				}
+			}
 
+			if (State.ShouldTickTasks(bHasEvents))
+			{
 				// Update Tasks data and tick if possible (ie. if no task has yet failed and so bShouldTickTasks is true)
 				for (int32 TaskIndex = State.TasksBegin; TaskIndex < (State.TasksBegin + State.TasksNum); TaskIndex++)
 				{
@@ -2884,7 +2875,6 @@ EStateTreeRunStatus FStateTreeExecutionContext::TickTasks(const float DeltaTime)
 						CSV_SCOPED_TIMING_STAT_EXCLUSIVE(StateTree_Task_Tick);
 
 						TaskResult = Task.Tick(*this, DeltaTime);
-						bTaskTicked = true;
 					}
 
 					UE_STATETREE_DEBUG_TASK_EVENT(this, TaskIndex, TaskInstanceView,
