StateTree: Add description and change a property's category.
Set the minimal context destructor virtual to prevent memory issues if a user uses dynamic memory allocation.

#rnx

diff --git a/Engine/Plugins/Runtime/StateTree/Source/StateTreeEditorModule/Public/StateTreeState.h b/Engine/Plugins/Runtime/StateTree/Source/StateTreeEditorModule/Public/StateTreeState.h
--- a/Engine/Plugins/Runtime/StateTree/Source/StateTreeEditorModule/Public/StateTreeState.h
+++ b/Engine/Plugins/Runtime/StateTree/Source/StateTreeEditorModule/Public/StateTreeState.h
@@ -362,7 +362,7 @@ public:
 	float CustomTickRate = 0.0f;
 
 	/** Activate the CustomTickRate. */
-	UPROPERTY(EditDefaultsOnly, Category = "Enter Conditions", meta=(InlineEditConditionToggle))
+	UPROPERTY(EditDefaultsOnly, Category = "State", meta=(InlineEditConditionToggle))
 	bool bHasCustomTickRate = false;
 
 	/** Parameters of this state. If the state is linked to another state or asset, the parameters are for the linked state. */
diff --git a/Engine/Plugins/Runtime/StateTree/Source/StateTreeModule/Public/StateTreeExecutionContext.h b/Engine/Plugins/Runtime/StateTree/Source/StateTreeModule/Public/StateTreeExecutionContext.h
--- a/Engine/Plugins/Runtime/StateTree/Source/StateTreeModule/Public/StateTreeExecutionContext.h
+++ b/Engine/Plugins/Runtime/StateTree/Source/StateTreeModule/Public/StateTreeExecutionContext.h
@@ -47,7 +47,7 @@ struct STATETREEMODULE_API FStateTreeMinimalExecutionContext
 public:
 	explicit FStateTreeMinimalExecutionContext(UObject& Owner, const UStateTree& StateTree, FStateTreeInstanceData& InInstanceData);
 	explicit FStateTreeMinimalExecutionContext(UObject& Owner, const UStateTree& StateTree, FStateTreeInstanceStorage& Storage);
-	~FStateTreeMinimalExecutionContext();
+	virtual ~FStateTreeMinimalExecutionContext();
 
 public:
 	/**
@@ -176,9 +176,13 @@ public:
 	/** Sets callback used to collect external data views during State Tree execution. */
 	void SetCollectExternalDataCallback(const FOnCollectStateTreeExternalData& Callback);
 
-	/** */
+	/**
+	 * Overrides for linked State Trees. This table is used to override State Tree references on linked states.
+	 * If a linked state's tag is exact match of the tag specified on the table, the reference from the table is used instead.
+	 */
 	void SetLinkedStateTreeOverrides(const FStateTreeReferenceOverrides* InLinkedStateTreeOverrides);
 
+	/** @return the first state tree reference set by SetLinkedStateTreeOverrides that matches the StateTag. Or null if not found. */
 	const FStateTreeReference* GetLinkedStateTreeOverrideForTag(const FGameplayTag StateTag) const;
 
 	/** @return const references to the instance data in use, or nullptr if the context is not valid. */
@@ -605,7 +609,7 @@ protected:
 	 */
 	void StateCompleted();
 
-	/**
+	/** 
 	 * Tick evaluators and global tasks by delta time.
 	 */
 	EStateTreeRunStatus TickEvaluatorsAndGlobalTasks(const float DeltaTime, bool bTickGlobalTasks = true);
diff --git a/Engine/Plugins/Runtime/StateTree/Source/StateTreeModule/Public/StateTreeExecutionTypes.h b/Engine/Plugins/Runtime/StateTree/Source/StateTreeModule/Public/StateTreeExecutionTypes.h
--- a/Engine/Plugins/Runtime/StateTree/Source/StateTreeModule/Public/StateTreeExecutionTypes.h
+++ b/Engine/Plugins/Runtime/StateTree/Source/StateTreeModule/Public/StateTreeExecutionTypes.h
@@ -432,10 +432,10 @@ struct STATETREEMODULE_API FStateTreeActiveStates
 	FORCEINLINE const FStateTreeStateHandle* end  () const { return &States[0] + Num(); }
 
 
-	UPROPERTY(EditDefaultsOnly, Category = Default)
+	UPROPERTY(VisibleDefaultsOnly, Category = Default)
 	FStateTreeStateHandle States[MaxStates];
 
-	UPROPERTY(EditDefaultsOnly, Category = Default)
+	UPROPERTY(VisibleDefaultsOnly, Category = Default)
 	uint8 NumStates = 0;
 };
 
@@ -620,7 +620,7 @@ struct STATETREEMODULE_API FStateTreeExecutionFrame
 
 	/** Handle to the global parameter data, exists in ParentFrame. */
 	UPROPERTY()
-	FStateTreeDataHandle GlobalParameterDataHandle = FStateTreeDataHandle::Invalid; 
+	FStateTreeDataHandle GlobalParameterDataHandle = FStateTreeDataHandle::Invalid;
 
 	/** Number of states in ActiveStates which have instance data. Used during state selection to decide which active state data is safe to access. */
 	uint8 NumCurrentlyActiveStates = 0;
