StateTree: Add test to detect the issue with copy binding when the task has only transitions. Update test for linked asset tick order. Add test for array binding that uses bad memory access (out of bound).
#rnx

diff --git a/Engine/Plugins/Runtime/StateTree/Source/StateTreeModule/Private/StateTreeExecutionContext.cpp b/Engine/Plugins/Runtime/StateTree/Source/StateTreeModule/Private/StateTreeExecutionContext.cpp
--- a/Engine/Plugins/Runtime/StateTree/Source/StateTreeModule/Private/StateTreeExecutionContext.cpp
+++ b/Engine/Plugins/Runtime/StateTree/Source/StateTreeModule/Private/StateTreeExecutionContext.cpp
@@ -3471,10 +3471,16 @@ bool FStateTreeExecutionContext::TriggerTransitions()
 			const FStateTreeDataView TaskInstanceView = GetDataView(CurrentParentFrame, CurrentFrame, Task.InstanceDataHandle);
 			FNodeInstanceDataScope DataScope(*this, Task.InstanceDataHandle, TaskInstanceView);
 
+			// Copy bound properties.
+			if (Task.BindingsBatch.IsValid())
+			{
+				CopyBatchOnActiveInstances(CurrentParentFrame, CurrentFrame, TaskInstanceView, Task.BindingsBatch);
+			}
+
 			STATETREE_LOG(VeryVerbose, TEXT("%*sTriggerTransitions: '%s'"), UE::StateTree::Debug::IndentSize, TEXT(""), *Task.Name.ToString());
 			UE_STATETREE_DEBUG_TASK_EVENT(this, Handler.TaskIndex.Get(), TaskInstanceView, EStateTreeTraceEventType::OnEvaluating, EStateTreeRunStatus::Running);
 			check(TaskInstanceView.IsValid());
-			// TODO: bug the batch data might not be copied if !State.ShouldTickTasks() && State.bHasTransitionTasks and the same for bellow, need to copy at least once
+
 			Task.TriggerTransitions(*this);
 		}
 		else if (Handler.StateHandle.IsValid())
diff --git a/Engine/Plugins/Runtime/StateTree/Source/StateTreeTestSuite/Private/StateTreeBindingTest.cpp b/Engine/Plugins/Runtime/StateTree/Source/StateTreeTestSuite/Private/StateTreeBindingTest.cpp
--- a/Engine/Plugins/Runtime/StateTree/Source/StateTreeTestSuite/Private/StateTreeBindingTest.cpp
+++ b/Engine/Plugins/Runtime/StateTree/Source/StateTreeTestSuite/Private/StateTreeBindingTest.cpp
@@ -471,6 +471,246 @@ struct FStateTreeTest_ReferencesConstness : FStateTreeTestBase
 };
 IMPLEMENT_AI_INSTANT_TEST(FStateTreeTest_ReferencesConstness, "System.StateTree.Binding.ReferencesConstness");
 
+struct FStateTreeTest_MutableArray : FStateTreeTestBase
+{
+	virtual bool InstantTest() override
+	{
+		//Tree 1
+		//	Root
+		//		StateA -> Succeeded(Root)
+
+		FStateTreeCompilerLog Log;
+		FStateTreePropertyBindings Bindings;
+		FStateTreePropertyBindingCompiler BindingCompiler;
+
+		UStateTree& StateTree = NewStateTree();
+		FGuid RootParameter_ArrayValueID;
+		{
+			UStateTreeEditorData& EditorData = *Cast<UStateTreeEditorData>(StateTree.EditorData);
+			{
+				// Parameters
+				EditorData.RootParameters.Parameters.AddProperty("Value", EPropertyBagPropertyType::Int32);
+				EditorData.RootParameters.Parameters.SetValueInt32("Value", -111);
+				EditorData.RootParameters.Parameters.AddContainerProperty("ArrayValue", FPropertyBagContainerTypes(EPropertyBagContainerType::Array), EPropertyBagPropertyType::Int32, nullptr);
+				EditorData.RootParameters.Parameters.AddProperty("ArrayValue", EPropertyBagPropertyType::Int32);
+				FPropertyBagArrayRef ValueArrayRef = EditorData.RootParameters.Parameters.GetMutableArrayRef("ArrayValue").GetValue();
+				ValueArrayRef.EmptyAndAddValues(4);
+				ValueArrayRef.SetValueInt32(0, -11);
+				ValueArrayRef.SetValueInt32(1, -22);
+				ValueArrayRef.SetValueInt32(2, -33);
+				ValueArrayRef.SetValueInt32(3, -44);
+
+				// Global
+				TStateTreeEditorNode<FTestTask_PrintValue>& TaskA = EditorData.AddGlobalTask<FTestTask_PrintValue>("Tree1GlobalTaskA");
+				TaskA.GetInstanceData().Value = -2;
+				TaskA.GetInstanceData().ArrayValue = {-1, -2};
+				EditorData.AddPropertyBinding(FStateTreePropertyPath(EditorData.RootParameters.ID, TEXT("Value")), FStateTreePropertyPath(TaskA.ID, TEXT("Value")));
+				EditorData.AddPropertyBinding(FStateTreePropertyPath(EditorData.RootParameters.ID, TEXT("ArrayValue")), FStateTreePropertyPath(TaskA.ID, TEXT("ArrayValue")));
+
+			}
+			UStateTreeState& Root = EditorData.AddSubTree("Tree1StateRoot");
+			{
+				UStateTreeState& State = Root.AddChildState("Tree1StateA", EStateTreeStateType::State);
+
+				FStateTreeTransition& Transition = State.AddTransition(EStateTreeTransitionTrigger::OnTick, EStateTreeTransitionType::Succeeded);
+				Transition.bDelayTransition = true;
+				Transition.DelayDuration = 60.0;
+
+				FStateTreePropertyPath RootParametersArrayValue3(EditorData.RootParameters.ID);
+				RootParametersArrayValue3.AddPathSegment("ArrayValue", 3);
+
+				TStateTreeEditorNode<FTestTask_PrintAndResetValue>& TaskA = State.AddTask<FTestTask_PrintAndResetValue>("Tree1StateATaskA");
+				TaskA.GetInstanceData().Value = -2;
+				TaskA.GetInstanceData().ArrayValue = { -1, -2, -3, -4 };
+				TaskA.GetNode().ResetValue = 22;
+				TaskA.GetNode().ResetArrayValue = {200, 201, 202, 203, 204, 205};
+
+				FStateTreePropertyPath TaskAArrayValue3(TaskA.ID);
+				TaskAArrayValue3.AddPathSegment("ArrayValue", 3);
+
+				EditorData.AddPropertyBinding(FStateTreePropertyPath(EditorData.RootParameters.ID, "Value"), FStateTreePropertyPath(TaskA.ID, TEXT("Value")));
+				EditorData.AddPropertyBinding(RootParametersArrayValue3, TaskAArrayValue3);
+				
+				TStateTreeEditorNode<FTestTask_PrintAndResetValue>& TaskB = State.AddTask<FTestTask_PrintAndResetValue>("Tree1StateATaskB");
+				TaskB.GetInstanceData().Value = -2;
+				TaskB.GetInstanceData().ArrayValue = { -1, -2, -3, -4 };
+				TaskB.GetNode().ResetValue = 33;
+				TaskB.GetNode().ResetArrayValue = { 300, 301, 302, 303, 304, 305, 306, 307, 308, 309, 310, 311, 312, 313, 314, 315 };
+
+				FStateTreePropertyPath TaskBArrayValue3(TaskB.ID);
+				TaskBArrayValue3.AddPathSegment("ArrayValue", 3);
+
+				EditorData.AddPropertyBinding(FStateTreePropertyPath(TaskA.ID, TEXT("Value")), FStateTreePropertyPath(TaskB.ID, TEXT("Value")));
+				EditorData.AddPropertyBinding(TaskAArrayValue3, TaskBArrayValue3);
+			}
+		}
+		{
+			FStateTreeCompiler Compiler(Log);
+			const bool bResult = Compiler.Compile(StateTree);
+			AITEST_TRUE("StateTree2 should get compiled", bResult);
+		}
+		{
+			FStateTreeInstanceData InstanceData;
+			FTestStateTreeExecutionContext Exec(StateTree, StateTree, InstanceData);
+
+			const bool bInitSucceeded = Exec.IsValid();
+			AITEST_TRUE("StateTree should init", bInitSucceeded);
+
+			EStateTreeRunStatus Status = EStateTreeRunStatus::Unset;
+			FInstancedPropertyBag GlobalParameters = StateTree.GetDefaultParameters();
+			{
+				GlobalParameters.SetValueInt32("Value", 11);
+				FPropertyBagArrayRef ValueArrayRef = GlobalParameters.GetMutableArrayRef("ArrayValue").GetValue();
+				ValueArrayRef.EmptyAndAddValues(2);
+				ValueArrayRef.SetValueInt32(0, 911);
+				ValueArrayRef.SetValueInt32(1, 922);
+				Status = Exec.Start(&GlobalParameters);
+			}
+			AITEST_EQUAL("Start should complete with Running", Status, EStateTreeRunStatus::Running);
+			AITEST_TRUE("Start should enter Global tasks", Exec.Expect("Tree1GlobalTaskA", TEXT("EnterState11"))
+				.Then("Tree1GlobalTaskA", TEXT("EnterState:{911, 922}")) // should copy the full array
+				.Then("Tree1StateATaskA", TEXT("EnterState11"))
+				.Then("Tree1StateATaskA", TEXT("EnterState:{-1, -2, -3, -4}")) // should not copy anything since [3] is out of scope
+				.Then("Tree1StateATaskB", TEXT("EnterState22")) // TaskA set the value to 22  and {200, 201, 202, 203, 204, 205} (in EnterTask)
+				.Then("Tree1StateATaskB", TEXT("EnterState:{-1, -2, -3, 203}"))
+			);
+			Exec.LogClear();
+
+			Exec.Stop();
+		}
+
+		return true;
+	}
+};
+IMPLEMENT_AI_INSTANT_TEST(FStateTreeTest_MutableArray, "System.StateTree.Binding.MutableArray");
+
+struct FStateTreeTest_TransitionTaskWithBinding : FStateTreeTestBase
+{
+	virtual bool InstantTest() override
+	{
+		//Tree 1
+		//	Root
+		//		StateA -> Succeeded(Root)
+
+		FStateTreeCompilerLog Log;
+		FStateTreePropertyBindings Bindings;
+		FStateTreePropertyBindingCompiler BindingCompiler;
+
+		UStateTree& StateTree = NewStateTree();
+		{
+			UStateTreeEditorData& EditorData = *Cast<UStateTreeEditorData>(StateTree.EditorData);
+			{
+				// Parameters
+				EditorData.RootParameters.Parameters.AddProperty("Value", EPropertyBagPropertyType::Int32);
+				EditorData.RootParameters.Parameters.SetValueInt32("Value", -111);
+
+				// Global
+				TStateTreeEditorNode<FTestTask_PrintValue>& TaskA = EditorData.AddGlobalTask<FTestTask_PrintValue>("Tree1GlobalTaskA");
+				TaskA.GetInstanceData().Value = -2;
+				EditorData.AddPropertyBinding(FStateTreePropertyPath(EditorData.RootParameters.ID, TEXT("Value")), FStateTreePropertyPath(TaskA.ID, TEXT("Value")));
+
+				TStateTreeEditorNode<FTestTask_PrintValue_TransitionTick>& TaskB = EditorData.AddGlobalTask<FTestTask_PrintValue_TransitionTick>("Tree1GlobalTaskB");
+				TaskB.GetInstanceData().Value = -2;
+				EditorData.AddPropertyBinding(FStateTreePropertyPath(EditorData.RootParameters.ID, TEXT("Value")), FStateTreePropertyPath(TaskB.ID, TEXT("Value")));
+
+				TStateTreeEditorNode<FTestTask_PrintValue_TransitionNoTick>& TaskC = EditorData.AddGlobalTask<FTestTask_PrintValue_TransitionNoTick>("Tree1GlobalTaskC");
+				TaskC.GetInstanceData().Value = -2;
+				EditorData.AddPropertyBinding(FStateTreePropertyPath(EditorData.RootParameters.ID, TEXT("Value")), FStateTreePropertyPath(TaskC.ID, TEXT("Value")));
+			}
+			UStateTreeState& Root = EditorData.AddSubTree("Tree1StateRoot");
+			{
+				TStateTreeEditorNode<FTestTask_PrintValue>& TaskA = Root.AddTask<FTestTask_PrintValue>("Tree1StateRootTaskA");
+				TaskA.GetInstanceData().Value = -2;
+				EditorData.AddPropertyBinding(FStateTreePropertyPath(EditorData.RootParameters.ID, TEXT("Value")), FStateTreePropertyPath(TaskA.ID, TEXT("Value")));
+				
+				TStateTreeEditorNode<FTestTask_PrintValue_TransitionTick>& TaskB = Root.AddTask<FTestTask_PrintValue_TransitionTick>("Tree1StateRootTaskB");
+				TaskB.GetInstanceData().Value = -2;
+				EditorData.AddPropertyBinding(FStateTreePropertyPath(EditorData.RootParameters.ID, TEXT("Value")), FStateTreePropertyPath(TaskB.ID, TEXT("Value")));
+				
+				TStateTreeEditorNode<FTestTask_PrintValue_TransitionNoTick>& TaskC = Root.AddTask<FTestTask_PrintValue_TransitionNoTick>("Tree1StateRootTaskC");
+				TaskC.GetInstanceData().Value = -2;
+				EditorData.AddPropertyBinding(FStateTreePropertyPath(EditorData.RootParameters.ID, TEXT("Value")), FStateTreePropertyPath(TaskC.ID, TEXT("Value")));
+			}
+			{
+				UStateTreeState& State = Root.AddChildState("Tree1StateA", EStateTreeStateType::State);
+
+				FStateTreeTransition& Transition = State.AddTransition(EStateTreeTransitionTrigger::OnTick, EStateTreeTransitionType::Succeeded);
+				Transition.bDelayTransition = true;
+				Transition.DelayDuration = 5.0;
+
+				TStateTreeEditorNode<FTestTask_PrintValue>& TaskA = State.AddTask<FTestTask_PrintValue>("Tree1StateATaskA");
+				TaskA.GetInstanceData().Value = -2;
+				EditorData.AddPropertyBinding(FStateTreePropertyPath(EditorData.RootParameters.ID, TEXT("Value")), FStateTreePropertyPath(TaskA.ID, TEXT("Value")));
+
+				TStateTreeEditorNode<FTestTask_PrintValue_TransitionTick>& TaskB = State.AddTask<FTestTask_PrintValue_TransitionTick>("Tree1StateATaskB");
+				TaskB.GetInstanceData().Value = -2;
+				EditorData.AddPropertyBinding(FStateTreePropertyPath(EditorData.RootParameters.ID, TEXT("Value")), FStateTreePropertyPath(TaskB.ID, TEXT("Value")));
+
+				TStateTreeEditorNode<FTestTask_PrintValue_TransitionNoTick>& TaskC = State.AddTask<FTestTask_PrintValue_TransitionNoTick>("Tree1StateATaskC");
+				TaskC.GetInstanceData().Value = -2;
+				EditorData.AddPropertyBinding(FStateTreePropertyPath(EditorData.RootParameters.ID, TEXT("Value")), FStateTreePropertyPath(TaskC.ID, TEXT("Value")));
+			}
+		}
+		{
+			FStateTreeCompiler Compiler(Log);
+			const bool bResult = Compiler.Compile(StateTree);
+			AITEST_TRUE("StateTree2 should get compiled", bResult);
+		}
+		{
+			FStateTreeInstanceData InstanceData;
+			FTestStateTreeExecutionContext Exec(StateTree, StateTree, InstanceData);
+
+			const bool bInitSucceeded = Exec.IsValid();
+			AITEST_TRUE("StateTree should init", bInitSucceeded);
+
+			EStateTreeRunStatus Status = EStateTreeRunStatus::Unset;
+			FInstancedPropertyBag GlobalParameters = StateTree.GetDefaultParameters();
+
+			{
+				GlobalParameters.SetValueInt32("Value", 99);
+				Status = Exec.Start(&GlobalParameters);
+			}
+			AITEST_EQUAL("Start should complete with Running", Status, EStateTreeRunStatus::Running);
+			AITEST_TRUE("Start should enter Global tasks", Exec.Expect("Tree1GlobalTaskA", TEXT("EnterState99"))
+				.Then("Tree1GlobalTaskB", TEXT("EnterState99"))
+				.Then("Tree1GlobalTaskC", TEXT("EnterState99"))
+				.Then("Tree1StateRootTaskA", TEXT("EnterState99"))
+				.Then("Tree1StateRootTaskB", TEXT("EnterState99"))
+				.Then("Tree1StateRootTaskC", TEXT("EnterState99"))
+				.Then("Tree1StateATaskA", TEXT("EnterState99"))
+				.Then("Tree1StateATaskB", TEXT("EnterState99"))
+				.Then("Tree1StateATaskC", TEXT("EnterState99"))
+			);
+			Exec.LogClear();
+
+			GlobalParameters.SetValueInt32("Value", 88);
+			InstanceData.GetMutableStorage().SetGlobalParameters(GlobalParameters);
+
+			Status = Exec.Tick(1.0f);
+			AITEST_EQUAL("2nd Tick should complete with Running", Status, EStateTreeRunStatus::Running);
+			AITEST_TRUE("2nd Tick should tick tasks", Exec.Expect("Tree1GlobalTaskA", TEXT("Tick88"))
+				.Then("Tree1GlobalTaskB", TEXT("Tick88"))
+				.Then("Tree1StateRootTaskA", TEXT("Tick88"))
+				.Then("Tree1StateRootTaskB", TEXT("Tick88"))
+				.Then("Tree1StateATaskA", TEXT("Tick88"))
+				.Then("Tree1StateATaskB", TEXT("Tick88"))
+				.Then("Tree1StateATaskC", TEXT("TriggerTransitions88"))
+				.Then("Tree1StateATaskB", TEXT("TriggerTransitions88"))
+				.Then("Tree1StateRootTaskC", TEXT("TriggerTransitions88"))
+				.Then("Tree1StateRootTaskB", TEXT("TriggerTransitions88"))
+				.Then("Tree1GlobalTaskC", TEXT("TriggerTransitions88"))
+				.Then("Tree1GlobalTaskB", TEXT("TriggerTransitions88"))
+			);
+			Exec.LogClear();
+
+			Exec.Stop();
+		}
+
+		return true;
+	}
+};
+IMPLEMENT_AI_INSTANT_TEST(FStateTreeTest_TransitionTaskWithBinding, "System.StateTree.Binding.TransitionTaskWithBinding");
+
 } // namespace UE::StateTree::Tests
 
 UE_ENABLE_OPTIMIZATION_SHIP
diff --git a/Engine/Plugins/Runtime/StateTree/Source/StateTreeTestSuite/Private/StateTreeLinkedStateTest.cpp b/Engine/Plugins/Runtime/StateTree/Source/StateTreeTestSuite/Private/StateTreeLinkedStateTest.cpp
--- a/Engine/Plugins/Runtime/StateTree/Source/StateTreeTestSuite/Private/StateTreeLinkedStateTest.cpp
+++ b/Engine/Plugins/Runtime/StateTree/Source/StateTreeTestSuite/Private/StateTreeLinkedStateTest.cpp
@@ -271,100 +271,127 @@ struct FStateTreeTest_MultipleSameLinkedAsset : FStateTreeTestBase
 			AITEST_TRUE("StateTree should get compiled", bResult);
 		}
 
-		const FString EnterStateStr(TEXT("EnterState"));
-		const FString ExitStateStr(TEXT("ExitState"));
-
 		{
-			EStateTreeRunStatus Status = EStateTreeRunStatus::Unset;
 			FStateTreeInstanceData InstanceData;
 			FTestStateTreeExecutionContext Exec(StateTree1, StateTree1, InstanceData);
-			const bool bInitSucceeded = Exec.IsValid();
-			AITEST_TRUE("StateTree should init", bInitSucceeded);
-
-			Status = Exec.Start();
-			AITEST_EQUAL("Start should complete with Running", Status, EStateTreeRunStatus::Running);
-			AITEST_TRUE("Start should enter Global tasks", Exec.Expect("Tree1GlobalTaskA", TEXT("EnterState99"))
-				.Then("Tree1StateRootTaskA", TEXT("EnterState88"))
-				.Then("Tree1StateATaskA", TEXT("EnterState1")));
-			Exec.LogClear();
-
-			Status = Exec.Tick(1.5f); // over tick, should trigger
-			AITEST_EQUAL("1st Tick should complete with Running", Status, EStateTreeRunStatus::Running);
-			AITEST_TRUE("1st should tick tasks", Exec.Expect("Tree1GlobalTaskA", TEXT("Tick99"))
-				.Then("Tree1StateRootTaskA", TEXT("Tick88"))
-				.Then("Tree1StateATaskA", TEXT("Tick1")));
-			Exec.LogClear();
-
-			Status = Exec.Tick(1.0f);
-			AITEST_EQUAL("2nd Tick should complete with Running", Status, EStateTreeRunStatus::Running);
-			AITEST_TRUE("2nd Tick should tick tasks", Exec.Expect("Tree1GlobalTaskA", TEXT("Tick99"))
-				.Then("Tree1StateRootTaskA", TEXT("Tick88"))
-				.Then("Tree1StateATaskA", TEXT("Tick1"))
-				.Then("Tree1StateATaskA", TEXT("ExitState1"))
-				.Then("Tree1StateBTaskA", TEXT("EnterState2")));
-			Exec.LogClear();
-
-			Status = Exec.Tick(1.0f);
-			AITEST_EQUAL("3rd Tick should complete with Running", Status, EStateTreeRunStatus::Running);
-			AITEST_TRUE("3rd should tick tasks", Exec.Expect("Tree1GlobalTaskA", TEXT("Tick99"))
-				.Then("Tree1StateRootTaskA", TEXT("Tick88"))
-				.Then("Tree1StateBTaskA", TEXT("Tick2")));
-			Exec.LogClear();
-
-			Status = Exec.Tick(1.0f);
-			AITEST_EQUAL("4th Tick should complete with Running", Status, EStateTreeRunStatus::Running);
-			AITEST_TRUE("4th Tick should tick tasks", Exec.Expect("Tree1GlobalTaskA", TEXT("Tick99"))
-				.Then("Tree1StateRootTaskA", TEXT("Tick88"))
-				.Then("Tree1StateBTaskA", TEXT("Tick2")));
-			//bug: binding are not executed
-			//AITEST_TRUE("4th Tick should enter the Global", Exec.Expect("Tree2GlobalTaskA", TEXT("EnterState111")));
-			AITEST_TRUE("4th Tick should exit the State", Exec.Expect("Tree1StateBTaskA", TEXT("ExitState2")));
-			AITEST_TRUE("4th Tick should enter the Root", Exec.Expect("Tree2StateRootTaskA", TEXT("EnterState111")));
-			AITEST_TRUE("4th Tick should enter the LinkedTree", Exec.Expect("Tree2StateATaskA", TEXT("EnterState111")));
-			Exec.LogClear();
-
-			Status = Exec.Tick(0.001f);
-			AITEST_EQUAL("5th Tick should complete with Running", Status, EStateTreeRunStatus::Running);
-			AITEST_TRUE("5h Tick should tick the tasks", Exec.Expect("Tree1GlobalTaskA", TEXT("Tick99")));
-			//bug: the task should be tick in the correct order
 			{
-				//.Then("Tree1StateRootTaskA", TEXT("Tick88"))
-				//.Then("Tree2GlobalTaskA", TEXT("Tick111"))
-				//.Then("Tree2StateRootTaskA", TEXT("Tick111"))
-				//.Then("Tree2StateATaskA", TEXT("Tick111")));
-				AITEST_TRUE("5th Tick should exit the State", Exec.Expect("Tree1StateRootTaskA", TEXT("Tick88")));
-				AITEST_TRUE("5th Tick should enter the Root", Exec.Expect("Tree2GlobalTaskA", TEXT("Tick111")));
-				AITEST_TRUE("5th Tick should enter the State", Exec.Expect("Tree2StateRootTaskA", TEXT("Tick111")));
-				AITEST_TRUE("5th Tick should enter the LinkedTree", Exec.Expect("Tree2StateATaskA", TEXT("Tick111")));
+				const bool bInitSucceeded = Exec.IsValid();
+				AITEST_TRUE("StateTree should init", bInitSucceeded);
+			}
+			{
+				EStateTreeRunStatus Status = Exec.Start();
+				AITEST_EQUAL("Start should complete with Running", Status, EStateTreeRunStatus::Running);
+				FTestStateTreeExecutionContext::FLogOrder LogOrder = Exec.Expect("Tree1GlobalTaskA", TEXT("EnterState99"));
+				AITEST_TRUE("Start should enter Tree1GlobalTaskA", LogOrder);
+				LogOrder = LogOrder.Then("Tree1StateRootTaskA", TEXT("EnterState88"));
+				AITEST_TRUE("Start should enter Tree1StateRootTaskA", LogOrder);
+				LogOrder = LogOrder.Then("Tree1StateATaskA", TEXT("EnterState1"));
+				AITEST_TRUE("Start should enter Tree1StateATaskA", LogOrder);
+				Exec.LogClear();
+			}
+			{
+				EStateTreeRunStatus Status = Exec.Tick(1.5f); // over tick, should trigger
+				AITEST_EQUAL("1st Tick should complete with Running", Status, EStateTreeRunStatus::Running);
+				FTestStateTreeExecutionContext::FLogOrder LogOrder = Exec.Expect("Tree1GlobalTaskA", TEXT("Tick99"));
+				AITEST_TRUE("1st should tick tasks Tree1GlobalTaskA", LogOrder);
+				LogOrder = LogOrder.Then("Tree1StateRootTaskA", TEXT("Tick88"));
+				AITEST_TRUE("1st should tick tasks Tree1StateRootTaskA", LogOrder);
+				LogOrder = LogOrder.Then("Tree1StateATaskA", TEXT("Tick1"));
+				AITEST_TRUE("1st should tick tasks Tree1StateATaskA", LogOrder);
+				Exec.LogClear();
+			}
+			{
+				EStateTreeRunStatus Status = Exec.Tick(1.0f);
+				AITEST_EQUAL("2nd Tick should complete with Running", Status, EStateTreeRunStatus::Running);
+				FTestStateTreeExecutionContext::FLogOrder LogOrder = Exec.Expect("Tree1GlobalTaskA", TEXT("Tick99"));
+				AITEST_TRUE("2nd Tick should tick Tree1GlobalTaskA", LogOrder);
+				LogOrder = LogOrder.Then("Tree1StateRootTaskA", TEXT("Tick88"));
+				AITEST_TRUE("2nd Tick should tick Tree1StateRootTaskA", LogOrder);
+				LogOrder = LogOrder.Then("Tree1StateATaskA", TEXT("Tick1"));
+				AITEST_TRUE("2nd Tick should tick Tree1StateATaskA", LogOrder);
+				LogOrder = LogOrder.Then("Tree1StateATaskA", TEXT("ExitState1"));
+				AITEST_TRUE("2nd Tick should exit Tree1StateATaskA", LogOrder);
+				LogOrder = LogOrder.Then("Tree1StateBTaskA", TEXT("EnterState2"));
+				AITEST_TRUE("2nd Tick should enter Tree1StateBTaskA", LogOrder);
+				Exec.LogClear();
+			}
+			{
+				EStateTreeRunStatus Status = Exec.Tick(1.0f);
+				AITEST_EQUAL("3rd Tick should complete with Running", Status, EStateTreeRunStatus::Running);
+				FTestStateTreeExecutionContext::FLogOrder LogOrder = Exec.Expect("Tree1GlobalTaskA", TEXT("Tick99"));
+				AITEST_TRUE("3rd should tick Tree1GlobalTaskA", LogOrder);
+				LogOrder = LogOrder.Then("Tree1StateRootTaskA", TEXT("Tick88"));
+				AITEST_TRUE("3rd should tick Tree1StateRootTaskA", LogOrder);
+				LogOrder = LogOrder.Then("Tree1StateBTaskA", TEXT("Tick2"));
+				AITEST_TRUE("3rd should tick Tree1StateBTaskA", LogOrder);
+				Exec.LogClear();
+			}
+			{
+				EStateTreeRunStatus Status = Exec.Tick(1.0f);
+				AITEST_EQUAL("4th Tick should complete with Running", Status, EStateTreeRunStatus::Running);
+				FTestStateTreeExecutionContext::FLogOrder LogOrder = Exec.Expect("Tree1GlobalTaskA", TEXT("Tick99"));
+				AITEST_TRUE("4th Tick should tick Tree1GlobalTaskA", LogOrder);
+				LogOrder = LogOrder.Then("Tree1StateRootTaskA", TEXT("Tick88"));
+				AITEST_TRUE("4th Tick should tick Tree1StateRootTaskA", LogOrder);
+				LogOrder = LogOrder.Then("Tree1StateBTaskA", TEXT("Tick2"));
+				AITEST_TRUE("4th Tick should tick Tree1StateBTaskA", LogOrder);
+				LogOrder = LogOrder.Then("Tree1StateBTaskA", TEXT("ExitState2"));
+				AITEST_TRUE("4th Tick should exit Tree1StateBTaskA", LogOrder);
+				//@todo: bug: order is wrong
+				//LogOrder = LogOrder.Then("Tree2GlobalTaskA", TEXT("EnterState111"))
+				//AITEST_TRUE("4th Tick should enter Tree2GlobalTaskA", LogOrder);
+				LogOrder = LogOrder.Then("Tree2StateRootTaskA", TEXT("EnterState111"));
+				AITEST_TRUE("4th Tick should enter Tree2StateRootTaskA", LogOrder);
+				LogOrder = LogOrder.Then("Tree2StateATaskA", TEXT("EnterState111"));
+				AITEST_TRUE("4th Tick should enter Tree2StateATaskA", LogOrder);
+				//@todo: bug: Binding are not executed
+				//AITEST_TRUE("4th Tick should tick tasks", Exec.Expect("Tree2GlobalTaskA", TEXT("EnterState111")));
+				Exec.LogClear();
+			}
+			{
+				EStateTreeRunStatus Status = Exec.Tick(0.001f);
+				AITEST_EQUAL("5th Tick should complete with Running", Status, EStateTreeRunStatus::Running);
+				FTestStateTreeExecutionContext::FLogOrder LogOrder = Exec.Expect("Tree1GlobalTaskA", TEXT("Tick99"));
+				AITEST_TRUE("5h Tick should tick Tree1GlobalTaskA", LogOrder);
+				LogOrder = LogOrder.Then("Tree2GlobalTaskA", TEXT("Tick111"));
+				AITEST_TRUE("5h Tick should tick Tree2GlobalTaskA", LogOrder);
+				LogOrder = LogOrder.Then("Tree1StateRootTaskA", TEXT("Tick88"));
+				AITEST_TRUE("5h Tick should tick Tree1StateRootTaskA", LogOrder);
+				LogOrder = LogOrder.Then("Tree2StateRootTaskA", TEXT("Tick111"));
+				AITEST_TRUE("5h Tick should tick Tree2StateRootTaskA", LogOrder);
+				LogOrder = LogOrder.Then("Tree2StateATaskA", TEXT("Tick111"));
+				AITEST_TRUE("5h Tick should tick Tree2StateATaskA", LogOrder);
+				Exec.LogClear();
+			}
+			{
+				EStateTreeRunStatus Status = Exec.Tick(1.0f);
+				AITEST_EQUAL("6th Tick should complete with Running", Status, EStateTreeRunStatus::Running);
+				FTestStateTreeExecutionContext::FLogOrder LogOrder = Exec.Expect("Tree1GlobalTaskA", TEXT("Tick99"));
+				AITEST_TRUE("6th Tick should tick Tree1GlobalTaskA", LogOrder);
+				LogOrder = LogOrder.Then("Tree2GlobalTaskA", TEXT("Tick111"));
+				AITEST_TRUE("6th Tick should tick Tree2GlobalTaskA", LogOrder);
+				LogOrder = LogOrder.Then("Tree1StateRootTaskA", TEXT("Tick88"));
+				AITEST_TRUE("6th Tick should tick Tree1StateRootTaskA", LogOrder);
+				LogOrder = LogOrder.Then("Tree2StateRootTaskA", TEXT("Tick111"));
+				AITEST_TRUE("6th Tick should tick Tree2StateRootTaskA", LogOrder);
+				LogOrder = LogOrder.Then("Tree2StateATaskA", TEXT("Tick111"));
+				AITEST_TRUE("6th Tick should tick Tree2StateATaskA", LogOrder);
+				LogOrder = LogOrder.Then("Tree2StateATaskA", TEXT("ExitState111"));
+				AITEST_TRUE("6th Tick should exit Tree2StateATaskA", LogOrder);
+				LogOrder = LogOrder.Then("Tree2StateRootTaskA", TEXT("ExitState111"));
+				AITEST_TRUE("6th Tick should exit Tree2StateRootTaskA", LogOrder);
+				//@todo: bug: The exit state is not called on the global, it's a new frame.
+				//LogOrder = LogOrder.Then("Tree2GlobalTaskA", TEXT("ExitState111"));
+				//AITEST_TRUE("6th Tick should exit Tree2GlobalTaskA", LogOrder);
+				//@todo: bug: The enter state is not called on the global, it's a new frame.
+				//LogOrder = LogOrder.Then("Tree2GlobalTaskA", TEXT("EnterState222"));
+				//AITEST_TRUE("6th Tick should enter Tree2GlobalTaskA", LogOrder);
+				LogOrder = LogOrder.Then("Tree2StateRootTaskA", TEXT("EnterState222"));
+				AITEST_TRUE("6th Tick should enter Tree2StateRootTaskA", LogOrder);
+				LogOrder = LogOrder.Then("Tree2StateATaskA", TEXT("EnterState222"));
+				AITEST_TRUE("6th Tick should enter Tree2StateATaskA", LogOrder);
+				Exec.LogClear();
 			}
-			Exec.LogClear();
-
-			Status = Exec.Tick(1.0f);
-			AITEST_EQUAL("6th Tick should complete with Running", Status, EStateTreeRunStatus::Running);
-			AITEST_TRUE("6th Tick should tick the tasks", Exec.Expect("Tree1GlobalTaskA", TEXT("Tick99")));
-			//bug: the task should be tick in the correct order
-			{
-				//.Then("Tree1StateRootTaskA", TEXT("Tick88"))
-				//.Then("Tree2GlobalTaskA", TEXT("Tick111"))
-				//.Then("Tree2StateRootTaskA", TEXT("Tick111"))
-				//.Then("Tree2StateATaskA", TEXT("Tick111")));
-				AITEST_TRUE("6th Tick should exit the State", Exec.Expect("Tree1StateRootTaskA", TEXT("Tick88")));
-				AITEST_TRUE("6th Tick should enter the Root", Exec.Expect("Tree2GlobalTaskA", TEXT("Tick111")));
-				AITEST_TRUE("6th Tick should enter the Root", Exec.Expect("Tree2StateRootTaskA", TEXT("Tick111")));
-				AITEST_TRUE("6th Tick should enter the LinkedTree", Exec.Expect("Tree2StateATaskA", TEXT("Tick111")));
-			}
-			//bug: the global task are not exited
-			//AITEST_TRUE("6th Tick should exit the global tasks", Exec.Expect("Tree2GlobalTaskA", TEXT("ExitState111")));
-			AITEST_TRUE("6th Tick should exit the state", Exec.Expect("Tree2StateRootTaskA", TEXT("ExitState111")));
-			AITEST_TRUE("6th Tick should exit the StateLinkedTree", Exec.Expect("Tree2StateATaskA", TEXT("ExitState111")));
-			//AITEST_TRUE("6th Tick should enter the global tasks", Exec.Expect("Tree2GlobalTaskA", TEXT("EnterState222")));
-			AITEST_TRUE("6th Tick should enter the state", Exec.Expect("Tree2StateRootTaskA", TEXT("EnterState222")));
-			AITEST_TRUE("6th Tick should enter the StateLinkedTree", Exec.Expect("Tree2StateATaskA", TEXT("EnterState222")));
-			Exec.LogClear();
-
-			Status = Exec.Tick(1.0f);
-			Exec.LogClear();
-
 			Exec.Stop();
 		}
 
@@ -472,77 +499,106 @@ struct FStateTreeTest_EmptyStateWithTickTransitionLinkedAsset : FStateTreeTestBa
 			AITEST_TRUE("StateTree should get compiled", bResult);
 		}
 
-		const FString EnterStateStr(TEXT("EnterState"));
-		const FString ExitStateStr(TEXT("ExitState"));
-
 		{
-			EStateTreeRunStatus Status = EStateTreeRunStatus::Unset;
 			FStateTreeInstanceData InstanceData;
 			FTestStateTreeExecutionContext Exec(StateTree1, StateTree1, InstanceData);
-			const bool bInitSucceeded = Exec.IsValid();
-			AITEST_TRUE("StateTree should init", bInitSucceeded);
-
-			Status = Exec.Start();
-			AITEST_EQUAL("Start should complete with Running", Status, EStateTreeRunStatus::Running);
-			AITEST_TRUE("Start tick in the correct order", Exec.Expect("Tree1GlobalTaskA", TEXT("EnterState11"))
-				.Then("Tree1StateRootTaskA", TEXT("EnterState12"))
-				.Then("Tree1StateATaskA", TEXT("EnterState13")));
-			AITEST_TRUE("Start should be in the correct states", Exec.ExpectInActiveStates("Tree1StateRoot", "Tree1StateA"));
-			Exec.LogClear();
-
-			Status = Exec.Tick(1.5f); // over tick, should trigger
-			AITEST_EQUAL("1st Tick should complete with Running", Status, EStateTreeRunStatus::Running);
-			AITEST_TRUE("1st Tick should tick the Global task", Exec.Expect("Tree1GlobalTaskA", TEXT("Tick11"))
-				.Then("Tree1StateRootTaskA", TEXT("Tick12"))
-				.Then("Tree1StateATaskA", TEXT("Tick13")));
-			AITEST_TRUE("Start should be in the correct states", Exec.ExpectInActiveStates("Tree1StateRoot", "Tree1StateA"));
-			Exec.LogClear();
-
-			Status = Exec.Tick(1.0f);
-			AITEST_EQUAL("2nd Tick should complete with Running", Status, EStateTreeRunStatus::Running);
-			AITEST_TRUE("2nd Tick should tick the Global task", Exec.Expect("Tree1GlobalTaskA", TEXT("Tick11")));
-			AITEST_TRUE("2nd Tick should tick the Root task", Exec.Expect("Tree1StateRootTaskA", TEXT("Tick12")));
-			AITEST_TRUE("2nd Tick should tick the TaskA", Exec.Expect("Tree1StateATaskA", TEXT("Tick13")));
-
-			AITEST_TRUE("2nd Tick should enter the Global task", Exec.Expect("Tree2GlobalTaskA", TEXT("EnterState21")));
-			AITEST_TRUE("2nd Tick should exit the TaskA", Exec.Expect("Tree1StateATaskA", TEXT("ExitState13")));
-			AITEST_TRUE("2nd Tick should enter the Root task", Exec.Expect("Tree2StateRootTaskA", TEXT("EnterState22")));
-			AITEST_FALSE("2nd Tick should NOT enter the FailTask", Exec.Expect("Tree2StateFailTaskA"));
-
-			AITEST_TRUE("2nd Tick should be in the correct states", Exec.ExpectInActiveStates("Tree1StateRoot", "Tree1StateLinkedTree", "Tree2StateRoot", "Tree2StateB"));
-			Exec.LogClear();
-
-			Status = Exec.Tick(1.0f);
-			AITEST_EQUAL("3rd Tick should complete with Running", Status, EStateTreeRunStatus::Running);
-			AITEST_TRUE("3rd Tick should tick Global tasks", Exec.Expect("Tree1GlobalTaskA", TEXT("Tick11")));
-			AITEST_TRUE("3rd Tick should tick the Root task", Exec.Expect("Tree1StateRootTaskA", TEXT("Tick12")));
-			AITEST_TRUE("3rd Tick should tick Global tasks", Exec.Expect("Tree2GlobalTaskA", TEXT("Tick21")));
-			AITEST_TRUE("3rd Tick should tick the Root task", Exec.Expect("Tree2StateRootTaskA", TEXT("Tick22")));
-
-			AITEST_TRUE("3rd Tick should be in the correct states", Exec.ExpectInActiveStates("Tree1StateRoot", "Tree1StateLinkedTree", "Tree2StateRoot", "Tree2StateB"));
-			Exec.LogClear();
-
-			Status = Exec.Tick(1.0f);
-			AITEST_EQUAL("4th Tick should complete with Running", Status, EStateTreeRunStatus::Running);
-			AITEST_TRUE("4th Tick should tick Global tasks", Exec.Expect("Tree1GlobalTaskA", TEXT("Tick11")));
-			AITEST_TRUE("4th Tick should tick the Root task", Exec.Expect("Tree1StateRootTaskA", TEXT("Tick12")));
-			AITEST_TRUE("4th Tick should tick Global tasks", Exec.Expect("Tree2GlobalTaskA", TEXT("Tick21")));
-			AITEST_TRUE("4th Tick should tick the Root task", Exec.Expect("Tree2StateRootTaskA", TEXT("Tick22")));
-			AITEST_TRUE("4th Tick should tick Global tasks", Exec.Expect("Tree2StateRootTaskA", TEXT("ExitState22")));
-			AITEST_TRUE("4th Tick should tick Global tasks", Exec.Expect("Tree2GlobalTaskA", TEXT("ExitState21")));
-			AITEST_TRUE("4th Tick should tick the Root task", Exec.Expect("Tree1StateBTaskA", TEXT("EnterState14")));
-
-			AITEST_TRUE("4th Tick should be in the correct states", Exec.ExpectInActiveStates("Tree1StateRoot", "Tree1StateB"));
-			Exec.LogClear();
-
-			Status = Exec.Tick(0.001f);
-			AITEST_EQUAL("5th Tick should complete with Running", Status, EStateTreeRunStatus::Running);
-			AITEST_TRUE("Start tick in the correct order", Exec.Expect("Tree1GlobalTaskA", TEXT("Tick11"))
-				.Then("Tree1StateRootTaskA", TEXT("Tick12"))
-				.Then("Tree1StateBTaskA", TEXT("Tick14")));
-			AITEST_TRUE("5th Tick should be in the correct states", Exec.ExpectInActiveStates("Tree1StateRoot", "Tree1StateB"));
-			Exec.LogClear();
-
+			{
+				const bool bInitSucceeded = Exec.IsValid();
+				AITEST_TRUE("StateTree should init", bInitSucceeded);
+			}
+			{
+				EStateTreeRunStatus Status = Exec.Start();
+				AITEST_EQUAL("Start should complete with Running", Status, EStateTreeRunStatus::Running);
+				FTestStateTreeExecutionContext::FLogOrder LogOrder = Exec.Expect("Tree1GlobalTaskA", TEXT("EnterState11"));
+				AITEST_TRUE("Start enters in the correct order Tree1GlobalTaskA", LogOrder);
+				LogOrder = LogOrder.Then("Tree1StateRootTaskA", TEXT("EnterState12"));
+				AITEST_TRUE("Start enters in the correct order Tree1StateRootTaskA", LogOrder);
+				LogOrder = LogOrder.Then("Tree1StateATaskA", TEXT("EnterState13"));
+				AITEST_TRUE("Start enters in the correct order Tree1StateATaskA", LogOrder);
+				AITEST_TRUE("Start should be in the correct states", Exec.ExpectInActiveStates("Tree1StateRoot", "Tree1StateA"));
+				Exec.LogClear();
+			}
+			{
+				EStateTreeRunStatus Status = Exec.Tick(1.5f); // over tick, should trigger
+				AITEST_EQUAL("1st Tick should complete with Running", Status, EStateTreeRunStatus::Running);
+				FTestStateTreeExecutionContext::FLogOrder LogOrder = Exec.Expect("Tree1GlobalTaskA", TEXT("Tick11"));
+				AITEST_TRUE("1st Tick should tick Tree1GlobalTaskA", LogOrder);
+				LogOrder = LogOrder.Then("Tree1StateRootTaskA", TEXT("Tick12"));
+				AITEST_TRUE("1st Tick should tick Tree1StateRootTaskA", LogOrder);
+				LogOrder = LogOrder.Then("Tree1StateATaskA", TEXT("Tick13"));
+				AITEST_TRUE("1st Tick should tick Tree1StateATaskA", LogOrder);
+				AITEST_TRUE("Start should be in the correct states", Exec.ExpectInActiveStates("Tree1StateRoot", "Tree1StateA"));
+				Exec.LogClear();
+			}
+			{
+				EStateTreeRunStatus Status = Exec.Tick(1.0f);
+				AITEST_EQUAL("2nd Tick should complete with Running", Status, EStateTreeRunStatus::Running);
+				FTestStateTreeExecutionContext::FLogOrder LogOrder = Exec.Expect("Tree1GlobalTaskA", TEXT("Tick11"));
+				AITEST_TRUE("2nd Tick should tick Tree1GlobalTaskA", LogOrder);
+				LogOrder = LogOrder.Then("Tree1StateRootTaskA", TEXT("Tick12"));
+				AITEST_TRUE("2nd Tick should tick Tree1StateRootTaskA", LogOrder);
+				LogOrder = LogOrder.Then("Tree1StateATaskA", TEXT("Tick13"));
+				AITEST_TRUE("2nd Tick should tick Tree1StateATaskA", LogOrder);
+				LogOrder = LogOrder.Then("Tree2GlobalTaskA", TEXT("EnterState21"));
+				AITEST_TRUE("2nd Tick should enter Tree2GlobalTaskA", LogOrder);
+				LogOrder = LogOrder.Then("State Tree Test Boolean Condition", TEXT("TestCondition=0"));
+				AITEST_TRUE("2nd Tick should test Bool", LogOrder);
+				LogOrder = LogOrder.Then("State Tree Test Boolean Condition", TEXT("TestCondition=1"));
+				AITEST_TRUE("2nd Tick should test Bool", LogOrder);
+				LogOrder = LogOrder.Then("Tree1StateATaskA", TEXT("ExitState13"));
+				AITEST_TRUE("2nd Tick should exit Tree1StateATaskA", LogOrder);
+				LogOrder = LogOrder.Then("Tree2StateRootTaskA", TEXT("EnterState22"));
+				AITEST_TRUE("2nd Tick should enter Tree2StateRootTaskA", LogOrder);
+				AITEST_FALSE("2nd Tick should not enter the fail state.", Exec.Expect("Tree2StateFailTaskA"));
+				AITEST_TRUE("2nd Tick should be in the correct states", Exec.ExpectInActiveStates("Tree1StateRoot", "Tree1StateLinkedTree", "Tree2StateRoot", "Tree2StateB"));
+				Exec.LogClear();
+			}
+			{
+				EStateTreeRunStatus Status = Exec.Tick(1.0f);
+				AITEST_EQUAL("3rd Tick should complete with Running", Status, EStateTreeRunStatus::Running);
+				FTestStateTreeExecutionContext::FLogOrder LogOrder = Exec.Expect("Tree1GlobalTaskA", TEXT("Tick11"));
+				AITEST_TRUE("3rd Tick should tick Tree1GlobalTaskA", LogOrder);
+				LogOrder = LogOrder.Then("Tree2GlobalTaskA", TEXT("Tick21"));
+				AITEST_TRUE("3rd Tick should tick Tree2GlobalTaskA", LogOrder);
+				LogOrder = LogOrder.Then("Tree1StateRootTaskA", TEXT("Tick12"));
+				AITEST_TRUE("3rd Tick should tick Tree1StateRootTaskA", LogOrder);
+				LogOrder = LogOrder.Then("Tree2StateRootTaskA", TEXT("Tick22"));
+				AITEST_TRUE("3rd Tick should tick Tree2StateRootTaskA", LogOrder);
+				AITEST_TRUE("3rd Tick should be in the correct states", Exec.ExpectInActiveStates("Tree1StateRoot", "Tree1StateLinkedTree", "Tree2StateRoot", "Tree2StateB"));
+				Exec.LogClear();
+			}
+			{
+				EStateTreeRunStatus Status = Exec.Tick(1.0f);
+				AITEST_EQUAL("4th Tick should complete with Running", Status, EStateTreeRunStatus::Running);
+				FTestStateTreeExecutionContext::FLogOrder LogOrder = Exec.Expect("Tree1GlobalTaskA", TEXT("Tick11"));
+				AITEST_TRUE("4th Tick should tick Tree1GlobalTaskA", LogOrder);
+				LogOrder = LogOrder.Then("Tree2GlobalTaskA", TEXT("Tick21"));
+				AITEST_TRUE("4th Tick should tick Tree2GlobalTaskA", LogOrder);
+				LogOrder = LogOrder.Then("Tree1StateRootTaskA", TEXT("Tick12"));
+				AITEST_TRUE("4th Tick should tick Tree1StateRootTaskA", LogOrder);
+				LogOrder = LogOrder.Then("Tree2StateRootTaskA", TEXT("Tick22"));
+				AITEST_TRUE("4th Tick should tick Tree2StateRootTaskA", LogOrder);
+				LogOrder = LogOrder.Then("Tree2StateRootTaskA", TEXT("ExitState22"));
+				AITEST_TRUE("4th Tick should exit Tree2StateRootTaskA", LogOrder);
+				LogOrder = LogOrder.Then("Tree2GlobalTaskA", TEXT("ExitState21"));
+				AITEST_TRUE("4th Tick should exit Tree2GlobalTaskA", LogOrder);
+				LogOrder = LogOrder.Then("Tree1StateBTaskA", TEXT("EnterState14"));
+				AITEST_TRUE("4th Tick should enter Tree1StateBTaskA", LogOrder);
+				AITEST_TRUE("4th Tick should be in the correct states", Exec.ExpectInActiveStates("Tree1StateRoot", "Tree1StateB"));
+				Exec.LogClear();
+			}
+			{
+				EStateTreeRunStatus Status = Exec.Tick(0.001f);
+				AITEST_EQUAL("5th Tick should complete with Running", Status, EStateTreeRunStatus::Running);
+				FTestStateTreeExecutionContext::FLogOrder LogOrder = Exec.Expect("Tree1GlobalTaskA", TEXT("Tick11"));
+				AITEST_TRUE("5th Tick should tick Tree1GlobalTaskA", LogOrder);
+				LogOrder = LogOrder.Then("Tree1StateRootTaskA", TEXT("Tick12"));
+				AITEST_TRUE("5th Tick should tick Tree1StateRootTaskA", LogOrder);
+				LogOrder = LogOrder.Then("Tree1StateBTaskA", TEXT("Tick14"));
+				AITEST_TRUE("5th Tick should tick Tree1StateBTaskA", LogOrder);
+				AITEST_TRUE("5th Tick should be in the correct states", Exec.ExpectInActiveStates("Tree1StateRoot", "Tree1StateB"));
+				Exec.LogClear();
+			}
 			Exec.Stop();
 		}
 
diff --git a/Engine/Plugins/Runtime/StateTree/Source/StateTreeTestSuite/Private/StateTreeTestTypes.h b/Engine/Plugins/Runtime/StateTree/Source/StateTreeTestSuite/Private/StateTreeTestTypes.h
--- a/Engine/Plugins/Runtime/StateTree/Source/StateTreeTestSuite/Private/StateTreeTestTypes.h
+++ b/Engine/Plugins/Runtime/StateTree/Source/StateTreeTestSuite/Private/StateTreeTestTypes.h
@@ -69,41 +69,47 @@ struct FTestStateTreeExecutionContext : public FStateTreeExecutionContext
 
 	struct FLogOrder
 	{
-		FLogOrder(const FTestStateTreeExecutionContext& InContext, const int32 InIndex) : Context(InContext), Index(InIndex) {}
+		FLogOrder(const FTestStateTreeExecutionContext& InContext, const int32 InIndex)
+			: Context(&InContext), Index(InIndex)
+		{}
 
 		FLogOrder Then(const FName& Name) const
 		{
 			int32 NextIndex = Index;
-			while (NextIndex < Context.Log.Items.Num())
+			while (NextIndex < Context->Log.Items.Num())
 			{
-				const FStateTreeTestLog::FLogItem& Item = Context.Log.Items[NextIndex];
+				const FStateTreeTestLog::FLogItem& Item = Context->Log.Items[NextIndex];
 				if (Item.Name == Name)
 				{
 					break;
 				}
 				NextIndex++;
 			}
-			return FLogOrder(Context, NextIndex);
+			return FLogOrder(*Context, NextIndex);
 		}
 
 		FLogOrder Then(const FName& Name, const FString& Message) const
 		{
 			int32 NextIndex = Index;
-			while (NextIndex < Context.Log.Items.Num())
+			while (NextIndex < Context->Log.Items.Num())
 			{
-				const FStateTreeTestLog::FLogItem& Item = Context.Log.Items[NextIndex];
+				const FStateTreeTestLog::FLogItem& Item = Context->Log.Items[NextIndex];
 				if (Item.Name == Name && Item.Message == Message)
 				{
 					break;
 				}
 				NextIndex++;
 			}
-			return FLogOrder(Context, NextIndex);
+			return FLogOrder(*Context, NextIndex);
 		}
 
-		operator bool() const { return Index < Context.Log.Items.Num(); }
-		
-		const FTestStateTreeExecutionContext& Context;
+		operator bool() const
+		{
+			return Index < Context->Log.Items.Num();
+		}
+
+	private:
+		const FTestStateTreeExecutionContext* Context;
 		int32 Index = 0;
 	};
 
@@ -223,6 +229,31 @@ struct FTestTask_PrintValueInstanceData
 
 	UPROPERTY(EditAnywhere, Category = "Parameter")
 	int32 Value = 0;
+	
+	UPROPERTY(EditAnywhere, Category = "Parameter")
+	TArray<int32> ArrayValue;
+
+	FString GetArrayDescription() const
+	{
+		if (ArrayValue.IsEmpty())
+		{
+			return TEXT("{}");
+		}
+		TStringBuilder<256> StringBuilder;
+		StringBuilder << TEXT('{');
+		bool bFirst = true;
+		for (int32 It : ArrayValue)
+		{
+			if (!bFirst)
+			{
+				StringBuilder << TEXT(", ");
+			}
+			bFirst = false;
+			StringBuilder << It;
+		}
+		StringBuilder << TEXT('}');
+		return StringBuilder.ToString();
+	}
 };
 
 USTRUCT()
@@ -243,6 +274,7 @@ struct FTestTask_PrintValue : public FStateTreeTaskBase
 		FStateTreeTestLog& TestLog = Context.GetExternalData(LogHandle);
 		const FInstanceDataType& InstanceData = Context.GetInstanceData(*this);
 		TestLog.Log(Name,  FString::Printf(TEXT("EnterState%d"), InstanceData.Value));
+		TestLog.Log(Name, FString::Printf(TEXT("EnterState:%s"), *InstanceData.GetArrayDescription()));
 		return EStateTreeRunStatus::Running;
 	}
 
@@ -251,6 +283,7 @@ struct FTestTask_PrintValue : public FStateTreeTaskBase
 		FStateTreeTestLog& TestLog = Context.GetExternalData(LogHandle);
 		const FInstanceDataType& InstanceData = Context.GetInstanceData(*this);
 		TestLog.Log(Name,  FString::Printf(TEXT("ExitState%d"), InstanceData.Value));
+		TestLog.Log(Name, FString::Printf(TEXT("ExitState:%s"), *InstanceData.GetArrayDescription()));
 	}
 
 	virtual EStateTreeRunStatus Tick(FStateTreeExecutionContext& Context, const float DeltaTime) const override
@@ -258,9 +291,18 @@ struct FTestTask_PrintValue : public FStateTreeTaskBase
 		FStateTreeTestLog& TestLog = Context.GetExternalData(LogHandle);
 		const FInstanceDataType& InstanceData = Context.GetInstanceData(*this);
 		TestLog.Log(Name,  FString::Printf(TEXT("Tick%d"), InstanceData.Value));
+		TestLog.Log(Name,  FString::Printf(TEXT("Tick:%s"), *InstanceData.GetArrayDescription()));
 		
 		return EStateTreeRunStatus::Running;
-	};
+	}
+
+	virtual void TriggerTransitions(FStateTreeExecutionContext& Context) const override
+	{
+		FStateTreeTestLog& TestLog = Context.GetExternalData(LogHandle);
+		const FInstanceDataType& InstanceData = Context.GetInstanceData(*this);
+		TestLog.Log(Name, FString::Printf(TEXT("TriggerTransitions%d"), InstanceData.Value));
+		TestLog.Log(Name, FString::Printf(TEXT("TriggerTransitions:%s"), *InstanceData.GetArrayDescription()));
+	}
 
 	virtual bool Link(FStateTreeLinker& Linker) override
 	{
@@ -282,7 +324,8 @@ struct FTestTask_PrintAndResetValue : public FTestTask_PrintValue
 	{
 		EStateTreeRunStatus Status = Super::EnterState(Context, Transition);
 		FInstanceDataType& InstanceData = Context.GetInstanceData(*this);
-		InstanceData.Value = 0;
+		InstanceData.Value = ResetValue;
+		InstanceData.ArrayValue = ResetArrayValue;
 		return Status;
 	}
 
@@ -290,9 +333,60 @@ struct FTestTask_PrintAndResetValue : public FTestTask_PrintValue
 	{
 		EStateTreeRunStatus Status = Super::Tick(Context, DeltaTime);
 		FInstanceDataType& InstanceData = Context.GetInstanceData(*this);
-		InstanceData.Value = 0;
+		InstanceData.Value = ResetValue;
+		InstanceData.ArrayValue = ResetArrayValue;
 		return Status;
-	};
+	}
+
+	virtual void ExitState(FStateTreeExecutionContext& Context, const FStateTreeTransitionResult& Transition) const override
+	{
+		Super::ExitState(Context, Transition);
+		FInstanceDataType& InstanceData = Context.GetInstanceData(*this);
+		InstanceData.Value = ResetValue;
+		InstanceData.ArrayValue = ResetArrayValue;
+	}
+
+	UPROPERTY()
+	int32 ResetValue = 0;
+
+	UPROPERTY()
+	TArray<int32> ResetArrayValue;
+};
+
+USTRUCT()
+struct FTestTask_PrintValue_TransitionTick : public FTestTask_PrintValue
+{
+	GENERATED_BODY()
+
+	FTestTask_PrintValue_TransitionTick()
+	{
+		bShouldCallTick = true;
+		bShouldAffectTransitions = true;
+	}
+	FTestTask_PrintValue_TransitionTick(const FName InName)
+		: FTestTask_PrintValue(InName)
+	{
+		bShouldCallTick = true;
+		bShouldAffectTransitions = true;
+	}
+};
+
+USTRUCT()
+struct FTestTask_PrintValue_TransitionNoTick : public FTestTask_PrintValue
+{
+	GENERATED_BODY()
+
+	FTestTask_PrintValue_TransitionNoTick()
+	{
+		bShouldCallTick = false;
+		bShouldAffectTransitions = true;
+	}
+	FTestTask_PrintValue_TransitionNoTick(const FName InName)
+		: FTestTask_PrintValue(InName)
+	{
+		bShouldCallTick = false;
+		bShouldAffectTransitions = true;
+	}
 };
 
 
