StateTree: Make sure there is always one editor color. Fix states color ref when it does not exist.

diff --git a/Engine/Plugins/Runtime/StateTree/Source/StateTreeEditorModule/Private/StateTreeEditorData.cpp b/Engine/Plugins/Runtime/StateTree/Source/StateTreeEditorModule/Private/StateTreeEditorData.cpp
--- a/Engine/Plugins/Runtime/StateTree/Source/StateTreeEditorModule/Private/StateTreeEditorData.cpp
+++ b/Engine/Plugins/Runtime/StateTree/Source/StateTreeEditorModule/Private/StateTreeEditorData.cpp
@@ -52,16 +52,20 @@ namespace UE::StateTree::Editor
 
 		return false;
 	}
+
+	FStateTreeEditorColor CreateDefaultColor()
+	{
+		FStateTreeEditorColor DefaultColor;
+		DefaultColor.ColorRef = FStateTreeEditorColorRef();
+		DefaultColor.Color = FLinearColor(FColor(31, 151, 167));
+		DefaultColor.DisplayName = TEXT("Default Color");
+		return DefaultColor;
+	}
 }
 
 UStateTreeEditorData::UStateTreeEditorData()
 {
-	FStateTreeEditorColor DefaultColor;
-	DefaultColor.ColorRef = FStateTreeEditorColorRef();
-	DefaultColor.Color = FLinearColor(FColor(31, 151, 167));
-	DefaultColor.DisplayName = TEXT("Default Color");
-
-	Colors.Add(MoveTemp(DefaultColor));
+	Colors.Add(UE::StateTree::Editor::CreateDefaultColor());
 
 	EditorBindings.SetBindingsOwner(this);
 }
@@ -286,9 +290,24 @@ void UStateTreeEditorData::PostEditChangeChainProperty(FPropertyChangedChainEven
 			UE::StateTree::Delegates::OnGlobalDataChanged.Broadcast(*StateTree);
 		}
 
-		// Notify that the color data has changed
+		// Notify that the color data has changed and fix existing data
 		if (MemberName == GET_MEMBER_NAME_CHECKED(UStateTreeEditorData, Colors))
 		{
+			if (Colors.IsEmpty())
+			{
+				// Add default color
+				Colors.Add(UE::StateTree::Editor::CreateDefaultColor());
+			}
+			VisitHierarchy([Self=this](UStateTreeState& State, UStateTreeState* ParentState)
+			{
+				if (!Self->FindColor(State.ColorRef))
+				{
+					State.Modify();
+					State.ColorRef = FStateTreeEditorColorRef();
+				}
+				return EStateTreeVisitor::Continue;
+			});
+
 			UE::StateTree::Delegates::OnVisualThemeChanged.Broadcast(*StateTree);
 		}
 	}
