StateTree: Add the delegate feature. Delegates allow event-based communication between tasks or to trigger a transition.

diff --git a/Engine/Plugins/Runtime/StateTree/Source/StateTreeEditorModule/Private/Customizations/StateTreeBindingExtension.cpp b/Engine/Plugins/Runtime/StateTree/Source/StateTreeEditorModule/Private/Customizations/StateTreeBindingExtension.cpp
--- a/Engine/Plugins/Runtime/StateTree/Source/StateTreeEditorModule/Private/Customizations/StateTreeBindingExtension.cpp
+++ b/Engine/Plugins/Runtime/StateTree/Source/StateTreeEditorModule/Private/Customizations/StateTreeBindingExtension.cpp
@@ -26,6 +26,7 @@
 #include "StateTreeNodeClassCache.h"
 #include "StateTreePropertyFunctionBase.h"
 #include "StateTreeEditorData.h"
+#include "StateTreeDelegate.h"
 
 #define LOCTEXT_NAMESPACE "StateTreeEditor"
 
@@ -300,6 +301,16 @@ TSharedRef<SWidget> MakeContextStructWidget(const FStateTreeBindableStructDesc&
 		];
 }
 
+bool IsDelegateListenerProperty(const FProperty& Property)
+{
+	if (const FStructProperty* StructProperty = CastField<FStructProperty>(&Property))
+	{
+		return StructProperty->Struct == FStateTreeDelegateListener::StaticStruct() || StructProperty->Struct == FStateTreeTransitionDelegateListener::StaticStruct();
+	}
+
+	return false;
+}
+
 TSharedRef<SWidget> MakeBindingPropertyInfoWidget(const FText& InDisplayText, const FEdGraphPinType& InPinType)
 {
 	const FSlateBrush* Icon = FBlueprintEditorUtils::GetIconFromPin(InPinType, /*bIsLarge*/true);
@@ -689,8 +700,6 @@ struct FCachedBindingData : public TSharedFromThis<FCachedBindingData>
 
 	void UpdateData()
 	{
-		static FName PropertyIcon(TEXT("Kismet.Tabs.Variables"));
-
 		SourceStructName = FText::GetEmpty();
 		FormatableText = FText::GetEmpty();
 		FormatableTooltipText = FText::GetEmpty();
@@ -733,6 +742,7 @@ struct FCachedBindingData : public TSharedFromThis<FCachedBindingData>
 		BindingOwner->GetDataViewByID(TargetPath.GetStructID(), TargetDataView);
 
 		FEdGraphPinType PinType;
+		FName IconName;
 		const bool bIsPropertyRef = UE::StateTree::PropertyRefHelpers::IsPropertyRef(*Property);
 		if (bIsPropertyRef && TargetDataView.IsValid())
 		{
@@ -743,10 +753,17 @@ struct FCachedBindingData : public TSharedFromThis<FCachedBindingData>
 				const void* PropertyRef = TargetIndirections.Last().GetPropertyAddress();
 				PinType = UE::StateTree::PropertyRefHelpers::GetPropertyRefInternalTypeAsPin(*Property, PropertyRef);
 			}
+			IconName = TEXT("Kismet.Tabs.Variables");
+		}
+		else if (UE::StateTree::PropertyBinding::IsDelegateListenerProperty(*Property))
+		{
+			PinType.PinCategory = UEdGraphSchema_K2::PC_Delegate;
+			IconName = TEXT("Icons.Event");
 		}
 		else
 		{
 			Schema->ConvertPropertyToPinType(Property, PinType);
+			IconName = TEXT("Kismet.Tabs.Variables");
 		}
 
 
@@ -839,7 +856,7 @@ struct FCachedBindingData : public TSharedFromThis<FCachedBindingData>
 						Schema->ConvertPropertyToPinType(SourceLeafProperty, PinType);
 					}
 
-					Image = FAppStyle::GetBrush(PropertyIcon);
+					Image = FAppStyle::GetBrush(IconName);
 					Color = Schema->GetPinTypeColor(PinType);
 				}
 				else
@@ -894,7 +911,7 @@ struct FCachedBindingData : public TSharedFromThis<FCachedBindingData>
 			FormatableText = FText::GetEmpty();
 			TooltipBuilder.AppendLineFormat(LOCTEXT("BindTooltip", "Bind {0} to value from another property."), UE::StateTree::PropertyBinding::GetPropertyTypeText(Property));
 
-			Image = FAppStyle::GetBrush(PropertyIcon);
+			Image = FAppStyle::GetBrush(IconName);
 			Color = Schema->GetPinTypeColor(PinType);
 
 			CachedSourcePath.Reset();
@@ -1140,6 +1157,13 @@ struct FCachedBindingData : public TSharedFromThis<FCachedBindingData>
 			check(TargetPropertyValue);
 			bCanBind = UE::StateTree::PropertyRefHelpers::IsPropertyRefCompatibleWithProperty(*TargetStructProperty, *SourceProperty, TargetPropertyValue, SourcePropertyValue);
 		}
+		else if (TargetStructProperty && UE::StateTree::PropertyBinding::IsDelegateListenerProperty(*TargetStructProperty))
+		{
+			if (const FStructProperty* SourceStructProperty = CastField<FStructProperty>(SourceProperty))
+			{
+				bCanBind = SourceStructProperty->Struct == FStateTreeDelegateDispatcher::StaticStruct();
+			}
+		}
 		else
 		{
 			// Note: We support type promotion here
@@ -1293,26 +1317,26 @@ struct FCachedBindingData : public TSharedFromThis<FCachedBindingData>
 		// Bound PropertyFunction is allowed to override it's icon.
 		if (IStateTreeEditorPropertyBindingsOwner* BindingOwner = Cast<IStateTreeEditorPropertyBindingsOwner>(WeakOwnerObject.Get()))
 		{
-				if(FStateTreeEditorPropertyBindings* EditorBindings = BindingOwner->GetPropertyEditorBindings())
+			if(FStateTreeEditorPropertyBindings* EditorBindings = BindingOwner->GetPropertyEditorBindings())
+			{
+				const FStateTreePropertyPathBinding* CurrentBinding = EditorBindings->GetBindings().FindByPredicate([this](const FStateTreePropertyPathBinding& Binding)
 				{
-					const FStateTreePropertyPathBinding* CurrentBinding = EditorBindings->GetBindings().FindByPredicate([this](const FStateTreePropertyPathBinding& Binding)
-					{
-						return Binding.GetTargetPath() == TargetPath;
-					});
+					return Binding.GetTargetPath() == TargetPath;
+				});
 
-					if (CurrentBinding)
+				if (CurrentBinding)
+				{
+					const FConstStructView PropertyFunctionEditorNodeView = CurrentBinding->GetPropertyFunctionNode();
+					if (PropertyFunctionEditorNodeView.IsValid())
 					{
-						const FConstStructView PropertyFunctionEditorNodeView = CurrentBinding->GetPropertyFunctionNode();
-						if (PropertyFunctionEditorNodeView.IsValid())
+						const FStateTreeEditorNode& EditorNode = PropertyFunctionEditorNodeView.Get<const FStateTreeEditorNode>();
+						if (const FStateTreeNodeBase* Node = EditorNode.Node.GetPtr<const FStateTreeNodeBase>())
 						{
-							const FStateTreeEditorNode& EditorNode = PropertyFunctionEditorNodeView.Get<const FStateTreeEditorNode>();
-							if (const FStateTreeNodeBase* Node = EditorNode.Node.GetPtr<const FStateTreeNodeBase>())
-							{
-								return UE::StateTreeEditor::EditorNodeUtils::ParseIcon(Node->GetIconName()).GetIcon();
-							}
+							return UE::StateTreeEditor::EditorNodeUtils::ParseIcon(Node->GetIconName()).GetIcon();
 						}
 					}
 				}
+			}
 		}
 
 		return Image;
@@ -1484,6 +1508,15 @@ bool FStateTreeBindingExtension::IsPropertyExtendable(const UClass* InObjectClas
 	// Figure out the structs we're editing, and property path relative to current property.
 	const EStateTreePropertyUsage Usage = UE::StateTree::PropertyBinding::MakeStructPropertyPathFromPropertyHandle(PropertyHandle.AsShared(), TargetPath);
 
+	// Binding from Delegate Dispatcher is not allowed.
+	if (const FStructProperty* StructProperty = CastField<FStructProperty>(PropertyHandle.GetProperty()))
+	{
+		if (StructProperty->Struct == FStateTreeDelegateDispatcher::StaticStruct())
+		{
+			return false;
+		}
+	}
+
 	if (Usage == EStateTreePropertyUsage::Input || Usage == EStateTreePropertyUsage::Context)
 	{
 		// Allow to bind only to the main level on input and context properties.
diff --git a/Engine/Plugins/Runtime/StateTree/Source/StateTreeEditorModule/Private/Customizations/StateTreeEditorNodeDetails.cpp b/Engine/Plugins/Runtime/StateTree/Source/StateTreeEditorModule/Private/Customizations/StateTreeEditorNodeDetails.cpp
--- a/Engine/Plugins/Runtime/StateTree/Source/StateTreeEditorModule/Private/Customizations/StateTreeEditorNodeDetails.cpp
+++ b/Engine/Plugins/Runtime/StateTree/Source/StateTreeEditorModule/Private/Customizations/StateTreeEditorNodeDetails.cpp
@@ -42,6 +42,7 @@
 #include "Styling/SlateTypes.h"
 #include "TextStyleDecorator.h"
 #include "Framework/Application/SlateApplication.h"
+#include "StateTreeDelegate.h"
 
 #define LOCTEXT_NAMESPACE "StateTreeEditor"
 
@@ -108,6 +109,17 @@ namespace UE::StateTreeEditor::Internal
 		return UEdGraphSchema_K2::GetCategoryText(PinType.PinCategory, NAME_None, true);
 	}
 
+	/** @return if property is struct property of DelegateDispatcher type. */
+	bool IsDelegateDispatcherProperty(const FProperty& Property)
+	{
+		if (const FStructProperty* StructProperty = CastField<FStructProperty>(&Property))
+		{
+			return StructProperty->Struct == FStateTreeDelegateDispatcher::StaticStruct();
+		}
+
+		return false;
+	}
+
 	/** @return UClass or UScriptStruct of class or struct property, nullptr for others. */
 	UStruct* GetPropertyStruct(TSharedPtr<IPropertyHandle> PropHandle)
 	{
@@ -164,7 +176,9 @@ namespace UE::StateTreeEditor::Internal
 
 			const bool bHasChildPropertyBinding = EditorPropBindings->HasPropertyBinding(Path, FStateTreeEditorPropertyBindings::ESearchMode::Includes);
 			const bool bValidUsage = Usage == EStateTreePropertyUsage::Input || Usage == EStateTreePropertyUsage::Output || Usage == EStateTreePropertyUsage::Context;
-			if (bHasChildPropertyBinding || bValidUsage)
+			const bool bIsDelegateDispatcher = UE::StateTreeEditor::Internal::IsDelegateDispatcherProperty(*Property);
+
+			if (bHasChildPropertyBinding || bValidUsage || bIsDelegateDispatcher)
 			{				
 				FEdGraphPinType PinType;
 				const UEdGraphSchema_K2* Schema = GetDefault<UEdGraphSchema_K2>();
@@ -203,7 +217,16 @@ namespace UE::StateTreeEditor::Internal
 				FText LabelToolTip;
 				FSlateColor TextColor = FSlateColor::UseForeground();
 
-				if (Usage == EStateTreePropertyUsage::Input)
+				if (bIsDelegateDispatcher)
+				{
+					Label = LOCTEXT("LabelDelegate", "DELEGATE");
+					LabelToolTip = LOCTEXT("DelegateToolTip", "This is Delegate Dispatcher. You can bind to it from listeners.");
+
+					FEdGraphPinType DelegatePinType;
+					DelegatePinType.PinCategory = UEdGraphSchema_K2::PC_Delegate;
+					IconColor = Schema->GetPinTypeColor(DelegatePinType);
+				}
+				else if (Usage == EStateTreePropertyUsage::Input)
 				{
 					Label = LOCTEXT("LabelInput", "IN");
 					LabelToolTip = LOCTEXT("InputToolTip", "This is Input property. It is always expected to be bound to some other property.");
diff --git a/Engine/Plugins/Runtime/StateTree/Source/StateTreeEditorModule/Private/Customizations/StateTreeTransitionDetails.cpp b/Engine/Plugins/Runtime/StateTree/Source/StateTreeEditorModule/Private/Customizations/StateTreeTransitionDetails.cpp
--- a/Engine/Plugins/Runtime/StateTree/Source/StateTreeEditorModule/Private/Customizations/StateTreeTransitionDetails.cpp
+++ b/Engine/Plugins/Runtime/StateTree/Source/StateTreeEditorModule/Private/Customizations/StateTreeTransitionDetails.cpp
@@ -24,6 +24,7 @@
 #include "Widgets/Input/SButton.h"
 #include "TextStyleDecorator.h"
 #include "Framework/Application/SlateApplication.h"
+#include "StateTreeBindingExtension.h"
 
 #define LOCTEXT_NAMESPACE "StateTreeEditor"
 
@@ -55,6 +56,7 @@ void FStateTreeTransitionDetails::CustomizeHeader(TSharedRef<IPropertyHandle> St
 	TriggerProperty = StructProperty->GetChildHandle(GET_MEMBER_NAME_CHECKED(FStateTreeTransition, Trigger));
 	PriorityProperty = StructProperty->GetChildHandle(GET_MEMBER_NAME_CHECKED(FStateTreeTransition, Priority));
 	RequiredEventProperty = StructProperty->GetChildHandle(GET_MEMBER_NAME_CHECKED(FStateTreeTransition, RequiredEvent));
+	DelegateListener = StructProperty->GetChildHandle(GET_MEMBER_NAME_CHECKED(FStateTreeTransition, DelegateListener));
 	StateProperty = StructProperty->GetChildHandle(GET_MEMBER_NAME_CHECKED(FStateTreeTransition, State));
 	DelayTransitionProperty = StructProperty->GetChildHandle(GET_MEMBER_NAME_CHECKED(FStateTreeTransition, bDelayTransition));
 	DelayDurationProperty = StructProperty->GetChildHandle(GET_MEMBER_NAME_CHECKED(FStateTreeTransition, DelayDuration));
@@ -148,6 +150,7 @@ void FStateTreeTransitionDetails::CustomizeChildren(TSharedRef<class IPropertyHa
 {
 	check(TriggerProperty);
 	check(RequiredEventProperty);
+	check(DelegateListener);
 	check(DelayTransitionProperty);
 	check(DelayDurationProperty);
 	check(DelayRandomVarianceProperty);
@@ -156,7 +159,7 @@ void FStateTreeTransitionDetails::CustomizeChildren(TSharedRef<class IPropertyHa
 	check(IDProperty);
 
 	TWeakPtr<FStateTreeTransitionDetails> WeakSelf = SharedThis(this);
-	auto IsTickOrEventTransition = [WeakSelf]()
+	auto IsNotCompletionTransition = [WeakSelf]()
 	{
 		if (const TSharedPtr<FStateTreeTransitionDetails> Self = WeakSelf.Pin())
 		{
@@ -184,20 +187,34 @@ void FStateTreeTransitionDetails::CustomizeChildren(TSharedRef<class IPropertyHa
 			return EVisibility::Collapsed;
 		})));
 
+	IDetailPropertyRow& DelegateDispatcherRow = StructBuilder.AddProperty(DelegateListener.ToSharedRef())
+		.Visibility(TAttribute<EVisibility>::Create(TAttribute<EVisibility>::FGetter::CreateLambda([WeakSelf]()
+		{
+			if (const TSharedPtr<FStateTreeTransitionDetails> Self = WeakSelf.Pin())
+			{
+				return (Self->GetTrigger() == EStateTreeTransitionTrigger::OnDelegate) ? EVisibility::Visible : EVisibility::Collapsed;
+			}
+			return EVisibility::Collapsed;
+		})));
+
+	FGuid ID;
+	UE::StateTree::PropertyHelpers::GetStructValue<FGuid>(IDProperty, ID);
+	DelegateListener->SetInstanceMetaData(UE::StateTree::PropertyBinding::StateTreeNodeIDName, LexToString(ID));
+
 	// State
 	StructBuilder.AddProperty(StateProperty.ToSharedRef());
 
 	// Priority
 	StructBuilder.AddProperty(PriorityProperty.ToSharedRef())
-		.Visibility(TAttribute<EVisibility>::Create(TAttribute<EVisibility>::FGetter::CreateLambda(IsTickOrEventTransition)));
+		.Visibility(TAttribute<EVisibility>::Create(TAttribute<EVisibility>::FGetter::CreateLambda(IsNotCompletionTransition)));
 
 	// Delay
 	StructBuilder.AddProperty(DelayTransitionProperty.ToSharedRef())
-		.Visibility(TAttribute<EVisibility>::Create(TAttribute<EVisibility>::FGetter::CreateLambda(IsTickOrEventTransition)));
+		.Visibility(TAttribute<EVisibility>::Create(TAttribute<EVisibility>::FGetter::CreateLambda(IsNotCompletionTransition)));
 	StructBuilder.AddProperty(DelayDurationProperty.ToSharedRef())
-		.Visibility(TAttribute<EVisibility>::Create(TAttribute<EVisibility>::FGetter::CreateLambda(IsTickOrEventTransition)));
+		.Visibility(TAttribute<EVisibility>::Create(TAttribute<EVisibility>::FGetter::CreateLambda(IsNotCompletionTransition)));
 	StructBuilder.AddProperty(DelayRandomVarianceProperty.ToSharedRef())
-		.Visibility(TAttribute<EVisibility>::Create(TAttribute<EVisibility>::FGetter::CreateLambda(IsTickOrEventTransition)));
+		.Visibility(TAttribute<EVisibility>::Create(TAttribute<EVisibility>::FGetter::CreateLambda(IsNotCompletionTransition)));
 
 	// Show conditions always expanded, with simplified header (remove item count)
 	IDetailPropertyRow& ConditionsRow = StructBuilder.AddProperty(ConditionsProperty.ToSharedRef());
diff --git a/Engine/Plugins/Runtime/StateTree/Source/StateTreeEditorModule/Private/Customizations/StateTreeTransitionDetails.h b/Engine/Plugins/Runtime/StateTree/Source/StateTreeEditorModule/Private/Customizations/StateTreeTransitionDetails.h
--- a/Engine/Plugins/Runtime/StateTree/Source/StateTreeEditorModule/Private/Customizations/StateTreeTransitionDetails.h
+++ b/Engine/Plugins/Runtime/StateTree/Source/StateTreeEditorModule/Private/Customizations/StateTreeTransitionDetails.h
@@ -50,6 +50,7 @@ private:
 	TSharedPtr<IPropertyHandle> TriggerProperty;
 	TSharedPtr<IPropertyHandle> PriorityProperty;
 	TSharedPtr<IPropertyHandle> RequiredEventProperty;
+	TSharedPtr<IPropertyHandle> DelegateListener;
 	TSharedPtr<IPropertyHandle> StateProperty;
 	TSharedPtr<IPropertyHandle> DelayTransitionProperty;
 	TSharedPtr<IPropertyHandle> DelayDurationProperty;
diff --git a/Engine/Plugins/Runtime/StateTree/Source/StateTreeEditorModule/Private/SStateTreeViewRow.cpp b/Engine/Plugins/Runtime/StateTree/Source/StateTreeEditorModule/Private/SStateTreeViewRow.cpp
--- a/Engine/Plugins/Runtime/StateTree/Source/StateTreeEditorModule/Private/SStateTreeViewRow.cpp
+++ b/Engine/Plugins/Runtime/StateTree/Source/StateTreeEditorModule/Private/SStateTreeViewRow.cpp
@@ -64,7 +64,7 @@ void SStateTreeViewRow::Construct(const FArguments& InArgs, const TSharedRef<STa
 	auto MakeTransitionWidget = [this](const EStateTreeTransitionTrigger Trigger, const FSlateBrush* Icon)
 	{
 		FTransitionDescFilterOptions FilterOptions;
-		FilterOptions.bUseMask = EnumHasAnyFlags(Trigger, EStateTreeTransitionTrigger::OnTick | EStateTreeTransitionTrigger::OnEvent);
+		FilterOptions.bUseMask = EnumHasAnyFlags(Trigger, EStateTreeTransitionTrigger::OnTick | EStateTreeTransitionTrigger::OnEvent| EStateTreeTransitionTrigger::OnDelegate);
 		
 			return SNew(SBox)
 			.HeightOverride(GetDefault<UStateTreeEditorUserSettings>()->GetStatesViewStateRowHeight())
@@ -443,7 +443,7 @@ void SStateTreeViewRow::Construct(const FArguments& InArgs, const TSharedRef<STa
 			.VAlign(VAlign_Top)
 			.AutoWidth()
 			[
-				MakeTransitionWidget(EStateTreeTransitionTrigger::OnTick | EStateTreeTransitionTrigger::OnEvent, FStateTreeEditorStyle::Get().GetBrush("StateTreeEditor.Transition.Condition"))
+				MakeTransitionWidget(EStateTreeTransitionTrigger::OnTick | EStateTreeTransitionTrigger::OnEvent | EStateTreeTransitionTrigger::OnDelegate, FStateTreeEditorStyle::Get().GetBrush("StateTreeEditor.Transition.Condition"))
 			]
 		];
 	}
@@ -1606,7 +1606,7 @@ TSharedRef<SWidget> SStateTreeViewRow::MakeTransitionWidgets(const EStateTreeTra
 	}
 
 	// Find states from transition tasks
-	if (EnumHasAnyFlags(Trigger, EStateTreeTransitionTrigger::OnTick | EStateTreeTransitionTrigger::OnEvent))
+	if (EnumHasAnyFlags(Trigger, EStateTreeTransitionTrigger::OnTick | EStateTreeTransitionTrigger::OnEvent | EStateTreeTransitionTrigger::OnDelegate))
 	{
 		auto AddLinksFromStruct = [&DescItems, TreeEditorData](FStateTreeDataView Struct, const FGuid NodeID)
 		{
@@ -1772,7 +1772,7 @@ FText SStateTreeViewRow::GetTransitionsDesc(const EStateTreeTransitionTrigger Tr
 	}
 
 	// Find states from transition tasks
-	if (EnumHasAnyFlags(Trigger, EStateTreeTransitionTrigger::OnTick | EStateTreeTransitionTrigger::OnEvent))
+	if (EnumHasAnyFlags(Trigger, EStateTreeTransitionTrigger::OnTick | EStateTreeTransitionTrigger::OnEvent| EStateTreeTransitionTrigger::OnDelegate))
 	{
 		auto AddLinksFromStruct = [EditorData, &DescItems](FStateTreeDataView Struct)
 		{
@@ -1829,7 +1829,7 @@ const FSlateBrush* SStateTreeViewRow::GetTransitionsIcon(const EStateTreeTransit
 		return nullptr;
 	}
 
-	if (EnumHasAnyFlags(Trigger, EStateTreeTransitionTrigger::OnTick | EStateTreeTransitionTrigger::OnEvent))
+	if (EnumHasAnyFlags(Trigger, EStateTreeTransitionTrigger::OnTick | EStateTreeTransitionTrigger::OnEvent| EStateTreeTransitionTrigger::OnDelegate))
 	{
 		return FStateTreeEditorStyle::Get().GetBrush("StateTreeEditor.Transition.Goto");
 	}
@@ -1965,7 +1965,7 @@ EVisibility SStateTreeViewRow::GetTransitionsVisibility(const EStateTreeTransiti
 	}
 
 	// Find states from transition tasks
-	if (EnumHasAnyFlags(Trigger, EStateTreeTransitionTrigger::OnTick | EStateTreeTransitionTrigger::OnEvent))
+	if (EnumHasAnyFlags(Trigger, EStateTreeTransitionTrigger::OnTick | EStateTreeTransitionTrigger::OnEvent| EStateTreeTransitionTrigger::OnDelegate))
 	{
 		auto HasAnyLinksInStruct = [](FStateTreeDataView Struct) -> bool
 		{
diff --git a/Engine/Plugins/Runtime/StateTree/Source/StateTreeEditorModule/Private/StateTreeCompiler.cpp b/Engine/Plugins/Runtime/StateTree/Source/StateTreeEditorModule/Private/StateTreeCompiler.cpp
--- a/Engine/Plugins/Runtime/StateTree/Source/StateTreeEditorModule/Private/StateTreeCompiler.cpp
+++ b/Engine/Plugins/Runtime/StateTree/Source/StateTreeEditorModule/Private/StateTreeCompiler.cpp
@@ -16,6 +16,7 @@
 #include "StateTreePropertyRefHelpers.h"
 #include "StateTreePropertyHelpers.h"
 #include "StateTreePropertyFunctionBase.h"
+#include "StateTreeDelegate.h"
 
 
 namespace UE::StateTree::Compiler
@@ -991,6 +992,33 @@ bool FStateTreeCompiler::CreateStateTransitions()
 			CompactTransition.Trigger = Transition.Trigger;
 			CompactTransition.Priority = Transition.Priority;
 
+			if (Transition.Trigger == EStateTreeTransitionTrigger::OnDelegate)
+			{
+				FStateTreePropertyPath DelegateBindingPath(Transition.ID, GET_MEMBER_NAME_CHECKED(FStateTreeTransition, DelegateListener));
+
+				const FStateTreePropertyPathBinding* Binding = EditorData->EditorBindings.FindPropertyBinding(DelegateBindingPath);
+
+				if (Binding == nullptr)
+				{
+					bTransitionSucceeded = false;
+					Log.Reportf(EMessageSeverity::Error,
+						TEXT("On Delegate Transition to '%s' requires to be bound to some delegate dispatcher."),
+						*Transition.State.Name.ToString());
+					continue;
+				}
+
+				CompactTransition.DelegateDispatcherID = BindingsCompiler.GetDispatcherIDFromPath(Binding->GetSourcePath());
+
+				if (!CompactTransition.DelegateDispatcherID.IsValid())
+				{
+					bTransitionSucceeded = false;
+					Log.Reportf(EMessageSeverity::Error,
+						TEXT("On Delegate Transition to '%s' is bound to unknown delegate dispatcher"),
+						*Transition.State.Name.ToString());
+					continue;
+				}
+			}
+
 			CompactTransition.bTransitionEnabled = Transition.bTransitionEnabled;
 
 			if (Transition.bDelayTransition)
@@ -1090,6 +1118,7 @@ bool FStateTreeCompiler::CreateStateTransitions()
 			{
 				CompactState.bHasTickTriggerTransitions |= EnumHasAnyFlags(Transition.Trigger, EStateTreeTransitionTrigger::OnTick);
 				CompactState.bHasEventTriggerTransitions |= EnumHasAnyFlags(Transition.Trigger, EStateTreeTransitionTrigger::OnEvent);
+				CompactState.bHasDelegateTriggerTransitions |= EnumHasAnyFlags(Transition.Trigger, EStateTreeTransitionTrigger::OnDelegate);
 				CompactState.bHasSucceededTriggerTransitions |= EnumHasAnyFlags(Transition.Trigger, EStateTreeTransitionTrigger::OnStateCompleted);
 				CompactState.bHasSucceededTriggerTransitions |= EnumHasAnyFlags(Transition.Trigger, EStateTreeTransitionTrigger::OnStateSucceeded);
 				CompactState.bHasFailedTriggerTransitions |= EnumHasAnyFlags(Transition.Trigger, EStateTreeTransitionTrigger::OnStateFailed);
@@ -1261,32 +1290,41 @@ bool FStateTreeCompiler::CreateBindingsForNodes(TConstArrayView<FStateTreeEditor
 
 bool FStateTreeCompiler::CreateBindingsForStruct(const FStateTreeBindableStructDesc& TargetStruct, FStateTreeDataView TargetValue, FStateTreeIndex16 PropertyFuncsBegin, FStateTreeIndex16 PropertyFuncsEnd, FStateTreeIndex16& OutBatchIndex)
 {
-	TArray<FStateTreePropertyPathBinding> CopyBindings;
-	TArray<FStateTreePropertyPathBinding> ReferenceBindings;
+	FValidatedPathBindings Bindings;
 
 	// Check that the bindings for this struct are still all valid.
-	if (!GetAndValidateBindings(TargetStruct, TargetValue, CopyBindings, ReferenceBindings))
+	if (!GetAndValidateBindings(TargetStruct, TargetValue, Bindings))
 	{
 		return false;
 	}
 
-	int32 BatchIndex = INDEX_NONE;
-
-	// Compile batch copy for this struct, we pass in all the bindings, the compiler will pick up the ones for the target structs.
-	if (!BindingsCompiler.CompileBatch(TargetStruct, CopyBindings, PropertyFuncsBegin, PropertyFuncsEnd, BatchIndex))
+	// Copy Bindings
 	{
-		return false;
+		int32 BatchIndex = INDEX_NONE;
+
+		// Compile batch copy for this struct, we pass in all the bindings, the compiler will pick up the ones for the target structs.
+		if (!BindingsCompiler.CompileBatch(TargetStruct, Bindings.CopyBindings, PropertyFuncsBegin, PropertyFuncsEnd, BatchIndex))
+		{
+			return false;
+		}
+
+		if (const auto Validation = UE::StateTree::Compiler::IsValidIndex16(BatchIndex); Validation.DidFail())
+		{
+			Validation.Log(Log, TEXT("CopiesBatchIndex"), TargetStruct);
+			return false;
+		}
+
+		OutBatchIndex = FStateTreeIndex16(BatchIndex);
 	}
 
-	if (const auto Validation = UE::StateTree::Compiler::IsValidIndex16(BatchIndex); Validation.DidFail())
+	// Delegate Bindings
+	if (!BindingsCompiler.CompileDelegates(TargetStruct, Bindings.DelegateSourceBindings, Bindings.DelegateTargetBindings, TargetValue))
 	{
-		Validation.Log(Log, TEXT("BatchIndex"), TargetStruct);
 		return false;
 	}
 
-	OutBatchIndex = FStateTreeIndex16(BatchIndex);
-
-	if (!BindingsCompiler.CompileReferences(TargetStruct, ReferenceBindings, TargetValue, IDToStructValue))
+	// Reference Bindings
+	if (!BindingsCompiler.CompileReferences(TargetStruct, Bindings.ReferenceBindings, TargetValue, IDToStructValue))
 	{
 		return false;
 	}
@@ -1941,12 +1979,11 @@ bool FStateTreeCompiler::ValidateStructRef(const FStateTreeBindableStructDesc& S
 }
 
 
-bool FStateTreeCompiler::GetAndValidateBindings(const FStateTreeBindableStructDesc& TargetStruct, FStateTreeDataView TargetValue, TArray<FStateTreePropertyPathBinding>& OutCopyBindings, TArray<FStateTreePropertyPathBinding>& OutReferenceBindings) const
+bool FStateTreeCompiler::GetAndValidateBindings(const FStateTreeBindableStructDesc& TargetStruct, FStateTreeDataView TargetValue, FValidatedPathBindings& OutValidatedBindings) const
 {
 	check(EditorData);
 	
-	OutCopyBindings.Reset();
-	OutReferenceBindings.Reset();
+	OutValidatedBindings = FValidatedPathBindings();
 
 	// If target struct is not set, nothing to do.
 	if (TargetStruct.Struct == nullptr)
@@ -1960,9 +1997,34 @@ bool FStateTreeCompiler::GetAndValidateBindings(const FStateTreeBindableStructDe
 	{
 		if (Binding.GetTargetPath().GetStructID() != TargetStruct.ID)
 		{
+			if(Binding.GetSourcePath().GetStructID() == TargetStruct.ID)
+			{
+				if (IsPropertyOfTypeOrChild(*FStateTreeDelegateDispatcher::StaticStruct(), TargetStruct, Binding.GetSourcePath()))
+				{
+					FStateTreePropertyPathBinding& BindingCopy = OutValidatedBindings.DelegateSourceBindings.Add_GetRef(Binding);
+					BindingCopy.SetSourceDataHandle(TargetStruct.DataHandle);
+				}
+			}
+
 			continue;
 		}
 
+		TArray<FStateTreePropertyPathBinding>* OutputBindings = nullptr;
+		if (IsPropertyOfTypeOrChild(*FStateTreeDelegateListener::StaticStruct(), TargetStruct, Binding.GetTargetPath()))
+		{
+			OutputBindings = &OutValidatedBindings.DelegateTargetBindings;
+		}
+		else if (IsPropertyOfTypeOrChild(*FStateTreePropertyRef::StaticStruct(), TargetStruct, Binding.GetTargetPath()))
+		{
+			OutputBindings = &OutValidatedBindings.ReferenceBindings;
+		}
+		else
+		{
+			OutputBindings = &OutValidatedBindings.CopyBindings;
+		}
+
+		check(OutputBindings);
+
 		// Source must be one of the source structs we have discovered in the tree.
 		const FGuid SourceStructID = Binding.GetSourcePath().GetStructID();
 		const FStateTreeBindableStructDesc* SourceStruct = BindingsCompiler.GetSourceStructDescByID(SourceStructID);
@@ -2039,7 +2101,7 @@ bool FStateTreeCompiler::GetAndValidateBindings(const FStateTreeBindableStructDe
 			continue;
 		}
 		
-		FStateTreePropertyPathBinding BindingCopy(Binding);
+		FStateTreePropertyPathBinding& BindingCopy = OutputBindings->Add_GetRef(Binding);
 		BindingCopy.SetSourceDataHandle(SourceStruct->DataHandle);
 
 		// Special case for AnyEnum. StateTreeBindingExtension allows AnyEnums to bind to other enum types.
@@ -2058,15 +2120,6 @@ bool FStateTreeCompiler::GetAndValidateBindings(const FStateTreeBindableStructDe
 				BindingCopy.GetMutableTargetPath().AddPathSegment(GET_MEMBER_NAME_STRING_CHECKED(FStateTreeAnyEnum, Value));
 			}
 		}
-
-		if (IsPropertyOfTypeOrChild(*FStateTreePropertyRef::StaticStruct(), TargetStruct, Binding.GetTargetPath()))
-		{
-			OutReferenceBindings.Add(BindingCopy);
-		}
-		else
-		{
-			OutCopyBindings.Add(BindingCopy);
-		}
 		
 		// Check if the bindings is for struct ref and validate the types.
 		if (!ValidateStructRef(*SourceStruct, Binding.GetSourcePath(), TargetStruct, Binding.GetTargetPath()))
@@ -2115,7 +2168,7 @@ bool FStateTreeCompiler::GetAndValidateBindings(const FStateTreeBindableStructDe
 				const void* PropertyRef = TargetIndirections.Last().GetPropertyAddress();
 				const bool bIsOptional = UE::StateTree::PropertyRefHelpers::IsPropertyRefMarkedAsOptional(*Property, PropertyRef);
 
-				if (bIsOptional == false && !IsPropertyBound(PropertyName, OutReferenceBindings))
+				if (bIsOptional == false && !IsPropertyBound(PropertyName, OutValidatedBindings.ReferenceBindings))
 				{
 					Log.Reportf(EMessageSeverity::Error, TargetStruct,
 						TEXT("Property reference '%s' on % s is expected to have a binding."),
@@ -2132,7 +2185,7 @@ bool FStateTreeCompiler::GetAndValidateBindings(const FStateTreeBindableStructDe
 			if (Usage == EStateTreePropertyUsage::Input)
 			{
 				// Make sure that an Input property is bound unless marked optional.
-				if (bIsOptional == false && !IsPropertyBound(PropertyName, OutCopyBindings))
+				if (bIsOptional == false && !IsPropertyBound(PropertyName, OutValidatedBindings.CopyBindings))
 				{
 					Log.Reportf(EMessageSeverity::Error, TargetStruct,
 						TEXT("Input property '%s' on %s is expected to have a binding."),
@@ -2163,7 +2216,7 @@ bool FStateTreeCompiler::GetAndValidateBindings(const FStateTreeBindableStructDe
 					continue;
 				}
 
-				const bool bIsBound = IsPropertyBound(PropertyName, OutCopyBindings);
+				const bool bIsBound = IsPropertyBound(PropertyName, OutValidatedBindings.CopyBindings);
 
 				if (!bIsBound)
 				{
@@ -2172,7 +2225,7 @@ bool FStateTreeCompiler::GetAndValidateBindings(const FStateTreeBindableStructDe
 					if (Desc.IsValid())
 					{
 						// Add automatic binding to Context data.
-						OutCopyBindings.Emplace(FStateTreePropertyPath(Desc.ID), FStateTreePropertyPath(TargetStruct.ID, PropertyName));
+						OutValidatedBindings.CopyBindings.Emplace(FStateTreePropertyPath(Desc.ID), FStateTreePropertyPath(TargetStruct.ID, PropertyName));
 					}
 					else
 					{
diff --git a/Engine/Plugins/Runtime/StateTree/Source/StateTreeEditorModule/Private/StateTreeDescriptionHelpers.cpp b/Engine/Plugins/Runtime/StateTree/Source/StateTreeEditorModule/Private/StateTreeDescriptionHelpers.cpp
--- a/Engine/Plugins/Runtime/StateTree/Source/StateTreeEditorModule/Private/StateTreeDescriptionHelpers.cpp
+++ b/Engine/Plugins/Runtime/StateTree/Source/StateTreeEditorModule/Private/StateTreeDescriptionHelpers.cpp
@@ -222,6 +222,18 @@ FText GetTransitionDesc(const UStateTreeEditorData* EditorData, const FStateTree
 		
 		TriggerText = FText::Format(TransitionFormat, FText::Join(INVTEXT(", "), PayloadItems));
 	}
+	else if (Transition.Trigger == EStateTreeTransitionTrigger::OnDelegate)
+	{
+		FStateTreeBindingLookup BindingLookup(EditorData);
+
+		const FText BoundDelegateText = BindingLookup.GetBindingSourceDisplayName(FStateTreePropertyPath(Transition.ID, GET_MEMBER_NAME_CHECKED(FStateTreeTransition, DelegateListener)), Formatting);
+
+		const FText TransitionFormat = Formatting == EStateTreeNodeFormatting::RichText
+			? LOCTEXT("TransitionOnDelegateRich", "<b>On Delegate</> ({0})")
+			: LOCTEXT("TransitionOnDelegate", "On Delegate ({0})");
+
+		TriggerText = FText::Format(TransitionFormat, BoundDelegateText);
+	}
 
 	FText ActionText = Formatting == EStateTreeNodeFormatting::RichText
 		? LOCTEXT("ActionGotoRichRich", "<s>go to</>")
diff --git a/Engine/Plugins/Runtime/StateTree/Source/StateTreeEditorModule/Private/StateTreeEditorData.cpp b/Engine/Plugins/Runtime/StateTree/Source/StateTreeEditorModule/Private/StateTreeEditorData.cpp
--- a/Engine/Plugins/Runtime/StateTree/Source/StateTreeEditorModule/Private/StateTreeEditorData.cpp
+++ b/Engine/Plugins/Runtime/StateTree/Source/StateTreeEditorModule/Private/StateTreeEditorData.cpp
@@ -1327,6 +1327,21 @@ EStateTreeVisitor UStateTreeEditorData::VisitStateNodes(const UStateTreeState& S
 			const FStateTreeTransition& Transition = State.Transitions[TransitionIndex];
 			const FString StatePathWithTransition = StatePath + FString::Printf(TEXT("/Transition[%d]"), TransitionIndex);
 
+			{
+				FStateTreeBindableStructDesc Desc;
+				Desc.StatePath = StatePathWithTransition;
+				Desc.Struct = FStateTreeTransition::StaticStruct();
+				Desc.Name = FName(TEXT("Transition"));
+				Desc.ID = Transition.ID;
+				Desc.DataSource = EStateTreeBindableStructSource::Transition;
+
+				if (InFunc(&State, Desc, FStateTreeDataView(FStructView::Make(const_cast<FStateTreeTransition&>(Transition)))) == EStateTreeVisitor::Break)
+				{
+					bContinue = false;
+					break;
+				}
+			}
+
 			if (Transition.Trigger == EStateTreeTransitionTrigger::OnEvent)
 			{
 				FStateTreeBindableStructDesc Desc;
diff --git a/Engine/Plugins/Runtime/StateTree/Source/StateTreeEditorModule/Private/StateTreePropertyBindingCompiler.cpp b/Engine/Plugins/Runtime/StateTree/Source/StateTreeEditorModule/Private/StateTreePropertyBindingCompiler.cpp
--- a/Engine/Plugins/Runtime/StateTree/Source/StateTreeEditorModule/Private/StateTreePropertyBindingCompiler.cpp
+++ b/Engine/Plugins/Runtime/StateTree/Source/StateTreeEditorModule/Private/StateTreePropertyBindingCompiler.cpp
@@ -10,6 +10,7 @@
 #include "StateTreePropertyRef.h"
 #include "StateTreePropertyRefHelpers.h"
 #include "StateTreePropertyHelpers.h"
+#include "StateTreeDelegate.h"
 
 #include UE_INLINE_GENERATED_CPP_BY_NAME(StateTreePropertyBindingCompiler)
 
@@ -84,7 +85,7 @@ bool FStateTreePropertyBindingCompiler::CompileBatch(const FStateTreeBindableStr
 
 		FSortedBinding& NewBinding = NewBindings.AddDefaulted_GetRef();
 		NewBinding.Binding = FStateTreePropertyPathBinding(SourceStruct->DataHandle, Binding.GetSourcePath(), Binding.GetTargetPath());
-		NewBinding.TargetIndirections = TargetIndirections;
+		NewBinding.TargetIndirections = MoveTemp(TargetIndirections);
 	}
 
 	if (!NewBindings.IsEmpty())
@@ -133,6 +134,114 @@ bool FStateTreePropertyBindingCompiler::CompileBatch(const FStateTreeBindableStr
 	return true;
 }
 
+bool FStateTreePropertyBindingCompiler::CompileDelegates(const FStateTreeBindableStructDesc& TargetStruct, TConstArrayView<FStateTreePropertyPathBinding> DelegateSourceBindings, TConstArrayView<FStateTreePropertyPathBinding> DelegateTargetBindings, FStateTreeDataView InstanceDataView)
+{
+	check(Log);
+	check(PropertyBindings);
+
+	StoreSourceStructs();
+
+	// Compiling dispatchers
+	{
+		TArray<FStateTreePropertyPath, TInlineAllocator<4>> FoundDispatchers;
+		for (const FStateTreePropertyPathBinding& Binding : DelegateSourceBindings)
+		{
+			if (Binding.GetSourcePath().GetStructID() != TargetStruct.ID)
+			{
+				continue;
+			}
+
+			if (FoundDispatchers.Contains(Binding.GetSourcePath()))
+			{
+				continue;
+			}
+
+			FoundDispatchers.Add(Binding.GetSourcePath());
+
+			FString Error;
+			TArray<FPropertyBindingPathIndirection> DispatcherIndirections;
+			if (!Binding.GetSourcePath().ResolveIndirectionsWithValue(InstanceDataView, DispatcherIndirections, &Error))
+			{
+				Log->Reportf(EMessageSeverity::Error, TargetStruct, TEXT("Resolving path in %s: %s"), *TargetStruct.ToString(), *Error);
+				return false;
+			}
+
+			FPropertyBindingPathIndirection& DispatcherLeafIndirection = DispatcherIndirections.Last();
+
+			const FStructProperty* LeafAsStructProperty = CastField<FStructProperty>(DispatcherLeafIndirection.GetProperty());
+			check(LeafAsStructProperty && LeafAsStructProperty->Struct == FStateTreeDelegateDispatcher::StaticStruct())
+
+			FStateTreeDelegateDispatcher* Dispatcher = reinterpret_cast<FStateTreeDelegateDispatcher*>(DispatcherLeafIndirection.GetMutablePropertyAddress());
+			Dispatcher->ID = FStateTreeIndex16(CompiledDelegateDispatchers.IsEmpty() ? 0 : CompiledDelegateDispatchers.Num());
+			CompiledDelegateDispatchers.Add({Binding.GetSourcePath(), Dispatcher->ID});
+		}
+	}
+
+	// Compiling listeners
+	{
+		for (const FStateTreePropertyPathBinding& Binding : DelegateTargetBindings)
+		{
+			if (Binding.GetTargetPath().GetStructID() != TargetStruct.ID)
+			{
+				continue;
+			}
+
+			// Source must be in the source array
+			const FStateTreeBindableStructDesc* DispatcherStruct = GetSourceStructDescByID(Binding.GetSourcePath().GetStructID());
+			if (!DispatcherStruct)
+			{
+				Log->Reportf(EMessageSeverity::Error, TargetStruct,
+					TEXT("Could not find a binding source."));
+				return false;
+			}
+
+			FString Error;
+			TArray<FPropertyBindingPathIndirection> DispatcherIndirections;
+			TArray<FPropertyBindingPathIndirection> ListenerIndirections;
+		
+			if (!Binding.GetSourcePath().ResolveIndirections(DispatcherStruct->Struct, DispatcherIndirections, &Error))
+			{
+				Log->Reportf(EMessageSeverity::Error, TargetStruct, TEXT("Resolving path in %s: %s"), *DispatcherStruct->ToString(), *Error);
+				return false;
+			}
+
+			if (!Binding.GetTargetPath().ResolveIndirectionsWithValue(InstanceDataView, ListenerIndirections, &Error))
+			{
+				Log->Reportf(EMessageSeverity::Error, TargetStruct, TEXT("Resolving path in %s: %s"), *TargetStruct.ToString(), *Error);
+				return false;
+			}
+
+			const FStateTreeIndex16 DispatcherID = GetDispatcherIDFromPath(Binding.GetSourcePath());
+
+			if (!DispatcherID.IsValid())
+			{
+				Log->Reportf(EMessageSeverity::Error, TargetStruct, TEXT("Delegate Listener %s is bound to unknown dispatcher %s"), *TargetStruct.ToString(), *DispatcherStruct->ToString());
+				return false;
+			}
+
+			FPropertyBindingPathIndirection& ListenerLeafIndirection = ListenerIndirections.Last();
+			const FStructProperty* LeafAsStructProperty = CastField<FStructProperty>(ListenerLeafIndirection.GetProperty());
+			check(LeafAsStructProperty && LeafAsStructProperty->Struct == FStateTreeDelegateListener::StaticStruct())
+			FStateTreeDelegateListener* Listener = reinterpret_cast<FStateTreeDelegateListener*>(ListenerLeafIndirection.GetMutablePropertyAddress());
+			Listener->DispatcherID = DispatcherID;
+			Listener->ListenerID = FStateTreeIndex16(ListenersNum++);
+		}
+	}
+
+	return true;
+}
+
+FStateTreeIndex16 FStateTreePropertyBindingCompiler::GetDispatcherIDFromPath(const FStateTreePropertyPath& PathToDispatcher) const
+{
+	const FCompiledDelegateDispatcher* FoundDispatcher = Algo::FindBy(CompiledDelegateDispatchers, PathToDispatcher, [](const FCompiledDelegateDispatcher& Dispatcher) { return Dispatcher.Path; });
+	if (FoundDispatcher)
+	{
+		return FoundDispatcher->ID;
+	}
+
+	return FStateTreeIndex16::Invalid;
+}
+
 bool FStateTreePropertyBindingCompiler::CompileReferences(const FStateTreeBindableStructDesc& TargetStruct, TConstArrayView<FStateTreePropertyPathBinding> PropertyReferenceBindings, FStateTreeDataView InstanceDataView, const TMap<FGuid, const FStateTreeDataView>& IDToStructValue)
 {
 	for (const FStateTreePropertyPathBinding& Binding : PropertyReferenceBindings)
@@ -142,7 +251,7 @@ bool FStateTreePropertyBindingCompiler::CompileReferences(const FStateTreeBindab
 			continue;
 		}
 
-		// Source must be in the source array/
+		// Source must be in the source array
 		const FStateTreeBindableStructDesc* SourceStruct = GetSourceStructDescByID(Binding.GetSourcePath().GetStructID());
 		if (!SourceStruct)
 		{
@@ -261,6 +370,9 @@ bool FStateTreePropertyBindingCompiler::CompileReferences(const FStateTreeBindab
 void FStateTreePropertyBindingCompiler::Finalize()
 {
 	StoreSourceStructs();
+
+	CompiledDelegateDispatchers.Reset();
+	CompiledReferences.Reset();
 }
 
 int32 FStateTreePropertyBindingCompiler::AddSourceStruct(const FStateTreeBindableStructDesc& SourceStruct)
diff --git a/Engine/Plugins/Runtime/StateTree/Source/StateTreeEditorModule/Public/StateTreeCompiler.h b/Engine/Plugins/Runtime/StateTree/Source/StateTreeEditorModule/Public/StateTreeCompiler.h
--- a/Engine/Plugins/Runtime/StateTree/Source/StateTreeEditorModule/Public/StateTreeCompiler.h
+++ b/Engine/Plugins/Runtime/StateTree/Source/StateTreeEditorModule/Public/StateTreeCompiler.h
@@ -66,7 +66,16 @@ private:
 
 	bool CreateTask(UStateTreeState* State, const FStateTreeEditorNode& TaskNode, const FStateTreeDataHandle TaskDataHandle);
 	bool CreateEvaluator(const FStateTreeEditorNode& EvalNode, const FStateTreeDataHandle EvalDataHandle);
-	bool GetAndValidateBindings(const FStateTreeBindableStructDesc& TargetStruct, FStateTreeDataView TargetValue, TArray<FStateTreePropertyPathBinding>& OutCopyBindings, TArray<FStateTreePropertyPathBinding>& OutReferenceBindings) const;
+
+	struct FValidatedPathBindings
+	{
+		TArray<FStateTreePropertyPathBinding> CopyBindings;
+		TArray<FStateTreePropertyPathBinding> DelegateSourceBindings;
+		TArray<FStateTreePropertyPathBinding> DelegateTargetBindings;
+		TArray<FStateTreePropertyPathBinding> ReferenceBindings;
+	};
+
+	bool GetAndValidateBindings(const FStateTreeBindableStructDesc& TargetStruct, FStateTreeDataView TargetValue, FValidatedPathBindings& OutValidatedBindings) const;
 	bool IsPropertyOfTypeOrChild(UScriptStruct& Type, const FStateTreeBindableStructDesc& Struct, FStateTreePropertyPath Path) const;
 	bool ValidateStructRef(const FStateTreeBindableStructDesc& SourceStruct, FStateTreePropertyPath SourcePath,
 							const FStateTreeBindableStructDesc& TargetStruct, FStateTreePropertyPath TargetPath) const;
diff --git a/Engine/Plugins/Runtime/StateTree/Source/StateTreeEditorModule/Public/StateTreeEditorPropertyBindings.h b/Engine/Plugins/Runtime/StateTree/Source/StateTreeEditorModule/Public/StateTreeEditorPropertyBindings.h
--- a/Engine/Plugins/Runtime/StateTree/Source/StateTreeEditorModule/Public/StateTreeEditorPropertyBindings.h
+++ b/Engine/Plugins/Runtime/StateTree/Source/StateTreeEditorModule/Public/StateTreeEditorPropertyBindings.h
@@ -207,7 +207,6 @@ struct STATETREEEDITORMODULE_API FStateTreeBindingLookup : public IStateTreeBind
 
 	const IStateTreeEditorPropertyBindingsOwner* BindingOwner = nullptr;
 
-protected:
 	virtual const FStateTreePropertyPath* GetPropertyBindingSource(const FStateTreePropertyPath& InTargetPath) const override;
 	virtual FText GetPropertyPathDisplayName(const FStateTreePropertyPath& InTargetPath, EStateTreeNodeFormatting Formatting) const override;
 	virtual FText GetBindingSourceDisplayName(const FStateTreePropertyPath& InTargetPath, EStateTreeNodeFormatting Formatting) const override;
diff --git a/Engine/Plugins/Runtime/StateTree/Source/StateTreeEditorModule/Public/StateTreePropertyBindingCompiler.h b/Engine/Plugins/Runtime/StateTree/Source/StateTreeEditorModule/Public/StateTreePropertyBindingCompiler.h
--- a/Engine/Plugins/Runtime/StateTree/Source/StateTreeEditorModule/Public/StateTreePropertyBindingCompiler.h
+++ b/Engine/Plugins/Runtime/StateTree/Source/StateTreeEditorModule/Public/StateTreePropertyBindingCompiler.h
@@ -36,6 +36,18 @@ struct STATETREEEDITORMODULE_API FStateTreePropertyBindingCompiler
 	 */
 	[[nodiscard]] bool CompileBatch(const FStateTreeBindableStructDesc& TargetStruct, TConstArrayView<FStateTreePropertyPathBinding> PropertyBindings, FStateTreeIndex16 PropertyFuncsBegin, FStateTreeIndex16 PropertyFuncsEnd, int32& OutBatchIndex);
 
+	 /**
+	  * Compiles delegates for selected struct.
+	  * @param TargetStruct - Description of the structs which contains the delegate dispatchers or listeners.
+	  * @param DelegateSourceBindings - Array of bindings to dispatchers to compile.
+	  * @param DelegateTargetBindings - Array of bindings from listeners to compile.
+	  * @param InstanceDataView - View to the instance data
+	  * @param TargetAccessibleStructs - View to all accessible structs.
+	  * @param OutBatchIndex - Resulting batch index, if index is INDEX_NONE, no bindings were found and no batch was generated.
+	  * @return True on success, false on failure.
+	 */
+	[[nodiscard]] bool CompileDelegates(const FStateTreeBindableStructDesc& TargetStruct, TConstArrayView<FStateTreePropertyPathBinding> DelegateSourceBindings, TConstArrayView<FStateTreePropertyPathBinding> DelegateTargetBindings, FStateTreeDataView InstanceDataView);
+
 	/**
 	  * Compiles references for selected struct
 	  * @param TargetStruct - Description of the structs which contains the target properties.
@@ -56,6 +68,9 @@ struct STATETREEEDITORMODULE_API FStateTreePropertyBindingCompiler
 	  */
 	int32 AddSourceStruct(const FStateTreeBindableStructDesc& SourceStruct);
 
+	/** @return Index of delegate dispatcher residing behind given path. */
+	FStateTreeIndex16 GetDispatcherIDFromPath(const FStateTreePropertyPath& PathToDispatcher) const;
+
 	const FStateTreeBindableStructDesc* GetSourceStructDescByID(const FGuid& ID) const
 	{
 		return SourceStructs.FindByPredicate([ID](const FStateTreeBindableStructDesc& Structs) { return (Structs.ID == ID); });
@@ -87,6 +102,15 @@ protected:
 
 	TArray<FCompiledReference> CompiledReferences;
 
+	struct FCompiledDelegateDispatcher
+	{
+		FStateTreePropertyPath Path;
+		FStateTreeIndex16 ID;
+	};
+
+	TArray<FCompiledDelegateDispatcher> CompiledDelegateDispatchers;
+	uint16 ListenersNum = 0u;
+
 	FStateTreePropertyBindings* PropertyBindings = nullptr;
 
 	FStateTreeCompilerLog* Log = nullptr;
diff --git a/Engine/Plugins/Runtime/StateTree/Source/StateTreeEditorModule/Public/StateTreeState.h b/Engine/Plugins/Runtime/StateTree/Source/StateTreeEditorModule/Public/StateTreeState.h
--- a/Engine/Plugins/Runtime/StateTree/Source/StateTreeEditorModule/Public/StateTreeState.h
+++ b/Engine/Plugins/Runtime/StateTree/Source/StateTreeEditorModule/Public/StateTreeState.h
@@ -62,6 +62,15 @@ private:
 	FStateTreeEvent TemporaryEvent;
 };
 
+/**
+ * StateTree's internal delegate listener used exclusively in transitions.
+ */
+USTRUCT()
+struct FStateTreeTransitionDelegateListener
+{
+	GENERATED_BODY()
+};
+
 /**
  * Editor representation of a transition in StateTree
  */
@@ -117,6 +126,10 @@ PRAGMA_ENABLE_DEPRECATION_WARNINGS
 	UPROPERTY(EditDefaultsOnly, Category = "Transition")
 	FGuid ID;
 
+	/** Listener to the selected delegate dispatcher. */
+	UPROPERTY(EditDefaultsOnly, Category = "Transition", DisplayName = "Delegate")
+	FStateTreeTransitionDelegateListener DelegateListener;
+
 	/**
 	 * Transition priority when multiple transitions happen at the same time.
 	 * During transition handling, the transitions are visited from leaf to root.
diff --git a/Engine/Plugins/Runtime/StateTree/Source/StateTreeModule/Private/Blueprint/StateTreeNodeBlueprintBase.cpp b/Engine/Plugins/Runtime/StateTree/Source/StateTreeModule/Private/Blueprint/StateTreeNodeBlueprintBase.cpp
--- a/Engine/Plugins/Runtime/StateTree/Source/StateTreeModule/Private/Blueprint/StateTreeNodeBlueprintBase.cpp
+++ b/Engine/Plugins/Runtime/StateTree/Source/StateTreeModule/Private/Blueprint/StateTreeNodeBlueprintBase.cpp
@@ -6,6 +6,7 @@
 #include "VisualLogger/VisualLogger.h"
 #include "StateTreePropertyRef.h"
 #include "StateTreePropertyRefHelpers.h"
+#include "StateTreeDelegate.h"
 
 #include UE_INLINE_GENERATED_CPP_BY_NAME(StateTreeNodeBlueprintBase)
 
@@ -100,7 +101,7 @@ void UStateTreeNodeBlueprintBase::SendEvent(const FStateTreeEvent& Event)
 {
 	if (!WeakExecutionContext.SendEvent(Event.Tag, Event.Payload, Event.Origin))
 	{
-		UE_VLOG_UELOG(this, LogStateTree, Error, TEXT("Trying to call SendEvent() while node is not active. Use SendEvent() on UStateTreeComponent instead for sending signals externally."));
+		UE_VLOG_UELOG(this, LogStateTree, Error, TEXT("Failed to send the event. The instance probably stopped."));
 	}
 }
 
@@ -108,7 +109,15 @@ void UStateTreeNodeBlueprintBase::RequestTransition(const FStateTreeStateLink& T
 {
 	if (!WeakExecutionContext.RequestTransition(TargetState.StateHandle, Priority, TargetState.Fallback))
 	{
-		UE_VLOG_UELOG(this, LogStateTree, Error, TEXT("Trying to call SendEvent() while node is not active. Use SendEvent() on UStateTreeComponent instead for sending signals externally."));
+		UE_VLOG_UELOG(this, LogStateTree, Error, TEXT("Failed to request a transition. The instance probably stopped."));
+	}
+}
+
+void UStateTreeNodeBlueprintBase::BroadcastDelegate(FStateTreeDelegateDispatcher Dispatcher)
+{
+	if (!WeakExecutionContext.BroadcastDelegate(Dispatcher))
+	{
+		UE_VLOG_UELOG(this, LogStateTree, Error, TEXT("Failed to broadcast the delegate. The instance probably stopped."));
 	}
 }
 
diff --git a/Engine/Plugins/Runtime/StateTree/Source/StateTreeModule/Private/StateTreeExecutionContext.cpp b/Engine/Plugins/Runtime/StateTree/Source/StateTreeModule/Private/StateTreeExecutionContext.cpp
--- a/Engine/Plugins/Runtime/StateTree/Source/StateTreeModule/Private/StateTreeExecutionContext.cpp
+++ b/Engine/Plugins/Runtime/StateTree/Source/StateTreeModule/Private/StateTreeExecutionContext.cpp
@@ -5,6 +5,7 @@
 #include "StateTreeEvaluatorBase.h"
 #include "StateTreeConditionBase.h"
 #include "StateTreeConsiderationBase.h"
+#include "StateTreeDelegate.h"
 #include "StateTreePropertyFunctionBase.h"
 #include "StateTreeReference.h"
 #include "Containers/StaticArray.h"
@@ -27,6 +28,39 @@ namespace UE::StateTree::Debug
 	static std::atomic<uint32> InstanceSerialNumber = 0;
 }; // UE::StateTree
 
+namespace UE::StateTree::ExecutionContext
+{
+	bool MarkDelegateAsBroadcasted(FStateTreeDelegateDispatcher Dispatcher, const FStateTreeExecutionFrame& CurrentFrame, FStateTreeInstanceStorage& Storage)
+	{
+		const UStateTree* StateTree = CurrentFrame.StateTree;
+		check(StateTree);
+
+		for (FStateTreeStateHandle ActiveState : CurrentFrame.ActiveStates)
+		{
+			if (const FCompactStateTreeState* State = StateTree->GetStateFromHandle(ActiveState))
+			{
+				if (!State->bHasDelegateTriggerTransitions)
+				{
+					continue;
+				}
+
+				for (uint16 TransitionIndex = State->TransitionsBegin; TransitionIndex < State->TransitionsBegin + State->TransitionsNum; ++TransitionIndex)
+				{
+					if (const FCompactStateTransition* Transition = StateTree->GetTransitionFromIndex(FStateTreeIndex16(TransitionIndex)))
+					{
+						if (Transition->DelegateDispatcherID == Dispatcher.GetDelegateID())
+						{
+							Storage.MarkDelegateAsBroadcasted(Dispatcher, CurrentFrame.StateTree);
+							return true;
+						}
+					}
+				}
+			}
+		}
+
+		return false;
+	}
+}
 
 FStateTreeMinimalExecutionContext::FStateTreeMinimalExecutionContext(UObject& InOwner, const UStateTree& InStateTree, FStateTreeInstanceData& InInstanceData)
 	: FStateTreeMinimalExecutionContext(InOwner, InStateTree, InInstanceData.GetMutableStorage())
@@ -90,7 +124,7 @@ FStateTreeScheduledTick FStateTreeMinimalExecutionContext::GetNextScheduledTick(
 	for (int32 FrameIndex = 0; FrameIndex < Exec.ActiveFrames.Num(); ++FrameIndex)
 	{
 		const FStateTreeExecutionFrame& CurrentFrame = Exec.ActiveFrames[FrameIndex];
-		if (!CurrentFrame.StateTree->bScheduledTickAllowed)
+		if (!CurrentFrame.StateTree->IsScheduledTickAllowed())
 		{
 			return FStateTreeScheduledTick::MakeEveryFrames();
 		}
@@ -98,6 +132,7 @@ FStateTreeScheduledTick FStateTreeMinimalExecutionContext::GetNextScheduledTick(
 
 	const FStateTreeEventQueue& EventQueue = Storage.GetEventQueue();
 	const bool bHasEvents = EventQueue.HasEvents();
+	const bool bHasBroadcastedDelegates = Storage.HasBroadcastedDelegates();
 
 	// Do we have a state that requires a tick or is waiting for an event.
 	{
@@ -129,7 +164,7 @@ FStateTreeScheduledTick FStateTreeMinimalExecutionContext::GetNextScheduledTick(
 					{
 						if (State.ShouldTickTasks(bHasEvents)
 							|| State.bHasTransitionTasks
-							|| State.ShouldTickTransitions(bHasEvents)
+							|| State.ShouldTickTransitions(bHasEvents, bHasBroadcastedDelegates)
 							)
 						{
 							// todo: ShouldTickTransitions has onevent or ontick. both can be already triggered and we are waiting for the delay
@@ -200,7 +235,7 @@ void FStateTreeMinimalExecutionContext::SendEvent(const FGameplayTag Tag, const
 void FStateTreeMinimalExecutionContext::ScheduleNextTick()
 {
 	TInstancedStruct<FStateTreeExecutionExtension>& ExecutionExtension = Storage.GetMutableExecutionState().ExecutionExtension;
-	if (RootStateTree.bScheduledTickAllowed && ExecutionExtension.IsValid())
+	if (RootStateTree.IsScheduledTickAllowed() && ExecutionExtension.IsValid())
 	{
 		ExecutionExtension.GetMutable().ScheduleNextTick(FStateTreeExecutionExtension::FContextParameters(Owner, RootStateTree, Storage));
 	}
@@ -550,6 +585,8 @@ EStateTreeRunStatus FStateTreeExecutionContext::Start(FStartParameters Parameter
 		// No active states or global tasks anymore, reset frames.
 		Exec.ActiveFrames.Reset();
 		
+		RemoveAllDelegateListeners();
+
 		// We are not considered as running yet so we only set the status without requiring a stop.
 		Exec.TreeRunStatus = GlobalTasksRunStatus;
 	}
@@ -830,6 +867,8 @@ void FStateTreeExecutionContext::TickTriggerTransitionsInternal()
 				// No active states or global tasks anymore, reset frames.
 				Exec.ActiveFrames.Reset();
 
+				RemoveAllDelegateListeners();
+
 				break;
 			}
 
@@ -855,6 +894,67 @@ void FStateTreeExecutionContext::TickTriggerTransitionsInternal()
 	}
 }
 
+void FStateTreeExecutionContext::BroadcastDelegate(FStateTreeDelegateDispatcher Dispatcher)
+{
+	if (!Dispatcher.IsValid())
+	{
+		return;
+	}
+
+	if (!IsValid())
+	{
+		STATETREE_LOG(Warning, TEXT("%hs: StateTree context is not initialized properly ('%s' using StateTree '%s')"),
+			__FUNCTION__, *GetNameSafe(&Owner), *GetFullNameSafe(&RootStateTree));
+		return;
+	}
+
+	const FStateTreeExecutionFrame* CurrentFrame = GetCurrentlyProcessedFrame();
+	check(CurrentFrame);
+
+	GetExecState().DelegateActiveListeners.BroadcastDelegate(Dispatcher, CurrentFrame->StateTree, GetExecState());
+	if (UE::StateTree::ExecutionContext::MarkDelegateAsBroadcasted(Dispatcher, *CurrentFrame, GetMutableInstanceData()->GetMutableStorage()))
+	{
+		ScheduleNextTick();
+	}
+}
+
+bool FStateTreeExecutionContext::AddDelegateListener(FStateTreeDelegateListener Listener, FSimpleDelegate Delegate)
+{
+	if (!Listener.IsValid())
+	{
+		return true;
+	}
+
+	if (!IsValid())
+	{
+		STATETREE_LOG(Warning, TEXT("%hs: StateTree context is not initialized properly ('%s' using StateTree '%s')"),
+			__FUNCTION__, *GetNameSafe(&Owner), *GetFullNameSafe(&RootStateTree));
+		return false;
+	}
+
+	return GetExecState().DelegateActiveListeners.AddListener(Listener, MoveTemp(Delegate), CurrentlyProcessedFrame->StateTree, CurrentlyProcessedState.Root, CurrentlyProcessedState, FStateTreeIndex16(CurrentNodeDataHandle.GetIndex()));
+}
+
+void FStateTreeExecutionContext::RemoveDelegateListener(FStateTreeDelegateListener Listener)
+{
+	if (!Listener.IsValid())
+	{
+		return;
+	}
+
+	if (!IsValid())
+	{
+		STATETREE_LOG(Warning, TEXT("%hs: StateTree context is not initialized properly ('%s' using StateTree '%s')"),
+			__FUNCTION__, *GetNameSafe(&Owner), *GetFullNameSafe(&RootStateTree));
+		return;
+	}
+
+	const FStateTreeExecutionFrame* CurrentFrame = GetCurrentlyProcessedFrame();
+	check(CurrentFrame);
+
+	GetExecState().DelegateActiveListeners.RemoveListener(Listener, CurrentFrame->StateTree);
+}
+
 void FStateTreeExecutionContext::RequestTransition(const FStateTreeTransitionRequest& Request)
 {
 	CSV_SCOPED_TIMING_STAT_EXCLUSIVE(StateTree_RequestTransition);
@@ -1953,6 +2053,9 @@ void FStateTreeExecutionContext::ExitState(const FStateTreeTransitionResult& Tra
 {
 	CSV_SCOPED_TIMING_STAT_EXCLUSIVE(StateTree_ExitState);
 
+	check(!GetExecState().LastExitedNodeIndex.IsValid());
+	ON_SCOPE_EXIT { GetExecState().LastExitedNodeIndex = FStateTreeIndex16::Invalid; };
+
 	FStateTreeExecutionState& Exec = GetExecState();
 
 	if (Exec.ActiveFrames.IsEmpty())
@@ -1984,6 +2087,9 @@ void FStateTreeExecutionContext::ExitState(const FStateTreeTransitionResult& Tra
 	// tree is reused and should be considered as a new instance.
 	bool bAnyParentStateChanged = false;
 
+	int32 FirstModifiedFrame = INDEX_NONE;
+	int32 FirstModifiedState = INDEX_NONE;
+
 	for (int32 FrameIndex = 0; FrameIndex < Exec.ActiveFrames.Num(); FrameIndex++)
 	{
 		FStateTreeExecutionFrame* CurrentParentFrame = FrameIndex > 0 ? &Exec.ActiveFrames[FrameIndex - 1] : nullptr;
@@ -2068,6 +2174,12 @@ void FStateTreeExecutionContext::ExitState(const FStateTreeTransitionResult& Tra
 			const bool bShouldCallExitState = (!bRemainsActive || bOnTargetBranch);
 			ExitStateCalls.Emplace(ChangeType, bShouldCallExitState);
 
+			if (bShouldCallExitState && FirstModifiedFrame == INDEX_NONE)
+			{
+				FirstModifiedFrame = FrameIndex;
+				FirstModifiedState = Index;
+			}
+
 			// Do property copies, ExitState() is called below.
 			for (int32 TaskIndex = State.TasksBegin; TaskIndex < (State.TasksBegin + State.TasksNum); TaskIndex++)
 			{
@@ -2159,6 +2271,7 @@ void FStateTreeExecutionContext::ExitState(const FStateTreeTransitionResult& Tra
 								Task.ExitState(*this, CurrentTransition);
 							}
 							UE_STATETREE_DEBUG_TASK_EVENT(this, TaskIndex, TaskInstanceView, EStateTreeTraceEventType::OnExited, Transition.CurrentRunStatus);
+							GetExecState().LastExitedNodeIndex = FStateTreeIndex16(TaskIndex);
 						}
 					}
 				}
@@ -2213,6 +2326,58 @@ void FStateTreeExecutionContext::ExitState(const FStateTreeTransitionResult& Tra
 		}
 	}
 
+	// Delegate Listeners Cleanup
+	{
+		if (Transition.NextActiveFrames.IsEmpty())
+		{
+			RemoveAllDelegateListeners();
+}
+		else
+		{
+			check(FirstModifiedFrame != INDEX_NONE);
+			check(FirstModifiedState != INDEX_NONE);
+
+			const FStateTreeExecutionFrame* ChangedFrame = &Exec.ActiveFrames[FirstModifiedFrame];
+			FStateTreeStateHandle ChangedState = ChangedFrame->ActiveStates[FirstModifiedState];
+			TConstArrayView<FStateTreeExecutionFrame> DeactivatedFrames;
+
+			if (Exec.ActiveFrames.Num() > FirstModifiedFrame + 1)
+			{
+				DeactivatedFrames = TConstArrayView<FStateTreeExecutionFrame>(&Exec.ActiveFrames[FirstModifiedFrame + 1], Exec.ActiveFrames.Num() - FirstModifiedFrame - 1);
+			}
+
+			RemoveDeactivatedDelegateListeners(ChangedFrame, ChangedState, DeactivatedFrames);
+		}
+	}
+}
+
+void FStateTreeExecutionContext::RemoveDeactivatedDelegateListeners(const FStateTreeExecutionFrame* ChangedFrame, FStateTreeStateHandle ChangedState, TConstArrayView<FStateTreeExecutionFrame> DeactivatedFrames)
+{
+	GetExecState().DelegateActiveListeners.RemoveListenersByPredicate([&](FStateTreeStateHandle OwningState, FStateTreeIndex16 OwningNodeIndex, const UStateTree* FrameStateTree, FStateTreeStateHandle FrameRootState)
+		{
+			if (FrameStateTree == nullptr)
+			{
+				return true;
+			}
+
+			if (ChangedFrame && ChangedFrame->RootState == FrameRootState)
+			{
+				check(ChangedState.IsValid());
+				const bool bIsGlobalListener = !OwningState.IsValid();
+				return !bIsGlobalListener && OwningState.Index >= ChangedState.Index;
+			}
+
+			return DeactivatedFrames.ContainsByPredicate([&](const FStateTreeExecutionFrame& Frame)
+				{
+					return Frame.RootState == FrameRootState && Frame.StateTree == FrameStateTree;
+				});
+
+		});
+}
+
+void FStateTreeExecutionContext::RemoveAllDelegateListeners()
+{
+	GetExecState().DelegateActiveListeners = FStateTreeDelegateActiveListeners();
 }
 
 void FStateTreeExecutionContext::StateCompleted()
@@ -2581,6 +2746,8 @@ void FStateTreeExecutionContext::CallStopOnEvaluatorsAndGlobalTasks(const FState
 {
 	check(Frame.bIsGlobalFrame);
 
+	ON_SCOPE_EXIT { GetExecState().LastExitedNodeIndex = FStateTreeIndex16::Invalid; };
+
 	FCurrentlyProcessedFrameScope FrameScope(*this, ParentFrame, Frame);
 	const UStateTree* CurrentStateTree = Frame.StateTree;
 
@@ -2607,6 +2774,7 @@ void FStateTreeExecutionContext::CallStopOnEvaluatorsAndGlobalTasks(const FState
 				Task.ExitState(*this, Transition);
 			}
 			UE_STATETREE_DEBUG_TASK_EVENT(this, TaskIndex, TaskInstanceView, EStateTreeTraceEventType::OnExited, Transition.CurrentRunStatus);
+			GetExecState().LastExitedNodeIndex = FStateTreeIndex16(TaskIndex);
 		}
 	}
 
@@ -2624,6 +2792,7 @@ void FStateTreeExecutionContext::CallStopOnEvaluatorsAndGlobalTasks(const FState
 
 			UE_STATETREE_DEBUG_EVALUATOR_EVENT(this, EvalIndex, EvalInstanceView, EStateTreeTraceEventType::OnTreeStopped);
 		}
+		GetExecState().LastExitedNodeIndex = FStateTreeIndex16(EvalIndex);
 	}
 }
 
@@ -3372,6 +3541,7 @@ bool FStateTreeExecutionContext::TriggerTransitions()
 	{
 		// Re-cache bHasEvents, RequestTransition above can create new events.
 		const bool bHasEvents = EventQueue && EventQueue->HasEvents();
+		const bool bHasBroadcastedDelegates = InstanceDataStorage->HasBroadcastedDelegates();
 		for (int32 FrameIndex = Exec.ActiveFrames.Num() - 1; FrameIndex >= 0; FrameIndex--)
 		{
 			FStateTreeExecutionFrame& CurrentFrame = Exec.ActiveFrames[FrameIndex];
@@ -3413,7 +3583,7 @@ bool FStateTreeExecutionContext::TriggerTransitions()
 				// Regular transitions on state
 				//or A transition task can trigger an event. We need to add the state if that is a possibility
 				//or Expired transition delayed
-				if (State.ShouldTickTransitions(bHasEvents) || State.bHasTransitionTasks || bHasActiveTransitionDelayed)
+				if (State.ShouldTickTransitions(bHasEvents, bHasBroadcastedDelegates) || State.bHasTransitionTasks || bHasActiveTransitionDelayed)
 				{
 					TransitionHandlers.Emplace((uint8)FrameIndex, StateHandle, EStateTreeTransitionPriority::Normal);
 				}
@@ -3561,11 +3731,19 @@ bool FStateTreeExecutionContext::TriggerTransitions()
 						}
 					}
 				}
-				else if (Transition.Trigger == EStateTreeTransitionTrigger::OnTick)
+				else if (EnumHasAnyFlags(Transition.Trigger, EStateTreeTransitionTrigger::OnTick))
 				{
 					// Dummy event to make sure we iterate to loop below once.
 					TransitionEvents.Emplace(nullptr);
 				}
+				else if (EnumHasAnyFlags(Transition.Trigger, EStateTreeTransitionTrigger::OnDelegate))
+				{
+					if (InstanceDataStorage->IsDelegateBroadcasted(Transition.DelegateDispatcherID, CurrentFrame))
+				{
+					// Dummy event to make sure we iterate to loop below once.
+					TransitionEvents.Emplace(nullptr);
+				}
+				}
 				else
 				{
 					ensureMsgf(false, TEXT("The trigger type is not supported."));
@@ -3654,6 +3832,8 @@ bool FStateTreeExecutionContext::TriggerTransitions()
 		EventQueue->Reset();
 	}
 
+	InstanceDataStorage->ResetBroadcastedDelegates();
+
 	//
 	// Check state completion transitions.
 	//
@@ -4211,6 +4391,8 @@ bool FStateTreeExecutionContext::SelectStateInternal(
 					NewFrame.GlobalInstanceIndexBase = ExistingFrame->GlobalInstanceIndexBase;
 					NewFrame.StateParameterDataHandle = ExistingFrame->StateParameterDataHandle;
 					NewFrame.GlobalParameterDataHandle = ExistingFrame->GlobalParameterDataHandle;
+					NewFrame.ActiveListenersIndexBase = ExistingFrame->ActiveListenersIndexBase;
+					NewFrame.DelegateDispatcherIndexBase = ExistingFrame->DelegateDispatcherIndexBase;
 				}
 				else
 				{
@@ -4291,6 +4473,8 @@ bool FStateTreeExecutionContext::SelectStateInternal(
 				NewFrame.StateParameterDataHandle = ExistingFrame->StateParameterDataHandle;
 				NewFrame.GlobalParameterDataHandle = ExistingFrame->GlobalParameterDataHandle;
 				NewFrame.ExternalDataBaseIndex = ExistingFrame->ExternalDataBaseIndex;
+				NewFrame.ActiveListenersIndexBase = ExistingFrame->ActiveListenersIndexBase;
+				NewFrame.DelegateDispatcherIndexBase = ExistingFrame->DelegateDispatcherIndexBase;
 			}
 			else
 			{
@@ -4317,6 +4501,8 @@ bool FStateTreeExecutionContext::SelectStateInternal(
 						__FUNCTION__, *GetSafeStateName(CurrentFrame, NextStateHandle), *GetFullNameSafe(NewFrame.StateTree), *GetNameSafe(&Owner), *GetFullNameSafe(CurrentFrame.StateTree));
 					
 					StopTemporaryEvaluatorsAndGlobalTasks(nullptr, NewFrame);
+					RemoveDeactivatedDelegateListeners(nullptr, FStateTreeStateHandle::Invalid, {NewFrame});
+
 					break;
 				}
 
@@ -4335,6 +4521,7 @@ bool FStateTreeExecutionContext::SelectStateInternal(
 			if (bStartedTemporaryEvaluatorsAndGlobalTasks)
 			{
 				StopTemporaryEvaluatorsAndGlobalTasks(nullptr, NewFrame);
+				RemoveDeactivatedDelegateListeners(nullptr, FStateTreeStateHandle::Invalid, {NewFrame});
 			}
 				
 			OutSelectionResult.PopFrame();
@@ -4416,7 +4603,8 @@ bool FStateTreeExecutionContext::SelectStateInternal(
 						}
 					}
 				}
-				else if (Transition.Trigger == EStateTreeTransitionTrigger::OnTick)
+				else if (EnumHasAnyFlags(Transition.Trigger, EStateTreeTransitionTrigger::OnTick)
+					|| (EnumHasAnyFlags(Transition.Trigger, EStateTreeTransitionTrigger::OnDelegate) && InstanceDataStorage->IsDelegateBroadcasted(Transition.DelegateDispatcherID, CurrentFrame)))
 				{
 					SelectedStateTransitionEvents.Emplace();
 				}
@@ -4970,6 +5158,18 @@ TArray<FName> FStateTreeExecutionContext::GetActiveStateNames() const
 	return Result;
 }
 
+FStateTreeWeakExecutionContext::FStateTreeWeakExecutionContext(const FStateTreeExecutionContext& Context)
+	: Owner(Context.GetOwner())
+	, StateTree(Context.GetStateTree())
+	, Storage(Context.GetMutableInstanceData()->GetWeakMutableStorage())
+{
+	if (const FStateTreeExecutionFrame* Frame = Context.GetCurrentlyProcessedFrame())
+	{
+		FrameStateTree = Frame->StateTree;
+		FrameRootState = Frame->RootState;
+	}
+}
+
 FStateTreeWeakExecutionContext::FStateTreeWeakExecutionContext(gsl::not_null<UObject*> InOwner, gsl::not_null<const UStateTree*> InStateTree, FStateTreeInstanceData& InInstanceData)
 	: Owner(InOwner)
 	, StateTree(InStateTree)
@@ -5030,17 +5230,111 @@ bool FStateTreeWeakExecutionContext::RequestTransition(FStateTreeStateHandle Tar
 
 			StoragePinned->AddTransitionRequest(OwnerPinned.Get(), Request);
 
-			if (Exec.ExecutionExtension.IsValid())
+			ScheduleNextTick(OwnerPinned.Get(), StateTreePinned.Get(), *StoragePinned);
+		}
+
+		StoragePinned->ReleaseWriteAccess();
+		return true;
+	}
+	return false;
+}
+
+bool FStateTreeWeakExecutionContext::BroadcastDelegate(FStateTreeDelegateDispatcher Dispatcher) const
+{
+	if (!Dispatcher.IsValid())
+	{
+		// Nothings binds to the delegate, not an error.
+		return true;
+	}
+
+	bool bBroadcasted = false;
+
+	TStrongObjectPtr<UObject> OwnerPinned = Owner.Pin();
+	TStrongObjectPtr<const UStateTree> RootStateTreePinned = StateTree.Pin();
+	TSharedPtr<FStateTreeInstanceStorage> StoragePinned = Storage.Pin();
+	if (OwnerPinned && RootStateTreePinned && StoragePinned)
+	{
+		// Manually acquire and release the write access.
+		//The context may already exist and we want to confirm that we are on the same thread.
+		StoragePinned->AcquireWriteAccess();
+
+		FStateTreeExecutionState& Exec = StoragePinned->GetMutableExecutionState();
+		if (Exec.CurrentExecutionContext)
+		{
+			Exec.CurrentExecutionContext->BroadcastDelegate(Dispatcher);
+			bBroadcasted = true;
+		}
+		else
+		{
+			TStrongObjectPtr<const UStateTree> FrameStateTreePinned = FrameStateTree.Pin();
+			if (FrameStateTreePinned)
+			{
+				const FStateTreeExecutionFrame* ParentFrame = nullptr;
+				const FStateTreeExecutionFrame* CurrentFrame = FStateTreeExecutionContext::FindFrame(FrameStateTreePinned.Get(), FrameRootState, StoragePinned->GetExecutionState().ActiveFrames, ParentFrame);
+				if (CurrentFrame)
+				{
+					Exec.DelegateActiveListeners.BroadcastDelegate(Dispatcher, CurrentFrame->StateTree, Exec);
+					if (UE::StateTree::ExecutionContext::MarkDelegateAsBroadcasted(Dispatcher, *CurrentFrame, *StoragePinned))
+					{
+						ScheduleNextTick(OwnerPinned.Get(), RootStateTreePinned.Get(), *StoragePinned);
+					}
+				}
+				bBroadcasted = true;
+			}
+		}
+
+		StoragePinned->ReleaseWriteAccess();
+	}
+
+	return bBroadcasted;
+}
+
+bool FStateTreeWeakExecutionContext::RemoveDelegateListener(FStateTreeDelegateListener Listener) const
+{
+	if (!Listener.IsValid())
+	{
+		// Listener is not bound, not an error.
+		return true;
+	}
+
+	TStrongObjectPtr<UObject> OwnerPinned = Owner.Pin();
+	TStrongObjectPtr<const UStateTree> StateTreePinned = StateTree.Pin();
+	TSharedPtr<FStateTreeInstanceStorage> StoragePinned = Storage.Pin();
+	if (OwnerPinned && StateTreePinned && StoragePinned)
+	{
+		// Manually acquire and release the write access.
+		//The context may already exist and we want to confirm that we are on the same thread.
+		StoragePinned->AcquireWriteAccess();
+
+		FStateTreeExecutionState& Exec = StoragePinned->GetMutableExecutionState();
+		if (Exec.CurrentExecutionContext)
+		{
+			Exec.CurrentExecutionContext->RemoveDelegateListener(Listener);
+		}
+		else
+		{
+			TStrongObjectPtr<const UStateTree> FrameStateTreePinned = FrameStateTree.Pin();
+			if (FrameStateTreePinned)
 			{
-				Exec.ExecutionExtension.GetMutable().ScheduleNextTick(FStateTreeExecutionExtension::FContextParameters(*OwnerPinned.Get(), *StateTreePinned.Get(), *StoragePinned.Get()));
+				StoragePinned->GetMutableExecutionState().DelegateActiveListeners.RemoveListener(Listener, FrameStateTreePinned.Get());
 			}
-			return true;
 		}
 
 		StoragePinned->ReleaseWriteAccess();
+		return true;
 	}
+
 	return false;
 }
 
+void FStateTreeWeakExecutionContext::ScheduleNextTick(gsl::not_null<UObject*> Owner, gsl::not_null<const UStateTree*> RootStateTree, FStateTreeInstanceStorage& Storage)
+{
+	TInstancedStruct<FStateTreeExecutionExtension>& ExecutionExtension = Storage.GetMutableExecutionState().ExecutionExtension;
+	if (RootStateTree->IsScheduledTickAllowed() && ExecutionExtension.IsValid())
+	{
+		ExecutionExtension.GetMutable().ScheduleNextTick(FStateTreeExecutionExtension::FContextParameters(*Owner, *RootStateTree, Storage));
+	}
+}
+
 #undef STATETREE_LOG
 #undef STATETREE_CLOG
diff --git a/Engine/Plugins/Runtime/StateTree/Source/StateTreeModule/Private/StateTreeExecutionTypes.cpp b/Engine/Plugins/Runtime/StateTree/Source/StateTreeModule/Private/StateTreeExecutionTypes.cpp
--- a/Engine/Plugins/Runtime/StateTree/Source/StateTreeModule/Private/StateTreeExecutionTypes.cpp
+++ b/Engine/Plugins/Runtime/StateTree/Source/StateTreeModule/Private/StateTreeExecutionTypes.cpp
@@ -2,6 +2,7 @@
 
 #include "StateTreeExecutionTypes.h"
 #include "StateTree.h"
+#include "StateTreeDelegate.h"
 
 const FStateTreeExternalDataHandle FStateTreeExternalDataHandle::Invalid = FStateTreeExternalDataHandle();
 
@@ -164,3 +165,109 @@ float FStateTreeScheduledTick::GetTickRate() const
 {
 	return NextDeltaTime;
 }
+
+FStateTreeDelegateActiveListeners::FActiveListener::FActiveListener(FStateTreeDelegateListener InListener, FSimpleDelegate InDelegate, gsl::not_null<const UStateTree*> InFrameStateTree, FStateTreeStateHandle InFrameRootState, FStateTreeStateHandle InOwningState, FStateTreeIndex16 InOwningNodeIndex)
+	: FrameStateTree(InFrameStateTree.get())
+	, Delegate(MoveTemp(InDelegate))
+	, DispatcherID(InListener.GetBoundDispatcherID())
+	, ListenerID(InListener.GetListenerID())
+	, OwningState(InOwningState)
+	, OwningNodeIndex(InOwningNodeIndex)
+	, FrameRootState(InFrameRootState)
+{}
+
+FStateTreeDelegateActiveListeners::~FStateTreeDelegateActiveListeners()
+{
+	check(BroadcastingLockCount == 0);
+}
+
+bool FStateTreeDelegateActiveListeners::AddListener(FStateTreeDelegateListener Listener, FSimpleDelegate Delegate, gsl::not_null<const UStateTree*> FrameStateTree, FStateTreeStateHandle FrameRootState, FStateTreeStateHandle OwningState, FStateTreeIndex16 OwningNodeIndex)
+{
+	check(Listener.IsValid());
+
+	for (FActiveListener& ActiveListener : Listeners)
+	{
+		if (ActiveListener.ListenerID == Listener.GetListenerID())
+		{
+			return false;
+		}
+	}
+
+	Listeners.Emplace(Listener, MoveTemp(Delegate), FrameStateTree, FrameRootState, OwningState, OwningNodeIndex);
+	return true;
+}
+
+void FStateTreeDelegateActiveListeners::RemoveListener(FStateTreeDelegateListener Listener, gsl::not_null<const UStateTree*> FrameStateTree)
+{
+	check(Listener.IsValid());
+
+	FObjectKey FrameStateTreeKey = FrameStateTree;
+	const int32 Index = Listeners.IndexOfByPredicate([FrameStateTreeKey, &Listener](const FActiveListener& ActiveListener)
+		{
+			return FrameStateTreeKey == ActiveListener.FrameStateTree && ActiveListener.ListenerID == Listener.GetListenerID();
+		});
+
+	if (Index == INDEX_NONE)
+	{
+		return;
+	}
+
+	if (BroadcastingLockCount > 0)
+	{
+		Listeners[Index] = FActiveListener();
+		bContainsUnboundListeners = true;
+	}
+	else
+	{
+		Listeners.RemoveAtSwap(Index);
+	}
+}
+
+void FStateTreeDelegateActiveListeners::RemoveListenersByPredicate(TFunctionRef<bool(FStateTreeStateHandle OwningState, FStateTreeIndex16 OwningNodeIndex, const UStateTree* FrameStateTree, FStateTreeStateHandle FrameRootState)> Predicate)
+{
+	check(BroadcastingLockCount == 0);
+	Listeners.RemoveAllSwap([&](const FActiveListener& Listener)
+		{
+			return Predicate(Listener.OwningState, Listener.OwningNodeIndex, CastChecked<const UStateTree>(Listener.FrameStateTree.ResolveObjectPtr()), Listener.FrameRootState);
+		});
+}
+
+void FStateTreeDelegateActiveListeners::BroadcastDelegate(FStateTreeDelegateDispatcher Dispatcher, gsl::not_null<const UStateTree*> FrameStateTree, const FStateTreeExecutionState& Exec)
+{
+	check(Dispatcher.IsValid());
+
+	++BroadcastingLockCount;
+
+
+	FObjectKey FrameStateTreeKey = FrameStateTree;
+	for (int32 i = Listeners.Num() - 1; i >= 0; --i)
+	{
+		const FActiveListener& Listener = Listeners[i];
+		if (Listener.OwningNodeIndex.Get() <= Exec.EnterStateFailedTaskIndex.Get()
+		&& Listener.OwningNodeIndex.Get() < Exec.LastExitedNodeIndex.Get()
+		&& Listener.DispatcherID == Dispatcher.GetDelegateID()
+		&& Listener.FrameStateTree == FrameStateTreeKey)
+		{
+			Listener.Delegate.ExecuteIfBound();
+		}
+	}
+
+	--BroadcastingLockCount;
+
+	if (BroadcastingLockCount == 0)
+	{
+		RemoveUnboundListeners();
+	}
+}
+
+void FStateTreeDelegateActiveListeners::RemoveUnboundListeners()
+{
+	check(BroadcastingLockCount == 0);
+	if (!bContainsUnboundListeners)
+	{
+		return;
+	}
+
+	Listeners.RemoveAllSwap([](const FActiveListener& Listener) { return !Listener.IsValid(); });
+	bContainsUnboundListeners = false;
+}
diff --git a/Engine/Plugins/Runtime/StateTree/Source/StateTreeModule/Private/StateTreeInstanceData.cpp b/Engine/Plugins/Runtime/StateTree/Source/StateTreeModule/Private/StateTreeInstanceData.cpp
--- a/Engine/Plugins/Runtime/StateTree/Source/StateTreeModule/Private/StateTreeInstanceData.cpp
+++ b/Engine/Plugins/Runtime/StateTree/Source/StateTreeModule/Private/StateTreeInstanceData.cpp
@@ -10,6 +10,7 @@
 #include "Serialization/ObjectAndNameAsStringProxyArchive.h"
 #include "Serialization/CustomVersion.h"
 #include "Serialization/PropertyLocalizationDataGathering.h"
+#include "StateTreeDelegate.h"
 
 #include UE_INLINE_GENERATED_CPP_BY_NAME(StateTreeInstanceData)
 
@@ -206,6 +207,26 @@ void FStateTreeInstanceStorage::AddTransitionRequest(const UObject* Owner, const
 	TransitionRequests.Add(Request);
 }
 
+void FStateTreeInstanceStorage::MarkDelegateAsBroadcasted(FStateTreeDelegateDispatcher Dispatcher, gsl::not_null<const UStateTree*> BelongingStateTree)
+{
+	BroadcastedDelegates.Add({BelongingStateTree, Dispatcher.GetDelegateID()});
+}
+
+bool FStateTreeInstanceStorage::IsDelegateBroadcasted(FStateTreeIndex16 DispatcherID, const FStateTreeExecutionFrame& BelongingFrame) const
+{
+	return BroadcastedDelegates.Contains(FStateTreeBroadcastedDelegate{BelongingFrame.StateTree, DispatcherID});
+}
+
+void FStateTreeInstanceStorage::ResetBroadcastedDelegates()
+{
+	BroadcastedDelegates.Reset();
+}
+
+bool FStateTreeInstanceStorage::HasBroadcastedDelegates() const
+{
+	return !BroadcastedDelegates.IsEmpty();
+}
+
 void FStateTreeInstanceStorage::ResetTransitionRequests()
 {
 	TransitionRequests.Reset();
diff --git a/Engine/Plugins/Runtime/StateTree/Source/StateTreeModule/Public/Blueprint/StateTreeNodeBlueprintBase.h b/Engine/Plugins/Runtime/StateTree/Source/StateTreeModule/Public/Blueprint/StateTreeNodeBlueprintBase.h
--- a/Engine/Plugins/Runtime/StateTree/Source/StateTreeModule/Public/Blueprint/StateTreeNodeBlueprintBase.h
+++ b/Engine/Plugins/Runtime/StateTree/Source/StateTreeModule/Public/Blueprint/StateTreeNodeBlueprintBase.h
@@ -53,6 +53,10 @@ public:
 	UFUNCTION(CustomThunk)
 	void GetPropertyReference(const FStateTreeBlueprintPropertyRef& PropertyRef) const;
 
+	/** Broadcasts provided delegate. */
+	UFUNCTION(BlueprintCallable, Category = "StateTree", meta = (HideSelfPin = "true"))
+	void BroadcastDelegate(FStateTreeDelegateDispatcher Dispatcher);
+
 	/** Returns true if reference to selected property in State Tree is accessible. */
 	UFUNCTION()
 	bool IsPropertyRefValid(const FStateTreeBlueprintPropertyRef& PropertyRef) const;
diff --git a/Engine/Plugins/Runtime/StateTree/Source/StateTreeModule/Public/StateTree.h b/Engine/Plugins/Runtime/StateTree/Source/StateTreeModule/Public/StateTree.h
--- a/Engine/Plugins/Runtime/StateTree/Source/StateTreeModule/Public/StateTree.h
+++ b/Engine/Plugins/Runtime/StateTree/Source/StateTreeModule/Public/StateTree.h
@@ -170,6 +170,12 @@ public:
 		return bHasGlobalTickTasks || (bHasEvents && bHasGlobalTickTasksOnlyOnEvents);
 	}
 
+	/** @return true if the tree can use the scheduled tick feature. */
+	bool IsScheduledTickAllowed() const
+	{
+		return bScheduledTickAllowed;
+	}
+
 #if WITH_EDITOR
 	/** Resets the compiled data to empty. */
 	void ResetCompiled();
diff --git a/Engine/Plugins/Runtime/StateTree/Source/StateTreeModule/Public/StateTreeDelegate.h b/Engine/Plugins/Runtime/StateTree/Source/StateTreeModule/Public/StateTreeDelegate.h
new file mode 100644
--- /dev/null
+++ b/Engine/Plugins/Runtime/StateTree/Source/StateTreeModule/Public/StateTreeDelegate.h
@@ -0,0 +1,86 @@
+// Copyright Epic Games, Inc. All Rights Reserved.
+
+#pragma once
+
+#include "StateTreeIndexTypes.h"
+#include "StateTreeDelegate.generated.h"
+
+ /**
+  * StateTree's internal delegate dispatcher.
+  */
+USTRUCT(BlueprintType)
+struct STATETREEMODULE_API FStateTreeDelegateDispatcher
+{
+	GENERATED_BODY()
+
+public:
+	/**
+	 * Used internally.
+	 * @return ID of the delegate.
+	 */
+	FStateTreeIndex16 GetDelegateID() const
+	{
+		return ID;
+	}
+
+	/**
+	 * Used internally.
+	 * @return if dispatcher is valid.
+	 */
+	bool IsValid() const
+	{
+		return ID.IsValid();
+	}
+
+private:
+	UPROPERTY()
+	FStateTreeIndex16 ID = FStateTreeIndex16::Invalid;
+
+	friend struct FStateTreePropertyBindingCompiler;
+};
+
+/**
+ * StateTree's internal delegate listener.
+ */
+USTRUCT()
+struct STATETREEMODULE_API FStateTreeDelegateListener
+{
+	GENERATED_BODY()
+
+public:
+	/**
+	 * Used internally.
+	 * @return ID of the bound delegate.
+	 */
+	FStateTreeIndex16 GetBoundDispatcherID() const
+	{
+		return DispatcherID;
+	}
+
+	/**
+	 * Used internally.
+	 * @return ID of the listener.
+	 */
+	FStateTreeIndex16 GetListenerID() const
+	{
+		return ListenerID;
+	}
+
+	/**
+	 * Used internally.
+	 * @return if the listener is valid.
+	 */
+	bool IsValid() const
+	{
+		return DispatcherID.IsValid() && ListenerID.IsValid();
+	}
+
+private:
+	UPROPERTY()
+	FStateTreeIndex16 DispatcherID = FStateTreeIndex16::Invalid;
+
+	UPROPERTY()
+	FStateTreeIndex16 ListenerID = FStateTreeIndex16::Invalid;
+
+	friend struct FStateTreePropertyBindingCompiler;
+};
\ No newline at end of file
diff --git a/Engine/Plugins/Runtime/StateTree/Source/StateTreeModule/Public/StateTreeExecutionContext.h b/Engine/Plugins/Runtime/StateTree/Source/StateTreeModule/Public/StateTreeExecutionContext.h
--- a/Engine/Plugins/Runtime/StateTree/Source/StateTreeModule/Public/StateTreeExecutionContext.h
+++ b/Engine/Plugins/Runtime/StateTree/Source/StateTreeModule/Public/StateTreeExecutionContext.h
@@ -16,6 +16,7 @@ struct FStateTreeReferenceOverrides;
 struct FStateTreeEvaluatorBase;
 struct FStateTreeTaskBase;
 struct FStateTreeConditionBase;
+struct FStateTreeDelegateDispatcher;
 struct FStateTreeEvent;
 struct FStateTreeMinimalExecutionContext;
 struct FStateTreeTransitionRequest;
@@ -281,6 +282,15 @@ public:
 	/** @return the names of all the active state. */
 	TArray<FName> GetActiveStateNames() const;
 
+	/** Broadcasts the delegate. */
+	void BroadcastDelegate(FStateTreeDelegateDispatcher Dispatcher);
+
+	/** Adds delegate listener. */
+	bool AddDelegateListener(FStateTreeDelegateListener Listener, FSimpleDelegate Delegate);
+
+	/** Removes delegate listener. */
+	void RemoveDelegateListener(FStateTreeDelegateListener Listener);
+
 	/**
 	 * Iterates over all events.
 	 * @param Function a lambda which takes const FStateTreeSharedEvent& Event, and returns EStateTreeLoopEvents.
@@ -604,12 +614,25 @@ protected:
 	 */
 	void ExitState(const FStateTreeTransitionResult& Transition);
 
+	/**
+	 * Removes all delegate listeners of deactivated frames and states.
+	 * @param ChangedFrame The frame where the change of the state took place.
+	 * @param ChangedState First state in ChangedFrame that became inactive.
+	 * @param DeactivatedFrames Frames that were entirely deactivated.
+	 */
+	void RemoveDeactivatedDelegateListeners(const FStateTreeExecutionFrame* ChangedFrame, FStateTreeStateHandle ChangedState, TConstArrayView<FStateTreeExecutionFrame> DeactivatedFrames);
+
+	/**
+	 * Removes all delegate listeners.
+	 */
+	void RemoveAllDelegateListeners();
+
 	/**
 	 * Handles logic for signaling State completed. StateCompleted is called on current active Evaluators and Tasks in reverse order (from leaf to root).
 	 */
 	void StateCompleted();
 
-	/** 
+	/**
 	 * Tick evaluators and global tasks by delta time.
 	 */
 	EStateTreeRunStatus TickEvaluatorsAndGlobalTasks(const float DeltaTime, bool bTickGlobalTasks = true);
@@ -1034,6 +1057,7 @@ struct STATETREEMODULE_API FStateTreeWeakExecutionContext
 {
 public:
 	FStateTreeWeakExecutionContext() = default;
+	explicit FStateTreeWeakExecutionContext(const FStateTreeExecutionContext& Context);
 	explicit FStateTreeWeakExecutionContext(gsl::not_null<UObject*> Owner, gsl::not_null<const UStateTree*> StateTree, FStateTreeInstanceData& Storage);
 	explicit FStateTreeWeakExecutionContext(gsl::not_null<UObject*> Owner, gsl::not_null<const UStateTree*> StateTree, TWeakPtr<FStateTreeInstanceStorage> Storage);
 
@@ -1064,8 +1088,20 @@ public:
 	 */
 	bool RequestTransition(FStateTreeStateHandle TargetState, EStateTreeTransitionPriority Priority = EStateTreeTransitionPriority::Normal, EStateTreeSelectionFallback Fallback = EStateTreeSelectionFallback::None) const;
 
+	/** Broadcasts the delegate. */
+	bool BroadcastDelegate(FStateTreeDelegateDispatcher Dispatcher) const;
+
+	/** Removes Delegate Listener. */
+	bool RemoveDelegateListener(FStateTreeDelegateListener Listener) const;
+
 private:
+	static void ScheduleNextTick(gsl::not_null<UObject*> Owner, gsl::not_null<const UStateTree*> RootStateTree, FStateTreeInstanceStorage& Storage);
+
 	TWeakObjectPtr<UObject> Owner;
 	TWeakObjectPtr<const UStateTree> StateTree;
 	TWeakPtr<FStateTreeInstanceStorage> Storage;
+
+	// This will be replace by StatePath
+	TWeakObjectPtr<const UStateTree> FrameStateTree;
+	FStateTreeStateHandle FrameRootState = FStateTreeStateHandle::Invalid;
 };
diff --git a/Engine/Plugins/Runtime/StateTree/Source/StateTreeModule/Public/StateTreeExecutionTypes.h b/Engine/Plugins/Runtime/StateTree/Source/StateTreeModule/Public/StateTreeExecutionTypes.h
--- a/Engine/Plugins/Runtime/StateTree/Source/StateTreeModule/Public/StateTreeExecutionTypes.h
+++ b/Engine/Plugins/Runtime/StateTree/Source/StateTreeModule/Public/StateTreeExecutionTypes.h
@@ -4,11 +4,15 @@
 #include "StateTreeTypes.h"
 #include "StateTreeEvents.h"
 #include "StateTreeExecutionExtension.h"
+#include "UObject/ObjectKey.h"
 
 #include "StateTreeExecutionTypes.generated.h"
 
 struct FStateTreeExecutionContext;
 class UStateTree;
+struct FStateTreeInstanceStorage;
+struct FStateTreeDelegateDispatcher;
+struct FStateTreeDelegateListener;
 
 /**
  * Enumeration for the different update phases.
@@ -576,6 +580,56 @@ struct STATETREEMODULE_API FStateTreeTransitionDelayedState
 	FStateTreeIndex16 TransitionIndex = FStateTreeIndex16::Invalid;
 };
 
+/** Describes added delegate listeners. */
+USTRUCT()
+struct STATETREEMODULE_API FStateTreeDelegateActiveListeners
+{
+	GENERATED_BODY()
+
+	~FStateTreeDelegateActiveListeners();
+
+	/** Adds the listener to the list. */
+	bool AddListener(FStateTreeDelegateListener Listener, FSimpleDelegate Delegate, gsl::not_null<const UStateTree*> FrameStateTree, FStateTreeStateHandle FrameRootState, FStateTreeStateHandle OwningState, FStateTreeIndex16 OwningNodeIndex);
+	
+	/** Removes the listener. Safe to be called during broadcasting. */
+	void RemoveListener(FStateTreeDelegateListener Listener, gsl::not_null<const UStateTree*> FrameStateTree);
+
+	/** Removes the listener by predicate. */
+	void RemoveListenersByPredicate(TFunctionRef<bool(FStateTreeStateHandle OwningState, FStateTreeIndex16 OwningNodeIndex, const UStateTree* FrameStateTree, FStateTreeStateHandle FrameRootState)> Predicate);
+
+	/** Broadcasts the listener by predicate. */
+	void BroadcastDelegate(FStateTreeDelegateDispatcher Dispatcher, gsl::not_null<const UStateTree*> FrameStateTree, const FStateTreeExecutionState& Exec);
+
+private:
+	void RemoveUnboundListeners();
+
+	struct FActiveListener
+	{
+		FActiveListener() = default;
+		FActiveListener(FStateTreeDelegateListener Listener, FSimpleDelegate InDelegate, gsl::not_null<const UStateTree*> FrameStateTree, FStateTreeStateHandle FrameRootState, FStateTreeStateHandle OwningState, FStateTreeIndex16 OwningNodeIndex);
+
+		bool IsValid() const
+		{
+			return Delegate.IsBound();
+		}
+
+		FObjectKey FrameStateTree;
+		FSimpleDelegate Delegate;
+
+		FStateTreeIndex16 DispatcherID = FStateTreeIndex16::Invalid;
+		FStateTreeIndex16 ListenerID = FStateTreeIndex16::Invalid;
+
+		FStateTreeStateHandle OwningState = FStateTreeStateHandle::Invalid;
+		FStateTreeIndex16 OwningNodeIndex = FStateTreeIndex16::Invalid;
+
+		FStateTreeStateHandle FrameRootState;
+	};
+
+	TArray<FActiveListener> Listeners;
+	uint32 BroadcastingLockCount : 31 = 0;
+	uint32 bContainsUnboundListeners : 1 = false;
+};
+
 /** Describes an active branch of a State Tree. */
 USTRUCT(BlueprintType)
 struct STATETREEMODULE_API FStateTreeExecutionFrame
@@ -614,6 +668,14 @@ struct STATETREEMODULE_API FStateTreeExecutionFrame
 	UPROPERTY()
 	FStateTreeIndex16 ActiveInstanceIndexBase = FStateTreeIndex16::Invalid;
 
+	/** Index within the active listeners to the first listener belonging to the frame. */
+	UPROPERTY()
+	FStateTreeIndex16 ActiveListenersIndexBase = FStateTreeIndex16::Invalid;
+
+	/** Index of the first delegate dispatcher belonging to the frame. */
+	UPROPERTY()
+	FStateTreeIndex16 DelegateDispatcherIndexBase = FStateTreeIndex16::Invalid;
+
 	/** Handle to the state parameter data, exists in ParentFrame. */
 	UPROPERTY()
 	FStateTreeDataHandle StateParameterDataHandle = FStateTreeDataHandle::Invalid; 
@@ -680,6 +742,10 @@ PRAGMA_ENABLE_DEPRECATION_WARNINGS
 	UPROPERTY()
 	FRandomStream RandomStream;
 
+	/** Active delegate listeners. */
+	UPROPERTY()
+	FStateTreeDelegateActiveListeners DelegateActiveListeners;
+
 #if WITH_STATETREE_TRACE
 	/** Id for the active instance used for debugging. */
 	mutable FStateTreeInstanceDebugId InstanceDebugId;
@@ -700,6 +766,10 @@ PRAGMA_ENABLE_DEPRECATION_WARNINGS
 	UPROPERTY()
 	FStateTreeIndex16 EnterStateFailedTaskIndex = FStateTreeIndex16::Invalid;
 
+	/** The index of the node that just finished exiting. */
+	UPROPERTY()
+	FStateTreeIndex16 LastExitedNodeIndex = FStateTreeIndex16::Invalid;
+
 	/** Result of last tick */
 	UPROPERTY()
 	EStateTreeRunStatus LastTickStatus = EStateTreeRunStatus::Failed;
diff --git a/Engine/Plugins/Runtime/StateTree/Source/StateTreeModule/Public/StateTreeInstanceData.h b/Engine/Plugins/Runtime/StateTree/Source/StateTreeModule/Public/StateTreeInstanceData.h
--- a/Engine/Plugins/Runtime/StateTree/Source/StateTreeModule/Public/StateTreeInstanceData.h
+++ b/Engine/Plugins/Runtime/StateTree/Source/StateTreeModule/Public/StateTreeInstanceData.h
@@ -14,6 +14,7 @@
 
 struct FStateTreeTransitionRequest;
 struct FStateTreeExecutionState;
+struct FStateTreeDelegateDispatcher;
 
 #if WITH_STATETREE_DEBUG
 namespace UE::StateTree::Debug
@@ -60,6 +61,22 @@ struct STATETREEMODULE_API FStateTreeTemporaryInstanceData
 	FInstancedStruct Instance;
 };
 
+/**
+ * Represents a delegate that was broadcasted.
+ */
+USTRUCT()
+struct FStateTreeBroadcastedDelegate
+{
+	GENERATED_BODY()
+
+	bool operator==(const FStateTreeBroadcastedDelegate&) const = default;
+
+	FObjectKey  FrameStateTree;
+
+	UPROPERTY()
+	FStateTreeIndex16 DelegateID = FStateTreeIndex16::Invalid;
+};
+
 struct STATETREEMODULE_API FStateTreeInstanceStorageCustomVersion
 {
 	enum Type
@@ -139,6 +156,18 @@ struct STATETREEMODULE_API FStateTreeInstanceStorage
 	*/
 	void AddTransitionRequest(const UObject* Owner, const FStateTreeTransitionRequest& Request);
 
+	/** Marks delegate as broadcasted. */
+	void MarkDelegateAsBroadcasted(FStateTreeDelegateDispatcher Dispatcher, gsl::not_null<const UStateTree*> BelongingStateTree);
+
+	/** @return true if the delegate of the given index was broadcasted. */
+	bool IsDelegateBroadcasted(FStateTreeIndex16 DispatcherID, const FStateTreeExecutionFrame& BelongingFrame) const;
+
+	/** Resets the list of broadcasted delegates. */
+	void ResetBroadcastedDelegates();
+
+	/** @return true if there's any broadcasted delegates. */
+	bool HasBroadcastedDelegates() const;
+
 	/** @return currently pending transition requests. */
 	TConstArrayView<FStateTreeTransitionRequest> GetTransitionRequests() const
 	{
@@ -274,6 +303,9 @@ protected:
 	/** Events (Transient) */
 	TSharedRef<FStateTreeEventQueue> EventQueue = MakeShared<FStateTreeEventQueue>();
 	
+	/** Array of broadcasted delegates. */
+	TArray<FStateTreeBroadcastedDelegate> BroadcastedDelegates;
+
 	/** Requested transitions */
 	UPROPERTY(Transient)
 	TArray<FStateTreeTransitionRequest> TransitionRequests;
diff --git a/Engine/Plugins/Runtime/StateTree/Source/StateTreeModule/Public/StateTreePropertyBindings.h b/Engine/Plugins/Runtime/StateTree/Source/StateTreeModule/Public/StateTreePropertyBindings.h
--- a/Engine/Plugins/Runtime/StateTree/Source/StateTreeModule/Public/StateTreePropertyBindings.h
+++ b/Engine/Plugins/Runtime/StateTree/Source/StateTreeModule/Public/StateTreePropertyBindings.h
@@ -41,6 +41,8 @@ enum class EStateTreeBindableStructSource : uint8
 	StateEvent,
 	/** Source is Property Function */
 	PropertyFunction,
+	/** Source is Transition */
+	Transition,
 };
 
 
diff --git a/Engine/Plugins/Runtime/StateTree/Source/StateTreeModule/Public/StateTreeTypes.h b/Engine/Plugins/Runtime/StateTree/Source/StateTreeModule/Public/StateTreeTypes.h
--- a/Engine/Plugins/Runtime/StateTree/Source/StateTreeModule/Public/StateTreeTypes.h
+++ b/Engine/Plugins/Runtime/StateTree/Source/StateTreeModule/Public/StateTreeTypes.h
@@ -172,6 +172,9 @@ enum class EStateTreeTransitionTrigger : uint8
 	/** Try trigger transition on specific State Tree event. */
 	OnEvent = 0x8,
 
+	/** Try trigger transition on specific State Tree delegate. */
+	OnDelegate = 0x10,
+
 	MAX
 };
 ENUM_CLASS_FLAGS(EStateTreeTransitionTrigger)
@@ -619,6 +622,10 @@ struct STATETREEMODULE_API FCompactStateTransition
 	UPROPERTY()
 	uint16 ConditionsBegin = 0;
 
+	/** ID of the delegate dispatcher the transition is listening to. */
+	UPROPERTY()
+	FStateTreeIndex16 DelegateDispatcherID = FStateTreeIndex16::Invalid;
+
 	/** Target state of the transition */
 	UPROPERTY()
 	FStateTreeStateHandle State = FStateTreeStateHandle::Invalid;
@@ -679,9 +686,9 @@ struct STATETREEMODULE_API FCompactStateTreeState
 	}
 
 	/** @return True if the state has any transitions that need ticking. */
-	bool ShouldTickTransitions(bool bHasEvent) const
+	bool ShouldTickTransitions(bool bHasEvent, bool bHasBroadcastedDelegates) const
 	{
-		return bHasTickTriggerTransitions || (bHasEvent && bHasEventTriggerTransitions);
+		return bHasTickTriggerTransitions || (bHasEvent && bHasEventTriggerTransitions) || (bHasBroadcastedDelegates && bHasDelegateTriggerTransitions);
 	}
 
 	/** @return True if the state has any transitions that tick on completion. */
@@ -755,7 +762,7 @@ struct STATETREEMODULE_API FCompactStateTreeState
 	/** Weight used to scale the normalized final utility score for this state */
 	UPROPERTY()
 	float Weight = 1.f;
-	
+
 	/**
 	 * Tick rate in seconds the state tasks and transitions should tick
 	 * If set the state cannot sleep.
@@ -828,6 +835,10 @@ struct STATETREEMODULE_API FCompactStateTreeState
 	UPROPERTY()	
 	uint8 bHasEventTriggerTransitions : 1 = false;
 
+	/** True if the state contains transitions with delegate trigger. */
+	UPROPERTY()	
+	uint8 bHasDelegateTriggerTransitions : 1 = false;
+
 	/** True if the state contains transitions with completed trigger. */
 	UPROPERTY()
 	uint8 bHasCompletedTriggerTransitions : 1 = false;
diff --git a/Engine/Plugins/Runtime/StateTree/Source/StateTreeTestSuite/Private/StateTreeDelegateTest.cpp b/Engine/Plugins/Runtime/StateTree/Source/StateTreeTestSuite/Private/StateTreeDelegateTest.cpp
new file mode 100644
--- /dev/null
+++ b/Engine/Plugins/Runtime/StateTree/Source/StateTreeTestSuite/Private/StateTreeDelegateTest.cpp
@@ -0,0 +1,396 @@
+// Copyright Epic Games, Inc. All Rights Reserved.
+
+#include "StateTreeTest.h"
+#include "StateTreeTestBase.h"
+#include "StateTreeTestTypes.h"
+
+#include "StateTreeCompilerLog.h"
+#include "StateTreeEditorData.h"
+#include "StateTreeCompiler.h"
+#include "Conditions/StateTreeCommonConditions.h"
+
+#define LOCTEXT_NAMESPACE "AITestSuite_StateTreeTest"
+
+UE_DISABLE_OPTIMIZATION_SHIP
+
+namespace UE::StateTree::Tests
+{
+struct FStateTreeTest_Delegate_ConcurrentListeners : FStateTreeTestBase
+{
+	virtual bool InstantTest() override
+	{
+		UStateTree& StateTree = NewStateTree();
+		UStateTreeEditorData& EditorData = *Cast<UStateTreeEditorData>(StateTree.EditorData);
+		UStateTreeState& Root = EditorData.AddSubTree(FName("Root"));
+
+		auto& DispatcherTask = Root.AddTask<FTestTask_BroadcastDelegate>(FName("DispatcherTask"));
+		auto& ListenerTaskA = Root.AddTask<FTestTask_ListenDelegate>(FName("ListenerTaskA"));
+		auto& ListenerTaskB = Root.AddTask<FTestTask_ListenDelegate>(FName("ListenerTaskB"));
+
+		EditorData.AddPropertyBinding(DispatcherTask, GET_MEMBER_NAME_STRING_CHECKED(FTestTask_BroadcastDelegate_InstanceData, OnTickDelegate), ListenerTaskB, GET_MEMBER_NAME_STRING_CHECKED(FTestTask_ListenDelegate_InstanceData, Listener));
+		EditorData.AddPropertyBinding(DispatcherTask, GET_MEMBER_NAME_STRING_CHECKED(FTestTask_BroadcastDelegate_InstanceData, OnTickDelegate), ListenerTaskA, GET_MEMBER_NAME_STRING_CHECKED(FTestTask_ListenDelegate_InstanceData, Listener));
+
+		FStateTreeCompilerLog Log;
+		FStateTreeCompiler Compiler(Log);
+		const bool bResult = Compiler.Compile(StateTree);
+		AITEST_TRUE("StateTree should get compiled", bResult);
+
+		FStateTreeInstanceData InstanceData;
+		FTestStateTreeExecutionContext Exec(StateTree, StateTree, InstanceData);
+		const bool bInitSucceeded = Exec.IsValid();
+		AITEST_TRUE("StateTree should init", bInitSucceeded);
+
+		Exec.Start();
+		Exec.LogClear();
+
+		Exec.Tick(0.1f);
+		AITEST_TRUE(TEXT("StateTree ListenerTaskA should be triggered once."), Exec.Expect(ListenerTaskA.GetName(), *FString::Printf(TEXT("OnDelegate%d"), 1)));
+		AITEST_TRUE(TEXT("StateTree ListenerTaskB should be triggered once."), Exec.Expect(ListenerTaskB.GetName(), *FString::Printf(TEXT("OnDelegate%d"), 1)));
+
+		Exec.LogClear();
+
+		Exec.Tick(0.1f);
+		AITEST_TRUE(TEXT("StateTree ListenerTaskA should be triggered twice."), Exec.Expect(ListenerTaskA.GetName(), *FString::Printf(TEXT("OnDelegate%d"), 2)));
+		AITEST_TRUE(TEXT("StateTree ListenerTaskB should be triggered twice."), Exec.Expect(ListenerTaskB.GetName(), *FString::Printf(TEXT("OnDelegate%d"), 2)));
+		Exec.LogClear();
+
+		Exec.Stop();
+
+		return true;
+	}
+};
+IMPLEMENT_AI_INSTANT_TEST(FStateTreeTest_Delegate_ConcurrentListeners, "System.StateTree.Delegate.ConcurrentListeners");
+
+struct FStateTreeTest_Delegate_MutuallyExclusiveListeners : FStateTreeTestBase
+{
+	virtual bool InstantTest() override
+	{
+		UStateTree& StateTree = NewStateTree();
+		UStateTreeEditorData& EditorData = *Cast<UStateTreeEditorData>(StateTree.EditorData);
+		UStateTreeState& Root = EditorData.AddSubTree(FName(TEXT("Root")));
+
+		UStateTreeState& StateA = Root.AddChildState("A");
+		UStateTreeState& StateB = Root.AddChildState("B");
+
+		StateA.AddTransition(EStateTreeTransitionTrigger::OnTick, EStateTreeTransitionType::GotoState, &StateB);
+		StateB.AddTransition(EStateTreeTransitionTrigger::OnTick, EStateTreeTransitionType::GotoState, &StateA);
+
+		auto& DispatcherTask = Root.AddTask<FTestTask_BroadcastDelegate>(FName(TEXT("DispatcherTask")));
+		auto& ListenerTaskA0 = StateA.AddTask<FTestTask_ListenDelegate>(FName(TEXT("ListenerTaskA0")));
+		auto& ListenerTaskA1 = StateA.AddTask<FTestTask_ListenDelegate>(FName(TEXT("ListenerTaskA1")));
+		auto& ListenerTaskB = StateB.AddTask<FTestTask_ListenDelegate>(FName(TEXT("ListenerTaskB")));
+
+		EditorData.AddPropertyBinding(DispatcherTask, GET_MEMBER_NAME_STRING_CHECKED(FTestTask_BroadcastDelegate_InstanceData, OnTickDelegate), ListenerTaskA0, GET_MEMBER_NAME_STRING_CHECKED(FTestTask_ListenDelegate_InstanceData, Listener));
+		EditorData.AddPropertyBinding(DispatcherTask, GET_MEMBER_NAME_STRING_CHECKED(FTestTask_BroadcastDelegate_InstanceData, OnTickDelegate), ListenerTaskA1, GET_MEMBER_NAME_STRING_CHECKED(FTestTask_ListenDelegate_InstanceData, Listener));
+		EditorData.AddPropertyBinding(DispatcherTask, GET_MEMBER_NAME_STRING_CHECKED(FTestTask_BroadcastDelegate_InstanceData, OnTickDelegate), ListenerTaskB, GET_MEMBER_NAME_STRING_CHECKED(FTestTask_ListenDelegate_InstanceData, Listener));
+
+		FStateTreeCompilerLog Log;
+		FStateTreeCompiler Compiler(Log);
+		const bool bResult = Compiler.Compile(StateTree);
+		AITEST_TRUE("StateTree should get compiled", bResult);
+
+		FStateTreeInstanceData InstanceData;
+		FTestStateTreeExecutionContext Exec(StateTree, StateTree, InstanceData);
+		const bool bInitSucceeded = Exec.IsValid();
+		AITEST_TRUE("StateTree should init", bInitSucceeded);
+
+		Exec.Start();
+		AITEST_TRUE("StateTree Active States should be in Root/A", Exec.ExpectInActiveStates(Root.Name, StateA.Name));
+		Exec.LogClear();
+
+		Exec.Tick(0.1f);
+		AITEST_TRUE(TEXT("StateTree ListenerTaskA0 should be triggered once"), Exec.Expect(ListenerTaskA0.GetName(), *FString::Printf(TEXT("OnDelegate%d"), 1)));
+		AITEST_TRUE(TEXT("StateTree ListenerTaskA1 should be triggered once"), Exec.Expect(ListenerTaskA1.GetName(), *FString::Printf(TEXT("OnDelegate%d"), 1)));
+		AITEST_TRUE(TEXT("StateTree ListenerTaskB shouldn't be triggered."), !Exec.Expect(ListenerTaskB.GetName(), *FString::Printf(TEXT("OnDelegate%d"), 1)));
+		AITEST_TRUE("StateTree Active States should be in Root/B", Exec.ExpectInActiveStates(Root.Name, StateB.Name));
+		Exec.LogClear();
+
+		Exec.Tick(0.1f);
+		AITEST_TRUE(TEXT("StateTree ListenerTaskB should be triggered once"), Exec.Expect(ListenerTaskB.GetName(), *FString::Printf(TEXT("OnDelegate%d"), 1)));
+		AITEST_TRUE(TEXT("StateTree ListenerTaskA0 shouldn't be triggered."), !Exec.Expect(ListenerTaskA0.GetName(), *FString::Printf(TEXT("OnDelegate%d"), 1)));
+		AITEST_TRUE(TEXT("StateTree ListenerTaskA1 shouldn't be triggered."), !Exec.Expect(ListenerTaskA1.GetName(), *FString::Printf(TEXT("OnDelegate%d"), 1)));
+		AITEST_TRUE("StateTree Active States should be in Root/A", Exec.ExpectInActiveStates(Root.Name, StateA.Name));
+		Exec.LogClear();
+
+		Exec.Stop();
+
+		return true;
+	}
+};
+IMPLEMENT_AI_INSTANT_TEST(FStateTreeTest_Delegate_MutuallyExclusiveListeners, "System.StateTree.Delegate.MutuallyExclusiveListeners");
+
+struct FStateTreeTest_Delegate_Transitions : FStateTreeTestBase
+{
+	virtual bool InstantTest() override
+	{
+		UStateTree& StateTree = NewStateTree();
+		UStateTreeEditorData& EditorData = *Cast<UStateTreeEditorData>(StateTree.EditorData);
+		UStateTreeState& Root = EditorData.AddSubTree(FName(TEXT("Root")));
+
+		UStateTreeState& StateA = Root.AddChildState("A");
+		UStateTreeState& StateB = Root.AddChildState("B");
+
+		FStateTreeTransition& TransitionAToB = StateA.AddTransition(EStateTreeTransitionTrigger::OnDelegate, EStateTreeTransitionType::GotoState, &StateB);
+		FStateTreeTransition& TransitionBToA = StateB.AddTransition(EStateTreeTransitionTrigger::OnDelegate, EStateTreeTransitionType::GotoState, &StateA);
+
+		auto& DispatcherTask0 = Root.AddTask<FTestTask_BroadcastDelegate>(FName(TEXT("DispatcherTask0")));
+		auto& DispatcherTask1 = Root.AddTask<FTestTask_BroadcastDelegate>(FName(TEXT("DispatcherTask1")));
+		
+		EditorData.AddPropertyBinding(FStateTreePropertyPath(DispatcherTask0.ID, GET_MEMBER_NAME_CHECKED(FTestTask_BroadcastDelegate_InstanceData, OnTickDelegate)),  FStateTreePropertyPath(TransitionAToB.ID, GET_MEMBER_NAME_CHECKED(FStateTreeTransition, DelegateListener)));
+		EditorData.AddPropertyBinding(FStateTreePropertyPath(DispatcherTask1.ID, GET_MEMBER_NAME_CHECKED(FTestTask_BroadcastDelegate_InstanceData, OnTickDelegate)), FStateTreePropertyPath(TransitionBToA.ID, GET_MEMBER_NAME_CHECKED(FStateTreeTransition, DelegateListener)));
+
+		FStateTreeCompilerLog Log;
+		FStateTreeCompiler Compiler(Log);
+		const bool bResult = Compiler.Compile(StateTree);
+		AITEST_TRUE("StateTree should get compiled", bResult);
+
+		FStateTreeInstanceData InstanceData;
+		FTestStateTreeExecutionContext Exec(StateTree, StateTree, InstanceData);
+		const bool bInitSucceeded = Exec.IsValid();
+		AITEST_TRUE("StateTree should init", bInitSucceeded);
+
+		Exec.Start();
+		AITEST_TRUE("StateTree Active States should be in Root/A", Exec.ExpectInActiveStates(Root.Name, StateA.Name));
+		Exec.LogClear();
+
+		Exec.Tick(0.1f);
+		AITEST_TRUE("StateTree Active States should be in Root/B", Exec.ExpectInActiveStates(Root.Name, StateB.Name));
+		Exec.LogClear();
+
+		Exec.Tick(0.1f);
+		AITEST_TRUE("StateTree Active States should be in Root/A", Exec.ExpectInActiveStates(Root.Name, StateA.Name));
+		Exec.LogClear();
+
+		Exec.Stop();
+
+		return true;
+	}
+};
+IMPLEMENT_AI_INSTANT_TEST(FStateTreeTest_Delegate_Transitions, "System.StateTree.Delegate.Transitions");
+
+struct FStateTreeTest_Delegate_Rebroadcasting : FStateTreeTestBase
+{
+	virtual bool InstantTest() override
+	{
+		UStateTree& StateTree = NewStateTree();
+		UStateTreeEditorData& EditorData = *Cast<UStateTreeEditorData>(StateTree.EditorData);
+		UStateTreeState& Root = EditorData.AddSubTree(FName(TEXT("Root")));
+
+		auto& DispatcherTask = Root.AddTask<FTestTask_BroadcastDelegate>(FName(TEXT("DispatcherTask")));
+		auto& RedispatcherTask = Root.AddTask<FTestTask_RebroadcastDelegate>(FName(TEXT("RedispatcherTask")));
+		auto& ListenerTask = Root.AddTask<FTestTask_ListenDelegate>(FName(TEXT("ListenerTask")));
+
+		EditorData.AddPropertyBinding(DispatcherTask, GET_MEMBER_NAME_STRING_CHECKED(FTestTask_BroadcastDelegate_InstanceData, OnTickDelegate), RedispatcherTask, TEXT("Listener"));
+		EditorData.AddPropertyBinding(RedispatcherTask, TEXT("Dispatcher"), ListenerTask, TEXT("Listener"));
+
+		FStateTreeCompilerLog Log;
+		FStateTreeCompiler Compiler(Log);
+		const bool bResult = Compiler.Compile(StateTree);
+		AITEST_TRUE("StateTree should get compiled", bResult);
+
+		FStateTreeInstanceData InstanceData;
+		FTestStateTreeExecutionContext Exec(StateTree, StateTree, InstanceData);
+		const bool bInitSucceeded = Exec.IsValid();
+		AITEST_TRUE("StateTree should init", bInitSucceeded);
+
+		Exec.Start();
+		Exec.LogClear();
+
+		Exec.Tick(0.1f);
+		AITEST_TRUE(TEXT("StateTree ListenerTask should be triggered once."), Exec.Expect(ListenerTask.GetName(), *FString::Printf(TEXT("OnDelegate%d"), 1)));
+		Exec.LogClear();
+
+		Exec.Tick(0.1f);
+		AITEST_TRUE(TEXT("StateTree ListenerTask should be triggered twice."), Exec.Expect(ListenerTask.GetName(), *FString::Printf(TEXT("OnDelegate%d"), 2)));
+		Exec.LogClear();
+
+		Exec.Stop();
+
+		return true;
+	}
+};
+IMPLEMENT_AI_INSTANT_TEST(FStateTreeTest_Delegate_Rebroadcasting, "System.StateTree.Delegate.Rebroadcasting");
+
+struct FStateTreeTest_Delegate_SelfRemoval : FStateTreeTestBase
+{
+	virtual bool InstantTest() override
+	{
+		UStateTree& StateTree = NewStateTree();
+		UStateTreeEditorData& EditorData = *Cast<UStateTreeEditorData>(StateTree.EditorData);
+		UStateTreeState& Root = EditorData.AddSubTree(FName(TEXT("Root")));
+
+		auto& DispatcherTask = Root.AddTask<FTestTask_BroadcastDelegate>(FName(TEXT("DispatcherTask")));
+		auto& CustomFuncTask = Root.AddTask<FTestTask_CustomFuncOnDelegate>(FName(TEXT("CustomFuncTask")));
+
+		uint32 TriggersCounter = 0;
+
+		EditorData.AddPropertyBinding(DispatcherTask, GET_MEMBER_NAME_STRING_CHECKED(FTestTask_BroadcastDelegate_InstanceData, OnTickDelegate), CustomFuncTask, GET_MEMBER_NAME_STRING_CHECKED(FTestTask_CustomFuncOnDelegate_InstanceData, Listener));
+		CustomFuncTask.GetNode().CustomFunc = [&](const FStateTreeWeakExecutionContext& WeakContext, FStateTreeDelegateListener Listener)
+			{
+				++TriggersCounter;
+				WeakContext.RemoveDelegateListener(Listener);
+			};
+
+		FStateTreeCompilerLog Log;
+		FStateTreeCompiler Compiler(Log);
+		const bool bResult = Compiler.Compile(StateTree);
+		AITEST_TRUE("StateTree should get compiled", bResult);
+
+		FStateTreeInstanceData InstanceData;
+		FTestStateTreeExecutionContext Exec(StateTree, StateTree, InstanceData);
+		const bool bInitSucceeded = Exec.IsValid();
+		AITEST_TRUE("StateTree should init", bInitSucceeded);
+
+		Exec.Start();
+		Exec.LogClear();
+
+		Exec.Tick(0.1f);
+		AITEST_EQUAL(TEXT("StateTree Delegate should be triggered once"), TriggersCounter, 1);
+		Exec.LogClear();
+
+		Exec.Tick(0.1f);
+		AITEST_EQUAL(TEXT("StateTree Delegate should be triggered once"), TriggersCounter, 1);
+		Exec.LogClear();
+
+		Exec.Stop();
+
+		return true;
+	}
+};
+IMPLEMENT_AI_INSTANT_TEST(FStateTreeTest_Delegate_SelfRemoval, "System.StateTree.Delegate.SelfRemoval");
+
+struct FStateTreeTest_Delegate_WithoutRemoval : FStateTreeTestBase
+{
+	virtual bool InstantTest() override
+	{
+		UStateTree& StateTree = NewStateTree();
+		UStateTreeEditorData& EditorData = *Cast<UStateTreeEditorData>(StateTree.EditorData);
+		UStateTreeState& Root = EditorData.AddSubTree(FName(TEXT("Root")));
+
+		UStateTreeState& StateA = Root.AddChildState("A");
+		UStateTreeState& StateB = Root.AddChildState("B");
+
+		FStateTreeTransition& TransitionAToB = StateA.AddTransition(EStateTreeTransitionTrigger::OnTick, EStateTreeTransitionType::GotoState, &StateB);
+
+		auto& DispatcherTask = Root.AddTask<FTestTask_BroadcastDelegate>(FName(TEXT("DispatcherTask")));
+		auto& ListenerTask = StateA.AddTask<FTestTask_ListenDelegate>(FName(TEXT("ListenerTask")));
+		EditorData.AddPropertyBinding(DispatcherTask, GET_MEMBER_NAME_STRING_CHECKED(FTestTask_BroadcastDelegate_InstanceData, OnTickDelegate), ListenerTask, GET_MEMBER_NAME_STRING_CHECKED(FTestTask_ListenDelegate_InstanceData, Listener));
+		ListenerTask.GetNode().bRemoveOnExit = false;
+
+		FStateTreeCompilerLog Log;
+		FStateTreeCompiler Compiler(Log);
+		const bool bResult = Compiler.Compile(StateTree);
+		AITEST_TRUE("StateTree should get compiled", bResult);
+
+		FStateTreeInstanceData InstanceData;
+		FTestStateTreeExecutionContext Exec(StateTree, StateTree, InstanceData);
+		const bool bInitSucceeded = Exec.IsValid();
+		AITEST_TRUE("StateTree should init", bInitSucceeded);
+
+		Exec.Start();
+		Exec.LogClear();
+		AITEST_TRUE("StateTree Active States should be in Root/A", Exec.ExpectInActiveStates(Root.Name, StateA.Name));
+
+		Exec.Tick(0.1f);
+		AITEST_TRUE(TEXT("StateTree Delegate should be triggered once."), Exec.Expect(ListenerTask.GetName(), *FString::Printf(TEXT("OnDelegate%d"), 1)));
+		AITEST_TRUE("StateTree Active States should be in Root/B", Exec.ExpectInActiveStates(Root.Name, StateB.Name));
+		Exec.LogClear();
+
+		Exec.Tick(0.1f);
+		AITEST_FALSE(TEXT("StateTree Delegate shouldn't be triggered again."), Exec.Expect(ListenerTask.GetName()));
+		AITEST_TRUE("StateTree Active States should be in Root/B", Exec.ExpectInActiveStates(Root.Name, StateB.Name));
+		Exec.LogClear();
+
+		Exec.Stop();
+
+		return true;
+	}
+};
+IMPLEMENT_AI_INSTANT_TEST(FStateTreeTest_Delegate_WithoutRemoval, "System.StateTree.Delegate.WithoutRemoval");
+
+struct FStateTreeTest_Delegate_GlobalDispatcherAndListener : FStateTreeTestBase
+{
+	virtual bool InstantTest() override
+	{
+		UStateTree& StateTree = NewStateTree();
+		UStateTreeEditorData& EditorData = *Cast<UStateTreeEditorData>(StateTree.EditorData);
+		UStateTreeState& Root = EditorData.AddSubTree(FName(TEXT("Root")));
+		auto& RootTask = Root.AddTask<FTestTask_Stand>();
+		RootTask.GetNode().TicksToCompletion = 100;
+
+		auto& DispatcherTask = EditorData.AddGlobalTask<FTestTask_BroadcastDelegate>(FName(TEXT("DispatcherTask")));
+		auto& ListenerTask = EditorData.AddGlobalTask<FTestTask_ListenDelegate>(FName(TEXT("ListenerTask")));
+		ListenerTask.GetNode().bRemoveOnExit = false;
+
+		EditorData.AddPropertyBinding(DispatcherTask, GET_MEMBER_NAME_STRING_CHECKED(FTestTask_BroadcastDelegate_InstanceData, OnTickDelegate), ListenerTask, GET_MEMBER_NAME_STRING_CHECKED(FTestTask_CustomFuncOnDelegate_InstanceData, Listener));
+
+		FStateTreeCompilerLog Log;
+		FStateTreeCompiler Compiler(Log);
+		const bool bResult = Compiler.Compile(StateTree);
+		AITEST_TRUE("StateTree should get compiled", bResult);
+
+		FStateTreeInstanceData InstanceData;
+		FTestStateTreeExecutionContext Exec(StateTree, StateTree, InstanceData);
+		const bool bInitSucceeded = Exec.IsValid();
+		AITEST_TRUE("StateTree should init", bInitSucceeded);
+
+		Exec.Start();
+		Exec.LogClear();
+
+		Exec.Tick(0.1f);
+		AITEST_TRUE(TEXT("StateTree Delegate should be triggered once."), Exec.Expect(ListenerTask.GetName(), *FString::Printf(TEXT("OnDelegate%d"), 1)));
+		Exec.LogClear();
+
+		Exec.Tick(0.1f);
+		AITEST_TRUE(TEXT("StateTree Delegate should be triggered twice."), Exec.Expect(ListenerTask.GetName(), *FString::Printf(TEXT("OnDelegate%d"), 2)));
+		Exec.LogClear();
+
+		Exec.Stop();
+
+		return true;
+	}
+};
+IMPLEMENT_AI_INSTANT_TEST(FStateTreeTest_Delegate_GlobalDispatcherAndListener, "System.StateTree.Delegate.GlobalDispatcherAndListener");
+
+struct FStateTreeTest_Delegate_ListeningToDelegateOnExit : FStateTreeTestBase
+{
+	virtual bool InstantTest() override
+	{
+		UStateTree& StateTree = NewStateTree();
+		UStateTreeEditorData& EditorData = *Cast<UStateTreeEditorData>(StateTree.EditorData);
+		UStateTreeState& Root = EditorData.AddSubTree(FName(TEXT("Root")));
+
+		auto& DispatcherTask = Root.AddTask<FTestTask_BroadcastDelegate>(FName(TEXT("DispatcherTask")));
+		auto& ListenerTask = Root.AddTask<FTestTask_ListenDelegate>(FName(TEXT("ListenerTask")));
+		ListenerTask.GetNode().bRemoveOnExit = false;
+
+		EditorData.AddPropertyBinding(DispatcherTask, GET_MEMBER_NAME_STRING_CHECKED(FTestTask_BroadcastDelegate_InstanceData, OnExitDelegate), ListenerTask, GET_MEMBER_NAME_STRING_CHECKED(FTestTask_CustomFuncOnDelegate_InstanceData, Listener));
+
+		FStateTreeCompilerLog Log;
+		FStateTreeCompiler Compiler(Log);
+		const bool bResult = Compiler.Compile(StateTree);
+		AITEST_TRUE("StateTree should get compiled", bResult);
+
+		FStateTreeInstanceData InstanceData;
+		FTestStateTreeExecutionContext Exec(StateTree, StateTree, InstanceData);
+		const bool bInitSucceeded = Exec.IsValid();
+		AITEST_TRUE("StateTree should init", bInitSucceeded);
+
+		Exec.Start();
+		AITEST_TRUE("Expected Root to be active.", Exec.ExpectInActiveStates(Root.Name));
+		Exec.LogClear();
+
+		Exec.Stop();
+		AITEST_FALSE(TEXT("StateTree Delegate shouldn't be triggered"), Exec.Expect(ListenerTask.GetName()));
+		Exec.LogClear();
+
+		return true;
+	}
+};
+IMPLEMENT_AI_INSTANT_TEST(FStateTreeTest_Delegate_ListeningToDelegateOnExit, "System.StateTree.Delegate.ListeningToDelegateOnExit");
+
+} // namespace UE::StateTree::Tests
+
+UE_ENABLE_OPTIMIZATION_SHIP
+
+#undef LOCTEXT_NAMESPACE
diff --git a/Engine/Plugins/Runtime/StateTree/Source/StateTreeTestSuite/Private/StateTreeTestTypes.h b/Engine/Plugins/Runtime/StateTree/Source/StateTreeTestSuite/Private/StateTreeTestTypes.h
--- a/Engine/Plugins/Runtime/StateTree/Source/StateTreeTestSuite/Private/StateTreeTestTypes.h
+++ b/Engine/Plugins/Runtime/StateTree/Source/StateTreeTestSuite/Private/StateTreeTestTypes.h
@@ -9,6 +9,7 @@
 #include "StateTreeLinker.h"
 #include "StateTreePropertyRef.h"
 #include "StateTreePropertyFunctionBase.h"
+#include "StateTreeDelegate.h"
 #include "StateTreeTestTypes.generated.h"
 
 class UStateTree;
@@ -797,4 +798,211 @@ struct FStateTreeTestBooleanCondition : public FStateTreeConditionCommonBase
 	}
 
 	TStateTreeExternalDataHandle<FStateTreeTestLog> LogHandle;
+};
+
+USTRUCT()
+struct FTestTask_BroadcastDelegate_InstanceData
+{
+	GENERATED_BODY()
+
+	UPROPERTY(EditAnywhere, Category = Parameter)
+	FStateTreeDelegateDispatcher OnEnterDelegate;
+
+	UPROPERTY(EditAnywhere, Category = Parameter)
+	FStateTreeDelegateDispatcher OnTickDelegate;
+
+	UPROPERTY(EditAnywhere, Category = Parameter)
+	FStateTreeDelegateDispatcher OnExitDelegate;
+};
+
+USTRUCT()
+struct FTestTask_BroadcastDelegate : public FStateTreeTaskBase
+{
+	GENERATED_BODY()
+
+	FTestTask_BroadcastDelegate() = default;
+	FTestTask_BroadcastDelegate(const FName InName) { Name = InName; }
+
+	using FInstanceDataType = FTestTask_BroadcastDelegate_InstanceData;
+
+	virtual const UStruct* GetInstanceDataType() const override
+	{
+		return FInstanceDataType::StaticStruct();
+	}
+
+	virtual EStateTreeRunStatus EnterState(FStateTreeExecutionContext& Context, const FStateTreeTransitionResult& Transition) const override
+	{
+		FTestTask_BroadcastDelegate_InstanceData& InstanceData = Context.GetInstanceData<FTestTask_BroadcastDelegate_InstanceData>(*this);
+		Context.BroadcastDelegate(InstanceData.OnEnterDelegate);
+		return EStateTreeRunStatus::Running;
+	}
+
+	virtual EStateTreeRunStatus Tick(FStateTreeExecutionContext& Context, const float DeltaTime) const override
+	{
+		FTestTask_BroadcastDelegate_InstanceData& InstanceData = Context.GetInstanceData<FTestTask_BroadcastDelegate_InstanceData>(*this);
+		Context.BroadcastDelegate(InstanceData.OnTickDelegate);
+		return EStateTreeRunStatus::Running;
+	}
+
+	virtual void ExitState(FStateTreeExecutionContext& Context, const FStateTreeTransitionResult& Transition) const override
+	{
+		FTestTask_BroadcastDelegate_InstanceData& InstanceData = Context.GetInstanceData<FTestTask_BroadcastDelegate_InstanceData>(*this);
+		Context.BroadcastDelegate(InstanceData.OnExitDelegate);
+	}
+};
+
+USTRUCT()
+struct FTestTask_ListenDelegate_InstanceData
+{
+	GENERATED_BODY()
+
+	UPROPERTY(EditAnywhere, Category = Parameter)
+	FStateTreeDelegateListener Listener;
+
+	uint32 TriggersAmount = 0u;
+};
+
+USTRUCT()
+struct FTestTask_ListenDelegate : public FStateTreeTaskBase
+{
+	GENERATED_BODY()
+
+	FTestTask_ListenDelegate() = default;
+	FTestTask_ListenDelegate(const FName InName) { Name = InName; }
+
+	using FInstanceDataType = FTestTask_ListenDelegate_InstanceData;
+
+	virtual const UStruct* GetInstanceDataType() const override
+	{
+		return FInstanceDataType::StaticStruct();
+	}
+
+	virtual EStateTreeRunStatus EnterState(FStateTreeExecutionContext& Context, const FStateTreeTransitionResult& Transition) const override
+	{
+		FTestTask_ListenDelegate_InstanceData& InstanceData = Context.GetInstanceData<FTestTask_ListenDelegate_InstanceData>(*this);
+		FStateTreeTestLog& TestLog = Context.GetExternalData(LogHandle);
+		Context.AddDelegateListener(InstanceData.Listener, FSimpleDelegate::CreateLambda([&TestLog, this, InstanceDataRef = Context.GetInstanceDataStructRef(*this)]()
+			{
+				if (FInstanceDataType* InstanceData = InstanceDataRef.GetPtr())
+				{
+					++InstanceData->TriggersAmount;
+					TestLog.Log(Name, FString::Printf(TEXT("OnDelegate%d"), InstanceData->TriggersAmount));
+				}
+			}));
+
+		return EStateTreeRunStatus::Running;
+	}
+
+	virtual bool Link(FStateTreeLinker& Linker) override
+	{
+		Linker.LinkExternalData(LogHandle);
+		return true;
+	}
+
+	virtual void ExitState(FStateTreeExecutionContext& Context, const FStateTreeTransitionResult& Transition) const override
+	{
+		if (bRemoveOnExit)
+		{
+			FTestTask_ListenDelegate_InstanceData& InstanceData = Context.GetInstanceData<FTestTask_ListenDelegate_InstanceData>(*this);
+			Context.RemoveDelegateListener(InstanceData.Listener);
+		}
+	}
+	
+	bool bRemoveOnExit = true;
+	TStateTreeExternalDataHandle<FStateTreeTestLog> LogHandle;
+};
+
+USTRUCT()
+struct FTestTask_RebroadcastDelegate_InstanceData
+{
+	GENERATED_BODY()
+
+	UPROPERTY(EditAnywhere, Category = Parameter)
+	FStateTreeDelegateListener Listener;
+
+	UPROPERTY(EditAnywhere, Category = Parameter)
+	FStateTreeDelegateDispatcher Dispatcher;
+};
+
+USTRUCT()
+struct FTestTask_RebroadcastDelegate : public FStateTreeTaskBase
+{
+	GENERATED_BODY()
+
+	FTestTask_RebroadcastDelegate() = default;
+	FTestTask_RebroadcastDelegate(const FName InName) { Name = InName; }
+
+	using FInstanceDataType = FTestTask_RebroadcastDelegate_InstanceData;
+
+	virtual const UStruct* GetInstanceDataType() const override
+	{
+		return FInstanceDataType::StaticStruct();
+	}
+
+	virtual EStateTreeRunStatus EnterState(FStateTreeExecutionContext& Context, const FStateTreeTransitionResult& Transition) const override
+	{
+		FTestTask_RebroadcastDelegate_InstanceData& InstanceData = Context.GetInstanceData<FTestTask_RebroadcastDelegate_InstanceData>(*this);
+
+		Context.AddDelegateListener(InstanceData.Listener, FSimpleDelegate::CreateLambda([Dispatcher = InstanceData.Dispatcher, WeakContext = Context.MakeWeakExecutionContext()]()
+			{
+				WeakContext.BroadcastDelegate(Dispatcher);
+			}));
+
+		return EStateTreeRunStatus::Running;
+	}
+
+	virtual void ExitState(FStateTreeExecutionContext& Context, const FStateTreeTransitionResult& Transition) const override
+	{
+		FTestTask_RebroadcastDelegate_InstanceData& InstanceData = Context.GetInstanceData<FTestTask_RebroadcastDelegate_InstanceData>(*this);
+		Context.RemoveDelegateListener(InstanceData.Listener);
+	}
+};
+
+USTRUCT()
+struct FTestTask_CustomFuncOnDelegate_InstanceData
+{
+	GENERATED_BODY()
+
+	UPROPERTY(EditAnywhere, Category = Parameter)
+	FStateTreeDelegateListener Listener;
+};
+
+USTRUCT()
+struct FTestTask_CustomFuncOnDelegate : public FStateTreeTaskBase
+{
+	GENERATED_BODY()
+
+	FTestTask_CustomFuncOnDelegate() = default;
+	FTestTask_CustomFuncOnDelegate(const FName InName) { Name = InName; }
+
+	using FInstanceDataType = FTestTask_CustomFuncOnDelegate_InstanceData;
+
+	virtual const UStruct* GetInstanceDataType() const override
+	{
+		return FInstanceDataType::StaticStruct();
+	}
+
+	virtual EStateTreeRunStatus EnterState(FStateTreeExecutionContext& Context, const FStateTreeTransitionResult& Transition) const override
+	{
+		FTestTask_CustomFuncOnDelegate_InstanceData& InstanceData = Context.GetInstanceData<FTestTask_CustomFuncOnDelegate_InstanceData>(*this);
+
+		Context.AddDelegateListener(InstanceData.Listener, FSimpleDelegate::CreateLambda([Func = CustomFunc, Listener = InstanceData.Listener, WeakContext = Context.MakeWeakExecutionContext()]()
+			{
+				Func(WeakContext, Listener);
+			}));
+		return EStateTreeRunStatus::Running;
+	}
+
+	virtual void ExitState(FStateTreeExecutionContext& Context, const FStateTreeTransitionResult& Transition) const override
+	{
+		if (bRemoveOnExit)
+		{
+			FTestTask_CustomFuncOnDelegate_InstanceData& InstanceData = Context.GetInstanceData<FTestTask_CustomFuncOnDelegate_InstanceData>(*this);
+			Context.RemoveDelegateListener(InstanceData.Listener);
+		}
+	}
+
+	TFunction<void(const FStateTreeWeakExecutionContext&, FStateTreeDelegateListener)> CustomFunc;
+
+	bool bRemoveOnExit = true;
 };
\ No newline at end of file
