StateTree: Remove V547 warning. Introduced with CL38388747

#rnx

diff --git a/Engine/Plugins/Runtime/StateTree/Source/StateTreeModule/Private/StateTreeExecutionContext.cpp b/Engine/Plugins/Runtime/StateTree/Source/StateTreeModule/Private/StateTreeExecutionContext.cpp
--- a/Engine/Plugins/Runtime/StateTree/Source/StateTreeModule/Private/StateTreeExecutionContext.cpp
+++ b/Engine/Plugins/Runtime/StateTree/Source/StateTreeModule/Private/StateTreeExecutionContext.cpp
@@ -3920,7 +3920,7 @@ bool FStateTreeExecutionContext::TriggerTransitions()
 
 			// if a valid completion transition has already been found, the remaining transitions in parent frames won't have a higher priority than the found one
 			// so skip the remainder. this also prevented false positive warnings and ensures from STDebugger
-			if (NextTransition.Priority != EStateTreeTransitionPriority::None)
+			if (NextTransition.Priority != EStateTreeTransitionPriority::None) //-V547
 			{
 				break;
 			}
