Change type from gsl::not_null to TNotNull. Only updated private or newly introduced not_null that doesn't impact compilation.

diff --git a/Engine/Plugins/Runtime/GameplayStateTree/Source/GameplayStateTreeModule/Private/Components/StateTreeComponentSchema.cpp b/Engine/Plugins/Runtime/GameplayStateTree/Source/GameplayStateTreeModule/Private/Components/StateTreeComponentSchema.cpp
--- a/Engine/Plugins/Runtime/GameplayStateTree/Source/GameplayStateTreeModule/Private/Components/StateTreeComponentSchema.cpp
+++ b/Engine/Plugins/Runtime/GameplayStateTree/Source/GameplayStateTreeModule/Private/Components/StateTreeComponentSchema.cpp
@@ -96,7 +96,7 @@ void UStateTreeComponentSchema::PostEditChangeChainProperty(FPropertyChangedChai
 }
 #endif
 
-UStateTreeComponentSchema::FContextDataSetter::FContextDataSetter(gsl::not_null<const UBrainComponent*> InBrainComponent, FStateTreeExecutionContext& InContext)
+UStateTreeComponentSchema::FContextDataSetter::FContextDataSetter(TNotNull<const UBrainComponent*> InBrainComponent, FStateTreeExecutionContext& InContext)
 	: BrainComponent(InBrainComponent)
 	, ExecutionContext(InContext)
 {}
@@ -106,12 +106,12 @@ bool UStateTreeComponentSchema::FContextDataSetter::SetContextDataByName(FName N
 	return ExecutionContext.IsValid() ? ExecutionContext.SetContextDataByName(Name, DataView) : false;
 }
 
-gsl::not_null<const UStateTree*> UStateTreeComponentSchema::FContextDataSetter::GetStateTree() const
+TNotNull<const UStateTree*> UStateTreeComponentSchema::FContextDataSetter::GetStateTree() const
 {
 	return ExecutionContext.GetStateTree();
 }
 
-gsl::not_null<const UStateTreeComponentSchema*> UStateTreeComponentSchema::FContextDataSetter::GetSchema() const
+TNotNull<const UStateTreeComponentSchema*> UStateTreeComponentSchema::FContextDataSetter::GetSchema() const
 {
 	return Cast<UStateTreeComponentSchema>(ExecutionContext.GetStateTree()->GetSchema());
 }
diff --git a/Engine/Plugins/Runtime/GameplayStateTree/Source/GameplayStateTreeModule/Public/Components/StateTreeComponentSchema.h b/Engine/Plugins/Runtime/GameplayStateTree/Source/GameplayStateTreeModule/Public/Components/StateTreeComponentSchema.h
--- a/Engine/Plugins/Runtime/GameplayStateTree/Source/GameplayStateTreeModule/Public/Components/StateTreeComponentSchema.h
+++ b/Engine/Plugins/Runtime/GameplayStateTree/Source/GameplayStateTreeModule/Public/Components/StateTreeComponentSchema.h
@@ -46,19 +46,19 @@ protected:
 	struct FContextDataSetter
 	{
 	public:
-		FContextDataSetter(gsl::not_null<const UBrainComponent*> BrainComponent, FStateTreeExecutionContext& Context);
+		FContextDataSetter(TNotNull<const UBrainComponent*> BrainComponent, FStateTreeExecutionContext& Context);
 
-		gsl::not_null<const UBrainComponent*> GetComponent() const
+		TNotNull<const UBrainComponent*> GetComponent() const
 		{
 			return BrainComponent;
 		}
-		gsl::not_null<const UStateTree*> GetStateTree() const;
-		gsl::not_null<const UStateTreeComponentSchema*> GetSchema() const;
+		TNotNull<const UStateTree*> GetStateTree() const;
+		TNotNull<const UStateTreeComponentSchema*> GetSchema() const;
 
 		bool SetContextDataByName(FName Name, FStateTreeDataView DataView);
 
 	private:
-		gsl::not_null<const UBrainComponent*> BrainComponent;
+		TNotNull<const UBrainComponent*> BrainComponent;
 		FStateTreeExecutionContext& ExecutionContext;
 	};
 	virtual void SetContextData(FContextDataSetter& ContextDataSetter, bool bLogErrors) const;
diff --git a/Engine/Plugins/Runtime/StateTree/Source/StateTreeEditorModule/Private/SStateTreeDiff.cpp b/Engine/Plugins/Runtime/StateTree/Source/StateTreeEditorModule/Private/SStateTreeDiff.cpp
--- a/Engine/Plugins/Runtime/StateTree/Source/StateTreeEditorModule/Private/SStateTreeDiff.cpp
+++ b/Engine/Plugins/Runtime/StateTree/Source/StateTreeEditorModule/Private/SStateTreeDiff.cpp
@@ -197,7 +197,7 @@ void SDiffWidget::HandleAssetEditorRequestClose(UObject* Asset, const EAssetEdit
 	}
 }
 
-TSharedRef<SDiffWidget> SDiffWidget::CreateDiffWindow(const FText WindowTitle, gsl::not_null<const UStateTree*> OldStateTree, gsl::not_null<const UStateTree*> NewStateTree, const FRevisionInfo& OldRevision, const FRevisionInfo& NewRevision)
+TSharedRef<SDiffWidget> SDiffWidget::CreateDiffWindow(const FText WindowTitle, TNotNull<const UStateTree*> OldStateTree, TNotNull<const UStateTree*> NewStateTree, const FRevisionInfo& OldRevision, const FRevisionInfo& NewRevision)
 {
 	// sometimes we're comparing different revisions of one single asset (other
 	// times we're comparing two completely separate assets altogether)
@@ -231,21 +231,20 @@ TSharedRef<SDiffWidget> SDiffWidget::CreateDiffWindow(const FText WindowTitle, g
 	return StateTreeDiff;
 }
 
-TSharedRef<SDiffWidget> SDiffWidget::CreateDiffWindow(gsl::not_null<const UStateTree*> OldStateTree, gsl::not_null<const UStateTree*> NewStateTree,
+TSharedRef<SDiffWidget> SDiffWidget::CreateDiffWindow(TNotNull<const UStateTree*> OldStateTree, TNotNull<const UStateTree*> NewStateTree,
 	const FRevisionInfo& OldRevision, const FRevisionInfo& NewRevision, const UClass* StateTreeClass)
 {
-	check(OldStateTree || NewStateTree);
-
 	// sometimes we're comparing different revisions of one single asset (other
 	// times we're comparing two completely separate assets altogether)
-	const bool bIsSingleAsset = !OldStateTree || !NewStateTree || (NewStateTree->GetName() == OldStateTree->GetName());
+	//@TODO use pathname instead of asset name.
+	const bool bIsSingleAsset = NewStateTree->GetFName() == OldStateTree->GetFName();
 
 	FText WindowTitle = FText::Format(LOCTEXT("NamelessStateTreeDiff", "{0} Diff (experimental)"), StateTreeClass->GetDisplayNameText());
 	// if we're diffing one asset against itself
 	if (bIsSingleAsset)
 	{
 		// identify the assumed single asset in the window's title
-		const FString STName = NewStateTree ? NewStateTree->GetName() : OldStateTree->GetName();
+		const FString STName = NewStateTree->GetName();
 		WindowTitle = FText::Format(LOCTEXT("NamedStateTreeDiff", "{0} - {1} Diff (experimental)"), FText::FromString(STName), StateTreeClass->GetDisplayNameText());
 	}
 
diff --git a/Engine/Plugins/Runtime/StateTree/Source/StateTreeEditorModule/Private/SStateTreeDiff.h b/Engine/Plugins/Runtime/StateTree/Source/StateTreeEditorModule/Private/SStateTreeDiff.h
--- a/Engine/Plugins/Runtime/StateTree/Source/StateTreeEditorModule/Private/SStateTreeDiff.h
+++ b/Engine/Plugins/Runtime/StateTree/Source/StateTreeEditorModule/Private/SStateTreeDiff.h
@@ -58,10 +58,10 @@ public:
 	virtual ~SDiffWidget();
 
 	/** Helper function to create a window that holds a diff widget */
-	static TSharedRef<SDiffWidget> CreateDiffWindow(FText WindowTitle, gsl::not_null<const UStateTree*> OldStateTree, gsl::not_null<const UStateTree*> NewStateTree, const FRevisionInfo& OldRevision, const FRevisionInfo& NewRevision);
+	static TSharedRef<SDiffWidget> CreateDiffWindow(FText WindowTitle, TNotNull<const UStateTree*> OldStateTree, TNotNull<const UStateTree*> NewStateTree, const FRevisionInfo& OldRevision, const FRevisionInfo& NewRevision);
 
 	/** Helper function to create a window that holds a diff widget (default window title) */
-	static TSharedRef<SDiffWidget> CreateDiffWindow(gsl::not_null<const UStateTree*> OldStateTree, gsl::not_null<const UStateTree*> NewStateTree, const FRevisionInfo& OldRevision, const FRevisionInfo& NewRevision, const UClass* StateTreeClass);
+	static TSharedRef<SDiffWidget> CreateDiffWindow(TNotNull<const UStateTree*> OldStateTree, TNotNull<const UStateTree*> NewStateTree, const FRevisionInfo& OldRevision, const FRevisionInfo& NewRevision, const UClass* StateTreeClass);
 
 protected:
 	/** Called when user clicks button to go to next difference */
diff --git a/Engine/Plugins/Runtime/StateTree/Source/StateTreeEditorModule/Private/StateTreeDiffControl.cpp b/Engine/Plugins/Runtime/StateTree/Source/StateTreeEditorModule/Private/StateTreeDiffControl.cpp
--- a/Engine/Plugins/Runtime/StateTree/Source/StateTreeEditorModule/Private/StateTreeDiffControl.cpp
+++ b/Engine/Plugins/Runtime/StateTree/Source/StateTreeEditorModule/Private/StateTreeDiffControl.cpp
@@ -91,7 +91,7 @@ FDiffControl::~FDiffControl()
 {
 }
 
-TSharedRef<SStateTreeView> FDiffControl::InsertObject(gsl::not_null<const UStateTree*> StateTree)
+TSharedRef<SStateTreeView> FDiffControl::InsertObject(TNotNull<const UStateTree*> StateTree)
 {
 	const FDiffWidgets DiffWidgets(StateTree);
 	TSharedRef<SStateTreeView> TreeView = DiffWidgets.GetStateTreeWidget();
diff --git a/Engine/Plugins/Runtime/StateTree/Source/StateTreeEditorModule/Private/StateTreeDiffControl.h b/Engine/Plugins/Runtime/StateTree/Source/StateTreeEditorModule/Private/StateTreeDiffControl.h
--- a/Engine/Plugins/Runtime/StateTree/Source/StateTreeEditorModule/Private/StateTreeDiffControl.h
+++ b/Engine/Plugins/Runtime/StateTree/Source/StateTreeEditorModule/Private/StateTreeDiffControl.h
@@ -62,7 +62,7 @@ protected:
 	static FText RightRevision;
 	static TSharedRef<SWidget> GenerateSingleEntryWidget(FSingleDiffEntry DiffEntry, FText ObjectName);
 
-	TSharedRef<SStateTreeView> InsertObject(gsl::not_null<const UStateTree*> StateTree);
+	TSharedRef<SStateTreeView> InsertObject(TNotNull<const UStateTree*> StateTree);
 
 	void OnSelectDiffEntry(const FSingleDiffEntry StateDiff);
 
diff --git a/Engine/Plugins/Runtime/StateTree/Source/StateTreeEditorModule/Private/StateTreeDiffHelper.cpp b/Engine/Plugins/Runtime/StateTree/Source/StateTreeEditorModule/Private/StateTreeDiffHelper.cpp
--- a/Engine/Plugins/Runtime/StateTree/Source/StateTreeEditorModule/Private/StateTreeDiffHelper.cpp
+++ b/Engine/Plugins/Runtime/StateTree/Source/StateTreeEditorModule/Private/StateTreeDiffHelper.cpp
@@ -32,7 +32,7 @@ FStateSoftPath::FChainElement::FChainElement(const UStateTreeState* State)
 	}
 }
 
-FStateSoftPath::FStateSoftPath(gsl::not_null<const UStateTreeState*> InState)
+FStateSoftPath::FStateSoftPath(TNotNull<const UStateTreeState*> InState)
 {
 	int32 ChainCount = 1;
 	const UStateTreeState* CurState = InState;
@@ -68,7 +68,7 @@ FString FStateSoftPath::ToDisplayName(bool Short) const
 	return Ret;
 }
 
-UStateTreeState* FStateSoftPath::ResolvePath(gsl::not_null<const UStateTree*> StateTree) const
+UStateTreeState* FStateSoftPath::ResolvePath(TNotNull<const UStateTree*> StateTree) const
 {
 	if (StateChain.Num() == 0)
 	{
@@ -78,7 +78,7 @@ UStateTreeState* FStateSoftPath::ResolvePath(gsl::not_null<const UStateTree*> St
 
 }
 
-UStateTreeState* FStateSoftPath::ResolvePath(gsl::not_null<const UStateTreeEditorData*> EditorData) const
+UStateTreeState* FStateSoftPath::ResolvePath(TNotNull<const UStateTreeEditorData*> EditorData) const
 {
 	if (StateChain.Num() == 0)
 	{
diff --git a/Engine/Plugins/Runtime/StateTree/Source/StateTreeEditorModule/Private/StateTreeDiffHelper.h b/Engine/Plugins/Runtime/StateTree/Source/StateTreeEditorModule/Private/StateTreeDiffHelper.h
--- a/Engine/Plugins/Runtime/StateTree/Source/StateTreeEditorModule/Private/StateTreeDiffHelper.h
+++ b/Engine/Plugins/Runtime/StateTree/Source/StateTreeEditorModule/Private/StateTreeDiffHelper.h
@@ -35,10 +35,10 @@ enum class EStateDiffType : uint8
 struct FStateSoftPath
 {
 	FStateSoftPath() = default;
-	explicit FStateSoftPath(gsl::not_null<const UStateTreeState*> InState);
+	explicit FStateSoftPath(TNotNull<const UStateTreeState*> InState);
 
-	UStateTreeState* ResolvePath(gsl::not_null<const UStateTree*> StateTree) const;
-	UStateTreeState* ResolvePath(gsl::not_null<const UStateTreeEditorData*> EditorData) const;
+	UStateTreeState* ResolvePath(TNotNull<const UStateTree*> StateTree) const;
+	UStateTreeState* ResolvePath(TNotNull<const UStateTreeEditorData*> EditorData) const;
 	FString ToDisplayName(bool Short = false) const;
 
 	/**
diff --git a/Engine/Source/Runtime/CoreUObject/Internal/UObject/UObjectArchetypeHelper.h b/Engine/Source/Runtime/CoreUObject/Internal/UObject/UObjectArchetypeHelper.h
--- a/Engine/Source/Runtime/CoreUObject/Internal/UObject/UObjectArchetypeHelper.h
+++ b/Engine/Source/Runtime/CoreUObject/Internal/UObject/UObjectArchetypeHelper.h
@@ -3,6 +3,7 @@
 #pragma once
 
 #include "CoreFwd.h"
+#include "Misc/NotNull.h"
 #include "UObject/Object.h"
 #include "UObject/UObjectAnnotation.h"
 
@@ -51,7 +52,7 @@ public:
 	 * as it renames and changes outers which makes impossible to retrieve it during the process.
 	 * @param Object top cache its archetype
 	 * @param Archetype to use as a cache, if null it will use GetArchetype on the owner of the specified object */
-	void CacheArchetype(gsl::not_null<const UObject*> Object, UObject* Archetype = nullptr)
+	void CacheArchetype(TNotNull<const UObject*> Object, UObject* Archetype = nullptr)
 	{
 		FCacheArchetype Annotation = ObjectCachedArchetypeAnnotations.GetAnnotation(Object);
 		if (Annotation.CachedArchetype == nullptr)
@@ -61,7 +62,7 @@ public:
 		}
 	}
 
-	void ResetCacheArchetype(gsl::not_null<const UObject*> Object)
+	void ResetCacheArchetype(TNotNull<const UObject*> Object)
 	{
 		ObjectCachedArchetypeAnnotations.RemoveAnnotation(Object);
 	}
@@ -69,7 +70,7 @@ public:
 	/**
 	 * Gets the archetype that was cached when the CacheArchetype method was called
 	 * @return the cached Archetype if any */
-	UObject* GetCachedArchetype(gsl::not_null<const UObject*> Object)
+	UObject* GetCachedArchetype(TNotNull<const UObject*> Object)
 	{
 		return ObjectCachedArchetypeAnnotations.GetAnnotation(Object).CachedArchetype;
 	}
