StateTree: Fix the binding picker for property reference from a null pointer. The data needed is inside the type. If the type is null, then the binding can't exist. The runtime also check if the data is valid. That take cares of virtual function that returns different types at runtime vs editor.

diff --git a/Engine/Plugins/Runtime/StateTree/Source/StateTreeModule/Private/StateTreePropertyRefHelpers.cpp b/Engine/Plugins/Runtime/StateTree/Source/StateTreeModule/Private/StateTreePropertyRefHelpers.cpp
--- a/Engine/Plugins/Runtime/StateTree/Source/StateTreeModule/Private/StateTreePropertyRefHelpers.cpp
+++ b/Engine/Plugins/Runtime/StateTree/Source/StateTreeModule/Private/StateTreePropertyRefHelpers.cpp
@@ -33,7 +33,7 @@ namespace UE::StateTree::PropertyRefHelpers
 	bool ArePropertyRefsCompatible(const FProperty& TargetRefProperty, const FProperty& SourceRefProperty, const void* TargetRefAddress, const void* SourceRefAddress)
 	{
 		check(IsPropertyRef(SourceRefProperty) && IsPropertyRef(TargetRefProperty));
-		check(SourceRefAddress && TargetRefAddress);
+		check(TargetRefAddress);
 
 		FEdGraphPinType SourceRefPin = GetPropertyRefInternalTypeAsPin(SourceRefProperty, SourceRefAddress);
 		FEdGraphPinType TargetRefPin = GetPropertyRefInternalTypeAsPin(TargetRefProperty, TargetRefAddress);
@@ -197,7 +197,6 @@ namespace UE::StateTree::PropertyRefHelpers
 	bool IsPropertyRefCompatibleWithProperty(const FProperty& RefProperty, const FProperty& SourceProperty, const void* PropertyRefAddress, const void* SourceAddress)
 	{
 		check(PropertyRefAddress);
-		check(SourceAddress);
 		check(IsPropertyRef(RefProperty));
 
 		if (IsPropertyRef(SourceProperty))
@@ -384,8 +383,11 @@ PRAGMA_ENABLE_DEPRECATION_WARNINGS
 			}
 			else if (StructProperty->Struct == FStateTreeBlueprintPropertyRef::StaticStruct())
 			{
-				check(PropertyRefAddress);
-				return GetBlueprintPropertyRefInternalTypeAsPin(*reinterpret_cast<const FStateTreeBlueprintPropertyRef*>(PropertyRefAddress));
+				// The source of the chain can be an uninitialized object.
+				if (PropertyRefAddress)
+				{
+					return GetBlueprintPropertyRefInternalTypeAsPin(*reinterpret_cast<const FStateTreeBlueprintPropertyRef*>(PropertyRefAddress));
+				}
 			}
 		}
 
