StateTree: Prevents warning and allow linking when the object reinstanced and the object is the FStruct::StaticStruct.

diff --git a/Engine/Plugins/Runtime/StateTree/Source/StateTreeModule/Private/Debugger/StateTreeRuntimeValidation.cpp b/Engine/Plugins/Runtime/StateTree/Source/StateTreeModule/Private/Debugger/StateTreeRuntimeValidation.cpp
--- a/Engine/Plugins/Runtime/StateTree/Source/StateTreeModule/Private/Debugger/StateTreeRuntimeValidation.cpp
+++ b/Engine/Plugins/Runtime/StateTree/Source/StateTreeModule/Private/Debugger/StateTreeRuntimeValidation.cpp
@@ -14,17 +14,24 @@ namespace UE::StateTree::Debug
 
 namespace Private
 {
-bool RuntimeValidationContext = true;
+bool bRuntimeValidationContext = true;
 static FAutoConsoleVariableRef CVarRuntimeValidationContext(
 	TEXT("StateTree.RuntimeValidation.Context"),
-	RuntimeValidationContext,
+	bRuntimeValidationContext,
 	TEXT("Test if the context creation parameters are the same between each creation of StateTreeExecutionContext.")
 );
 
-bool RuntimeValidationEnterExitState = true;
+bool bRuntimeValidationDoesNewerVersionExists = true;
+static FAutoConsoleVariableRef CVarRuntimeValidationDoesNewerVersionExists(
+	TEXT("StateTree.RuntimeValidation.DoesNewerVersionExists"),
+	bRuntimeValidationDoesNewerVersionExists,
+	TEXT("Test if a StateTreeExecutionContext started with an old version of a blueprint type.")
+);
+
+bool bRuntimeValidationEnterExitState = true;
 static FAutoConsoleVariableRef CVarRuntimeValidationEnterExitState(
 	TEXT("StateTree.RuntimeValidation.EnterExitState"),
-	RuntimeValidationEnterExitState,
+	bRuntimeValidationEnterExitState,
 	TEXT("Test that if a node get a EnterState, it will receive an ExitState.\n"
 		"Test that if a node get a ExitState, it did receive an EnterState before.")
 );
@@ -44,14 +51,14 @@ FString NodeToString(const UObject* Obj, FGuid Id)
 
 FRuntimeValidationInstanceData::~FRuntimeValidationInstanceData()
 {
-	if (Private::RuntimeValidationEnterExitState)
+	if (Private::bRuntimeValidationEnterExitState)
 	{
 		for (const FNodeStatePair& Pair : NodeStates)
 		{
 			if (EnumHasAllFlags(Pair.State, EState::BetweenEnterExitState))
 			{
 				ensureAlwaysMsgf(false, TEXT("Tree exited. Missing ExitState on %s."), *Private::NodeToString(StateTree.Get(), Pair.Id));
-				Private::RuntimeValidationEnterExitState = false;
+				Private::bRuntimeValidationEnterExitState = false;
 			}
 		}
 	}
@@ -61,7 +68,7 @@ void FRuntimeValidationInstanceData::SetContext(const UObject* InNewOwner, const
 {
 	TWeakObjectPtr<const UStateTree> NewStateTree = InNewStateTree;
 	TWeakObjectPtr<const UObject> NewOwner = InNewOwner;
-	if (Private::RuntimeValidationContext)
+	if (Private::bRuntimeValidationContext)
 	{
 		if (StateTree.IsValid() && StateTree != NewStateTree)
 		{
@@ -69,7 +76,7 @@ void FRuntimeValidationInstanceData::SetContext(const UObject* InNewOwner, const
 				"Make sure you initialize FStateTreeExecutionContext with the same value every time.\n"
 				"Auto deactivate Runtime check StateTree.RuntimeValidation.Context to prevent reporting the same error multiple times.")
 				, InNewStateTree ? *InNewStateTree->GetFullName() : TEXT("StateTree"), *StateTree.Get()->GetFullName());
-			Private::RuntimeValidationContext = false;
+			Private::bRuntimeValidationContext = false;
 		}
 		if (Owner.IsValid() && Owner != NewOwner)
 		{
@@ -77,9 +84,12 @@ void FRuntimeValidationInstanceData::SetContext(const UObject* InNewOwner, const
 				"Make sure you initialize FStateTreeExecutionContext with the same values every time.\n"
 				"Auto deactivate Runtime check StateTree.RuntimeValidation.Context to prevent reporting the same error multiple times.")
 				, InNewOwner ? *InNewOwner->GetFullName() : TEXT("owner"), *Owner.Get()->GetFullName());
-			Private::RuntimeValidationContext = false;
+			Private::bRuntimeValidationContext = false;
 		}
 	}
+
+	ValidatesTreeNodes(InNewStateTree);
+
 	StateTree = NewStateTree;
 	Owner = NewOwner;
 }
@@ -92,12 +102,12 @@ void FRuntimeValidationInstanceData::NodeEnterState(FGuid Id)
 	});
 	if (Found)
 	{
-		if (Private::RuntimeValidationEnterExitState && EnumHasAllFlags(Found->State, EState::BetweenEnterExitState))
+		if (Private::bRuntimeValidationEnterExitState && EnumHasAllFlags(Found->State, EState::BetweenEnterExitState))
 		{
 			ensureAlwaysMsgf(false, TEXT("StateTree runtime check failed: EnterState executed on node %s without an ExitState.\n"
 				"Auto deactivate Runtime check StateTree.RuntimeValidation.EnterExitState to prevent reporting the same error multiple times.")
 				, *Private::NodeToString(Owner.Get(), Id));
-			Private::RuntimeValidationEnterExitState = false;
+			Private::bRuntimeValidationEnterExitState = false;
 		}
 		EnumAddFlags(Found->State, EState::BetweenEnterExitState);
 	}
@@ -115,21 +125,94 @@ void FRuntimeValidationInstanceData::NodeExitState(FGuid Id)
 	});
 	if (Found)
 	{
-		if (Private::RuntimeValidationEnterExitState && !EnumHasAllFlags(Found->State, EState::BetweenEnterExitState))
+		if (Private::bRuntimeValidationEnterExitState && !EnumHasAllFlags(Found->State, EState::BetweenEnterExitState))
 		{
 			ensureAlwaysMsgf(false, TEXT("StateTree runtime check failed: ExitState executed on node %s without an EnterState.\n"
 				"Auto deactivate Runtime check StateTree.RuntimeValidation.EnterExitState to prevent reporting the same error multiple times.")
 				, *Private::NodeToString(Owner.Get(), Id));
-			Private::RuntimeValidationEnterExitState = false;
+			Private::bRuntimeValidationEnterExitState = false;
 		}
 		EnumRemoveFlags(Found->State, EState::BetweenEnterExitState);
 	}
-	else if (Private::RuntimeValidationEnterExitState)
+	else if (Private::bRuntimeValidationEnterExitState)
 	{
 		ensureAlwaysMsgf(false, TEXT("StateTree runtime check failed: ExitState executed on node %s without an EnterState.\n"
 			"Auto deactivate Runtime check StateTree.RuntimeValidation.EnterExitState to prevent reporting the same error multiple times.")
 			, *Private::NodeToString(Owner.Get(), Id));
-		Private::RuntimeValidationEnterExitState = false;
+		Private::bRuntimeValidationEnterExitState = false;
+	}
+}
+
+void FRuntimeValidationInstanceData::ValidatesTreeNodes(const UStateTree* InNewStateTree) const
+{
+	if (Private::bRuntimeValidationDoesNewerVersionExists)
+	{
+		if (InNewStateTree && InNewStateTree->IsReadyToRun())
+		{
+			auto DoesNewerVersionExists = [](const UObject* InstanceDataType)
+			{
+				// Is the class/scriptstruct a blueprint that got replaced by another class.
+				bool bHasNewerVersionExistsFlag = InstanceDataType->HasAnyFlags(RF_NewerVersionExists);
+				if (!bHasNewerVersionExistsFlag)
+				{
+					if (const UClass* InstanceDataClass = Cast<UClass>(InstanceDataType))
+					{
+						bHasNewerVersionExistsFlag = InstanceDataClass->HasAnyClassFlags(CLASS_NewerVersionExists);
+					}
+					else if (const UScriptStruct* InstanceDataStruct = Cast<UScriptStruct>(InstanceDataType))
+					{
+						bHasNewerVersionExistsFlag = (InstanceDataStruct->StructFlags & STRUCT_NewerVersionExists) != 0;
+					}
+				}
+				return bHasNewerVersionExistsFlag;
+			};
+			{
+				const FStateTreeInstanceData& InstanceData = InNewStateTree->GetDefaultInstanceData();
+				const int32 InstanceDataNum = InstanceData.Num();
+				for (int32 Index = 0; Index < InstanceDataNum; ++Index)
+				{
+					bool bFailed = false;
+					const UObject* InstanceObject = nullptr;
+					if (InstanceData.IsObject(Index))
+					{
+						InstanceObject = InstanceData.GetObject(Index);
+						bFailed = DoesNewerVersionExists(InstanceObject)
+							|| (InstanceObject && DoesNewerVersionExists(InstanceObject->GetClass()));
+					}
+					else
+					{
+						InstanceObject = InstanceData.GetStruct(Index).GetScriptStruct();
+						bFailed = DoesNewerVersionExists(InstanceObject);
+					}
+
+					if (bFailed)
+					{
+						ensureAlwaysMsgf(false, TEXT("StateTree runtime check failed: The data '%s' has a newest version.\n"
+							"It should be detected in StateTree::Link.\n"
+							"Auto deactivate Runtime check StateTree.RuntimeValidation.DoesNewerVersionExists to prevent reporting the same error multiple times.")
+							, *InstanceObject->GetFullName());
+						Private::bRuntimeValidationDoesNewerVersionExists = false;
+					}
+				}
+			}
+
+			for (FConstStructView NodeView : InNewStateTree->GetNodes())
+			{
+				const FStateTreeNodeBase* Node = NodeView.GetPtr<const FStateTreeNodeBase>();
+				if (Node)
+				{
+					const UStruct* DesiredInstanceDataType = Node->GetInstanceDataType();
+					if (DoesNewerVersionExists(DesiredInstanceDataType))
+					{
+						ensureAlwaysMsgf(false, TEXT("StateTree runtime check failed: The node '%s' has a newest version.\n"
+							"It should be detected in StateTree::Link.\n"
+							"Auto deactivate Runtime check StateTree.RuntimeValidation.DoesNewerVersionExists to prevent reporting the same error multiple times.")
+							, *DesiredInstanceDataType->GetFullName());
+						Private::bRuntimeValidationDoesNewerVersionExists = false;
+					}
+				}
+			}
+		}
 	}
 }
 #endif //WITH_STATETREE_DEBUG
diff --git a/Engine/Plugins/Runtime/StateTree/Source/StateTreeModule/Private/Debugger/StateTreeRuntimeValidationInstanceData.h b/Engine/Plugins/Runtime/StateTree/Source/StateTreeModule/Private/Debugger/StateTreeRuntimeValidationInstanceData.h
--- a/Engine/Plugins/Runtime/StateTree/Source/StateTreeModule/Private/Debugger/StateTreeRuntimeValidationInstanceData.h
+++ b/Engine/Plugins/Runtime/StateTree/Source/StateTreeModule/Private/Debugger/StateTreeRuntimeValidationInstanceData.h
@@ -25,6 +25,9 @@ public:
 	void NodeEnterState(FGuid Id);
 	void NodeExitState(FGuid Id);
 
+private:
+	void ValidatesTreeNodes(const UStateTree* InNewStateTree) const;
+
 private:
 	enum class EState : uint8
 	{
diff --git a/Engine/Plugins/Runtime/StateTree/Source/StateTreeModule/Private/StateTree.cpp b/Engine/Plugins/Runtime/StateTree/Source/StateTreeModule/Private/StateTree.cpp
--- a/Engine/Plugins/Runtime/StateTree/Source/StateTreeModule/Private/StateTree.cpp
+++ b/Engine/Plugins/Runtime/StateTree/Source/StateTreeModule/Private/StateTree.cpp
@@ -533,40 +533,48 @@ bool UStateTree::ValidateInstanceData()
 				}
 			}
 
+			auto HasNewerVersionExists = [](const UStruct* InstanceDataType)
 			{
 				// Is the class/scriptstruct a blueprint that got replaced by another class.
-				bool bHasNewerVersionExists = false;
-				if (const UClass* CurrentInstanceDataClass = Cast<UClass>(CurrentInstanceDataType))
+				bool bHasNewerVersionExists = InstanceDataType->HasAnyFlags(RF_NewerVersionExists);
+				if (!bHasNewerVersionExists)
 				{
-					bHasNewerVersionExists = CurrentInstanceDataClass->HasAnyClassFlags(CLASS_NewerVersionExists);
-				}
-				else if (const UScriptStruct* CurrentInstanceDataStruct = Cast<UScriptStruct>(CurrentInstanceDataType))
-				{
-					bHasNewerVersionExists = (CurrentInstanceDataStruct->StructFlags & STRUCT_NewerVersionExists) != 0;
+					if (const UClass* InstanceDataClass = Cast<UClass>(InstanceDataType))
+					{
+						bHasNewerVersionExists = InstanceDataClass->HasAnyClassFlags(CLASS_NewerVersionExists);
+					}
+					else if (const UScriptStruct* InstanceDataStruct = Cast<UScriptStruct>(InstanceDataType))
+					{
+						bHasNewerVersionExists = (InstanceDataStruct->StructFlags & STRUCT_NewerVersionExists) != 0;
+					}
 				}
+				return bHasNewerVersionExists;
+			};
 
-				if (bHasNewerVersionExists)
-				{
-					bool bLogError = true;
+			if (HasNewerVersionExists(CurrentInstanceDataType))
+			{
+				bool bLogError = true;
 #if WITH_EDITOR
-					OutOfDateStructs.Add(CurrentInstanceDataType);
-					bLogError = false;
+				OutOfDateStructs.Add(CurrentInstanceDataType);
+				bLogError = false;
 #endif
 
-					if (bLogError)
-					{
-						UE_LOG(LogStateTree, Error, TEXT("%s: node '%s' failed. The source Instance Data type '%s' has a newer version."), *GetFullName(), *WriteToString<64>(Node->StaticStruct()->GetFName()), *WriteToString<64>(CurrentInstanceDataType->GetFName()));
-					}
-
-					bResult = false;
+				if (bLogError)
+				{
+					UE_LOG(LogStateTree, Error, TEXT("%s: node '%s' failed. The source Instance Data type '%s' has a newer version."), *GetFullName(), *WriteToString<64>(Node->StaticStruct()->GetFName()), *WriteToString<64>(CurrentInstanceDataType->GetFName()));
 				}
+
+				bResult = false;
 			}
 
 			{
 				const UStruct* DesiredInstanceDataType = Node->GetInstanceDataType();
+				// The FMyInstance::StaticStruct doesn't get a notification like the other objects when reinstanced.
+				const bool bDesiredHasNewerVersion = HasNewerVersionExists(DesiredInstanceDataType);
 
 				// Use strict testing so that the users will have option to initialize data mismatch if the type changes (even if potentially compatible).
-				if (CurrentInstanceDataType != DesiredInstanceDataType)
+				if (CurrentInstanceDataType != DesiredInstanceDataType
+					&& !bDesiredHasNewerVersion)
 				{
 					bool bLogError = true;
 #if WITH_EDITOR
@@ -574,7 +582,7 @@ bool UStateTree::ValidateInstanceData()
 					const UClass* DesiredInstanceDataClass = Cast<UClass>(DesiredInstanceDataType);
 					if (CurrentInstanceDataClass && DesiredInstanceDataClass)
 					{
-						// Because of the loading order.It's possible that the OnObjectsReinstanced did complete.
+						// Because of the loading order. It's possible that the OnObjectsReinstanced did not completed.
 						if (CurrentInstanceDataClass->ClassGeneratedBy == DesiredInstanceDataClass->ClassGeneratedBy)
 						{
 							OutOfDateStructs.Add(CurrentInstanceDataType);
