StateTree: Add FStateTreeExecutionContext::CompleteTask. Mark a task as completed same as if the task would returns a run status. Allows for async task to complete without ticking.

diff --git a/Engine/Plugins/Runtime/StateTree/Source/StateTreeModule/Private/Blueprint/StateTreeTaskBlueprintBase.cpp b/Engine/Plugins/Runtime/StateTree/Source/StateTreeModule/Private/Blueprint/StateTreeTaskBlueprintBase.cpp
--- a/Engine/Plugins/Runtime/StateTree/Source/StateTreeModule/Private/Blueprint/StateTreeTaskBlueprintBase.cpp
+++ b/Engine/Plugins/Runtime/StateTree/Source/StateTreeModule/Private/Blueprint/StateTreeTaskBlueprintBase.cpp
@@ -19,6 +19,7 @@ UStateTreeTaskBlueprintBase::UStateTreeTaskBlueprintBase(const FObjectInitialize
 	, bShouldCallTickOnlyOnEvents(false)
 	, bShouldCopyBoundPropertiesOnTick(true)
 	, bShouldCopyBoundPropertiesOnExitState(true)
+	, bIsProcessingEnterStateOrTick(false)
 {
 	bHasExitState = BlueprintNodeHelpers::HasBlueprintFunction(TEXT("ReceiveExitState"), *this, *StaticClass());
 	bHasStateCompleted = BlueprintNodeHelpers::HasBlueprintFunction(TEXT("ReceiveStateCompleted"), *this, *StaticClass());
@@ -34,6 +35,9 @@ EStateTreeRunStatus UStateTreeTaskBlueprintBase::EnterState(FStateTreeExecutionC
 {
 	// Task became active, cache event queue and owner.
 	SetCachedInstanceDataFromContext(Context);
+	WeakTaskRef = Context.MakeWeakTaskRefFromInstanceDataPtr(this);
+
+	FGuardValue_Bitfield(bIsProcessingEnterStateOrTick, true);
 
 	// Reset status to running since the same task may be restarted.
 	RunStatus = EStateTreeRunStatus::Running;
@@ -43,10 +47,10 @@ EStateTreeRunStatus UStateTreeTaskBlueprintBase::EnterState(FStateTreeExecutionC
 		// Note: the name contains latent just to differentiate it from the deprecated version (the old version did not allow latent actions to be started).
 		ReceiveLatentEnterState(Transition);
 	}
-PRAGMA_DISABLE_DEPRECATION_WARNINGS		
+PRAGMA_DISABLE_DEPRECATION_WARNINGS
 	else if (bHasEnterState_DEPRECATED)
 	{
-		return ReceiveEnterState(Transition);
+		RunStatus = ReceiveEnterState(Transition);
 	}
 PRAGMA_ENABLE_DEPRECATION_WARNINGS
 	
@@ -67,6 +71,7 @@ void UStateTreeTaskBlueprintBase::ExitState(FStateTreeExecutionContext& Context,
 	}
 
 	// Task became inactive, clear cached event queue and owner.
+	WeakTaskRef.Release();
 	ClearCachedInstanceData();
 }
 
@@ -80,6 +85,8 @@ void UStateTreeTaskBlueprintBase::StateCompleted(FStateTreeExecutionContext& Con
 
 EStateTreeRunStatus UStateTreeTaskBlueprintBase::Tick(FStateTreeExecutionContext& Context, const float DeltaTime)
 {
+	FGuardValue_Bitfield(bIsProcessingEnterStateOrTick, true);
+
 	if (bHasLatentTick)
 	{
 		// Note: the name contains latent just to differentiate it from the deprecated version (the old version did not allow latent actions to be started).
@@ -88,13 +95,23 @@ EStateTreeRunStatus UStateTreeTaskBlueprintBase::Tick(FStateTreeExecutionContext
 PRAGMA_DISABLE_DEPRECATION_WARNINGS
 	else if (bHasTick_DEPRECATED)
 	{
-		return ReceiveTick(DeltaTime);
+		RunStatus = ReceiveTick(DeltaTime);
 	}
 PRAGMA_ENABLE_DEPRECATION_WARNINGS
 	
 	return RunStatus;
 }
 
+void UStateTreeTaskBlueprintBase::FinishTask(const bool bSucceeded)
+{
+	RunStatus = bSucceeded ? EStateTreeRunStatus::Succeeded : EStateTreeRunStatus::Failed;
+	if (!bIsProcessingEnterStateOrTick)
+	{
+		EStateTreeFinishTaskType CompletionType = bSucceeded ? EStateTreeFinishTaskType::Succeeded : EStateTreeFinishTaskType::Failed;
+		GetWeakExecutionContext().FinishTask(WeakTaskRef, CompletionType);
+	}
+}
+
 //----------------------------------------------------------------------//
 //  FStateTreeBlueprintTaskWrapper
 //----------------------------------------------------------------------//
diff --git a/Engine/Plugins/Runtime/StateTree/Source/StateTreeModule/Private/StateTreeExecutionContext.cpp b/Engine/Plugins/Runtime/StateTree/Source/StateTreeModule/Private/StateTreeExecutionContext.cpp
--- a/Engine/Plugins/Runtime/StateTree/Source/StateTreeModule/Private/StateTreeExecutionContext.cpp
+++ b/Engine/Plugins/Runtime/StateTree/Source/StateTreeModule/Private/StateTreeExecutionContext.cpp
@@ -22,13 +22,19 @@
 #define STATETREE_LOG(Verbosity, Format, ...) UE_VLOG_ALWAYS_UELOG(GetOwner(), LogStateTree, Verbosity, TEXT("%s: ") Format, *GetInstanceDescriptionInternal(), ##__VA_ARGS__)
 #define STATETREE_CLOG(Condition, Verbosity, Format, ...) UE_CVLOG_ALWAYS_UELOG((Condition), GetOwner(), LogStateTree, Verbosity, TEXT("%s: ") Format, *GetInstanceDescriptionInternal(), ##__VA_ARGS__)
 
+#if WITH_STATETREE_DEBUG
+#define UE_STATETREE_ENSURE_COMPLETED_STATE_RUN_STATUS() ensureMsgf(!CompletedStateRunStatus.IsSet(), TEXT("No public entry point should be able to set the completed state. Internal error."))
+#else
+#define UE_STATETREE_ENSURE_COMPLETED_STATE_RUN_STATUS()
+#endif //WITH_STATETREE_DEBUG
+
 namespace UE::StateTree::Debug
 {
 	constexpr int32 IndentSize = 2;	// Debug printing indent for hierarchical data.
 	static std::atomic<uint32> InstanceSerialNumber = 0;
 }; // UE::StateTree
 
-namespace UE::StateTree::ExecutionContext
+namespace UE::StateTree::ExecutionContext::Private
 {
 	bool MarkDelegateAsBroadcasted(FStateTreeDelegateDispatcher Dispatcher, const FStateTreeExecutionFrame& CurrentFrame, FStateTreeInstanceStorage& Storage)
 	{
@@ -60,6 +66,61 @@ namespace UE::StateTree::ExecutionContext
 
 		return false;
 	}
+
+	EStateTreeRunStatus FinishTaskTypeToRunStatus(EStateTreeFinishTaskType FinishTask)
+	{
+		static_assert((int32)EStateTreeRunStatus::Succeeded == (int32)EStateTreeFinishTaskType::Succeeded, "Succeeded do not match in both types.");
+		static_assert((int32)EStateTreeRunStatus::Failed == (int32)EStateTreeFinishTaskType::Failed, "Failed do not match in both types.");
+		return (EStateTreeRunStatus)FinishTask;
+	}
+
+	struct FActiveStateIndex
+	{
+		int32 FrameIndex = INDEX_NONE;
+		int32 ActiveStateIndex = INDEX_NONE;
+	};
+	TOptional<FActiveStateIndex> FindActiveStateIndex(const FStateTreeExecutionState& ExecutionState, const FStateTreeDelayedFinishedTask& Task)
+	{
+		TOptional<FActiveStateIndex> Result;
+
+		const int32 FrameIndex = ExecutionState.ActiveFrames.IndexOfByPredicate([&Task](const FStateTreeExecutionFrame& Frame)
+			{
+				return FObjectKey(Frame.StateTree) == Task.FrameStateTree && Frame.RootState == Task.FrameRootState;
+			});
+		if (FrameIndex == INDEX_NONE)
+		{
+			return Result;
+		}
+
+		const UStateTree* StateTree = ExecutionState.ActiveFrames[FrameIndex].StateTree;
+		if (!StateTree->GetStates().IsValidIndex(Task.State.Index))
+		{
+			return Result;
+		}
+
+		const int32 ActiveStateIndex = ExecutionState.ActiveFrames[FrameIndex].ActiveStates.IndexOfReverse(Task.State);
+		if (ActiveStateIndex == INDEX_NONE)
+		{
+			return Result;
+		}
+
+		if (Task.TaskIndex.IsValid())
+		{
+			const FCompactStateTreeState& State = StateTree->GetStates()[Task.State.Index];
+			const int32 BeginTaskIndex = State.TasksBegin;
+			const int32 EndTaskIndex = State.TasksBegin + State.TasksNum;
+			if (ensure(Task.TaskIndex.AsInt32() >= BeginTaskIndex && Task.TaskIndex.AsInt32() < EndTaskIndex))
+			{
+				Result = FActiveStateIndex{ .FrameIndex = FrameIndex, .ActiveStateIndex = ActiveStateIndex };
+			}
+		}
+		else
+		{
+			//The TaskIndex can be invalid if no task is responsible for the state completion.
+			Result = FActiveStateIndex{ .FrameIndex = FrameIndex, .ActiveStateIndex = ActiveStateIndex };
+		}
+		return Result;
+	}
 }
 
 FStateTreeMinimalExecutionContext::FStateTreeMinimalExecutionContext(UObject& InOwner, const UStateTree& InStateTree, FStateTreeInstanceData& InInstanceData)
@@ -175,6 +236,7 @@ FStateTreeScheduledTick FStateTreeMinimalExecutionContext::GetNextScheduledTick(
 				}
 			}
 		}
+		//@TODO: This doesn't work. A transition request might be pending it should tick before the requested CustomTickRate.
 		if (StateCustomTickRate.IsSet())
 		{
 			return FStateTreeScheduledTick::MakeCustomTickRate(StateCustomTickRate.GetValue());
@@ -197,6 +259,12 @@ FStateTreeScheduledTick FStateTreeMinimalExecutionContext::GetNextScheduledTick(
 		return FStateTreeScheduledTick::MakeNextFrame();
 	}
 
+	// Completed task. For EnterState or for user that only called TickTasks and not TickTransitions.
+	if (Exec.DelayedFinishedTasks.Num() > 0)
+	{
+		return FStateTreeScheduledTick::MakeNextFrame();
+	}
+
 	// Min of all delayed transitions.
 	if (Exec.DelayedTransitions.Num() > 0)
 	{
@@ -235,7 +303,7 @@ void FStateTreeMinimalExecutionContext::SendEvent(const FGameplayTag Tag, const
 void FStateTreeMinimalExecutionContext::ScheduleNextTick()
 {
 	TInstancedStruct<FStateTreeExecutionExtension>& ExecutionExtension = Storage.GetMutableExecutionState().ExecutionExtension;
-	if (RootStateTree.IsScheduledTickAllowed() && ExecutionExtension.IsValid())
+	if (bAllowedToScheduleNextTick && RootStateTree.IsScheduledTickAllowed() && ExecutionExtension.IsValid())
 	{
 		ExecutionExtension.GetMutable().ScheduleNextTick(FStateTreeExecutionExtension::FContextParameters(Owner, RootStateTree, Storage));
 	}
@@ -299,6 +367,26 @@ FStateTreeExecutionContext::FCurrentlyProcessedFrameScope::~FCurrentlyProcessedF
 	}
 }
 
+FStateTreeExecutionContext::FNodeInstanceDataScope::FNodeInstanceDataScope(FStateTreeExecutionContext& InContext, const FStateTreeNodeBase* InNode, const int32 InNodeIndex, const FStateTreeDataHandle InNodeDataHandle, const FStateTreeDataView InNodeInstanceData)
+	: Context(InContext)
+{
+	SavedNode = Context.CurrentNode;
+	SavedNodeIndex = Context.CurrentNodeIndex;
+	SavedNodeDataHandle = Context.CurrentNodeDataHandle;
+	SavedNodeInstanceData = Context.CurrentNodeInstanceData;
+	Context.CurrentNode = InNode;
+	Context.CurrentNodeIndex = InNodeIndex;
+	Context.CurrentNodeDataHandle = InNodeDataHandle;
+	Context.CurrentNodeInstanceData = InNodeInstanceData;
+}
+
+FStateTreeExecutionContext::FNodeInstanceDataScope::~FNodeInstanceDataScope()
+{
+	Context.CurrentNodeDataHandle = SavedNodeDataHandle;
+	Context.CurrentNodeInstanceData = SavedNodeInstanceData;
+	Context.CurrentNodeIndex = SavedNodeIndex;
+	Context.CurrentNode = SavedNode;
+}
 
 FStateTreeExecutionContext::FStateTreeExecutionContext(UObject& InOwner, const UStateTree& InStateTree, FStateTreeInstanceData& InInstanceData, const FOnCollectStateTreeExternalData& InCollectExternalDataDelegate, const EStateTreeRecordTransitions RecordTransitions)
 	: FStateTreeMinimalExecutionContext(InOwner, InStateTree, InInstanceData)
@@ -435,7 +523,27 @@ bool FStateTreeExecutionContext::SetContextDataByName(const FName Name, FStateTr
 
 FStateTreeWeakExecutionContext FStateTreeExecutionContext::MakeWeakExecutionContext() const
 {
-	return FStateTreeWeakExecutionContext(GetOwner(), GetStateTree(), InstanceData);
+	return FStateTreeWeakExecutionContext(*this);
+}
+
+FStateTreeWeakTaskRef FStateTreeExecutionContext::MakeWeakTaskRef(const FStateTreeTaskBase& Node) const
+{
+	check(CurrentNode == &Node);
+	return MakeWeakTaskRefInternal();
+}
+
+FStateTreeWeakTaskRef FStateTreeExecutionContext::MakeWeakTaskRefInternal() const
+{
+	FStateTreeWeakTaskRef Result;
+	if (const FStateTreeExecutionFrame* Frame = GetCurrentlyProcessedFrame())
+	{
+		if (Frame->StateTree->Nodes.IsValidIndex(CurrentNodeIndex)
+			&& Frame->StateTree->Nodes[CurrentNodeIndex].GetPtr<const FStateTreeTaskBase>() != nullptr)
+		{
+			Result = FStateTreeWeakTaskRef(Frame->StateTree, FStateTreeIndex16(CurrentNodeIndex));
+		}
+	}
+	return Result;
 }
 
 EStateTreeRunStatus FStateTreeExecutionContext::Start(const FInstancedPropertyBag* InitialParameters, int32 RandomSeed)
@@ -489,6 +597,7 @@ EStateTreeRunStatus FStateTreeExecutionContext::Start(FStartParameters Parameter
 	}
 
 	FCurrentlyExecutionContextScope ContextScope(*this);
+	TGuardValue<bool> ScheduledNextTickScope(bAllowedToScheduleNextTick, false);
 
 	// Initialize for the init frame.
 	ensure(Exec.ActiveFrames.Num() == 0);
@@ -528,7 +637,8 @@ EStateTreeRunStatus FStateTreeExecutionContext::Start(FStartParameters Parameter
 	{
 		// First tick.
 		// Tasks are not ticked here, since their behavior is that EnterState() (called above) is treated as a tick.  
-		TickEvaluatorsAndGlobalTasks(0.0f, /*bTickGlobalTasks*/false);
+		constexpr bool bTickGlobalTasks = false;
+		TickEvaluatorsAndGlobalTasks(0.0f, bTickGlobalTasks);
 
 		// Initialize to unset running state.
 		Exec.TreeRunStatus = EStateTreeRunStatus::Running;
@@ -626,6 +736,7 @@ EStateTreeRunStatus FStateTreeExecutionContext::Stop(EStateTreeRunStatus Complet
 	}
 
 	FCurrentlyExecutionContextScope ContextScope(*this);
+	TGuardValue<bool> ScheduledNextTickScope(bAllowedToScheduleNextTick, false);
 
 	// Set scoped phase only for properly initialized context with valid Instance data
 	// since we need it to output the InstanceId 
@@ -717,6 +828,8 @@ EStateTreeRunStatus FStateTreeExecutionContext::TickPrelude()
 	Exec.CurrentPhase = EStateTreeUpdatePhase::TickStateTree;
 	UE_STATETREE_DEBUG_ENTER_PHASE(this, EStateTreeUpdatePhase::TickStateTree);
 
+	UpdateCompletedStateList();
+
 	return EStateTreeRunStatus::Running;
 }
 
@@ -748,6 +861,7 @@ EStateTreeRunStatus FStateTreeExecutionContext::Tick(const float DeltaTime)
 	CSV_SCOPED_TIMING_STAT_EXCLUSIVE(StateTree_Tick);
 
 	FCurrentlyExecutionContextScope ContextScope(*this);
+	TGuardValue<bool> ScheduledNextTickScope(bAllowedToScheduleNextTick, false);
 
 	const EStateTreeRunStatus PreludeResult = TickPrelude();
 	if (PreludeResult != EStateTreeRunStatus::Running)
@@ -766,6 +880,7 @@ EStateTreeRunStatus FStateTreeExecutionContext::TickUpdateTasks(const float Delt
 	CSV_SCOPED_TIMING_STAT_EXCLUSIVE(StateTree_Tick);
 
 	FCurrentlyExecutionContextScope ContextScope(*this);
+	TGuardValue<bool> ScheduledNextTickScope(bAllowedToScheduleNextTick, false);
 
 	const EStateTreeRunStatus PreludeResult = TickPrelude();
 	if (PreludeResult != EStateTreeRunStatus::Running)
@@ -783,6 +898,7 @@ EStateTreeRunStatus FStateTreeExecutionContext::TickTriggerTransitions()
 	CSV_SCOPED_TIMING_STAT_EXCLUSIVE(StateTree_Tick);
 
 	FCurrentlyExecutionContextScope ContextScope(*this);
+	TGuardValue<bool> ScheduledNextTickScope(bAllowedToScheduleNextTick, false);
 
 	const EStateTreeRunStatus PreludeResult = TickPrelude();
 	if (PreludeResult != EStateTreeRunStatus::Running)
@@ -790,7 +906,21 @@ EStateTreeRunStatus FStateTreeExecutionContext::TickTriggerTransitions()
 		return PreludeResult;
 	}
 
-	TickTriggerTransitionsInternal();
+	// If we have pending global task that completd, then stop the tree (like it would with a regular Tick).
+	//This should not occurs and its probably a user mistake to not tick the system in the proper order.
+	const EStateTreeRunStatus GlobalTaskStatus = GetGlobalTasksCompletedStatesStatus();
+	if (GlobalTaskStatus == EStateTreeRunStatus::Running)
+	{
+		TickTriggerTransitionsInternal();
+	}
+	else
+	{
+		//@TODO check what we want to do with global task from linked asset state. Should they request a stop or close their own tree.
+		FStateTreeExecutionState& Exec = GetExecState();
+		UE_STATETREE_DEBUG_LOG_EVENT(this, Log, TEXT("Global tasks completed (%s), stopping the tree"), *UEnum::GetDisplayValueAsText(GlobalTaskStatus).ToString());
+		STATETREE_LOG(Log, TEXT("Global tasks completed (%s), stopping the tree"), *UEnum::GetDisplayValueAsText(GlobalTaskStatus).ToString());
+		Exec.RequestedStop = GlobalTaskStatus;
+	}
 
 	return TickPostlude();
 }
@@ -806,7 +936,8 @@ void FStateTreeExecutionContext::TickUpdateTasksInternal(const float DeltaTime)
 	}
 
 	// Tick global evaluators and tasks.
-	const EStateTreeRunStatus EvalAndGlobalTaskStatus = TickEvaluatorsAndGlobalTasks(DeltaTime);
+	const bool bTickGlobalTasks = true;
+	const EStateTreeRunStatus EvalAndGlobalTaskStatus = TickEvaluatorsAndGlobalTasks(DeltaTime, bTickGlobalTasks);
 	if (EvalAndGlobalTaskStatus == EStateTreeRunStatus::Running)
 	{
 		if (Exec.LastTickStatus == EStateTreeRunStatus::Running)
@@ -823,6 +954,7 @@ void FStateTreeExecutionContext::TickUpdateTasksInternal(const float DeltaTime)
 	}
 	else
 	{
+		//@TODO check what we want to do with global task from linked asset state. Should they request a stop or close their own tree.
 		UE_STATETREE_DEBUG_LOG_EVENT(this, Log, TEXT("Global tasks completed (%s), stopping the tree"), *UEnum::GetDisplayValueAsText(EvalAndGlobalTaskStatus).ToString());
 		STATETREE_LOG(Log, TEXT("Global tasks completed (%s), stopping the tree"), *UEnum::GetDisplayValueAsText(EvalAndGlobalTaskStatus).ToString());
 		Exec.RequestedStop = EvalAndGlobalTaskStatus;
@@ -912,7 +1044,7 @@ void FStateTreeExecutionContext::BroadcastDelegate(FStateTreeDelegateDispatcher
 	check(CurrentFrame);
 
 	GetExecState().DelegateActiveListeners.BroadcastDelegate(Dispatcher, CurrentFrame->StateTree, GetExecState());
-	if (UE::StateTree::ExecutionContext::MarkDelegateAsBroadcasted(Dispatcher, *CurrentFrame, GetMutableInstanceData()->GetMutableStorage()))
+	if (UE::StateTree::ExecutionContext::Private::MarkDelegateAsBroadcasted(Dispatcher, *CurrentFrame, GetMutableInstanceData()->GetMutableStorage()))
 	{
 		ScheduleNextTick();
 	}
@@ -1012,6 +1144,195 @@ void FStateTreeExecutionContext::RequestTransition(FStateTreeStateHandle InTarge
 	RequestTransition(FStateTreeTransitionRequest(InTargetState, InPriority, InFallback));
 }
 
+void FStateTreeExecutionContext::FinishTask(const FStateTreeTaskBase& Task, EStateTreeFinishTaskType FinishType)
+{
+	if (!IsValid())
+	{
+		STATETREE_LOG(Warning, TEXT("%hs: StateTree context is not initialized properly ('%s' using StateTree '%s')"),
+			__FUNCTION__, *GetNameSafe(&Owner), *GetFullNameSafe(&RootStateTree));
+		return;
+	}
+
+	// Like GetInstanceData, only accept task if we are currently processing them.
+	if (!ensure(CurrentNode == &Task))
+	{
+		return;
+	}
+
+	using namespace UE::StateTree::ExecutionContext::Private;
+
+	FStateTreeExecutionState& Exec = GetExecState();
+
+	// Then we have the current frame and the current state.
+	const int32 ActiveFrameIndex = Exec.ActiveFrames.IndexOfByPredicate([FramePtr = CurrentlyProcessedFrame](const FStateTreeExecutionFrame& Other) { return &Other == FramePtr; });
+	check(ActiveFrameIndex != INDEX_NONE);
+	const int32 ActiveStateIndex = Exec.ActiveFrames[ActiveFrameIndex].ActiveStates.IndexOfReverse(CurrentlyProcessedState);
+	check(ActiveStateIndex != INDEX_NONE);
+	FinishTaskInternal(ActiveFrameIndex, ActiveStateIndex, CurrentNodeIndex, FinishTaskTypeToRunStatus(FinishType));
+
+	// The FStateTreeTaskBase::Tick or OnEnter returns the value. Set it here and test it later.
+	CompletedStateRunStatus = FinishTaskTypeToRunStatus(FinishType);
+}
+
+void FStateTreeExecutionContext::FinishTask(const FStateTreeDelayedFinishedTask& Task, EStateTreeFinishTaskType FinishType)
+{
+	if (!IsValid())
+	{
+		STATETREE_LOG(Warning, TEXT("%hs: StateTree context is not initialized properly ('%s' using StateTree '%s')"),
+			__FUNCTION__, *GetNameSafe(&Owner), *GetFullNameSafe(&RootStateTree));
+		return;
+	}
+
+	using namespace UE::StateTree::ExecutionContext::Private;
+
+	FStateTreeExecutionState& Exec = GetExecState();
+	//Is it still valid to use
+	TOptional<FActiveStateIndex> FrameAndState = FindActiveStateIndex(Exec, Task);
+	if (FrameAndState.IsSet())
+	{
+		if (CurrentNodeIndex == Task.TaskIndex.AsInt32())
+		{
+			// We are currently ticking that task. We are either in Tick, EnterState, ExitState, Transition.
+			//NB. Completed tasks are not considered during ExitState or Transition.
+			//The node will return the value and it will be processed once returned. Set it here and test it later.
+			CompletedStateRunStatus = FinishTaskTypeToRunStatus(FinishType);
+		}
+		else
+		{
+			FinishTaskInternal(FrameAndState.GetValue().FrameIndex, FrameAndState.GetValue().ActiveStateIndex, Task.TaskIndex.AsInt32(), FinishTaskTypeToRunStatus(FinishType));
+			ScheduleNextTick();
+		}
+	}
+}
+
+void FStateTreeExecutionContext::FinishTaskInternal(int32 ActiveFrameIndex, int32 ActiveStateIndex, int32 TaskIndex, EStateTreeRunStatus RunStatus)
+{
+	FStateTreeExecutionState& Exec = GetExecState();
+	FStateTreeExecutionFrame& ActiveFrame = Exec.ActiveFrames[ActiveFrameIndex];
+
+	// Add it for later. In case the user called TickTasks and will call TickTransitions later.
+	const FStateTreeIndex16 TaskIndex16 = TaskIndex >= 0 ? FStateTreeIndex16(TaskIndex) : FStateTreeIndex16::Invalid;
+	const EStateTreeFinishTaskType CompletionType = RunStatus == EStateTreeRunStatus::Succeeded ? EStateTreeFinishTaskType::Succeeded : EStateTreeFinishTaskType::Failed;
+	const FStateTreeDelayedFinishedTask& DelayedCompletedTask = Exec.DelayedFinishedTasks.Emplace_GetRef(FObjectKey(ActiveFrame.StateTree), ActiveFrame.RootState, ActiveFrame.ActiveStates[ActiveStateIndex], TaskIndex16, CompletionType);
+
+	AddUniqueCompletedStateList(ActiveFrameIndex, ActiveStateIndex, DelayedCompletedTask);
+}
+
+void FStateTreeExecutionContext::UpdateCompletedStateList()
+{
+	FStateTreeExecutionState& Exec = GetExecState();
+
+	// Do not modify the PendingCompletedState until transitions completion.
+	//User can call TickTask and wait a couple of frame before calling TickTransitions.
+
+	// Recreate the completed states list from the pending list
+	CompletedStates.Empty();
+
+	for (const FStateTreeDelayedFinishedTask& PendingCompletedState : Exec.DelayedFinishedTasks)
+	{
+		if (const UObject* CompletedTree = PendingCompletedState.FrameStateTree.ResolveObjectPtr())
+		{
+			for (int32 FrameIndex = 0; FrameIndex < Exec.ActiveFrames.Num(); ++FrameIndex)
+			{
+				const FStateTreeExecutionFrame& ActiveFrame = Exec.ActiveFrames[FrameIndex];
+				if (ActiveFrame.StateTree == CompletedTree && ActiveFrame.RootState == PendingCompletedState.FrameRootState)
+				{
+					bool bFound = false;
+					int32 StateIndex = INDEX_NONE;
+					if (PendingCompletedState.State.IsValid())
+					{
+						StateIndex = ActiveFrame.ActiveStates.IndexOfReverse(PendingCompletedState.State);
+						bFound = ActiveFrame.ActiveStates.IsValidIndex(StateIndex);
+					}
+					else
+					{
+						bFound = true; // It's a global task.
+					}
+
+					// It's valid. Add Unique entry.
+					if (bFound)
+					{
+						AddUniqueCompletedStateList(FrameIndex, StateIndex, PendingCompletedState);
+					}
+					break; // Next pending
+				}
+			}
+		}
+	}
+}
+
+void FStateTreeExecutionContext::AddUniqueCompletedStateList(int32 ActiveFrameIndex, int32 ActiveStateIndex, const FStateTreeDelayedFinishedTask& DelayedFinishedTask)
+{
+	bool bAdd = true;
+	// The priority is global tasks and failure. Only keep 1 per Frame per State.
+	for (int32 CompletedStateIndex = CompletedStates.Num() - 1; CompletedStateIndex >= 0; --CompletedStateIndex)
+	{
+		FCompletedState& ExistingCompletedState = CompletedStates[CompletedStateIndex];
+		if (ExistingCompletedState.ActiveFrameIndex == ActiveFrameIndex)
+		{
+			// If it's a global task.
+			if (!DelayedFinishedTask.State.IsValid())
+			{
+				// Remove or update the existing entry.
+				if (ExistingCompletedState.ActiveStateIndex == INDEX_NONE)
+				{
+					// Set failure, otherwise keep as it is
+					if (DelayedFinishedTask.Type == EStateTreeFinishTaskType::Failed)
+					{
+						ExistingCompletedState.RunStatus = EStateTreeRunStatus::Failed;
+					}
+					bAdd = false;
+				}
+				else
+				{
+					CompletedStates.RemoveAtSwap(CompletedStateIndex);
+				}
+				break;
+			}
+			else if (ExistingCompletedState.ActiveStateIndex == INDEX_NONE)
+			{
+				// Keep the existing entry.
+				bAdd = false;
+				break;
+			}
+			else if (ExistingCompletedState.ActiveStateIndex == ActiveStateIndex)
+			{
+				// Set failure, otherwise keep as it is
+				if (DelayedFinishedTask.Type == EStateTreeFinishTaskType::Failed)
+				{
+					ExistingCompletedState.RunStatus = EStateTreeRunStatus::Failed;
+				}
+				bAdd = false;
+				break;
+			}
+		}
+	}
+
+	if (bAdd)
+	{
+		const EStateTreeRunStatus RunStatus = UE::StateTree::ExecutionContext::Private::FinishTaskTypeToRunStatus(DelayedFinishedTask.Type);
+		CompletedStates.Emplace(ActiveFrameIndex, ActiveStateIndex, RunStatus, DelayedFinishedTask);
+	}
+}
+
+EStateTreeRunStatus FStateTreeExecutionContext::GetGlobalTasksCompletedStatesStatus() const
+{
+	EStateTreeRunStatus Result = EStateTreeRunStatus::Running;
+	for (const FCompletedState& CompletedState : CompletedStates)
+	{
+		if (CompletedState.ActiveStateIndex == INDEX_NONE)
+		{
+			// return failure in priority
+			if (CompletedState.RunStatus == EStateTreeRunStatus::Failed)
+			{
+				return EStateTreeRunStatus::Failed;
+			}
+			Result = CompletedState.RunStatus;
+		}
+	}
+	return Result;
+}
+
 void FStateTreeExecutionContext::UpdateInstanceData(TConstArrayView<FStateTreeExecutionFrame> CurrentActiveFrames, TArrayView<FStateTreeExecutionFrame> NextActiveFrames)
 {
 	// Estimate how many new instance data items we might have.
@@ -1836,8 +2157,7 @@ EStateTreeRunStatus FStateTreeExecutionContext::EnterState(FStateTreeTransitionR
 	CaptureNewStateEvents(Exec.ActiveFrames, Transition.NextActiveFrames, Transition.NextActiveFrameEvents);
 
 	Exec.StateChangeCount++;
-	Exec.CompletedFrameIndex = FStateTreeIndex16::Invalid;
-	Exec.CompletedStateHandle = FStateTreeStateHandle::Invalid;
+	Exec.DelayedFinishedTasks.Empty();
 	Exec.EnterStateFailedFrameIndex = FStateTreeIndex16::Invalid; // This will make all tasks to be accepted.
 	Exec.EnterStateFailedTaskIndex = FStateTreeIndex16::Invalid; // This will make all tasks to be accepted.
 	
@@ -1954,7 +2274,7 @@ EStateTreeRunStatus FStateTreeExecutionContext::EnterState(FStateTreeTransitionR
 						if (bShouldCallEnterState)
 						{
 							const FStateTreeDataView ConditionInstanceView = GetDataView(CurrentParentFrame, CurrentFrame, Cond.InstanceDataHandle);
-							FNodeInstanceDataScope DataScope(*this, Cond.InstanceDataHandle, ConditionInstanceView);
+							FNodeInstanceDataScope DataScope(*this, &Cond, ConditionIndex, Cond.InstanceDataHandle, ConditionInstanceView);
 
 							if (Cond.BindingsBatch.IsValid())
 							{
@@ -1981,7 +2301,7 @@ EStateTreeRunStatus FStateTreeExecutionContext::EnterState(FStateTreeTransitionR
 				const FStateTreeTaskBase& Task = NextFrame.StateTree->Nodes[TaskIndex].Get<const FStateTreeTaskBase>();
 				const FStateTreeDataView TaskInstanceView = GetDataView(CurrentParentFrame, CurrentFrame, Task.InstanceDataHandle);
 
-				FNodeInstanceDataScope DataScope(*this, Task.InstanceDataHandle, TaskInstanceView);
+				FNodeInstanceDataScope DataScope(*this, &Task, TaskIndex, Task.InstanceDataHandle, TaskInstanceView);
 
 				// Copy bound properties.
 				if (Task.BindingsBatch.IsValid())
@@ -2002,6 +2322,7 @@ EStateTreeRunStatus FStateTreeExecutionContext::EnterState(FStateTreeTransitionR
 				if (bIsEnteringState && bShouldCallEnterState)
 				{
 					STATETREE_LOG(Verbose, TEXT("%*sTask '%s'.EnterState()"), (FrameIndex + StateIndex + 1)*UE::StateTree::Debug::IndentSize, TEXT(""), *Task.Name.ToString());
+					UE_STATETREE_ENSURE_COMPLETED_STATE_RUN_STATUS();
 					UE_STATETREE_DEBUG_TASK_ENTER_STATE(this, NextFrame.StateTree, FStateTreeIndex16(TaskIndex));
 
 					EStateTreeRunStatus Status = EStateTreeRunStatus::Unset;
@@ -2014,24 +2335,27 @@ EStateTreeRunStatus FStateTreeExecutionContext::EnterState(FStateTreeTransitionR
 
 					UE_STATETREE_DEBUG_TASK_EVENT(this, TaskIndex, TaskInstanceView, EStateTreeTraceEventType::OnEntered, Status);
 
+					// Delayed from CompleteTask(), override the Tick result unless it's a failure.
+					if (CompletedStateRunStatus.IsSet()
+						&& (Status == EStateTreeRunStatus::Succeeded || Status == EStateTreeRunStatus::Running))
+					{
+						Status = CompletedStateRunStatus.GetValue();
+						CompletedStateRunStatus.Reset();
+					}
+
 					if (Status != EStateTreeRunStatus::Running)
 					{
-						// Store the first state that completed, will be used to decide where to trigger transitions.
-						if (!Exec.CompletedStateHandle.IsValid())
+						Result = Status;
+						FinishTaskInternal(FrameIndex, StateIndex, TaskIndex, Result);
+						if (Status == EStateTreeRunStatus::Failed)
 						{
-							Exec.CompletedFrameIndex = FStateTreeIndex16(FrameIndex);
-							Exec.CompletedStateHandle = CurrentHandle;
+							Exec.EnterStateFailedFrameIndex = FStateTreeIndex16(FrameIndex);
+							Exec.EnterStateFailedTaskIndex = FStateTreeIndex16(TaskIndex);
+							break;
 						}
-						Result = Status;
-					}
-					
-					if (Status == EStateTreeRunStatus::Failed)
-					{
-						// Store how far in the enter state we got. This will be used to match the StateCompleted() and ExitState() calls.
-						Exec.EnterStateFailedFrameIndex = FStateTreeIndex16(FrameIndex); 
-						Exec.EnterStateFailedTaskIndex = FStateTreeIndex16(TaskIndex);
-						break;
 					}
+
+					UE_STATETREE_ENSURE_COMPLETED_STATE_RUN_STATUS();
 				}
 			}
 
@@ -2114,7 +2438,7 @@ void FStateTreeExecutionContext::ExitState(const FStateTreeTransitionResult& Tra
 			{
 				const FStateTreeEvaluatorBase& Eval = CurrentStateTree->Nodes[EvalIndex].Get<const FStateTreeEvaluatorBase>();
 				const FStateTreeDataView EvalInstanceView = GetDataView(CurrentParentFrame, CurrentFrame, Eval.InstanceDataHandle);
-				FNodeInstanceDataScope DataScope(*this, Eval.InstanceDataHandle, EvalInstanceView);
+				FNodeInstanceDataScope DataScope(*this, &Eval, EvalIndex, Eval.InstanceDataHandle, EvalInstanceView);
 
 				if (Eval.BindingsBatch.IsValid())
 				{
@@ -2126,7 +2450,7 @@ void FStateTreeExecutionContext::ExitState(const FStateTreeTransitionResult& Tra
 			{
 				const FStateTreeTaskBase& Task = CurrentStateTree->Nodes[TaskIndex].Get<const FStateTreeTaskBase>();
 				const FStateTreeDataView TaskInstanceView = GetDataView(CurrentParentFrame, CurrentFrame, Task.InstanceDataHandle);
-				FNodeInstanceDataScope DataScope(*this, Task.InstanceDataHandle, TaskInstanceView);
+				FNodeInstanceDataScope DataScope(*this, &Task, TaskIndex, Task.InstanceDataHandle, TaskInstanceView);
 
 				if (Task.BindingsBatch.IsValid() && Task.bShouldCopyBoundPropertiesOnExitState)
 				{
@@ -2249,7 +2573,7 @@ void FStateTreeExecutionContext::ExitState(const FStateTreeTransitionResult& Tra
 						const FStateTreeTaskBase& Task = CurrentStateTree->Nodes[TaskIndex].Get<const FStateTreeTaskBase>();
 						const FStateTreeDataView TaskInstanceView = GetDataView(CurrentParentFrame, CurrentFrame, Task.InstanceDataHandle);
 
-						FNodeInstanceDataScope DataScope(*this, Task.InstanceDataHandle, TaskInstanceView);
+						FNodeInstanceDataScope DataScope(*this, &Task, TaskIndex, Task.InstanceDataHandle, TaskInstanceView);
 
 						// Ignore disabled task
 						if (Task.bTaskEnabled == false)
@@ -2290,7 +2614,7 @@ void FStateTreeExecutionContext::ExitState(const FStateTreeTransitionResult& Tra
 							if (bShouldCallStateChange)
 							{
 								const FStateTreeDataView ConditionInstanceView = GetDataView(CurrentParentFrame, CurrentFrame, Cond.InstanceDataHandle);
-								FNodeInstanceDataScope DataScope(*this, Cond.InstanceDataHandle, ConditionInstanceView);
+								FNodeInstanceDataScope DataScope(*this, &Cond, ConditionIndex, Cond.InstanceDataHandle, ConditionInstanceView);
 
 								if (Cond.BindingsBatch.IsValid())
 								{
@@ -2427,7 +2751,7 @@ void FStateTreeExecutionContext::StateCompleted()
 					{
 						const FStateTreeTaskBase& Task = CurrentStateTree->Nodes[TaskIndex].Get<const FStateTreeTaskBase>();
 						const FStateTreeDataView TaskInstanceView = GetDataView(CurrentParentFrame, CurrentFrame, Task.InstanceDataHandle);
-						FNodeInstanceDataScope DataScope(*this, Task.InstanceDataHandle, TaskInstanceView);
+						FNodeInstanceDataScope DataScope(*this, &Task, TaskIndex, Task.InstanceDataHandle, TaskInstanceView);
 
 						// Ignore disabled task
 						if (Task.bTaskEnabled == false)
@@ -2450,7 +2774,7 @@ void FStateTreeExecutionContext::StateCompleted()
 						if (Cond.bHasShouldCallStateChangeEvents)
 						{
 							const FStateTreeDataView ConditionInstanceView = GetDataView(CurrentParentFrame, CurrentFrame, Cond.InstanceDataHandle);
-							FNodeInstanceDataScope DataScope(*this, Cond.InstanceDataHandle, ConditionInstanceView);
+							FNodeInstanceDataScope DataScope(*this, &Cond, ConditionIndex, Cond.InstanceDataHandle, ConditionInstanceView);
 
 							if (Cond.BindingsBatch.IsValid())
 							{
@@ -2499,7 +2823,7 @@ EStateTreeRunStatus FStateTreeExecutionContext::TickEvaluatorsAndGlobalTasks(con
 			{
 				const FStateTreeEvaluatorBase& Eval = CurrentStateTree->Nodes[EvalIndex].Get<const FStateTreeEvaluatorBase>();
 				const FStateTreeDataView EvalInstanceView = GetDataView(CurrentParentFrame, CurrentFrame, Eval.InstanceDataHandle);
-				FNodeInstanceDataScope DataScope(*this, Eval.InstanceDataHandle, EvalInstanceView);
+				FNodeInstanceDataScope DataScope(*this, &Eval, EvalIndex, Eval.InstanceDataHandle, EvalInstanceView);
 
 				// Copy bound properties.
 				if (Eval.BindingsBatch.IsValid())
@@ -2516,10 +2840,19 @@ EStateTreeRunStatus FStateTreeExecutionContext::TickEvaluatorsAndGlobalTasks(con
 				}
 			}
 
+			// Used to stop ticking tasks after one fails, but we still want to keep updating the data views so that property binding works properly.
+			bool bShouldTickTasks = true;
+
+			const FCompletedState* FrameCompletedState = CompletedStates.FindByPredicate([FrameIndex](FCompletedState& Other) { return Other.ActiveFrameIndex == FrameIndex && Other.ActiveStateIndex == INDEX_NONE; });
+			if (FrameCompletedState)
+			{
+				bShouldTickTasks = false;
+				//@TODO: use a priority system? What if another frame did return fail and this frame success.
+				Result = FrameCompletedState->RunStatus;
+			}
+
 			if (bTickGlobalTasks)
 			{
-				// Used to stop ticking tasks after one fails, but we still want to keep updating the data views so that property binding works properly.
-				bool bShouldTickTasks = true;
 				const bool bHasEvents = EventQueue && EventQueue->HasEvents();
 
 				if (CurrentStateTree->ShouldTickGlobalTasks(bHasEvents))
@@ -2528,7 +2861,7 @@ EStateTreeRunStatus FStateTreeExecutionContext::TickEvaluatorsAndGlobalTasks(con
 					{
 						const FStateTreeTaskBase& Task = CurrentStateTree->Nodes[TaskIndex].Get<const FStateTreeTaskBase>();
 						const FStateTreeDataView TaskInstanceView = GetDataView(CurrentParentFrame, CurrentFrame, Task.InstanceDataHandle);
-						FNodeInstanceDataScope DataScope(*this, Task.InstanceDataHandle, TaskInstanceView);
+						FNodeInstanceDataScope DataScope(*this, &Task, TaskIndex, Task.InstanceDataHandle, TaskInstanceView);
 
 						// Ignore disabled task
 						if (Task.bTaskEnabled == false)
@@ -2552,6 +2885,7 @@ EStateTreeRunStatus FStateTreeExecutionContext::TickEvaluatorsAndGlobalTasks(con
 						}
 
 						//UE_STATETREE_DEBUG_TASK_EVENT(this, TaskIndex, TaskDataView, EStateTreeTraceEventType::OnTickingTask, EStateTreeRunStatus::Running);
+						UE_STATETREE_ENSURE_COMPLETED_STATE_RUN_STATUS();
 						UE_STATETREE_DEBUG_TASK_TICK(this, CurrentStateTree, TaskIndex);
 
 						EStateTreeRunStatus TaskResult = EStateTreeRunStatus::Unset;
@@ -2562,6 +2896,14 @@ EStateTreeRunStatus FStateTreeExecutionContext::TickEvaluatorsAndGlobalTasks(con
 							TaskResult = Task.Tick(*this, DeltaTime);
 						}
 
+						// Delayed from CompleteTask(), override the Tick result unless it's a failure.
+						if (CompletedStateRunStatus.IsSet()
+							&& (TaskResult == EStateTreeRunStatus::Succeeded || TaskResult == EStateTreeRunStatus::Running))
+						{
+							TaskResult = CompletedStateRunStatus.GetValue();
+							CompletedStateRunStatus.Reset();
+						}
+
 						UE_STATETREE_DEBUG_TASK_EVENT(this, TaskIndex, TaskInstanceView,
 							TaskResult != EStateTreeRunStatus::Running ? EStateTreeTraceEventType::OnTaskCompleted : EStateTreeTraceEventType::OnTicked,
 							TaskResult);
@@ -2569,13 +2911,16 @@ EStateTreeRunStatus FStateTreeExecutionContext::TickEvaluatorsAndGlobalTasks(con
 						// If a global task succeeds or fails, it will stop the whole tree.
 						if (TaskResult != EStateTreeRunStatus::Running)
 						{
+							//@TODO: use a priority system? What if another frame did return fail and this one success.
 							Result = TaskResult;
+
+							if (TaskResult == EStateTreeRunStatus::Failed)
+							{
+								bShouldTickTasks = false;
+							}
 						}
-						
-						if (TaskResult == EStateTreeRunStatus::Failed)
-						{
-							bShouldTickTasks = false;
-						}
+
+						UE_STATETREE_ENSURE_COMPLETED_STATE_RUN_STATUS();
 					}
 				}
 			}
@@ -2612,7 +2957,7 @@ EStateTreeRunStatus FStateTreeExecutionContext::StartEvaluatorsAndGlobalTasks(FS
 			{
 				const FStateTreeEvaluatorBase& Eval = CurrentStateTree->Nodes[EvalIndex].Get<const FStateTreeEvaluatorBase>();
 				const FStateTreeDataView EvalInstanceView = GetDataView(CurrentParentFrame, CurrentFrame, Eval.InstanceDataHandle);
-				FNodeInstanceDataScope DataScope(*this, Eval.InstanceDataHandle, EvalInstanceView);
+				FNodeInstanceDataScope DataScope(*this, &Eval, EvalIndex, Eval.InstanceDataHandle, EvalInstanceView);
 
 				// Copy bound properties.
 				if (Eval.BindingsBatch.IsValid())
@@ -2637,7 +2982,7 @@ EStateTreeRunStatus FStateTreeExecutionContext::StartEvaluatorsAndGlobalTasks(FS
 			{
 				const FStateTreeTaskBase& Task =  CurrentStateTree->Nodes[TaskIndex].Get<const FStateTreeTaskBase>();
 				const FStateTreeDataView TaskInstanceView = GetDataView(CurrentParentFrame, CurrentFrame, Task.InstanceDataHandle);
-				FNodeInstanceDataScope DataScope(*this, Task.InstanceDataHandle, TaskInstanceView);
+				FNodeInstanceDataScope DataScope(*this, &Task, TaskIndex, Task.InstanceDataHandle, TaskInstanceView);
 
 				// Copy bound properties.
 				if (Task.BindingsBatch.IsValid())
@@ -2698,7 +3043,7 @@ void FStateTreeExecutionContext::StopEvaluatorsAndGlobalTasks(const EStateTreeRu
 			{
 				const FStateTreeEvaluatorBase& Eval = CurrentStateTree->Nodes[EvalIndex].Get<const FStateTreeEvaluatorBase>();
 				const FStateTreeDataView EvalInstanceView = GetDataView(CurrentParentFrame, CurrentFrame, Eval.InstanceDataHandle);
-				FNodeInstanceDataScope DataScope(*this, Eval.InstanceDataHandle, EvalInstanceView);
+				FNodeInstanceDataScope DataScope(*this, &Eval, EvalIndex, Eval.InstanceDataHandle, EvalInstanceView);
 
 				// Copy bound properties.
 				if (Eval.BindingsBatch.IsValid())
@@ -2711,7 +3056,7 @@ void FStateTreeExecutionContext::StopEvaluatorsAndGlobalTasks(const EStateTreeRu
 			{
 				const FStateTreeTaskBase& Task = CurrentStateTree->Nodes[TaskIndex].Get<const FStateTreeTaskBase>();
 				const FStateTreeDataView TaskInstanceView = GetDataView(CurrentParentFrame, CurrentFrame, Task.InstanceDataHandle);
-				FNodeInstanceDataScope DataScope(*this, Task.InstanceDataHandle, TaskInstanceView);
+				FNodeInstanceDataScope DataScope(*this, &Task, TaskIndex, Task.InstanceDataHandle, TaskInstanceView);
 
 				// Copy bound properties.
 				if (Task.BindingsBatch.IsValid() && Task.bShouldCopyBoundPropertiesOnExitState)
@@ -2755,7 +3100,7 @@ void FStateTreeExecutionContext::CallStopOnEvaluatorsAndGlobalTasks(const FState
 	{
 		const FStateTreeTaskBase& Task =  CurrentStateTree->Nodes[TaskIndex].Get<const FStateTreeTaskBase>();
 		const FStateTreeDataView TaskInstanceView = GetDataView(ParentFrame, Frame, Task.InstanceDataHandle);
-		FNodeInstanceDataScope DataScope(*this, Task.InstanceDataHandle, TaskInstanceView);
+		FNodeInstanceDataScope DataScope(*this, &Task, TaskIndex, Task.InstanceDataHandle, TaskInstanceView);
 
 		// Ignore disabled task
 		if (Task.bTaskEnabled == false)
@@ -2782,7 +3127,7 @@ void FStateTreeExecutionContext::CallStopOnEvaluatorsAndGlobalTasks(const FState
 	{
 		const FStateTreeEvaluatorBase& Eval = CurrentStateTree->Nodes[EvalIndex].Get<const FStateTreeEvaluatorBase>();
 		const FStateTreeDataView EvalInstanceView = GetDataView(ParentFrame, Frame, Eval.InstanceDataHandle);
-		FNodeInstanceDataScope DataScope(*this, Eval.InstanceDataHandle, EvalInstanceView);
+		FNodeInstanceDataScope DataScope(*this, &Eval, EvalIndex, Eval.InstanceDataHandle, EvalInstanceView);
 
 		STATETREE_LOG(Verbose, TEXT("  Stop: '%s'"), *Eval.Name.ToString());
 		UE_STATETREE_DEBUG_EVALUATOR_EXIT_TREE(this, CurrentStateTree, FStateTreeIndex16(EvalIndex));
@@ -2823,7 +3168,7 @@ EStateTreeRunStatus FStateTreeExecutionContext::StartTemporaryEvaluatorsAndGloba
 			check(EvalInstanceView.IsValid());
 		}
 		
-		FNodeInstanceDataScope DataScope(*this, Eval.InstanceDataHandle, EvalInstanceView);
+		FNodeInstanceDataScope DataScope(*this, &Eval, EvalIndex, Eval.InstanceDataHandle, EvalInstanceView);
 		// Copy bound properties.
 		if (Eval.BindingsBatch.IsValid())
 		{
@@ -2861,7 +3206,7 @@ EStateTreeRunStatus FStateTreeExecutionContext::StartTemporaryEvaluatorsAndGloba
 			check(TaskDataView.IsValid())
 		}
 
-		FNodeInstanceDataScope DataScope(*this, Task.InstanceDataHandle, TaskDataView);
+		FNodeInstanceDataScope DataScope(*this, &Task, TaskIndex, Task.InstanceDataHandle, TaskDataView);
 
 		// Copy bound properties.
 		if (Task.BindingsBatch.IsValid())
@@ -2921,11 +3266,12 @@ void FStateTreeExecutionContext::StopTemporaryEvaluatorsAndGlobalTasks(const FSt
 				NodeInstanceView = FStateTreeDataView(TempInstance.Instance);
 			}
 
-			FNodeInstanceDataScope DataScope(*this, TempInstance.DataHandle, NodeInstanceView);
 
 			FConstStructView NodeView = CurrentFrame.StateTree->Nodes[TempInstance.OwnerNodeIndex.Get()];
 			if (const FStateTreeTaskBase* Task = NodeView.GetPtr<const FStateTreeTaskBase>())
 			{
+				FNodeInstanceDataScope DataScope(*this, Task, TempInstance.OwnerNodeIndex.Get(), TempInstance.DataHandle, NodeInstanceView);
+
 				STATETREE_LOG(Verbose, TEXT("  Stop: '%s'"), *Task->Name.ToString());
 				UE_STATETREE_DEBUG_TASK_EXIT_STATE(this, CurrentFrame.StateTree, FStateTreeIndex16(TempInstance.OwnerNodeIndex.Get()));
 				{
@@ -2936,6 +3282,8 @@ void FStateTreeExecutionContext::StopTemporaryEvaluatorsAndGlobalTasks(const FSt
 			}
 			else if (const FStateTreeEvaluatorBase* Eval = NodeView.GetPtr<const FStateTreeEvaluatorBase>())
 			{
+				FNodeInstanceDataScope DataScope(*this, Eval, TempInstance.OwnerNodeIndex.Get(), TempInstance.DataHandle, NodeInstanceView);
+
 				STATETREE_LOG(Verbose, TEXT("  Stop: '%s'"), *Eval->Name.ToString());
 				UE_STATETREE_DEBUG_EVALUATOR_EXIT_TREE(this, CurrentFrame.StateTree, FStateTreeIndex16(TempInstance.OwnerNodeIndex.Get()));
 				{
@@ -2965,12 +3313,9 @@ EStateTreeRunStatus FStateTreeExecutionContext::TickTasks(const float DeltaTime)
 	int32 NumTotalEnabledTasks = 0;
 
 	const bool bHasEvents = EventQueue && EventQueue->HasEvents();
-
-	Exec.CompletedFrameIndex = FStateTreeIndex16::Invalid;
-	Exec.CompletedStateHandle = FStateTreeStateHandle::Invalid;
 	
 	// Used to stop ticking tasks after one fails, but we still want to keep updating the data views so that property binding works properly.
-	bool bShouldTickTasks = true;
+	bool bShouldTickTasks = CompletedStates.Num() == 0;
 
 	STATETREE_CLOG(Exec.ActiveFrames.Num() > 0, VeryVerbose, TEXT("Ticking Tasks"));
 
@@ -3009,7 +3354,7 @@ EStateTreeRunStatus FStateTreeExecutionContext::TickTasks(const float DeltaTime)
 				{
 					const FStateTreeTaskBase& Task = CurrentStateTree->Nodes[TaskIndex].Get<const FStateTreeTaskBase>();
 					const FStateTreeDataView TaskInstanceView = GetDataView(CurrentParentFrame, CurrentFrame, Task.InstanceDataHandle);
-					FNodeInstanceDataScope DataScope(*this, Task.InstanceDataHandle, TaskInstanceView);
+					FNodeInstanceDataScope DataScope(*this, &Task, TaskIndex, Task.InstanceDataHandle, TaskInstanceView);
 
 					// Ignore disabled task
 					if (Task.bTaskEnabled == false)
@@ -3036,6 +3381,7 @@ EStateTreeRunStatus FStateTreeExecutionContext::TickTasks(const float DeltaTime)
 					}
 
 					//UE_STATETREE_DEBUG_TASK_EVENT(this, TaskIndex, TaskDataView, EStateTreeTraceEventType::OnTickingTask, EStateTreeRunStatus::Running);
+					UE_STATETREE_ENSURE_COMPLETED_STATE_RUN_STATUS();
 					UE_STATETREE_DEBUG_TASK_TICK(this, CurrentStateTree, TaskIndex);
 
 					EStateTreeRunStatus TaskResult = EStateTreeRunStatus::Unset;
@@ -3046,27 +3392,29 @@ EStateTreeRunStatus FStateTreeExecutionContext::TickTasks(const float DeltaTime)
 						TaskResult = Task.Tick(*this, DeltaTime);
 					}
 
+					// Delayed from CompleteTask(), override the Tick result unless it's a failure.
+					if (CompletedStateRunStatus.IsSet()
+						&& (TaskResult == EStateTreeRunStatus::Succeeded || TaskResult == EStateTreeRunStatus::Running))
+					{
+						TaskResult = CompletedStateRunStatus.GetValue();
+						CompletedStateRunStatus.Reset();
+					}
+
 					UE_STATETREE_DEBUG_TASK_EVENT(this, TaskIndex, TaskInstanceView,
 						TaskResult != EStateTreeRunStatus::Running ? EStateTreeTraceEventType::OnTaskCompleted : EStateTreeTraceEventType::OnTicked,
 						TaskResult);
 
-
-					// TODO: Add more control over which states can control the failed/succeeded result.
 					if (TaskResult != EStateTreeRunStatus::Running)
 					{
-						// Store the first state that completed, will be used to decide where to trigger transitions.
-						if (!Exec.CompletedStateHandle.IsValid())
+						Result = TaskResult;
+						FinishTaskInternal(FrameIndex, StateIndex, TaskIndex, TaskResult);
+						if (TaskResult == EStateTreeRunStatus::Failed)
 						{
-							Exec.CompletedFrameIndex = FStateTreeIndex16(FrameIndex);
-							Exec.CompletedStateHandle = CurrentHandle;
+							bShouldTickTasks = false;
 						}
-						Result = TaskResult;
-					}
-				
-					if (TaskResult == EStateTreeRunStatus::Failed)
-					{
-						bShouldTickTasks = false;
 					}
+
+					UE_STATETREE_ENSURE_COMPLETED_STATE_RUN_STATUS();
 				}
 			}
 
@@ -3074,12 +3422,19 @@ EStateTreeRunStatus FStateTreeExecutionContext::TickTasks(const float DeltaTime)
 		}
 	}
 
-	if (NumTotalEnabledTasks == 0)
+	// Delayed completed task but no task were ticked to return the value.
+	if (Result == EStateTreeRunStatus::Running && CompletedStates.Num() > 0)
+	{
+		const FCompletedState* FirstFailureCompletedState = CompletedStates.FindByPredicate([](const FCompletedState& Other) { return Other.RunStatus != EStateTreeRunStatus::Succeeded; });
+		Result = FirstFailureCompletedState == nullptr ? EStateTreeRunStatus::Succeeded : FirstFailureCompletedState->RunStatus;
+	}
+	else if (NumTotalEnabledTasks == 0 && CompletedStates.Num() == 0)
 	{
 		// No enabled tasks, done ticking. Check from the bottom state in the bottom frame to trigger completion transitions
+		const int32 CompletedFrameIndex = Exec.ActiveFrames.Num() > 0 ? Exec.ActiveFrames.Num() - 1 : INDEX_NONE;
+		const int32 CompletedStateIndex = Exec.ActiveFrames.Num() > 0 ? (int32)Exec.ActiveFrames.Last().ActiveStates.Num() - 1 : INDEX_NONE;
+		FinishTaskInternal(CompletedFrameIndex, CompletedStateIndex, INDEX_NONE, EStateTreeRunStatus::Succeeded);
 		Result = EStateTreeRunStatus::Succeeded;
-		Exec.CompletedFrameIndex = Exec.ActiveFrames.Num() > 0 ? FStateTreeIndex16(Exec.ActiveFrames.Num() - 1) : FStateTreeIndex16::Invalid;
-		Exec.CompletedStateHandle = Exec.ActiveFrames.Num() > 0 ? Exec.ActiveFrames.Last().ActiveStates.Last() : FStateTreeStateHandle::Invalid;
 	}
 
 	return Result;
@@ -3104,7 +3459,7 @@ bool FStateTreeExecutionContext::TestAllConditions(const FStateTreeExecutionFram
 		const int32 ConditionIndex = ConditionsOffset + Index;
 		const FStateTreeConditionBase& Cond = CurrentFrame.StateTree->Nodes[ConditionIndex].Get<const FStateTreeConditionBase>();
 		const FStateTreeDataView ConditionInstanceView = GetDataView(CurrentParentFrame, CurrentFrame, Cond.InstanceDataHandle);
-		FNodeInstanceDataScope DataScope(*this, Cond.InstanceDataHandle, ConditionInstanceView);
+		FNodeInstanceDataScope DataScope(*this, &Cond, ConditionIndex, Cond.InstanceDataHandle, ConditionInstanceView);
 
 		bool bValue = false;
 		if (Cond.EvaluationMode == EStateTreeConditionEvaluationMode::Evaluated)
@@ -3199,7 +3554,7 @@ float FStateTreeExecutionContext::EvaluateUtility(const FStateTreeExecutionFrame
 		const int32 ConsiderationIndex = ConsiderationsOffset + Index;
 		const FStateTreeConsiderationBase& Consideration = CurrentFrame.StateTree->Nodes[ConsiderationIndex].Get<const FStateTreeConsiderationBase>();
 		const FStateTreeDataView ConsiderationInstanceView = GetDataView(CurrentParentFrame, CurrentFrame, Consideration.InstanceDataHandle);
-		FNodeInstanceDataScope DataScope(*this, Consideration.InstanceDataHandle, ConsiderationInstanceView);
+		FNodeInstanceDataScope DataScope(*this, &Consideration, ConsiderationIndex, Consideration.InstanceDataHandle, ConsiderationInstanceView);
 
 		// Copy bound properties.
 		if (Consideration.BindingsBatch.IsValid())
@@ -3265,7 +3620,7 @@ void FStateTreeExecutionContext::EvaluatePropertyFunctionsOnActiveInstances(cons
 	{
 		const FStateTreePropertyFunctionBase& Func = CurrentFrame.StateTree->Nodes[FuncIndex].Get<const FStateTreePropertyFunctionBase>();
 		const FStateTreeDataView FuncInstanceView = GetDataView(CurrentParentFrame, CurrentFrame, Func.InstanceDataHandle);
-		FNodeInstanceDataScope DataScope(*this, Func.InstanceDataHandle, FuncInstanceView);
+		FNodeInstanceDataScope DataScope(*this, &Func, FuncIndex, Func.InstanceDataHandle, FuncInstanceView);
 
 		// Copy bound properties.
 		if (Func.BindingsBatch.IsValid())
@@ -3290,7 +3645,7 @@ void FStateTreeExecutionContext::EvaluatePropertyFunctionsWithValidation(const F
 	{
 		const FStateTreePropertyFunctionBase& Func = CurrentFrame.StateTree->Nodes[FuncIndex].Get<const FStateTreePropertyFunctionBase>();
 		const FStateTreeDataView FuncInstanceView = GetDataView(CurrentParentFrame, CurrentFrame, Func.InstanceDataHandle);
-		FNodeInstanceDataScope DataScope(*this, Func.InstanceDataHandle, FuncInstanceView);
+		FNodeInstanceDataScope DataScope(*this, &Func, FuncIndex, Func.InstanceDataHandle, FuncInstanceView);
 
 		// Copy bound properties.
 		if (Func.BindingsBatch.IsValid())
@@ -3393,7 +3748,7 @@ bool FStateTreeExecutionContext::RequestTransition(
 		// Consume events from states, if required. 
 		for (int32 FrameIndex = 0; FrameIndex < NextTransition.NextActiveFrames.Num(); FrameIndex++)
 		{
-			const FStateTreeExecutionFrame& Frame = NextTransition.NextActiveFrames[FrameIndex]; 
+			const FStateTreeExecutionFrame& Frame = NextTransition.NextActiveFrames[FrameIndex];
 			const FStateTreeFrameStateSelectionEvents& FrameEvents = NextTransition.NextActiveFrameEvents[FrameIndex];
 
 			for (int32 StateIndex = 0; StateIndex < Frame.ActiveStates.Num(); StateIndex++)
@@ -3403,20 +3758,20 @@ bool FStateTreeExecutionContext::RequestTransition(
 					const FCompactStateTreeState& State = Frame.StateTree->States[StateIndex];
 					if (State.bConsumeEventOnSelect)
 					{
-						 ConsumeEvent(FrameEvents.Events[StateIndex]);
+						ConsumeEvent(FrameEvents.Events[StateIndex]);
 					}
 				}
 			}
 		}
-		
+
 		STATETREE_LOG(Verbose, TEXT("Transition on state '%s' -[%s]-> state '%s'"),
 			*GetSafeStateName(CurrentFrame, CurrentFrame.ActiveStates.Last()),
 			*GetSafeStateName(CurrentFrame, NextState),
 			*GetSafeStateName(NextTransition.NextActiveFrames.Last(), NextTransition.NextActiveFrames.Last().ActiveStates.Last()));
-		
+
 		return true;
 	}
-		
+
 	return false;
 }
 
@@ -3639,7 +3994,7 @@ bool FStateTreeExecutionContext::TriggerTransitions()
 			}
 
 			const FStateTreeDataView TaskInstanceView = GetDataView(CurrentParentFrame, CurrentFrame, Task.InstanceDataHandle);
-			FNodeInstanceDataScope DataScope(*this, Task.InstanceDataHandle, TaskInstanceView);
+			FNodeInstanceDataScope DataScope(*this, &Task, Handler.TaskIndex.Get(), Task.InstanceDataHandle, TaskInstanceView);
 
 			// Copy bound properties.
 			if (Task.BindingsBatch.IsValid())
@@ -3842,8 +4197,32 @@ bool FStateTreeExecutionContext::TriggerTransitions()
 	if (NextTransition.Priority == EStateTreeTransitionPriority::None
 		&& Exec.LastTickStatus != EStateTreeRunStatus::Running)
 	{
-		// Start from the last completed state if specified.
-		const int32 FrameStartIndex = Exec.CompletedFrameIndex.IsValid() ? Exec.CompletedFrameIndex.AsInt32() : (Exec.ActiveFrames.Num() - 1);
+		// Sort CompletedStates, only keep the smaller frame and the active state of the smaller frame.
+		TOptional<const FCompletedState> FirstCompletedState;
+		{
+			for (const FCompletedState& CompletedState : CompletedStates)
+			{
+				if (CompletedState.ActiveStateIndex == INDEX_NONE)
+				{
+					// Global task are not triggered in the transition loop.
+					FirstCompletedState.Reset();
+					break;
+				}
+				if (!FirstCompletedState.IsSet()
+					|| CompletedState.ActiveFrameIndex < FirstCompletedState.GetValue().ActiveFrameIndex
+					|| (CompletedState.ActiveFrameIndex == FirstCompletedState.GetValue().ActiveFrameIndex && CompletedState.ActiveStateIndex < FirstCompletedState.GetValue().ActiveStateIndex)
+					)
+				{
+					FirstCompletedState = CompletedState;
+				}
+			}
+
+			Exec.DelayedFinishedTasks.Empty();
+			CompletedStates.Empty();
+		}
+
+		// Start from the last completed state if specified, else only check the last state from the last frame.
+		const int32 FrameStartIndex = FirstCompletedState.IsSet() ? FirstCompletedState.GetValue().ActiveFrameIndex : (Exec.ActiveFrames.Num() - 1);
 		check(FrameStartIndex >= 0 && FrameStartIndex < Exec.ActiveFrames.Num());
 		
 		for (int32 FrameIndex = FrameStartIndex; FrameIndex >= 0; FrameIndex--)
@@ -3855,15 +4234,15 @@ bool FStateTreeExecutionContext::TriggerTransitions()
 			FCurrentlyProcessedFrameScope FrameScope(*this, CurrentParentFrame, CurrentFrame);
 
 			int32 StateStartIndex = CurrentFrame.ActiveStates.Num() - 1; // This is ok, even if the ActiveStates is 0, -1 will skip the whole state loop below.
-			if (FrameIndex == FrameStartIndex
-				&& Exec.CompletedStateHandle.IsValid())
+			if (FrameIndex == FrameStartIndex && FirstCompletedState.IsSet())
 			{
-				StateStartIndex = CurrentFrame.ActiveStates.IndexOfReverse(Exec.CompletedStateHandle);
-				// INDEX_NONE (-1) will skip the whole state loop below. We still want to warn.
-				ensureMsgf(StateStartIndex != INDEX_NONE, TEXT("If CompletedFrameIndex and CompletedStateHandle are specified, we expect that the state is found"));
+				StateStartIndex = FirstCompletedState.GetValue().ActiveStateIndex;
 			}
-			
-			const EStateTreeTransitionTrigger CompletionTrigger = Exec.LastTickStatus == EStateTreeRunStatus::Succeeded ? EStateTreeTransitionTrigger::OnStateSucceeded : EStateTreeTransitionTrigger::OnStateFailed;
+
+			// Use worst case of LastTickStatus and FoundCompletedState.
+			const bool bLastTickWasSuccess = Exec.LastTickStatus == EStateTreeRunStatus::Succeeded;
+			const EStateTreeRunStatus ExpectedRunStatus = bLastTickWasSuccess && FirstCompletedState.IsSet() ? FirstCompletedState.GetValue().RunStatus : Exec.LastTickStatus;
+			const EStateTreeTransitionTrigger CompletionTrigger = ExpectedRunStatus == EStateTreeRunStatus::Succeeded ? EStateTreeTransitionTrigger::OnStateSucceeded : EStateTreeTransitionTrigger::OnStateFailed;
 
 			// Check completion transitions
 			for (int32 StateIndex = StateStartIndex; StateIndex >= 0; StateIndex--)
@@ -3953,6 +4332,11 @@ bool FStateTreeExecutionContext::TriggerTransitions()
 			}
 		}
 	}
+	else
+	{
+		Exec.DelayedFinishedTasks.Empty();
+		CompletedStates.Empty();
+	}
 
 	// Check if the transition was succeed/failed, if we're on a sub-tree, complete the subtree instead of transition.
 	if (NextTransition.TargetState.IsCompletionState() && bProcessSubTreeCompletion)
@@ -3976,8 +4360,8 @@ bool FStateTreeExecutionContext::TriggerTransitions()
 					*GetSafeStateName(ParentFrame, ParentLinkedState), *GetSafeStateName(SourceFrame, NextTransition.SourceState), *UEnum::GetDisplayValueAsText(RunStatus).ToString());
 
 				// Set the parent linked state as last completed state, and update tick status to the status from the transition.
-				Exec.CompletedFrameIndex = FStateTreeIndex16(ParentFrameIndex);
-				Exec.CompletedStateHandle = ParentLinkedState;
+				const int32 TaskIndex = INDEX_NONE;
+				FinishTaskInternal(ParentFrameIndex, ParentFrame.ActiveStates.Num() - 1, TaskIndex, RunStatus);
 				Exec.LastTickStatus = RunStatus;
 
 				// Clear the transition and return that no transition took place.
@@ -5167,23 +5551,10 @@ FStateTreeWeakExecutionContext::FStateTreeWeakExecutionContext(const FStateTreeE
 	{
 		FrameStateTree = Frame->StateTree;
 		FrameRootState = Frame->RootState;
+		CurrentState = Context.GetCurrentlyProcessedState();
 	}
 }
 
-FStateTreeWeakExecutionContext::FStateTreeWeakExecutionContext(gsl::not_null<UObject*> InOwner, gsl::not_null<const UStateTree*> InStateTree, FStateTreeInstanceData& InInstanceData)
-	: Owner(InOwner)
-	, StateTree(InStateTree)
-	, Storage(InInstanceData.GetWeakMutableStorage())
-{
-}
-
-FStateTreeWeakExecutionContext::FStateTreeWeakExecutionContext(gsl::not_null<UObject*> InOwner, gsl::not_null<const UStateTree*> InStateTree, TWeakPtr<FStateTreeInstanceStorage> InStorage)
-	: Owner(InOwner)
-	, StateTree(InStateTree)
-	, Storage(InStorage)
-{
-}
-
 bool FStateTreeWeakExecutionContext::SendEvent(const FGameplayTag Tag, const FConstStructView Payload, const FName Origin) const
 {
 	TStrongObjectPtr<UObject> OwnerPinned = Owner.Pin();
@@ -5264,17 +5635,17 @@ bool FStateTreeWeakExecutionContext::BroadcastDelegate(FStateTreeDelegateDispatc
 			Exec.CurrentExecutionContext->BroadcastDelegate(Dispatcher);
 			bBroadcasted = true;
 		}
-		else
+		else if (Exec.ActiveFrames.Num() > 0)
 		{
 			TStrongObjectPtr<const UStateTree> FrameStateTreePinned = FrameStateTree.Pin();
 			if (FrameStateTreePinned)
 			{
 				const FStateTreeExecutionFrame* ParentFrame = nullptr;
-				const FStateTreeExecutionFrame* CurrentFrame = FStateTreeExecutionContext::FindFrame(FrameStateTreePinned.Get(), FrameRootState, StoragePinned->GetExecutionState().ActiveFrames, ParentFrame);
+				const FStateTreeExecutionFrame* CurrentFrame = FStateTreeExecutionContext::FindFrame(FrameStateTreePinned.Get(), FrameRootState, Exec.ActiveFrames, ParentFrame);
 				if (CurrentFrame)
 				{
 					Exec.DelegateActiveListeners.BroadcastDelegate(Dispatcher, CurrentFrame->StateTree, Exec);
-					if (UE::StateTree::ExecutionContext::MarkDelegateAsBroadcasted(Dispatcher, *CurrentFrame, *StoragePinned))
+					if (UE::StateTree::ExecutionContext::Private::MarkDelegateAsBroadcasted(Dispatcher, *CurrentFrame, *StoragePinned))
 					{
 						ScheduleNextTick(OwnerPinned.Get(), RootStateTreePinned.Get(), *StoragePinned);
 					}
@@ -5297,6 +5668,8 @@ bool FStateTreeWeakExecutionContext::RemoveDelegateListener(FStateTreeDelegateLi
 		return true;
 	}
 
+	bool bRemoved = false;
+
 	TStrongObjectPtr<UObject> OwnerPinned = Owner.Pin();
 	TStrongObjectPtr<const UStateTree> StateTreePinned = StateTree.Pin();
 	TSharedPtr<FStateTreeInstanceStorage> StoragePinned = Storage.Pin();
@@ -5310,21 +5683,69 @@ bool FStateTreeWeakExecutionContext::RemoveDelegateListener(FStateTreeDelegateLi
 		if (Exec.CurrentExecutionContext)
 		{
 			Exec.CurrentExecutionContext->RemoveDelegateListener(Listener);
+			bRemoved = true;
 		}
-		else
+		else if(Exec.ActiveFrames.Num() > 0)
 		{
 			TStrongObjectPtr<const UStateTree> FrameStateTreePinned = FrameStateTree.Pin();
 			if (FrameStateTreePinned)
 			{
-				StoragePinned->GetMutableExecutionState().DelegateActiveListeners.RemoveListener(Listener, FrameStateTreePinned.Get());
+				Exec.DelegateActiveListeners.RemoveListener(Listener, FrameStateTreePinned.Get());
+				bRemoved = true;
 			}
 		}
 
 		StoragePinned->ReleaseWriteAccess();
-		return true;
 	}
 
-	return false;
+	return bRemoved;
+}
+
+bool FStateTreeWeakExecutionContext::FinishTask(const FStateTreeWeakTaskRef& Task, EStateTreeFinishTaskType FinishType) const
+{
+	bool bCompleted = false;
+
+	TStrongObjectPtr<UObject> OwnerPinned = Owner.Pin();
+	TStrongObjectPtr<const UStateTree> StateTreePinned = StateTree.Pin();
+	TSharedPtr<FStateTreeInstanceStorage> StoragePinned = Storage.Pin();
+	if (OwnerPinned && StateTreePinned && StoragePinned)
+	{
+		// Manually acquire and release the write access.
+		//The context may already exist and we want to confirm that we are on the same thread.
+		StoragePinned->AcquireWriteAccess();
+
+		FStateTreeStrongTaskRef TaskPin = Task.Pin();
+		const FStateTreeTaskBase* TaskPtr = TaskPin.GetTask();
+		TStrongObjectPtr<const UStateTree> FrameStateTreePinned = FrameStateTree.Pin();
+		if (TaskPtr != nullptr && FrameStateTreePinned.Get() != nullptr && TaskPin.GetStateTree() == FrameStateTreePinned.Get())
+		{
+			FStateTreeDelayedFinishedTask CompletedTask{
+				.FrameStateTree = FObjectKey(FrameStateTreePinned.Get()), 
+				.FrameRootState = FrameRootState,
+				.State = CurrentState,
+				.TaskIndex = TaskPin.GetTaskIndex(),
+				.Type = FinishType
+				};
+
+			FStateTreeExecutionState& Exec = StoragePinned->GetMutableExecutionState();
+			if (Exec.CurrentExecutionContext)
+			{
+				Exec.CurrentExecutionContext->FinishTask(CompletedTask, FinishType);
+				bCompleted = true;
+			}
+			else
+			{
+				// Confirm that the task if from the correct state.
+				Exec.DelayedFinishedTasks.Add(CompletedTask);
+				ScheduleNextTick(OwnerPinned.Get(), StateTreePinned.Get(), *StoragePinned);
+				bCompleted = true;
+			}
+		}
+
+		StoragePinned->ReleaseWriteAccess();
+	}
+
+	return bCompleted;
 }
 
 void FStateTreeWeakExecutionContext::ScheduleNextTick(gsl::not_null<UObject*> Owner, gsl::not_null<const UStateTree*> RootStateTree, FStateTreeInstanceStorage& Storage)
@@ -5338,3 +5759,4 @@ void FStateTreeWeakExecutionContext::ScheduleNextTick(gsl::not_null<UObject*> Ow
 
 #undef STATETREE_LOG
 #undef STATETREE_CLOG
+#undef UE_STATETREE_ENSURE_COMPLETED_STATE_RUN_STATUS
diff --git a/Engine/Plugins/Runtime/StateTree/Source/StateTreeModule/Private/StateTreeNodeRef.cpp b/Engine/Plugins/Runtime/StateTree/Source/StateTreeModule/Private/StateTreeNodeRef.cpp
new file mode 100644
--- /dev/null
+++ b/Engine/Plugins/Runtime/StateTree/Source/StateTreeModule/Private/StateTreeNodeRef.cpp
@@ -0,0 +1,76 @@
+// Copyright Epic Games, Inc. All Rights Reserved.
+
+#include "StateTreeNodeRef.h"
+#include "StateTree.h"
+#include "StateTreeExecutionContext.h"
+#include "StateTreeInstanceData.h"
+#include "StateTreeTaskBase.h"
+
+
+#if WITH_STATETREE_DEBUG
+FStateTreeStrongTaskRef::FStateTreeStrongTaskRef(TStrongObjectPtr<const UStateTree> InStateTree, const FStateTreeTaskBase* InTask, FStateTreeIndex16 InNodeIndex, FGuid InNodeId)
+	: StateTree(InStateTree)
+	, Task(InTask)
+	, NodeIndex(InNodeIndex)
+	, NodeId(InNodeId)
+{
+}
+#else
+FStateTreeStrongTaskRef::FStateTreeStrongTaskRef(TStrongObjectPtr<const UStateTree> InStateTree, const FStateTreeTaskBase* InTask, FStateTreeIndex16 InNodeIndex)
+	: StateTree(InStateTree)
+	, Task(InTask)
+	, NodeIndex(InNodeIndex)
+	{
+	}
+#endif
+
+const UStateTree* FStateTreeStrongTaskRef::GetStateTree() const
+{
+	return IsValid() ? StateTree.Get() : nullptr;
+}
+	
+const FStateTreeTaskBase* FStateTreeStrongTaskRef::GetTask() const
+{
+	return IsValid() ? Task : nullptr;
+}
+
+bool FStateTreeStrongTaskRef::IsValid() const
+{
+	bool bIsValid = Task != nullptr && StateTree.Get() != nullptr;
+#if WITH_STATETREE_DEBUG
+	if (bIsValid)
+	{
+		bIsValid = NodeId == StateTree->GetNodeIdFromIndex(NodeIndex);
+		ensureMsgf(bIsValid, TEXT("The node id changed from the last use. Did the StateTree asset recompiled?"));
+	}
+#endif
+	return bIsValid;
+}
+
+FStateTreeWeakTaskRef::FStateTreeWeakTaskRef(TNotNull<const UStateTree*> InStateTree, FStateTreeIndex16 InTaskIndex)
+	: StateTree(InStateTree)
+	, NodeIndex(InTaskIndex)
+{
+#if WITH_STATETREE_DEBUG
+	NodeId = InStateTree->GetNodeIdFromIndex(NodeIndex);
+#endif
+}
+
+FStateTreeStrongTaskRef FStateTreeWeakTaskRef::Pin() const
+{
+	TStrongObjectPtr<const UStateTree> StateTreePinned = StateTree.Pin();
+	const FStateTreeTaskBase* Task = nullptr;
+	if (StateTreePinned && StateTreePinned->GetNodes().IsValidIndex(NodeIndex.AsInt32()))
+	{
+#if WITH_STATETREE_DEBUG
+		ensureMsgf(NodeId == StateTreePinned->GetNodeIdFromIndex(NodeIndex), TEXT("The node id changed from the last use. Did the StateTree asset recompiled?"));
+#endif
+		Task = StateTreePinned->GetNodes()[NodeIndex.AsInt32()].GetPtr<const FStateTreeTaskBase>();
+	}
+
+#if WITH_STATETREE_DEBUG
+	return Task ? FStateTreeStrongTaskRef(StateTreePinned, Task, NodeIndex, NodeId) : FStateTreeStrongTaskRef();
+#else
+	return Task ? FStateTreeStrongTaskRef(StateTreePinned, Task, NodeIndex) : FStateTreeStrongTaskRef();
+#endif
+}
diff --git a/Engine/Plugins/Runtime/StateTree/Source/StateTreeModule/Public/Blueprint/StateTreeNodeBlueprintBase.h b/Engine/Plugins/Runtime/StateTree/Source/StateTreeModule/Public/Blueprint/StateTreeNodeBlueprintBase.h
--- a/Engine/Plugins/Runtime/StateTree/Source/StateTreeModule/Public/Blueprint/StateTreeNodeBlueprintBase.h
+++ b/Engine/Plugins/Runtime/StateTree/Source/StateTreeModule/Public/Blueprint/StateTreeNodeBlueprintBase.h
@@ -49,14 +49,14 @@ public:
 	UFUNCTION(BlueprintCallable, Category = "StateTree", meta = (HideSelfPin = "true", DisplayName = "StateTree Request Transition"))
 	void RequestTransition(const FStateTreeStateLink& TargetState, const EStateTreeTransitionPriority Priority = EStateTreeTransitionPriority::Normal);
 
-	/** Returns a reference to selected property in State Tree. */
-	UFUNCTION(CustomThunk)
-	void GetPropertyReference(const FStateTreeBlueprintPropertyRef& PropertyRef) const;
-
 	/** Broadcasts provided delegate. */
 	UFUNCTION(BlueprintCallable, Category = "StateTree", meta = (HideSelfPin = "true"))
 	void BroadcastDelegate(FStateTreeDelegateDispatcher Dispatcher);
 
+	/** Returns a reference to selected property in State Tree. */
+	UFUNCTION(CustomThunk)
+	void GetPropertyReference(const FStateTreeBlueprintPropertyRef& PropertyRef) const;
+
 	/** Returns true if reference to selected property in State Tree is accessible. */
 	UFUNCTION()
 	bool IsPropertyRefValid(const FStateTreeBlueprintPropertyRef& PropertyRef) const;
@@ -92,15 +92,20 @@ protected:
 	void SetCachedInstanceDataFromContext(const FStateTreeExecutionContext& Context) const;
 	void ClearCachedInstanceData() const;
 
+	FStateTreeWeakExecutionContext GetWeakExecutionContext() const
+	{
+		return WeakExecutionContext;
+	}
+
 private:
 	DECLARE_FUNCTION(execGetPropertyReference);
 
 	void* GetMutablePtrToProperty(const FStateTreeBlueprintPropertyRef& PropertyRef, FProperty*& OutSourceProperty) const;
 	
-	/** Cached instance data while the node is active. */
+	/** Cached execution context while the node is active for async nodes. */
 	mutable FStateTreeWeakExecutionContext WeakExecutionContext;
 
-	/** Cached instance data while the node is active. */
+	/** Cached instance data while the node is active for async nodes. */
 	mutable TWeakPtr<FStateTreeInstanceStorage> WeakInstanceStorage;
 
 	/** Cached State Tree of owning execution frame. */
diff --git a/Engine/Plugins/Runtime/StateTree/Source/StateTreeModule/Public/Blueprint/StateTreeTaskBlueprintBase.h b/Engine/Plugins/Runtime/StateTree/Source/StateTreeModule/Public/Blueprint/StateTreeTaskBlueprintBase.h
--- a/Engine/Plugins/Runtime/StateTree/Source/StateTreeModule/Public/Blueprint/StateTreeTaskBlueprintBase.h
+++ b/Engine/Plugins/Runtime/StateTree/Source/StateTreeModule/Public/Blueprint/StateTreeTaskBlueprintBase.h
@@ -5,6 +5,7 @@
 #include "Templates/SubclassOf.h"
 #include "StateTreeTaskBase.h"
 #include "StateTreeNodeBlueprintBase.h"
+#include "StateTreeNodeRef.h"
 #include "StateTreeTaskBlueprintBase.generated.h"
 
 struct FStateTreeExecutionContext;
@@ -79,10 +80,10 @@ protected:
 
 	/** Finish the task and sets it's status. */
 	UFUNCTION(BlueprintCallable, Category = "StateTree", meta = (HideSelfPin = "true", DisplayName = "Finish Task"))
-	void FinishTask(const bool bSucceeded = true)
-	{
-		RunStatus = bSucceeded ? EStateTreeRunStatus::Succeeded : EStateTreeRunStatus::Failed;
-	}
+	void FinishTask(const bool bSucceeded = true);
+
+	/** Cached task while the node is active for async execution. */
+	mutable FStateTreeWeakTaskRef WeakTaskRef;
 
 	/** Run status when using latent EnterState and Tick */
 	EStateTreeRunStatus RunStatus = EStateTreeRunStatus::Running;
@@ -121,6 +122,8 @@ protected:
 	UE_DEPRECATED(all, "Use bHasLatentTick instead.")
 	uint8 bHasTick_DEPRECATED : 1;
 
+	mutable uint8 bIsProcessingEnterStateOrTick : 1;
+
 	friend struct FStateTreeBlueprintTaskWrapper;
 };
 
diff --git a/Engine/Plugins/Runtime/StateTree/Source/StateTreeModule/Public/StateTreeExecutionContext.h b/Engine/Plugins/Runtime/StateTree/Source/StateTreeModule/Public/StateTreeExecutionContext.h
--- a/Engine/Plugins/Runtime/StateTree/Source/StateTreeModule/Public/StateTreeExecutionContext.h
+++ b/Engine/Plugins/Runtime/StateTree/Source/StateTreeModule/Public/StateTreeExecutionContext.h
@@ -6,6 +6,7 @@
 #include "StateTreeExecutionExtension.h"
 #include "StateTreeExecutionTypes.h"
 #include "StateTreeNodeBase.h"
+#include "StateTreeNodeRef.h"
 #include "Experimental/ConcurrentLinearAllocator.h"
 #include "Templates/IsInvocable.h"
 
@@ -109,6 +110,9 @@ protected:
 
 	/** Data storage of the instance data. */
 	FStateTreeInstanceStorage& Storage;
+
+	/** The context is processing the tree. We do not need to inform the owner that something changed. */
+	bool bAllowedToScheduleNextTick = true;
 };
 
 /**
@@ -200,6 +204,34 @@ public:
 
 	/** @return a weak context to interact with the state tree instance data that can be stored for later uses. */
 	FStateTreeWeakExecutionContext MakeWeakExecutionContext() const;
+	
+	/**
+	 * @return a weak reference for a task that can be stored for later uses.
+	 * @note similar to GetInstanceData, the node needs to be the current processing node.
+	 */
+	FStateTreeWeakTaskRef MakeWeakTaskRef(const FStateTreeTaskBase& Node) const;
+
+	/**
+	 * @return a weak reference for a task that can be stored for later uses.
+	 * @note similar to GetInstanceData, the instance data needs to be the current processing node.
+	 */
+	template <typename T>
+	FStateTreeWeakTaskRef MakeWeakTaskRefFromInstanceData(const T& InInstanceData) const
+	{
+		check(&CurrentNodeInstanceData.template GetMutable<T>() == &InInstanceData);
+		return MakeWeakTaskRefInternal();
+	}
+
+	/**
+	 * @return a weak reference for a task that can be stored for later uses.
+	 * @note similar to GetInstanceData, the instance data needs to be the current processing node.
+	 */
+	template <typename T>
+	FStateTreeWeakTaskRef MakeWeakTaskRefFromInstanceDataPtr(const T* InInstanceData) const
+	{
+		check(CurrentNodeInstanceData.template GetMutablePtr<T>() == InInstanceData);
+		return MakeWeakTaskRefInternal();
+	}
 
 	struct FStartParameters
 	{
@@ -504,6 +536,21 @@ public:
 	 */
 	void RequestTransition(FStateTreeStateHandle TargetState, EStateTreeTransitionPriority Priority = EStateTreeTransitionPriority::Normal, EStateTreeSelectionFallback Fallback = EStateTreeSelectionFallback::None);
 
+	/**
+	 * Finishes a task. This fails if the Task is not currently the processed node.
+	 * ie. Must be called from inside a FStateTreeTaskBase EnterState, ExitState, StateCompleted, Tick, TriggerTransitions.
+	 * If called during tick processing, then the state completes immediately.
+	 * If called outside of the tick processing, then the request is buffered and handled on the next tick.
+	 */
+	void FinishTask(const FStateTreeTaskBase& Task, EStateTreeFinishTaskType FinishType);
+	
+	/**
+	 * Finishes a task.
+	 * If called during tick processing, then the state completes immediately.
+	 * If called outside of the tick processing, then the request is buffered and handled on the next tick.
+	 */
+	void FinishTask(const FStateTreeDelayedFinishedTask& Task, EStateTreeFinishTaskType FinishType);
+
 	/** @return data view of the specified handle relative to given frame. */
 	static FStateTreeDataView GetDataViewFromInstanceStorage(FStateTreeInstanceStorage& InstanceDataStorage, FStateTreeInstanceStorage* CurrentlyProcessedSharedInstanceStorage, const FStateTreeExecutionFrame* ParentFrame, const FStateTreeExecutionFrame& CurrentFrame, const FStateTreeDataHandle Handle);
 
@@ -718,6 +765,24 @@ protected:
 	 */
 	bool TriggerTransitions();
 
+	/** Adds a task to the finished and delayed list. */
+	void FinishTaskInternal(int32 ActiveFrameIndex, int32 ActiveStateIndex, int32 TaskIndex, EStateTreeRunStatus RunStatus);
+
+	/**
+	 * Removes stale entries and fills the completed states list.
+	 * The tree can be in a different state from when a finished task was added to the pending list and the finished task may not be valid.
+	 */
+	void UpdateCompletedStateList();
+
+	/**
+	 * Adds uniquely the task to the finished.
+	 * Adds the task to the delayed list. The delayed list may contains duplicate. They will be filtered with UpdateCompletedStateList when ticked.
+	 */
+	void AddUniqueCompletedStateList(int32 ActiveFrameIndex, int32 ActiveStateIndex, const FStateTreeDelayedFinishedTask& DelayedFinisedTask);
+
+	/** In the list of completed states, returns the run status when it's completed by a global task. */
+	EStateTreeRunStatus GetGlobalTasksCompletedStatesStatus() const;
+
 	/**
 	 * Runs state selection logic starting at the specified state, walking towards the leaf states.
 	 * If a state cannot be selected, false is returned. 
@@ -818,6 +883,9 @@ protected:
 	 */
 	void CaptureNewStateEvents(TConstArrayView<FStateTreeExecutionFrame> PrevFrames, TConstArrayView<FStateTreeExecutionFrame> NewFrames, TArrayView<FStateTreeFrameStateSelectionEvents> FramesStateSelectionEvents);
 
+	/** @return a weak reference for a task that can be stored for later uses. */
+	FStateTreeWeakTaskRef MakeWeakTaskRefInternal() const;
+
 	/** Instance data used during current tick. */
 	FStateTreeInstanceData& InstanceData;
 
@@ -843,7 +911,25 @@ protected:
 	TArray<FCollectedExternalDataCache, TConcurrentLinearArrayAllocator<FDefaultBlockAllocationTag>> CollectedExternalCache;
 
 	bool bActiveExternalDataCollected = false;
-	
+
+	/** Run status set in FinishTask and that used after the task ticks. */
+	TOptional<EStateTreeRunStatus> CompletedStateRunStatus;
+
+	/** Helper struct to store the resolved FStateTreeExecutionState::FCompletedState */
+	struct FCompletedState
+	{
+		int32 ActiveFrameIndex = INDEX_NONE;
+		int32 ActiveStateIndex = INDEX_NONE;
+		EStateTreeRunStatus RunStatus = EStateTreeRunStatus::Unset;
+		FStateTreeDelayedFinishedTask DelayedTask;
+	};
+
+	/**
+	 * Completed states while ticking.
+	 * Build from the PendingCompletedStates list and maintained by the CompleteTask function.
+	 */
+	TArray<FCompletedState, TInlineAllocator<4>> CompletedStates;
+
 	/** Next transition, used by RequestTransition(). */
 	FStateTreeTransitionResult NextTransition;
 
@@ -872,7 +958,6 @@ protected:
 	struct FCurrentlyProcessedFrameScope
 	{
 		FCurrentlyProcessedFrameScope(FStateTreeExecutionContext& InContext, const FStateTreeExecutionFrame* CurrentParentFrame, const FStateTreeExecutionFrame& CurrentFrame);
-
 		~FCurrentlyProcessedFrameScope();
 
 	private:
@@ -978,29 +1063,21 @@ protected:
 	};
 
 	/** Currently processed nodes instance data. Ideally we would pass these to the nodes directly, but do not want to change the API currently. */
+	const FStateTreeNodeBase* CurrentNode = nullptr;
+	int32 CurrentNodeIndex = FStateTreeIndex16::InvalidValue;
 	FStateTreeDataHandle CurrentNodeDataHandle;
 	FStateTreeDataView CurrentNodeInstanceData;
 
 	/** Helper struct to set current node data. */
 	struct FNodeInstanceDataScope
 	{
-		FNodeInstanceDataScope(FStateTreeExecutionContext& InContext, const FStateTreeDataHandle InNodeDataHandle, const FStateTreeDataView InNodeInstanceData)
-			: Context(InContext)
-		{
-			SavedNodeDataHandle = Context.CurrentNodeDataHandle;
-			SavedNodeInstanceData = Context.CurrentNodeInstanceData;
-			Context.CurrentNodeDataHandle = InNodeDataHandle;
-			Context.CurrentNodeInstanceData = InNodeInstanceData;
-		}
-
-		~FNodeInstanceDataScope()
-		{
-			Context.CurrentNodeDataHandle = SavedNodeDataHandle;
-			Context.CurrentNodeInstanceData = SavedNodeInstanceData;
-		}
+		FNodeInstanceDataScope(FStateTreeExecutionContext& InContext, const FStateTreeNodeBase* InNode, const int32 InNodeIndex, const FStateTreeDataHandle InNodeDataHandle, const FStateTreeDataView InNodeInstanceData);
+		~FNodeInstanceDataScope();
 
 	private:
 		FStateTreeExecutionContext& Context;
+		const FStateTreeNodeBase* SavedNode = nullptr;
+		int32 SavedNodeIndex = FStateTreeIndex16::InvalidValue;
 		FStateTreeDataHandle SavedNodeDataHandle;
 		FStateTreeDataView SavedNodeInstanceData;
 	};
@@ -1058,8 +1135,6 @@ struct STATETREEMODULE_API FStateTreeWeakExecutionContext
 public:
 	FStateTreeWeakExecutionContext() = default;
 	explicit FStateTreeWeakExecutionContext(const FStateTreeExecutionContext& Context);
-	explicit FStateTreeWeakExecutionContext(gsl::not_null<UObject*> Owner, gsl::not_null<const UStateTree*> StateTree, FStateTreeInstanceData& Storage);
-	explicit FStateTreeWeakExecutionContext(gsl::not_null<UObject*> Owner, gsl::not_null<const UStateTree*> StateTree, TWeakPtr<FStateTreeInstanceStorage> Storage);
 
 public:
 	/** @return The owner of the context */
@@ -1094,6 +1169,13 @@ public:
 	/** Removes Delegate Listener. */
 	bool RemoveDelegateListener(FStateTreeDelegateListener Listener) const;
 
+	/**
+	 * Finishes a task.
+	 * If called during tick processing, then the state completes immediately.
+	 * If called outside of the tick processing, then the request is buffered and handled on the next tick.
+	 */
+	bool FinishTask(const FStateTreeWeakTaskRef& Task, EStateTreeFinishTaskType FinishType) const;
+
 private:
 	static void ScheduleNextTick(gsl::not_null<UObject*> Owner, gsl::not_null<const UStateTree*> RootStateTree, FStateTreeInstanceStorage& Storage);
 
@@ -1104,4 +1186,5 @@ private:
 	// This will be replace by StatePath
 	TWeakObjectPtr<const UStateTree> FrameStateTree;
 	FStateTreeStateHandle FrameRootState = FStateTreeStateHandle::Invalid;
+	FStateTreeStateHandle CurrentState = FStateTreeStateHandle::Invalid;
 };
diff --git a/Engine/Plugins/Runtime/StateTree/Source/StateTreeModule/Public/StateTreeExecutionTypes.h b/Engine/Plugins/Runtime/StateTree/Source/StateTreeModule/Public/StateTreeExecutionTypes.h
--- a/Engine/Plugins/Runtime/StateTree/Source/StateTreeModule/Public/StateTreeExecutionTypes.h
+++ b/Engine/Plugins/Runtime/StateTree/Source/StateTreeModule/Public/StateTreeExecutionTypes.h
@@ -65,6 +65,17 @@ enum class EStateTreeRunStatus : uint8
 };
 
 
+/** Status describing how a task finished. */
+enum class EStateTreeFinishTaskType : uint8
+{	
+	/** The task execution failed. */
+	Failed = (uint8)EStateTreeRunStatus::Failed,
+	
+	/** The task execution succeed. */
+	Succeeded = (uint8)EStateTreeRunStatus::Succeeded,
+};
+
+
 /** State change type. Passed to EnterState() and ExitState() to indicate how the state change affects the state and Evaluator or Task is on. */
 UENUM()
 enum class EStateTreeStateChangeType : uint8
@@ -580,6 +591,23 @@ struct STATETREEMODULE_API FStateTreeTransitionDelayedState
 	FStateTreeIndex16 TransitionIndex = FStateTreeIndex16::Invalid;
 };
 
+
+/** Describes a finished task waiting to be processed by an execution context. */
+struct FStateTreeDelayedFinishedTask
+{
+	/** Tree that identifies the active frame. */
+	FObjectKey FrameStateTree;
+	/** State that identifies the active frame. */
+	FStateTreeStateHandle FrameRootState = FStateTreeStateHandle::Invalid;
+	/** State that contains the finished task. Invalid if the task is global. */
+	FStateTreeStateHandle State = FStateTreeStateHandle::Invalid;
+	/** Task that finished and need to be processed. */
+	FStateTreeIndex16 TaskIndex = FStateTreeIndex16::Invalid;
+	/** The result of the finished task. */
+	EStateTreeFinishTaskType Type = EStateTreeFinishTaskType::Failed;
+};
+
+
 /** Describes added delegate listeners. */
 USTRUCT()
 struct STATETREEMODULE_API FStateTreeDelegateActiveListeners
@@ -738,6 +766,9 @@ PRAGMA_ENABLE_DEPRECATION_WARNINGS
 	UPROPERTY()
 	TArray<FStateTreeTransitionDelayedState> DelayedTransitions;
 
+	/** Pending finished tasks that need processing. */
+	TArray<FStateTreeDelayedFinishedTask> DelayedFinishedTasks;
+
 	/** Used by state tree random-based operations. */
 	UPROPERTY()
 	FRandomStream RandomStream;
@@ -786,18 +817,23 @@ PRAGMA_ENABLE_DEPRECATION_WARNINGS
 	UPROPERTY()
 	EStateTreeUpdatePhase CurrentPhase = EStateTreeUpdatePhase::Unset;
 
+	/** Number of times a new state has been changed. */
+	UPROPERTY()
+	uint16 StateChangeCount = 0;
+
+#if WITH_EDITORONLY_DATA
+	UE_DEPRECATED(5.6, "Use DelayedCompletedTasks for trigerring completion transitions.")
 	/** Handle of the state that was first to report state completed (success or failure), used to trigger completion transitions. */
 	UPROPERTY()
 	FStateTreeIndex16 CompletedFrameIndex = FStateTreeIndex16::Invalid; 
 	
+	UE_DEPRECATED(5.6, "Use DelayedCompletedTasks for trigerring completion transitions.")
 	UPROPERTY()
 	FStateTreeStateHandle CompletedStateHandle = FStateTreeStateHandle::Invalid;
-
-	/** Number of times a new state has been changed. */
-	UPROPERTY()
-	uint16 StateChangeCount = 0;
+#endif //WITH_EDITORONLY_DATA
 };
 
+
 /** Contains StateTree events used during State Selection for a single execution frame. */
 struct FStateTreeFrameStateSelectionEvents
 {
@@ -998,4 +1034,4 @@ private:
 
 	UPROPERTY()
 	float NextDeltaTime = 0.0f;
-};
\ No newline at end of file
+};
diff --git a/Engine/Plugins/Runtime/StateTree/Source/StateTreeModule/Public/StateTreeNodeRef.h b/Engine/Plugins/Runtime/StateTree/Source/StateTreeModule/Public/StateTreeNodeRef.h
new file mode 100644
--- /dev/null
+++ b/Engine/Plugins/Runtime/StateTree/Source/StateTreeModule/Public/StateTreeNodeRef.h
@@ -0,0 +1,68 @@
+// Copyright Epic Games, Inc. All Rights Reserved.
+
+#pragma once
+
+#include "StateTreeTypes.h"
+#include "UObject/StrongObjectPtr.h"
+#include "UObject/WeakObjectPtr.h"
+
+class UStateTree;
+struct FStateTreeTaskBase;
+
+/**
+ * A reference to a task that can utilized in a async callback. Use FStateTreeWeakTaskRef to store the reference and Pin it to get the strong version. Similar to StrongPtr and WeakPtr.
+ */
+struct FStateTreeStrongTaskRef
+{
+	explicit FStateTreeStrongTaskRef() = default;
+#if WITH_STATETREE_DEBUG
+	explicit FStateTreeStrongTaskRef(TStrongObjectPtr<const UStateTree> StateTree, const FStateTreeTaskBase* Task, FStateTreeIndex16 NodeIndex, FGuid NodeId);
+#else
+	explicit FStateTreeStrongTaskRef(TStrongObjectPtr<const UStateTree> StateTree, const FStateTreeTaskBase* Task, FStateTreeIndex16 NodeIndex);
+#endif
+
+	STATETREEMODULE_API const UStateTree* GetStateTree() const;
+	STATETREEMODULE_API const FStateTreeTaskBase* GetTask() const;
+	FStateTreeIndex16 GetTaskIndex() const
+	{
+		return NodeIndex;
+	}
+
+	operator bool() const
+	{
+		return IsValid();
+	}
+
+	STATETREEMODULE_API bool IsValid() const;
+
+private:
+	TStrongObjectPtr<const UStateTree> StateTree;
+	const FStateTreeTaskBase* Task = nullptr;
+	FStateTreeIndex16 NodeIndex = FStateTreeIndex16::Invalid;
+#if WITH_STATETREE_DEBUG
+	FGuid NodeId;
+#endif
+};
+
+
+/**
+ * A reference to a task that can be retreive. Similar to StrongPtr and WeakPtr.
+ */
+struct FStateTreeWeakTaskRef
+{
+	explicit FStateTreeWeakTaskRef() = default;
+	STATETREEMODULE_API explicit FStateTreeWeakTaskRef(TNotNull<const UStateTree*> StateTree, FStateTreeIndex16 TaskIndex);
+
+	STATETREEMODULE_API FStateTreeStrongTaskRef Pin() const;
+	void Release()
+	{
+		*this = FStateTreeWeakTaskRef();
+	}
+
+private:
+	TWeakObjectPtr<const UStateTree> StateTree;
+	FStateTreeIndex16 NodeIndex = FStateTreeIndex16::Invalid;
+#if WITH_STATETREE_DEBUG
+	FGuid NodeId;
+#endif
+};
