Fixing build issue 709934.

diff --git a/Engine/Plugins/Runtime/StateTree/Source/StateTreeModule/Private/StateTreeExecutionContext.cpp b/Engine/Plugins/Runtime/StateTree/Source/StateTreeModule/Private/StateTreeExecutionContext.cpp
--- a/Engine/Plugins/Runtime/StateTree/Source/StateTreeModule/Private/StateTreeExecutionContext.cpp
+++ b/Engine/Plugins/Runtime/StateTree/Source/StateTreeModule/Private/StateTreeExecutionContext.cpp
@@ -1212,8 +1212,8 @@ void FStateTreeExecutionContext::FinishTaskInternal(int32 ActiveFrameIndex, int3
 
 	// Add it for later. In case the user called TickTasks and will call TickTransitions later.
 	const FStateTreeIndex16 TaskIndex16 = TaskIndex >= 0 ? FStateTreeIndex16(TaskIndex) : FStateTreeIndex16::Invalid;
-	const EStateTreeFinishTaskType CompletionType = RunStatus == EStateTreeRunStatus::Succeeded ? EStateTreeFinishTaskType::Succeeded : EStateTreeFinishTaskType::Failed;
-	const FStateTreeDelayedFinishedTask& DelayedCompletedTask = Exec.DelayedFinishedTasks.Emplace_GetRef(FObjectKey(ActiveFrame.StateTree), ActiveFrame.RootState, ActiveFrame.ActiveStates[ActiveStateIndex], TaskIndex16, CompletionType);
+	const EStateTreeFinishTaskType FinishType = RunStatus == EStateTreeRunStatus::Succeeded ? EStateTreeFinishTaskType::Succeeded : EStateTreeFinishTaskType::Failed;
+	const FStateTreeDelayedFinishedTask& DelayedCompletedTask = Exec.DelayedFinishedTasks.Add_GetRef(FStateTreeDelayedFinishedTask{.FrameStateTree = FObjectKey(ActiveFrame.StateTree), .FrameRootState = ActiveFrame.RootState, .State = ActiveFrame.ActiveStates[ActiveStateIndex], .TaskIndex = TaskIndex16, .Type = FinishType});
 
 	AddUniqueCompletedStateList(ActiveFrameIndex, ActiveStateIndex, DelayedCompletedTask);
 }
@@ -1311,7 +1311,7 @@ void FStateTreeExecutionContext::AddUniqueCompletedStateList(int32 ActiveFrameIn
 	if (bAdd)
 	{
 		const EStateTreeRunStatus RunStatus = UE::StateTree::ExecutionContext::Private::FinishTaskTypeToRunStatus(DelayedFinishedTask.Type);
-		CompletedStates.Emplace(ActiveFrameIndex, ActiveStateIndex, RunStatus, DelayedFinishedTask);
+		CompletedStates.Add(FCompletedState{.ActiveFrameIndex = ActiveFrameIndex, .ActiveStateIndex = ActiveStateIndex, .RunStatus = RunStatus, .DelayedTask = DelayedFinishedTask});
 	}
 }
 
