StateTree: Deprecate gsl::not_null
#rnx

diff --git a/Engine/Plugins/Runtime/StateTree/Source/StateTreeEditorModule/Private/StateTreeEditingSubsystem.cpp b/Engine/Plugins/Runtime/StateTree/Source/StateTreeEditorModule/Private/StateTreeEditingSubsystem.cpp
--- a/Engine/Plugins/Runtime/StateTree/Source/StateTreeEditorModule/Private/StateTreeEditingSubsystem.cpp
+++ b/Engine/Plugins/Runtime/StateTree/Source/StateTreeEditorModule/Private/StateTreeEditingSubsystem.cpp
@@ -24,7 +24,7 @@ void UStateTreeEditingSubsystem::BeginDestroy()
 	Super::BeginDestroy();
 }
 
-bool UStateTreeEditingSubsystem::CompileStateTree(const gsl::not_null<UStateTree*> InStateTree, FStateTreeCompilerLog& InOutLog)
+bool UStateTreeEditingSubsystem::CompileStateTree(TNotNull<UStateTree*> InStateTree, FStateTreeCompilerLog& InOutLog)
 {
 	ValidateStateTree(InStateTree);
 	const uint32 EditorDataHash = CalculateStateTreeHash(InStateTree);
@@ -61,7 +61,7 @@ bool UStateTreeEditingSubsystem::CompileStateTree(const gsl::not_null<UStateTree
 	return bCompilationResult;
 }
 
-TSharedRef<FStateTreeViewModel> UStateTreeEditingSubsystem::FindOrAddViewModel(const gsl::not_null<UStateTree*> InStateTree)
+TSharedRef<FStateTreeViewModel> UStateTreeEditingSubsystem::FindOrAddViewModel(TNotNull<UStateTree*> InStateTree)
 {
 	const FObjectKey StateTreeKey = InStateTree;
 	TSharedPtr<FStateTreeViewModel> ViewModelPtr = StateTreeViewModels.FindRef(StateTreeKey);
@@ -121,7 +121,7 @@ TSharedRef<SWidget> UStateTreeEditingSubsystem::GetStateTreeView(TSharedRef<FSta
 	return SNew(SStateTreeView, InViewModel, TreeViewCommandList);
 }
 
-void UStateTreeEditingSubsystem::ValidateStateTree(const gsl::not_null<UStateTree*> InStateTree)
+void UStateTreeEditingSubsystem::ValidateStateTree(TNotNull<const UStateTree*> InStateTree)
 {
 	auto FixChangedStateLinkName = [](FStateTreeStateLink& StateLink, const TMap<FGuid, FName>& IDToName) -> bool
 	{
@@ -316,7 +316,7 @@ void UStateTreeEditingSubsystem::ValidateStateTree(const gsl::not_null<UStateTre
 	UpdateLinkedStateParameters(*InStateTree);
 }
 
-uint32 UStateTreeEditingSubsystem::CalculateStateTreeHash(const gsl::not_null<const UStateTree*> InStateTree)
+uint32 UStateTreeEditingSubsystem::CalculateStateTreeHash(TNotNull<const UStateTree*> InStateTree)
 {
 	uint32 EditorDataHash = 0;
 	if (InStateTree->EditorData != nullptr)
diff --git a/Engine/Plugins/Runtime/StateTree/Source/StateTreeEditorModule/Public/StateTreeEditingSubsystem.h b/Engine/Plugins/Runtime/StateTree/Source/StateTreeEditorModule/Public/StateTreeEditingSubsystem.h
--- a/Engine/Plugins/Runtime/StateTree/Source/StateTreeEditorModule/Public/StateTreeEditingSubsystem.h
+++ b/Engine/Plugins/Runtime/StateTree/Source/StateTreeEditorModule/Public/StateTreeEditingSubsystem.h
@@ -22,9 +22,9 @@ public:
 	UStateTreeEditingSubsystem();
 	virtual void BeginDestroy() override;
 	
-	TSharedRef<FStateTreeViewModel> FindOrAddViewModel(const gsl::not_null<UStateTree*> InStateTree);
+	TSharedRef<FStateTreeViewModel> FindOrAddViewModel(TNotNull<UStateTree*> InStateTree);
 	
-	static bool CompileStateTree(const gsl::not_null<UStateTree*> InStateTree,  FStateTreeCompilerLog& InOutLog);
+	static bool CompileStateTree(TNotNull<UStateTree*> InStateTree,  FStateTreeCompilerLog& InOutLog);
 	
 	/** Create a StateTreeView widget for the viewmodel. */
 	static TSharedRef<SWidget> GetStateTreeView(TSharedRef<FStateTreeViewModel> InViewModel, const TSharedRef<FUICommandList>& TreeViewCommandList);
@@ -33,10 +33,10 @@ public:
 	 * Validates and applies the schema restrictions on the StateTree.
 	 * Updates state's link, removes the unused node while validating the StateTree asset.
 	 */
-	static void ValidateStateTree(const gsl::not_null<UStateTree*> InStateTree);
+	static void ValidateStateTree(TNotNull<const UStateTree*> InStateTree);
 
 	/** Calculates editor data hash of the asset. */
-	static uint32 CalculateStateTreeHash(const gsl::not_null<const UStateTree*> InStateTree);
+	static uint32 CalculateStateTreeHash(TNotNull<const UStateTree*> InStateTree);
 	
 private:
 	void HandlePostGarbageCollect();
