StateTree: Deprecate some gsl::not_null
#rnx

diff --git a/Engine/Plugins/Runtime/GameplayStateTree/Source/GameplayStateTreeModule/Private/Components/StateTreeComponentSchema.cpp b/Engine/Plugins/Runtime/GameplayStateTree/Source/GameplayStateTreeModule/Private/Components/StateTreeComponentSchema.cpp
--- a/Engine/Plugins/Runtime/GameplayStateTree/Source/GameplayStateTreeModule/Private/Components/StateTreeComponentSchema.cpp
+++ b/Engine/Plugins/Runtime/GameplayStateTree/Source/GameplayStateTreeModule/Private/Components/StateTreeComponentSchema.cpp
@@ -162,7 +162,7 @@ bool UStateTreeComponentSchema::SetContextRequirements(UBrainComponent& BrainCom
 		return false;
 	}
 
-	FContextDataSetter ContextDataSetter = { gsl::make_not_null(&BrainComponent), Context };
+	FContextDataSetter ContextDataSetter = FContextDataSetter(&BrainComponent, Context);
 	ContextDataSetter.GetSchema()->SetContextData(ContextDataSetter, bLogErrors);
 
 	bool bResult = Context.AreContextDataViewsValid();
diff --git a/Engine/Plugins/Runtime/StateTree/Source/StateTreeModule/Private/Debugger/StateTreeDebug.cpp b/Engine/Plugins/Runtime/StateTree/Source/StateTreeModule/Private/Debugger/StateTreeDebug.cpp
--- a/Engine/Plugins/Runtime/StateTree/Source/StateTreeModule/Private/Debugger/StateTreeDebug.cpp
+++ b/Engine/Plugins/Runtime/StateTree/Source/StateTreeModule/Private/Debugger/StateTreeDebug.cpp
@@ -27,7 +27,7 @@ FNodeDelegate OnTickTask_AnyThread;
 FNodeDelegate OnTaskExitState_AnyThread;
 FEventSentDelegate OnEventSent_AnyThread;
 
-FNodeReference::FNodeReference(const gsl::not_null<const UStateTree*> InStateTree, const FStateTreeIndex16 InNodeIndex)
+FNodeReference::FNodeReference(TNotNull<const UStateTree*> InStateTree, const FStateTreeIndex16 InNodeIndex)
 	: StateTree(InStateTree)
 	, Index(InNodeIndex)
 {
@@ -148,7 +148,7 @@ void TaskExitState(const FStateTreeExecutionContext& ExecutionContext, const FNo
 	Private::NodeExit(ExecutionContext, Node, OnTaskExitState_AnyThread);
 }
 
-void EventSent(const FStateTreeMinimalExecutionContext& ExecutionContext, const gsl::not_null<const UStateTree*> StateTree, const FGameplayTag Tag, const FConstStructView Payload, const FName Origin)
+void EventSent(const FStateTreeMinimalExecutionContext& ExecutionContext, TNotNull<const UStateTree*> StateTree, const FGameplayTag Tag, const FConstStructView Payload, const FName Origin)
 {
 	OnEventSent_AnyThread.Broadcast(ExecutionContext, FEventSentDelegateArgs{ .StateTree = StateTree, .Tag = Tag, .Payload = Payload, .Origin = Origin });
 }
diff --git a/Engine/Plugins/Runtime/StateTree/Source/StateTreeModule/Private/Debugger/StateTreeRuntimeValidation.cpp b/Engine/Plugins/Runtime/StateTree/Source/StateTreeModule/Private/Debugger/StateTreeRuntimeValidation.cpp
--- a/Engine/Plugins/Runtime/StateTree/Source/StateTreeModule/Private/Debugger/StateTreeRuntimeValidation.cpp
+++ b/Engine/Plugins/Runtime/StateTree/Source/StateTreeModule/Private/Debugger/StateTreeRuntimeValidation.cpp
@@ -223,7 +223,7 @@ void FRuntimeValidationInstanceData::ValidatesTreeNodes(const UStateTree* InNewS
 
 
 #if WITH_STATETREE_DEBUG
-FRuntimeValidation::FRuntimeValidation(gsl::not_null<FRuntimeValidationInstanceData*> Instance)
+FRuntimeValidation::FRuntimeValidation(TNotNull<FRuntimeValidationInstanceData*> Instance)
 	: RuntimeValidationData(Instance)
 {
 }
diff --git a/Engine/Plugins/Runtime/StateTree/Source/StateTreeModule/Private/StateTreeExecutionContext.cpp b/Engine/Plugins/Runtime/StateTree/Source/StateTreeModule/Private/StateTreeExecutionContext.cpp
--- a/Engine/Plugins/Runtime/StateTree/Source/StateTreeModule/Private/StateTreeExecutionContext.cpp
+++ b/Engine/Plugins/Runtime/StateTree/Source/StateTreeModule/Private/StateTreeExecutionContext.cpp
@@ -5748,7 +5748,7 @@ bool FStateTreeWeakExecutionContext::FinishTask(const FStateTreeWeakTaskRef& Tas
 	return bCompleted;
 }
 
-void FStateTreeWeakExecutionContext::ScheduleNextTick(gsl::not_null<UObject*> Owner, gsl::not_null<const UStateTree*> RootStateTree, FStateTreeInstanceStorage& Storage)
+void FStateTreeWeakExecutionContext::ScheduleNextTick(TNotNull<UObject*> Owner, TNotNull<const UStateTree*> RootStateTree, FStateTreeInstanceStorage& Storage)
 {
 	TInstancedStruct<FStateTreeExecutionExtension>& ExecutionExtension = Storage.GetMutableExecutionState().ExecutionExtension;
 	if (RootStateTree->IsScheduledTickAllowed() && ExecutionExtension.IsValid())
diff --git a/Engine/Plugins/Runtime/StateTree/Source/StateTreeModule/Private/StateTreeExecutionTypes.cpp b/Engine/Plugins/Runtime/StateTree/Source/StateTreeModule/Private/StateTreeExecutionTypes.cpp
--- a/Engine/Plugins/Runtime/StateTree/Source/StateTreeModule/Private/StateTreeExecutionTypes.cpp
+++ b/Engine/Plugins/Runtime/StateTree/Source/StateTreeModule/Private/StateTreeExecutionTypes.cpp
@@ -166,8 +166,8 @@ float FStateTreeScheduledTick::GetTickRate() const
 	return NextDeltaTime;
 }
 
-FStateTreeDelegateActiveListeners::FActiveListener::FActiveListener(FStateTreeDelegateListener InListener, FSimpleDelegate InDelegate, gsl::not_null<const UStateTree*> InFrameStateTree, FStateTreeStateHandle InFrameRootState, FStateTreeStateHandle InOwningState, FStateTreeIndex16 InOwningNodeIndex)
-	: FrameStateTree(InFrameStateTree.get())
+FStateTreeDelegateActiveListeners::FActiveListener::FActiveListener(FStateTreeDelegateListener InListener, FSimpleDelegate InDelegate, TNotNull<const UStateTree*> InFrameStateTree, FStateTreeStateHandle InFrameRootState, FStateTreeStateHandle InOwningState, FStateTreeIndex16 InOwningNodeIndex)
+	: FrameStateTree(InFrameStateTree)
 	, Delegate(MoveTemp(InDelegate))
 	, DispatcherID(InListener.GetBoundDispatcherID())
 	, ListenerID(InListener.GetListenerID())
@@ -181,7 +181,7 @@ FStateTreeDelegateActiveListeners::~FStateTreeDelegateActiveListeners()
 	check(BroadcastingLockCount == 0);
 }
 
-bool FStateTreeDelegateActiveListeners::AddListener(FStateTreeDelegateListener Listener, FSimpleDelegate Delegate, gsl::not_null<const UStateTree*> FrameStateTree, FStateTreeStateHandle FrameRootState, FStateTreeStateHandle OwningState, FStateTreeIndex16 OwningNodeIndex)
+bool FStateTreeDelegateActiveListeners::AddListener(FStateTreeDelegateListener Listener, FSimpleDelegate Delegate, TNotNull<const UStateTree*> FrameStateTree, FStateTreeStateHandle FrameRootState, FStateTreeStateHandle OwningState, FStateTreeIndex16 OwningNodeIndex)
 {
 	check(Listener.IsValid());
 
@@ -197,7 +197,7 @@ bool FStateTreeDelegateActiveListeners::AddListener(FStateTreeDelegateListener L
 	return true;
 }
 
-void FStateTreeDelegateActiveListeners::RemoveListener(FStateTreeDelegateListener Listener, gsl::not_null<const UStateTree*> FrameStateTree)
+void FStateTreeDelegateActiveListeners::RemoveListener(FStateTreeDelegateListener Listener, TNotNull<const UStateTree*> FrameStateTree)
 {
 	check(Listener.IsValid());
 
@@ -232,7 +232,7 @@ void FStateTreeDelegateActiveListeners::RemoveListenersByPredicate(TFunctionRef<
 		});
 }
 
-void FStateTreeDelegateActiveListeners::BroadcastDelegate(FStateTreeDelegateDispatcher Dispatcher, gsl::not_null<const UStateTree*> FrameStateTree, const FStateTreeExecutionState& Exec)
+void FStateTreeDelegateActiveListeners::BroadcastDelegate(FStateTreeDelegateDispatcher Dispatcher, TNotNull<const UStateTree*> FrameStateTree, const FStateTreeExecutionState& Exec)
 {
 	check(Dispatcher.IsValid());
 
diff --git a/Engine/Plugins/Runtime/StateTree/Source/StateTreeModule/Private/StateTreeInstanceData.cpp b/Engine/Plugins/Runtime/StateTree/Source/StateTreeModule/Private/StateTreeInstanceData.cpp
--- a/Engine/Plugins/Runtime/StateTree/Source/StateTreeModule/Private/StateTreeInstanceData.cpp
+++ b/Engine/Plugins/Runtime/StateTree/Source/StateTreeModule/Private/StateTreeInstanceData.cpp
@@ -207,7 +207,7 @@ void FStateTreeInstanceStorage::AddTransitionRequest(const UObject* Owner, const
 	TransitionRequests.Add(Request);
 }
 
-void FStateTreeInstanceStorage::MarkDelegateAsBroadcasted(FStateTreeDelegateDispatcher Dispatcher, gsl::not_null<const UStateTree*> BelongingStateTree)
+void FStateTreeInstanceStorage::MarkDelegateAsBroadcasted(FStateTreeDelegateDispatcher Dispatcher, TNotNull<const UStateTree*> BelongingStateTree)
 {
 	BroadcastedDelegates.Add({BelongingStateTree, Dispatcher.GetDelegateID()});
 }
diff --git a/Engine/Plugins/Runtime/StateTree/Source/StateTreeModule/Public/Debugger/StateTreeDebug.h b/Engine/Plugins/Runtime/StateTree/Source/StateTreeModule/Public/Debugger/StateTreeDebug.h
--- a/Engine/Plugins/Runtime/StateTree/Source/StateTreeModule/Public/Debugger/StateTreeDebug.h
+++ b/Engine/Plugins/Runtime/StateTree/Source/StateTreeModule/Public/Debugger/StateTreeDebug.h
@@ -90,8 +90,8 @@ namespace UE::StateTree::Debug
 {
 	struct FNodeReference
 	{
-		explicit FNodeReference(gsl::not_null<const UStateTree*> InStateTree, FStateTreeIndex16 InNodeIndex);
-		gsl::not_null<const UStateTree*> StateTree;
+		explicit FNodeReference(TNotNull<const UStateTree*> InStateTree, FStateTreeIndex16 InNodeIndex);
+		TNotNull<const UStateTree*> StateTree;
 		FStateTreeIndex16 Index;
 	};
 
@@ -103,7 +103,7 @@ namespace UE::StateTree::Debug
 
 	struct FEventSentDelegateArgs
 	{
-		gsl::not_null<const UStateTree*> StateTree;
+		TNotNull<const UStateTree*> StateTree;
 		FGameplayTag Tag;
 		FConstStructView Payload;
 		FName Origin;
@@ -127,7 +127,7 @@ namespace UE::StateTree::Debug
 	void TaskTick(const FStateTreeExecutionContext& ExecutionContext, FNodeReference Node);
 	void TaskExitState(const FStateTreeExecutionContext& ExecutionContext, FNodeReference Node);
 
-	void EventSent(const FStateTreeMinimalExecutionContext& ExecutionContext, gsl::not_null<const UStateTree*> StateTree, FGameplayTag Tag, FConstStructView Payload, FName Origin);
+	void EventSent(const FStateTreeMinimalExecutionContext& ExecutionContext, TNotNull<const UStateTree*> StateTree, FGameplayTag Tag, FConstStructView Payload, FName Origin);
 
 	void EnterPhase(const FStateTreeExecutionContext& ExecutionContext, EStateTreeUpdatePhase Phase, FStateTreeStateHandle StateHandle);
 	void ExitPhase(const FStateTreeExecutionContext& ExecutionContext, EStateTreeUpdatePhase Phase, FStateTreeStateHandle StateHandle);
diff --git a/Engine/Plugins/Runtime/StateTree/Source/StateTreeModule/Public/Debugger/StateTreeRuntimeValidation.h b/Engine/Plugins/Runtime/StateTree/Source/StateTreeModule/Public/Debugger/StateTreeRuntimeValidation.h
--- a/Engine/Plugins/Runtime/StateTree/Source/StateTreeModule/Public/Debugger/StateTreeRuntimeValidation.h
+++ b/Engine/Plugins/Runtime/StateTree/Source/StateTreeModule/Public/Debugger/StateTreeRuntimeValidation.h
@@ -30,7 +30,7 @@ public:
 	FRuntimeValidation() = default;
 
 #if WITH_STATETREE_DEBUG
-	FRuntimeValidation(gsl::not_null<FRuntimeValidationInstanceData*> Instance);
+	FRuntimeValidation(TNotNull<FRuntimeValidationInstanceData*> Instance);
 	FRuntimeValidationInstanceData* GetInstanceData() const;
 #endif
 
diff --git a/Engine/Plugins/Runtime/StateTree/Source/StateTreeModule/Public/StateTreeExecutionContext.h b/Engine/Plugins/Runtime/StateTree/Source/StateTreeModule/Public/StateTreeExecutionContext.h
--- a/Engine/Plugins/Runtime/StateTree/Source/StateTreeModule/Public/StateTreeExecutionContext.h
+++ b/Engine/Plugins/Runtime/StateTree/Source/StateTreeModule/Public/StateTreeExecutionContext.h
@@ -1177,7 +1177,7 @@ public:
 	bool FinishTask(const FStateTreeWeakTaskRef& Task, EStateTreeFinishTaskType FinishType) const;
 
 private:
-	static void ScheduleNextTick(gsl::not_null<UObject*> Owner, gsl::not_null<const UStateTree*> RootStateTree, FStateTreeInstanceStorage& Storage);
+	static void ScheduleNextTick(TNotNull<UObject*> Owner, TNotNull<const UStateTree*> RootStateTree, FStateTreeInstanceStorage& Storage);
 
 	TWeakObjectPtr<UObject> Owner;
 	TWeakObjectPtr<const UStateTree> StateTree;
diff --git a/Engine/Plugins/Runtime/StateTree/Source/StateTreeModule/Public/StateTreeExecutionTypes.h b/Engine/Plugins/Runtime/StateTree/Source/StateTreeModule/Public/StateTreeExecutionTypes.h
--- a/Engine/Plugins/Runtime/StateTree/Source/StateTreeModule/Public/StateTreeExecutionTypes.h
+++ b/Engine/Plugins/Runtime/StateTree/Source/StateTreeModule/Public/StateTreeExecutionTypes.h
@@ -617,16 +617,16 @@ struct STATETREEMODULE_API FStateTreeDelegateActiveListeners
 	~FStateTreeDelegateActiveListeners();
 
 	/** Adds the listener to the list. */
-	bool AddListener(FStateTreeDelegateListener Listener, FSimpleDelegate Delegate, gsl::not_null<const UStateTree*> FrameStateTree, FStateTreeStateHandle FrameRootState, FStateTreeStateHandle OwningState, FStateTreeIndex16 OwningNodeIndex);
+	bool AddListener(FStateTreeDelegateListener Listener, FSimpleDelegate Delegate, TNotNull<const UStateTree*> FrameStateTree, FStateTreeStateHandle FrameRootState, FStateTreeStateHandle OwningState, FStateTreeIndex16 OwningNodeIndex);
 	
 	/** Removes the listener. Safe to be called during broadcasting. */
-	void RemoveListener(FStateTreeDelegateListener Listener, gsl::not_null<const UStateTree*> FrameStateTree);
+	void RemoveListener(FStateTreeDelegateListener Listener, TNotNull<const UStateTree*> FrameStateTree);
 
 	/** Removes the listener by predicate. */
 	void RemoveListenersByPredicate(TFunctionRef<bool(FStateTreeStateHandle OwningState, FStateTreeIndex16 OwningNodeIndex, const UStateTree* FrameStateTree, FStateTreeStateHandle FrameRootState)> Predicate);
 
 	/** Broadcasts the listener by predicate. */
-	void BroadcastDelegate(FStateTreeDelegateDispatcher Dispatcher, gsl::not_null<const UStateTree*> FrameStateTree, const FStateTreeExecutionState& Exec);
+	void BroadcastDelegate(FStateTreeDelegateDispatcher Dispatcher, TNotNull<const UStateTree*> FrameStateTree, const FStateTreeExecutionState& Exec);
 
 private:
 	void RemoveUnboundListeners();
@@ -634,7 +634,7 @@ private:
 	struct FActiveListener
 	{
 		FActiveListener() = default;
-		FActiveListener(FStateTreeDelegateListener Listener, FSimpleDelegate InDelegate, gsl::not_null<const UStateTree*> FrameStateTree, FStateTreeStateHandle FrameRootState, FStateTreeStateHandle OwningState, FStateTreeIndex16 OwningNodeIndex);
+		FActiveListener(FStateTreeDelegateListener Listener, FSimpleDelegate InDelegate, TNotNull<const UStateTree*> FrameStateTree, FStateTreeStateHandle FrameRootState, FStateTreeStateHandle OwningState, FStateTreeIndex16 OwningNodeIndex);
 
 		bool IsValid() const
 		{
diff --git a/Engine/Plugins/Runtime/StateTree/Source/StateTreeModule/Public/StateTreeInstanceData.h b/Engine/Plugins/Runtime/StateTree/Source/StateTreeModule/Public/StateTreeInstanceData.h
--- a/Engine/Plugins/Runtime/StateTree/Source/StateTreeModule/Public/StateTreeInstanceData.h
+++ b/Engine/Plugins/Runtime/StateTree/Source/StateTreeModule/Public/StateTreeInstanceData.h
@@ -157,7 +157,7 @@ struct STATETREEMODULE_API FStateTreeInstanceStorage
 	void AddTransitionRequest(const UObject* Owner, const FStateTreeTransitionRequest& Request);
 
 	/** Marks delegate as broadcasted. */
-	void MarkDelegateAsBroadcasted(FStateTreeDelegateDispatcher Dispatcher, gsl::not_null<const UStateTree*> BelongingStateTree);
+	void MarkDelegateAsBroadcasted(FStateTreeDelegateDispatcher Dispatcher, TNotNull<const UStateTree*> BelongingStateTree);
 
 	/** @return true if the delegate of the given index was broadcasted. */
 	bool IsDelegateBroadcasted(FStateTreeIndex16 DispatcherID, const FStateTreeExecutionFrame& BelongingFrame) const;
