[PropertyBindingsUtils]

diff --git a/Engine/Plugins/Runtime/StateTree/Source/StateTreeEditorModule/Private/Customizations/StateTreeBindingExtension.cpp b/Engine/Plugins/Runtime/StateTree/Source/StateTreeEditorModule/Private/Customizations/StateTreeBindingExtension.cpp
--- a/Engine/Plugins/Runtime/StateTree/Source/StateTreeEditorModule/Private/Customizations/StateTreeBindingExtension.cpp
+++ b/Engine/Plugins/Runtime/StateTree/Source/StateTreeEditorModule/Private/Customizations/StateTreeBindingExtension.cpp
@@ -8,23 +8,21 @@
 #include "StateTreeCompiler.h"
 #include "StateTreeEditorPropertyBindings.h"
 #include "StateTreeNodeBase.h"
-#include "Styling/AppStyle.h"
 #include "UObject/EnumProperty.h"
 #include "Widgets/Images/SImage.h"
 #include "Widgets/Layout/SBox.h"
 #include "StateTreePropertyRef.h"
 #include "StateTreePropertyRefHelpers.h"
 #include "DetailLayoutBuilder.h"
+#if UE_ENABLE_INCLUDE_ORDER_DEPRECATED_IN_5_6
 #include "IPropertyUtilities.h"
+#endif // UE_ENABLE_INCLUDE_ORDER_DEPRECATED_IN_5_6
 #include "IDetailChildrenBuilder.h"
 #include "IStructureDataProvider.h"
 #include "Kismet2/BlueprintEditorUtils.h"
-#include "PropertyBagDetails.h"
-#include "ScopedTransaction.h"
+#include "PropertyBindingExtension.h"
 #include "StateTreeEditorNodeUtils.h"
 #include "StateTreeEditorModule.h"
-#include "StateTreeNodeClassCache.h"
-#include "StateTreePropertyFunctionBase.h"
 #include "StateTreeEditorData.h"
 #include "StateTreeDelegate.h"
 
@@ -33,17 +31,14 @@
 namespace UE::StateTree::PropertyBinding
 {
 
-/** Information for the types gathered from a FStateTreePropertyRef property meta-data */
-struct FRefTypeInfo
+/**
+ * Information for the types gathered from a FStateTreePropertyRef property meta-data 
+ * Kept this type to facilitate introduction of base type FPropertyInfoOverride.
+ */
+struct FRefTypeInfo : UE::PropertyBinding::FPropertyInfoOverride
 {
-	/** Display Name Text of the Ref Type */
-	FText TypeNameText;
-
-	/** Ref Type expressed as a Pin Type */
-	FEdGraphPinType PinType;
 };
 
-const FName StateTreeNodeIDName(TEXT("StateTreeNodeID"));
 const FName AllowAnyBindingName(TEXT("AllowAnyBinding"));
 
 UObject* FindEditorBindingsOwner(UObject* InObject)
@@ -62,244 +57,6 @@ UObject* FindEditorBindingsOwner(UObject* InObject)
 	return Result;
 }
 
-UStruct* ResolveLeafValueStructType(FStateTreeDataView ValueView, TConstArrayView<FBindingChainElement> InBindingChain)
-{
-	if (ValueView.GetMemory() == nullptr)
-	{
-		return nullptr;
-	}
-	
-	FStateTreePropertyPath Path;
-
-	for (const FBindingChainElement& Element : InBindingChain)
-	{
-		if (const FProperty* Property = Element.Field.Get<FProperty>())
-		{
-			Path.AddPathSegment(Property->GetFName(), Element.ArrayIndex);
-		}
-		else if (const UFunction* Function = Element.Field.Get<UFunction>())
-		{
-			// Cannot handle function calls
-			return nullptr;
-		}
-	}
-
-	TArray<FPropertyBindingPathIndirection> Indirections;
-	if (!Path.ResolveIndirectionsWithValue(ValueView, Indirections)
-		|| Indirections.IsEmpty())
-	{
-		return nullptr;
-	}
-
-	// Last indirection points to the value of the leaf property, check the type.
-	const FPropertyBindingPathIndirection& LastIndirection = Indirections.Last();
-
-	UStruct* Result = nullptr;
-
-	if (LastIndirection.GetContainerAddress())
-	{
-		if (const FStructProperty* StructProperty = CastField<FStructProperty>(LastIndirection.GetProperty()))
-		{
-			// Get the type of the instanced struct's value.
-			if (StructProperty->Struct == TBaseStructure<FInstancedStruct>::Get())
-			{
-				const FInstancedStruct& InstancedStruct = *reinterpret_cast<const FInstancedStruct*>(LastIndirection.GetPropertyAddress());
-				Result = const_cast<UScriptStruct*>(InstancedStruct.GetScriptStruct());
-			}
-		}
-		else if (const FObjectProperty* ObjectProperty = CastField<FObjectProperty>(LastIndirection.GetProperty()))
-		{
-			// Get type of the instanced object.
-			if (const UObject* Object = *reinterpret_cast<UObject* const*>(LastIndirection.GetPropertyAddress()))
-			{
-				Result = Object->GetClass();
-			}
-		}
-	}
-
-	return Result;
-}
-
-void MakeStructPropertyPathFromBindingChain(const FGuid StructID, TConstArrayView<FBindingChainElement> InBindingChain, FStateTreeDataView DataView, FStateTreePropertyPath& OutPath)
-{
-	OutPath.Reset();
-	OutPath.SetStructID(StructID);
-	
-	for (const FBindingChainElement& Element : InBindingChain)
-	{
-		if (const FProperty* Property = Element.Field.Get<FProperty>())
-		{
-			OutPath.AddPathSegment(Property->GetFName(), Element.ArrayIndex);
-		}
-		else if (const UFunction* Function = Element.Field.Get<UFunction>())
-		{
-			OutPath.AddPathSegment(Function->GetFName());
-		}
-	}
-
-	OutPath.UpdateSegmentsFromValue(DataView);
-}
-
-EStateTreePropertyUsage MakeStructPropertyPathFromPropertyHandle(TSharedPtr<const IPropertyHandle> InPropertyHandle, FStateTreePropertyPath& OutPath)
-{
-	OutPath.Reset();
-
-	FGuid StructID;
-	TArray<FPropertyBindingPathSegment> PathSegments;
-	EStateTreePropertyUsage ResultUsage = EStateTreePropertyUsage::Invalid; 
-
-	TSharedPtr<const IPropertyHandle> CurrentPropertyHandle = InPropertyHandle;
-	while (CurrentPropertyHandle.IsValid())
-	{
-		const FProperty* Property = CurrentPropertyHandle->GetProperty();
-		if (Property)
-		{
-			FPropertyBindingPathSegment& Segment = PathSegments.InsertDefaulted_GetRef(0); // Traversing from leaf to root, insert in reverse.
-
-			// Store path up to the property which has ID.
-			Segment.SetName(Property->GetFName());
-			Segment.SetArrayIndex(CurrentPropertyHandle->GetIndexInArray());
-
-			// Store type of the object (e.g. for instanced objects or instanced structs).
-			if (const FObjectProperty* ObjectProperty = CastField<FObjectProperty>(Property))
-			{
-				if (ObjectProperty->HasAnyPropertyFlags(CPF_PersistentInstance | CPF_InstancedReference))
-				{
-					const UObject* Object = nullptr;
-					if (CurrentPropertyHandle->GetValue(Object) == FPropertyAccess::Success)
-					{
-						if (Object)
-						{
-							Segment.SetInstanceStruct(Object->GetClass());
-						}
-					}
-				}
-			}
-			else if (const FStructProperty* StructProperty = CastField<FStructProperty>(Property))
-			{
-				if (StructProperty->Struct == TBaseStructure<FInstancedStruct>::Get())
-				{
-					void* Address = nullptr;
-					if (CurrentPropertyHandle->GetValueData(Address) == FPropertyAccess::Success)
-					{
-						if (Address)
-						{
-							FInstancedStruct& Struct = *static_cast<FInstancedStruct*>(Address);
-							Segment.SetInstanceStruct(Struct.GetScriptStruct());
-						}
-					}
-				}
-			}
-
-			// Array access is represented as: "Array, PropertyInArray[Index]", we're traversing from leaf to root, skip the node without index.
-			// Advancing the node before ID test, since the array is on the instance data, the ID will be on the Array node.
-			if (Segment.GetArrayIndex() != INDEX_NONE)
-			{
-				TSharedPtr<const IPropertyHandle> ParentPropertyHandle = CurrentPropertyHandle->GetParentHandle();
-				if (ParentPropertyHandle.IsValid())
-				{
-					const FProperty* ParentProperty = ParentPropertyHandle->GetProperty();
-					if (ParentProperty
-						&& ParentProperty->IsA<FArrayProperty>()
-						&& Property->GetFName() == ParentProperty->GetFName())
-					{
-						CurrentPropertyHandle = ParentPropertyHandle;
-					}
-				}
-			}
-
-			// Bindable property must have node ID
-			if (const FString* IDString = CurrentPropertyHandle->GetInstanceMetaData(UE::StateTree::PropertyBinding::StateTreeNodeIDName))
-			{
-				LexFromString(StructID, **IDString);
-				ResultUsage = UE::StateTree::GetUsageFromMetaData(Property);
-				break;
-			}
-		}
-		
-		CurrentPropertyHandle = CurrentPropertyHandle->GetParentHandle();
-	}
-
-	if (!StructID.IsValid())
-	{
-		ResultUsage = EStateTreePropertyUsage::Invalid;
-	}
-	else
-	{
-		OutPath = FStateTreePropertyPath(StructID, PathSegments);
-	}
-
-	return ResultUsage;
-}
-
-// @todo: there's a similar function in StateTreeNodeDetails.cpp, merge.
-FText GetPropertyTypeText(const FProperty* Property)
-{
-	FEdGraphPinType PinType;
-	const UEdGraphSchema_K2* Schema = GetDefault<UEdGraphSchema_K2>();
-	Schema->ConvertPropertyToPinType(Property, PinType);
-				
-	const FName PinSubCategory = PinType.PinSubCategory;
-	const UObject* PinSubCategoryObject = PinType.PinSubCategoryObject.Get();
-	if (PinSubCategory != UEdGraphSchema_K2::PSC_Bitmask && PinSubCategoryObject)
-	{
-		if (const UField* Field = Cast<const UField>(PinSubCategoryObject))
-		{
-			return Field->GetDisplayNameText();
-		}
-		return FText::FromString(PinSubCategoryObject->GetName());
-	}
-
-	return UEdGraphSchema_K2::GetCategoryText(PinType.PinCategory, NAME_None, true);
-}
-
-TSharedRef<SWidget> MakeContextStructWidget(const FStateTreeBindableStructDesc& InContextStruct)
-{
-	FEdGraphPinType PinType;
-
-	UStruct* Struct = const_cast<UStruct*>(InContextStruct.Struct.Get());
-
-	if (UClass* Class = Cast<UClass>(Struct))
-	{
-		PinType.PinCategory = UEdGraphSchema_K2::PC_Object;
-		PinType.PinSubCategory = NAME_None;
-		PinType.PinSubCategoryObject = Class;
-	}
-	else if (UScriptStruct* ScriptStruct = Cast<UScriptStruct>(Struct))
-	{
-		PinType.PinCategory = UEdGraphSchema_K2::PC_Struct;
-		PinType.PinSubCategory = NAME_None;
-		PinType.PinSubCategoryObject = ScriptStruct;
-	}
-
-	const FSlateBrush* Icon = FBlueprintEditorUtils::GetIconFromPin(PinType, true);
-	const FLinearColor IconColor = GetDefault<UEdGraphSchema_K2>()->GetPinTypeColor(PinType);
-
-	return SNew(SHorizontalBox)
-		+ SHorizontalBox::Slot()
-		.AutoWidth()
-		[
-			SNew(SSpacer)
-			.Size(FVector2D(18.0f, 0.0f))
-		]
-		+ SHorizontalBox::Slot()
-		.AutoWidth()
-		.VAlign(VAlign_Center)
-		.Padding(1.0f, 0.0f)
-		[
-			SNew(SImage)
-			.Image(Icon)
-			.ColorAndOpacity(IconColor)
-		]
-		+ SHorizontalBox::Slot()
-		.AutoWidth()
-		.VAlign(VAlign_Center)
-		.Padding(4.0f, 0.0f)
-		[
-			SNew(STextBlock)
-			.Text(FText::FromName(InContextStruct.Name))
-		];
-}
 
 bool IsDelegateListenerProperty(const FProperty& Property)
 {
@@ -311,254 +68,143 @@ bool IsDelegateListenerProperty(const FProperty& Property)
 	return false;
 }
 
-TSharedRef<SWidget> MakeBindingPropertyInfoWidget(const FText& InDisplayText, const FEdGraphPinType& InPinType)
-{
-	const FSlateBrush* Icon = FBlueprintEditorUtils::GetIconFromPin(InPinType, /*bIsLarge*/true);
-	const FLinearColor IconColor = GetDefault<UEdGraphSchema_K2>()->GetPinTypeColor(InPinType);
-
-	return SNew(SHorizontalBox)
-		+SHorizontalBox::Slot()
-		.AutoWidth()
-		[
-			SNew(SSpacer)
-			.Size(FVector2D(18.0f, 0.0f))
-		]
-		+SHorizontalBox::Slot()
-		.AutoWidth()
-		.VAlign(VAlign_Center)
-		.Padding(1.0f, 0.0f)
-		[
-			SNew(SImage)
-			.Image(Icon)
-			.ColorAndOpacity(IconColor)
-		]
-		+SHorizontalBox::Slot()
-		.AutoWidth()
-		.VAlign(VAlign_Center)
-		.Padding(4.0f, 0.0f)
-		[
-			SNew(STextBlock)
-			.Text(InDisplayText)
-		];
-}
-	
-/** Helper struct to Begin/End Sections */
-struct FSectionHelper
-{
-	FSectionHelper(FMenuBuilder& MenuBuilder)
-		: MenuBuilder(MenuBuilder)
-	{
-	}
-	~FSectionHelper()
-	{
-		if (bSectionOpened)
-		{
-			MenuBuilder.EndSection();
-		}
-	}
-	void SetSection(const FText& InSection)
-	{
-		if (!InSection.IdenticalTo(CurrentSection))
-		{
-			if (bSectionOpened)
-			{
-				MenuBuilder.EndSection();
-			}
-			CurrentSection = InSection;
-			MenuBuilder.BeginSection(NAME_None, CurrentSection);
-			bSectionOpened = true;
-		}
-	}
-private:
-	FText CurrentSection;
-	FMenuBuilder& MenuBuilder;
-	bool bSectionOpened = false;
-};
 
 FOnStateTreePropertyBindingChanged STATETREEEDITORMODULE_API OnStateTreePropertyBindingChanged;
 
-struct FCachedBindingData : public TSharedFromThis<FCachedBindingData>
+struct FStateTreeCachedBindingData : UE::PropertyBinding::FCachedBindingData, TSharedFromThis<FStateTreeCachedBindingData>
 {
-	FCachedBindingData(UObject* InOwnerObject, const FStateTreePropertyPath& InTargetPath, const TSharedPtr<const IPropertyHandle>& InPropertyHandle, TArrayView<FStateTreeBindableStructDesc> InAccessibleStructs)
-		: WeakOwnerObject(InOwnerObject)
-		, TargetPath(InTargetPath)
-		, PropertyHandle(InPropertyHandle)
-		, AccessibleStructs(InAccessibleStructs)
-	{
-	}
-
-	void AddBinding(TConstArrayView<FBindingChainElement> InBindingChain)
-	{
-		if (InBindingChain.IsEmpty())
-		{
-			return;
-		}
-		
-		if (!TargetPath.GetStructID().IsValid())
-		{
-			return;
-		}
-
-		UObject* OwnerObject = WeakOwnerObject.Get();
-		if (!OwnerObject)
-		{
-			return;
-		}
-		
-		IStateTreeEditorPropertyBindingsOwner* BindingOwner = Cast<IStateTreeEditorPropertyBindingsOwner>(OwnerObject);
-		if (!BindingOwner)
-		{
-			return;
-		}
-
-		FStateTreeEditorPropertyBindings* EditorBindings = BindingOwner->GetPropertyEditorBindings();
-		if (!EditorBindings)
-		{
-			return;
-		}
+	using Super = FCachedBindingData;
 
-		// First item in the binding chain is the index in AccessibleStructs.
-		const int32 SourceStructIndex = InBindingChain[0].ArrayIndex;
-		check(SourceStructIndex >= 0 && SourceStructIndex < AccessibleStructs.Num());
-
-		const FStateTreeBindableStructDesc& BindableStruct = AccessibleStructs[SourceStructIndex];
-				
-		TConstArrayView<FBindingChainElement> SourceBindingChain(InBindingChain.begin() + 1, InBindingChain.Num() - 1); // remove struct index.
-
-		FStateTreeDataView DataView;
-		BindingOwner->GetDataViewByID(BindableStruct.ID, DataView);
+	using FCachedBindingData::FCachedBindingData;
 
-		// If SourceBindingChain is empty at this stage, it means that the binding points to the source struct itself.
-		FStateTreePropertyPath SourcePath;
-		UE::StateTree::PropertyBinding::MakeStructPropertyPathFromBindingChain(BindableStruct.ID, SourceBindingChain, DataView, SourcePath);
-
-		OwnerObject->Modify();
-
-		if (BindableStruct.DataSource == EStateTreeBindableStructSource::PropertyFunction)
+	virtual bool AddBindingInternal(TConstStructView<FPropertyBindingBindableStructDescriptor> InDescriptor,
+		FPropertyBindingPath& InOutSourcePath,
+		const FPropertyBindingPath& InTargetPath) override
+	{
+		const FStateTreeBindableStructDesc& SourceDesc = InDescriptor.Get<FStateTreeBindableStructDesc>();
+		if (SourceDesc.DataSource == EStateTreeBindableStructSource::PropertyFunction)
 		{
+			IStateTreeEditorPropertyBindingsOwner* BindingsOwner = Cast<IStateTreeEditorPropertyBindingsOwner>(GetOwner());
 			const UScriptStruct* PropertyFunctionNodeStruct = nullptr;
 
-			BindingOwner->EnumerateBindablePropertyFunctionNodes([ID = BindableStruct.ID, &PropertyFunctionNodeStruct](const UScriptStruct* NodeStruct, const FStateTreeBindableStructDesc& Desc, const FStateTreeDataView Value)
-			{
-				if (Desc.ID == ID)
+			BindingsOwner->EnumerateBindablePropertyFunctionNodes([ID = SourceDesc.ID, &PropertyFunctionNodeStruct](const UScriptStruct* NodeStruct, const FStateTreeBindableStructDesc& Desc, const FPropertyBindingDataView Value)
 				{
-					PropertyFunctionNodeStruct = NodeStruct;
-					return EStateTreeVisitor::Break;
-				}
+					if (Desc.ID == ID)
+					{
+						PropertyFunctionNodeStruct = NodeStruct;
+						return EStateTreeVisitor::Break;
+					}
 
-				return EStateTreeVisitor::Continue;
-			});
+					return EStateTreeVisitor::Continue;
+				});
 
 			if (ensure(PropertyFunctionNodeStruct))
 			{
+				FStateTreeEditorPropertyBindings* EditorBindings = BindingsOwner->GetPropertyEditorBindings();
+
 				// If there are no segments, bindings leads directly into source struct's single output property. It's path has to be recovered.
-				if (SourcePath.NumSegments() == 0)
+				if (InOutSourcePath.NumSegments() == 0)
 				{
-					const FProperty* SingleOutputProperty = UE::StateTree::GetStructSingleOutputProperty(*BindableStruct.Struct);
+					const FProperty* SingleOutputProperty = UE::StateTree::GetStructSingleOutputProperty(*SourceDesc.Struct);
 					check(SingleOutputProperty);
 
 					FPropertyBindingPathSegment SingleOutputPropertySegment = FPropertyBindingPathSegment(SingleOutputProperty->GetFName());
-					SourcePath = EditorBindings->AddFunctionPropertyBinding(PropertyFunctionNodeStruct, {SingleOutputPropertySegment}, TargetPath);
+					InOutSourcePath = EditorBindings->AddFunctionBinding(PropertyFunctionNodeStruct, { SingleOutputPropertySegment }, InTargetPath);
 				}
 				else
 				{
-					SourcePath = EditorBindings->AddFunctionPropertyBinding(PropertyFunctionNodeStruct, SourcePath.GetSegments(), TargetPath);
+					InOutSourcePath = EditorBindings->AddFunctionBinding(PropertyFunctionNodeStruct, InOutSourcePath.GetSegments(), InTargetPath);
 				}
+
+				return true;
 			}
 		}
-		else
-		{
-			EditorBindings->AddPropertyBinding(SourcePath, TargetPath);
-		}
-
-		UpdateData();
-
-		UE::StateTree::PropertyBinding::OnStateTreePropertyBindingChanged.Broadcast(SourcePath, TargetPath);
+		return false;
 	}
 
-	bool HasBinding(FStateTreeEditorPropertyBindings::ESearchMode SearchMode) const
+	virtual void UpdatePropertyReferenceTooltip(const FProperty& InProperty, FTextBuilder& InOutTextBuilder) const override
 	{
-		UObject* OwnerObject = WeakOwnerObject.Get();
-		if (!OwnerObject)
+		if (InProperty.HasMetaData(PropertyRefHelpers::IsRefToArrayName))
 		{
-			return false;
+			InOutTextBuilder.AppendLineFormat(LOCTEXT("PropertyRefBindingTooltipArray", "Supported types are Array of {0}"),
+				FText::FromString(InProperty.GetMetaData(PropertyRefHelpers::RefTypeName)));
 		}
-
-		IStateTreeEditorPropertyBindingsOwner* BindingOwner = Cast<IStateTreeEditorPropertyBindingsOwner>(OwnerObject);
-		if (!BindingOwner)
+		else
 		{
-			return false;
+			InOutTextBuilder.AppendLineFormat(LOCTEXT("PropertyRefBindingTooltip", "Supported types are {0}"),
+				FText::FromString(InProperty.GetMetaData(PropertyRefHelpers::RefTypeName)));
+			if (InProperty.HasMetaData(PropertyRefHelpers::CanRefToArrayName))
+			{
+				InOutTextBuilder.AppendLine(LOCTEXT("PropertyRefBindingTooltipCanSupportArray", "Supports Arrays"));
+			}
 		}
+	}
 
-		FStateTreeEditorPropertyBindings* EditorBindings = BindingOwner->GetPropertyEditorBindings();
-		if (!EditorBindings)
+	virtual void UpdateSourcePropertyPath(const TConstStructView<FPropertyBindingBindableStructDescriptor> InDescriptor, const FPropertyBindingPath& InSourcePath, FString& OutString) override
+	{
+		const FStateTreeBindableStructDesc& SourceDesc = InDescriptor.Get<FStateTreeBindableStructDesc>();
+		// Making first segment of the path invisible for the user if it's property function's single output property.
+		if (SourceDesc.DataSource == EStateTreeBindableStructSource::PropertyFunction && GetStructSingleOutputProperty(*SourceDesc.Struct))
 		{
-			return false;
+			OutString = InSourcePath.ToString(/*HighlightedSegment*/ INDEX_NONE, /*HighlightPrefix*/ nullptr, /*HighlightPostfix*/ nullptr, /*bOutputInstances*/ false, 1);
 		}
-
-		return EditorBindings && EditorBindings->HasPropertyBinding(TargetPath, SearchMode);
 	}
 
-	void RemoveBinding(FStateTreeEditorPropertyBindings::ESearchMode RemoveMode)
+	virtual void GetSourceDataViewForNewBinding(TNotNull<IPropertyBindingBindingCollectionOwner*> InBindingsOwner, TConstStructView<FPropertyBindingBindableStructDescriptor> InDescriptor, FPropertyBindingDataView& OutSourceDataView) override
 	{
-		UObject* OwnerObject = WeakOwnerObject.Get();
-		if (!OwnerObject)
-		{
-			return;
-		}
-		
-		IStateTreeEditorPropertyBindingsOwner* BindingOwner = Cast<IStateTreeEditorPropertyBindingsOwner>(OwnerObject);
-		if (!BindingOwner)
+		const FStateTreeBindableStructDesc& SourceDesc = InDescriptor.Get<FStateTreeBindableStructDesc>();
+		if (SourceDesc.DataSource == EStateTreeBindableStructSource::PropertyFunction)
 		{
-			return;
+			OutSourceDataView = FPropertyBindingDataView(SourceDesc.Struct, nullptr);
 		}
-
-		FStateTreeEditorPropertyBindings* EditorBindings = BindingOwner->GetPropertyEditorBindings();
-		if (!EditorBindings)
+		else
 		{
-			return;
+			Super::GetSourceDataViewForNewBinding(InBindingsOwner, InDescriptor, OutSourceDataView);
 		}
-
-		OwnerObject->Modify();
-		EditorBindings->RemovePropertyBindings(TargetPath, RemoveMode);
-
-		UpdateData();
-		
-		const FStateTreePropertyPath SourcePath; // Null path
-		UE::StateTree::PropertyBinding::OnStateTreePropertyBindingChanged.Broadcast(SourcePath, TargetPath);
 	}
 
-	bool CanCreateParameter(const FStateTreeBindableStructDesc& InStructDesc, TArray<TSharedPtr<const FRefTypeInfo>>& OutRefTypeInfos) const
+	virtual bool GetPinTypeAndIconForProperty(const FProperty& InProperty, FPropertyBindingDataView InTargetDataView, FEdGraphPinType& OutPinType, FName& OutIconName) const override
 	{
-		const FProperty* Property = PropertyHandle->GetProperty();
-		if (!Property)
+		const bool bIsPropertyRef = PropertyRefHelpers::IsPropertyRef(InProperty);;
+		if (bIsPropertyRef && InTargetDataView.IsValid())
 		{
-			return false;
+			// Use internal type to construct PinType if it's property of PropertyRef type.
+			TArray<FPropertyBindingPathIndirection> TargetIndirections;
+			if (ensure(GetTargetPath().ResolveIndirectionsWithValue(InTargetDataView, TargetIndirections)))
+			{
+				const void* PropertyRef = TargetIndirections.Last().GetPropertyAddress();
+				OutPinType = PropertyRefHelpers::GetPropertyRefInternalTypeAsPin(InProperty, PropertyRef);
+			}
+			OutIconName = TEXT("Kismet.Tabs.Variables");
+			return true;
 		}
 
-		IStateTreeEditorPropertyBindingsOwner* BindingOwner = Cast<IStateTreeEditorPropertyBindingsOwner>(WeakOwnerObject.Get());
-		if (!BindingOwner)
+		if (IsDelegateListenerProperty(InProperty))
 		{
-			return false;
+			OutPinType.PinCategory = UEdGraphSchema_K2::PC_Delegate;
+			OutIconName = TEXT("Icons.Event");
+			return true;
 		}
 
-		if (!BindingOwner->CanCreateParameter(InStructDesc.ID))
-		{
-			return false;
-		}
+		return false;
+	}
 
+	virtual bool IsPropertyReference(const FProperty& InProperty) override
+	{
+		return PropertyRefHelpers::IsPropertyRef(InProperty);
+	}
+
+	virtual void AddPropertyInfoOverride(const FProperty& Property, TArray<TSharedPtr<const UE::PropertyBinding::FPropertyInfoOverride>>& OutPropertyInfoOverrides) const override
+	{
 		// Add the PropertyRef property type with its RefTypes
-		const FStructProperty* StructProperty = CastField<const FStructProperty>(Property);
+		const FStructProperty* StructProperty = CastField<const FStructProperty>(&Property);
 		if (StructProperty && StructProperty->Struct && StructProperty->Struct->IsChildOf(FStateTreePropertyRef::StaticStruct()))
 		{
 			TArray<FEdGraphPinType, TInlineAllocator<1>> PinTypes;
-
-			const bool bCanTargetRefArray = PropertyHandle->HasMetaData(PropertyRefHelpers::CanRefToArrayName);
-
 			if (StructProperty->Struct->IsChildOf(FStateTreeBlueprintPropertyRef::StaticStruct()))
 			{
 				void* PropertyRefAddress = nullptr;
-				if (PropertyHandle->GetValueData(PropertyRefAddress) == FPropertyAccess::Result::Success)
+				if (GetPropertyHandle()->GetValueData(PropertyRefAddress) == FPropertyAccess::Result::Success)
 				{
 					check(PropertyRefAddress);
 					PinTypes.Add(PropertyRefHelpers::GetBlueprintPropertyRefInternalTypeAsPin(*static_cast<const FStateTreeBlueprintPropertyRef*>(PropertyRefAddress)));
@@ -566,11 +212,11 @@ struct FCachedBindingData : public TSharedFromThis<FCachedBindingData>
 			}
 			else
 			{
-				PinTypes = PropertyRefHelpers::GetPropertyRefInternalTypesAsPins(*Property);
+				PinTypes = PropertyRefHelpers::GetPropertyRefInternalTypesAsPins(Property);
 			}
 
 			// If Property supports Arrays, add the Array version of these pin types
-			if (PropertyHandle->HasMetaData(PropertyRefHelpers::CanRefToArrayName))
+			if (GetPropertyHandle()->HasMetaData(PropertyRefHelpers::CanRefToArrayName))
 			{
 				const int32 PinTypeNum = PinTypes.Num();
 				for (int32 Index = 0; Index < PinTypeNum; ++Index)
@@ -584,15 +230,13 @@ struct FCachedBindingData : public TSharedFromThis<FCachedBindingData>
 				}
 			}
 
-			const UEdGraphSchema_K2* Schema = GetDefault<UEdGraphSchema_K2>();
-
 			for (const FEdGraphPinType& PinType : PinTypes)
 			{
 				TSharedRef<FRefTypeInfo> RefTypeInfo = MakeShared<FRefTypeInfo>();
 				RefTypeInfo->PinType = PinType;
 
 				FString TypeName;
-				if(UObject* SubCategoryObject = PinType.PinSubCategoryObject.Get()) 
+				if (UObject* SubCategoryObject = PinType.PinSubCategoryObject.Get())
 				{
 					TypeName = SubCategoryObject->GetName();
 				}
@@ -602,352 +246,20 @@ struct FCachedBindingData : public TSharedFromThis<FCachedBindingData>
 				}
 
 				RefTypeInfo->TypeNameText = FText::FromString(TypeName);
-				OutRefTypeInfos.Emplace(MoveTemp(RefTypeInfo));
-			}
-		}
-
-		return true;
-	}
-
-	void PromoteToParameter(FName InPropertyName, FStateTreeBindableStructDesc InStructDesc, TSharedPtr<const FRefTypeInfo> InPropertyInfoOverride)
-	{
-		if (!TargetPath.GetStructID().IsValid())
-		{
-			return;
-		}
-		
-		UObject* OwnerObject = WeakOwnerObject.Get();
-		if (!OwnerObject)
-		{
-			return;
-		}
-
-		IStateTreeEditorPropertyBindingsOwner* BindingOwner = Cast<IStateTreeEditorPropertyBindingsOwner>(OwnerObject);
-		if (!BindingOwner)
-		{
-			return;
-		}
-
-		const FProperty* Property = PropertyHandle->GetProperty();
-		if (!Property)
-		{
-			return;
-		}
-
-		const FProperty* TargetProperty = nullptr;
-		const void* TargetContainerAddress = nullptr;
-
-		FStateTreeDataView TargetDataView;
-		if (BindingOwner->GetDataViewByID(TargetPath.GetStructID(), TargetDataView) && TargetDataView.IsValid())
-		{
-			TArray<FPropertyBindingPathIndirection> TargetIndirections;
-			if (ensure(TargetPath.ResolveIndirectionsWithValue(TargetDataView, TargetIndirections)))
-			{
-				const FPropertyBindingPathIndirection& LastIndirection = TargetIndirections.Last();
-				TargetProperty = LastIndirection.GetProperty();
-				TargetContainerAddress = LastIndirection.GetContainerAddress();
+				OutPropertyInfoOverrides.Emplace(MoveTemp(RefTypeInfo));
 			}
 		}
-
-		FStateTreeEditorPropertyBindings* EditorBindings = BindingOwner->GetPropertyEditorBindings();
-		if (!EditorBindings)
-		{
-			return;
-		}
-
-		const FGuid StructID = InStructDesc.ID;
-
-		const FScopedTransaction Transaction(LOCTEXT("PromoteToParameter", "Promote to Parameter"));
-
-		TArray<FStateTreeEditorPropertyCreationDesc, TFixedAllocator<1>> PropertyCreationDescs;
-		{
-			FStateTreeEditorPropertyCreationDesc& PropertyCreationDesc = PropertyCreationDescs.AddDefaulted_GetRef();
-
-			if (InPropertyInfoOverride)
-			{
-				PropertyCreationDesc.PropertyDesc.Name = InPropertyName;
-				UE::StructUtils::SetPropertyDescFromPin(PropertyCreationDesc.PropertyDesc, InPropertyInfoOverride->PinType);
-			}
-			else
-			{
-				PropertyCreationDesc.PropertyDesc = FPropertyBagPropertyDesc(InPropertyName, Property);
-			}
-
-			// Create desc based on the Target Property, but without the meta-data.
-			// This functionality mirrors the user action of adding a new property from the UI, where meta-data is not available.
-			// Additionally, meta-data like EditCondition is not desirable here
-			PropertyCreationDesc.PropertyDesc.MetaClass = nullptr;
-			PropertyCreationDesc.PropertyDesc.MetaData.Reset();
-
-			// Set the Property & Container Address to copy
-			if (TargetProperty && TargetContainerAddress)
-			{
-				PropertyCreationDesc.SourceProperty = TargetProperty;
-				PropertyCreationDesc.SourceContainerAddress = TargetContainerAddress;
-			}
-		}
-
-		OwnerObject->Modify();
-		BindingOwner->CreateParameters(StructID, /*InOut*/PropertyCreationDescs);
-
-		// Use the name in PropertyDescs, as it might contain a different name than the desired InPropertyName (for uniquness)
-		FStateTreePropertyPath SourcePath(StructID, PropertyCreationDescs[0].PropertyDesc.Name);
-		EditorBindings->AddPropertyBinding(SourcePath, TargetPath);
-
-		UpdateData();
-		UE::StateTree::PropertyBinding::OnStateTreePropertyBindingChanged.Broadcast(SourcePath, TargetPath);
 	}
 
-	void UpdateData()
+	virtual bool CanBindToContextStructInternal(const UStruct* InStruct, const int32 InStructIndex) override
 	{
-		SourceStructName = FText::GetEmpty();
-		FormatableText = FText::GetEmpty();
-		FormatableTooltipText = FText::GetEmpty();
-		Color = FLinearColor::White;
-		Image = nullptr;
-
-		if (!PropertyHandle.IsValid())
-		{
-			return;
-		}
-
-		const FProperty* Property = PropertyHandle->GetProperty();
-		if (!Property)
-		{
-			return;
-		}
-
-		UObject* OwnerObject = WeakOwnerObject.Get();
-		if (!OwnerObject)
-		{
-			return;
-		}
-		
-		IStateTreeEditorPropertyBindingsOwner* BindingOwner = Cast<IStateTreeEditorPropertyBindingsOwner>(OwnerObject);
-		if (!BindingOwner)
-		{
-			return;
-		}
-
-		FStateTreeEditorPropertyBindings* EditorBindings = BindingOwner->GetPropertyEditorBindings();
-		if (!EditorBindings)
-		{
-			return;
-		}
-
-		const UEdGraphSchema_K2* Schema = GetDefault<UEdGraphSchema_K2>();
-		check(Schema);
-
-		FStateTreeDataView TargetDataView;
-		BindingOwner->GetDataViewByID(TargetPath.GetStructID(), TargetDataView);
-
-		FEdGraphPinType PinType;
-		FName IconName;
-		const bool bIsPropertyRef = UE::StateTree::PropertyRefHelpers::IsPropertyRef(*Property);
-		if (bIsPropertyRef && TargetDataView.IsValid())
-		{
-			// Use internal type to construct PinType if it's property of PropertyRef type.
-			TArray<FPropertyBindingPathIndirection> TargetIndirections;
-			if (ensure(TargetPath.ResolveIndirectionsWithValue(TargetDataView, TargetIndirections)))
-			{
-				const void* PropertyRef = TargetIndirections.Last().GetPropertyAddress();
-				PinType = UE::StateTree::PropertyRefHelpers::GetPropertyRefInternalTypeAsPin(*Property, PropertyRef);
-			}
-			IconName = TEXT("Kismet.Tabs.Variables");
-		}
-		else if (UE::StateTree::PropertyBinding::IsDelegateListenerProperty(*Property))
-		{
-			PinType.PinCategory = UEdGraphSchema_K2::PC_Delegate;
-			IconName = TEXT("Icons.Event");
-		}
-		else
-		{
-			Schema->ConvertPropertyToPinType(Property, PinType);
-			IconName = TEXT("Kismet.Tabs.Variables");
-		}
-
-
-		FTextBuilder TooltipBuilder;
-		const FStateTreePropertyPathBinding* CurrentBinding = EditorBindings->GetBindings().FindByPredicate([this](const FStateTreePropertyPathBinding& Binding)
-		{
-			return Binding.GetTargetPath() == TargetPath;
-		});
-
-		if (CurrentBinding)
-		{
-			const FStateTreePropertyPath& SourcePath = CurrentBinding->GetSourcePath();
-			FString SourcePropertyPathAsString = SourcePath.ToString();
-
-			// If source is a bound PropertyFunction, it will not be present in AccessibleStructs thus it has to be accessed through bindings owner.
-			FStateTreeBindableStructDesc SourceDesc;
-			if (BindingOwner->GetStructByID(SourcePath.GetStructID(), SourceDesc))
-			{
-				// Making first segment of the path invisible for the user if it's property function's single output property.
-				if (SourceDesc.DataSource == EStateTreeBindableStructSource::PropertyFunction && UE::StateTree::GetStructSingleOutputProperty(*SourceDesc.Struct))
-				{
-					SourcePropertyPathAsString = SourcePath.ToString(/*HighlightedSegment*/ INDEX_NONE, /*HighlightPrefix*/ nullptr, /*HighlightPostfix*/ nullptr, /*bOutputInstances*/ false, 1);
-				}
-
-				// Check that the binding is valid.
-				bool bIsValidBinding = false;
-				FStateTreeDataView SourceDataView;
-				const FProperty* SourceLeafProperty = nullptr;
-				const UStruct* SourceStruct = nullptr;
-				if (BindingOwner->GetDataViewByID(SourcePath.GetStructID(), SourceDataView)
-					&& TargetDataView.IsValid())
-				{
-					TArray<FPropertyBindingPathIndirection> SourceIndirections;
-					TArray<FPropertyBindingPathIndirection> TargetIndirections;
-
-					// Resolve source and target properties.
-					// Source path can be empty, when the binding binds directly to a context struct/class.
-					// Target path must always point to a valid property (at least one indirection).
-					if (SourcePath.ResolveIndirectionsWithValue(SourceDataView, SourceIndirections)
-						&& TargetPath.ResolveIndirectionsWithValue(TargetDataView, TargetIndirections)
-						&& !TargetIndirections.IsEmpty())
-					{
-						const FPropertyBindingPathIndirection TargetLeafIndirection = TargetIndirections.Last();
-						if (SourceIndirections.Num() > 0)
-						{
-							// Binding to a source property.
-							const FPropertyBindingPathIndirection SourceLeafIndirection = SourceIndirections.Last();
-							SourceLeafProperty = SourceLeafIndirection.GetProperty();
-							bIsValidBinding = ArePropertiesCompatible(SourceLeafProperty, TargetLeafIndirection.GetProperty(), SourceLeafIndirection.GetPropertyAddress(), TargetLeafIndirection.GetPropertyAddress());
-						}
-						else
-						{
-							// Binding to a source context struct.
-							SourceStruct = SourceDataView.GetStruct();
-							bIsValidBinding = ArePropertyAndContextStructCompatible(SourceStruct, TargetLeafIndirection.GetProperty());
-						}
-					}
-				}
-
-				FormatableText = FText::FormatNamed(LOCTEXT("ValidSourcePath", "{SourceStruct}{PropertyPath}"), TEXT("PropertyPath"), SourcePropertyPathAsString.IsEmpty() ? FText() : FText::FromString(TEXT(".") + SourcePropertyPathAsString));
-				SourceStructName = FText::FromString(SourceDesc.Name.ToString());
-
-				if (bIsValidBinding)
-				{
-					if (SourcePropertyPathAsString.IsEmpty())
-					{
-						if (CurrentBinding->GetPropertyFunctionNode().IsValid())
-						{
-							TooltipBuilder.AppendLine(LOCTEXT("ExistingBindingToFunctionTooltip", "Property is bound to function {SourceStruct}."));
-						}
-						else
-						{
-							TooltipBuilder.AppendLine(LOCTEXT("ExistingBindingTooltip", "Property is bound to {SourceStruct}."));
-						}
-					}
-					else
-					{
-						if (CurrentBinding->GetPropertyFunctionNode().IsValid())
-						{
-							TooltipBuilder.AppendLineFormat(LOCTEXT("ExistingBindingToFunctionWithPropertyTooltip", "Property is bound to function {SourceStruct} property {PropertyPath}."), {{TEXT("PropertyPath"), FText::FromString(SourcePropertyPathAsString)}});
-						}
-						else
-						{
-							TooltipBuilder.AppendLineFormat(LOCTEXT("ExistingBindingWithPropertyTooltip", "Property is bound to {SourceStruct} property {PropertyPath}."), {{TEXT("PropertyPath"), FText::FromString(SourcePropertyPathAsString)}});
-						}
-					}
-
-					if (bIsPropertyRef) // Update the pin type with source property so that property ref that can binds to multiple types display the binded one.
-					{
-						Schema->ConvertPropertyToPinType(SourceLeafProperty, PinType);
-					}
-
-					Image = FAppStyle::GetBrush(IconName);
-					Color = Schema->GetPinTypeColor(PinType);
-				}
-				else
-				{
-					FText SourceType;
-					if (SourceLeafProperty)
-					{
-						SourceType = UE::StateTree::PropertyBinding::GetPropertyTypeText(SourceLeafProperty);
-					}
-					else if (SourceStruct)
-					{
-						SourceType = SourceStruct->GetDisplayNameText();
-					}
-					FText TargetType = UE::StateTree::PropertyBinding::GetPropertyTypeText(Property);
-
-					if (SourcePath.IsPathEmpty())
-					{
-						TooltipBuilder.AppendLineFormat(LOCTEXT("MismatchingBindingTooltip", "Property is bound to {SourceStruct}, but binding source type '{SourceType}' does not match property type '{TargetType}'."),
-							{
-								{TEXT("SourceType"), SourceType},
-								{TEXT("TargetType"), TargetType}
-							});
-					}
-					else
-					{
-						TooltipBuilder.AppendLineFormat(LOCTEXT("MismatchingBindingTooltipWithProperty", "Property is bound to {SourceStruct} property {PropertyPath}, but binding source type '{SourceType}' does not match property type '{TargetType}'."),
-							{
-								{TEXT("PropertyPath"), FText::FromString(SourcePropertyPathAsString)},
-								{TEXT("SourceType"), SourceType},
-								{TEXT("TargetType"), TargetType}
-							});
-					}
-
-					Image = FCoreStyle::Get().GetBrush("Icons.ErrorWithColor");
-					Color = FLinearColor::White;
-				}
-			}
-			else
-			{
-				// Missing source
-				FormatableText = FText::Format(LOCTEXT("MissingSource", "???.{0}"), FText::FromString(SourcePropertyPathAsString));
-				TooltipBuilder.AppendLineFormat(LOCTEXT("MissingBindingTooltip", "Missing binding source for property path '{0}'."), FText::FromString(SourcePropertyPathAsString));
-				Image = FCoreStyle::Get().GetBrush("Icons.ErrorWithColor");
-				Color = FLinearColor::White;
-			}
-
-			CachedSourcePath = SourcePath;
-		}
-		else
-		{
-			// No bindings
-			FormatableText = FText::GetEmpty();
-			TooltipBuilder.AppendLineFormat(LOCTEXT("BindTooltip", "Bind {0} to value from another property."), UE::StateTree::PropertyBinding::GetPropertyTypeText(Property));
-
-			Image = FAppStyle::GetBrush(IconName);
-			Color = Schema->GetPinTypeColor(PinType);
-
-			CachedSourcePath.Reset();
-		}
-
-		if (bIsPropertyRef)
-		{
-			if (Property->HasMetaData(UE::StateTree::PropertyRefHelpers::IsRefToArrayName))
-			{
-				TooltipBuilder.AppendLineFormat(LOCTEXT("PropertyRefBindingTooltipArray", "Supported types are Array of {0}"), FText::FromString(Property->GetMetaData(UE::StateTree::PropertyRefHelpers::RefTypeName)));
-			}
-			else
-			{
-				TooltipBuilder.AppendLineFormat(LOCTEXT("PropertyRefBindingTooltip", "Supported types are {0}"), FText::FromString(Property->GetMetaData(UE::StateTree::PropertyRefHelpers::RefTypeName)));
-				if (Property->HasMetaData(UE::StateTree::PropertyRefHelpers::CanRefToArrayName))
-				{
-					TooltipBuilder.AppendLine(LOCTEXT("PropertyRefBindingTooltipCanSupportArray", "Supports Arrays"));
-				}
-			}
-		}
-
-		FormatableTooltipText = TooltipBuilder.ToText();
-
-		bIsDataCached = true;
-	}
-
-	bool CanBindToContextStruct(const UStruct* InStruct, int32 InStructIndex)
-	{
-		ConditionallyUpdateData();
-
 		// Do not allow to bind directly StateTree nodes
 		// @todo: find a way to more specifically call out the context structs, e.g. pass the property path to the callback.
 		if (InStruct != nullptr)
 		{
-			const bool bIsStateTreeNode = AccessibleStructs.ContainsByPredicate([InStruct](const FStateTreeBindableStructDesc& AccessibleStruct)
+			const bool bIsStateTreeNode = GetAccessibleStructs().ContainsByPredicate([InStruct](const TInstancedStruct<FPropertyBindingBindableStructDescriptor>& StructDesc)
 			{
+				const FStateTreeBindableStructDesc& AccessibleStruct = StructDesc.Get<FStateTreeBindableStructDesc>();
 				return AccessibleStruct.DataSource != EStateTreeBindableStructSource::Context
 					&& AccessibleStruct.DataSource != EStateTreeBindableStructSource::Parameter
 					&& AccessibleStruct.DataSource != EStateTreeBindableStructSource::TransitionEvent
@@ -958,109 +270,48 @@ struct FCachedBindingData : public TSharedFromThis<FCachedBindingData>
 
 			if (bIsStateTreeNode)
 			{
-				return false;
-			}
-		}
-
-		check(AccessibleStructs.IsValidIndex(InStructIndex));
-		// Binding directly into PropertyFunction's struct is allowed if it contains a compatible single output property.
-		if (AccessibleStructs[InStructIndex].DataSource == EStateTreeBindableStructSource::PropertyFunction)
-		{
-			IStateTreeEditorPropertyBindingsOwner* BindingOwner = Cast<IStateTreeEditorPropertyBindingsOwner>(WeakOwnerObject.Get());
-			FStateTreeDataView DataView;
-			// If DataView exists, struct is an instance of already bound function.
-			if (BindingOwner == nullptr || BindingOwner->GetDataViewByID(AccessibleStructs[InStructIndex].ID, DataView))
-			{
-				return false;
-			}
-
-			if (const FProperty* SingleOutputProperty = UE::StateTree::GetStructSingleOutputProperty(*AccessibleStructs[InStructIndex].Struct))
-			{
-				return CanBindToProperty(SingleOutputProperty, {FBindingChainElement(nullptr, InStructIndex), FBindingChainElement(const_cast<FProperty*>(SingleOutputProperty))});
-			}
-		}
-
-		return ArePropertyAndContextStructCompatible(InStruct, PropertyHandle->GetProperty());
-	}
-			
-	bool CanBindToProperty(const FProperty* SourceProperty, TConstArrayView<FBindingChainElement> InBindingChain)
-	{
-		ConditionallyUpdateData();
-
-		// Special case for binding widget calling OnCanBindProperty with Args.Property (i.e. self).
-		if (PropertyHandle->GetProperty() == SourceProperty)
-		{
-			return true;
-		}
-
-		UObject* OwnerObject = WeakOwnerObject.Get();
-		if (!OwnerObject)
-		{
-			return false;
-		}
-
-		IStateTreeEditorPropertyBindingsOwner* BindingOwner = Cast<IStateTreeEditorPropertyBindingsOwner>(OwnerObject);
-		if (!BindingOwner)
-		{
-			return false;
-		}
-
-		const int32 SourceStructIndex = InBindingChain[0].ArrayIndex;
-		check(AccessibleStructs.IsValidIndex(SourceStructIndex));
-
-		FStateTreeDataView SourceDataView;
-		if (AccessibleStructs[SourceStructIndex].DataSource == EStateTreeBindableStructSource::PropertyFunction)
-		{
-			SourceDataView = FStateTreeDataView(AccessibleStructs[SourceStructIndex].Struct, nullptr);
-		}
-		else
-		{
-			BindingOwner->GetDataViewByID(AccessibleStructs[SourceStructIndex].ID, SourceDataView);
-		}
-
-		FStateTreePropertyPath SourcePath;
-		UE::StateTree::PropertyBinding::MakeStructPropertyPathFromBindingChain(AccessibleStructs[SourceStructIndex].ID, InBindingChain, SourceDataView, SourcePath);
-
-		TArray<FPropertyBindingPathIndirection> SourceIndirections;
-		void* TargetValueAddress = nullptr;
-		if (PropertyHandle->GetValueData(TargetValueAddress) == FPropertyAccess::Success && SourcePath.ResolveIndirectionsWithValue(SourceDataView, SourceIndirections))
-		{
-			return ArePropertiesCompatible(SourceProperty, PropertyHandle->GetProperty(), SourceIndirections.Last().GetPropertyAddress(), TargetValueAddress);
-		}
-		
-		return false;
-	}
-
-	bool CanAcceptPropertyOrChildren(const FProperty* SourceProperty, TConstArrayView<FBindingChainElement> InBindingChain)
-	{
-		if (!SourceProperty)
-		{
-			return false;
+				return false;
+			}
 		}
 
-		ConditionallyUpdateData();
-
-		if (!PropertyHandle.IsValid() || PropertyHandle->GetProperty() == nullptr)
+		const FStateTreeBindableStructDesc& StructDesc = GetBindableStructDescriptor(InStructIndex).Get<FStateTreeBindableStructDesc>();
+		// Binding directly into PropertyFunction's struct is allowed if it contains a compatible single output property.
+		if (StructDesc.DataSource == EStateTreeBindableStructSource::PropertyFunction)
 		{
-			return false;
+			const IStateTreeEditorPropertyBindingsOwner* BindingOwner = Cast<IStateTreeEditorPropertyBindingsOwner>(GetOwner());
+			FStateTreeDataView DataView;
+			// If DataView exists, struct is an instance of already bound function.
+			if (BindingOwner == nullptr || BindingOwner->GetBindingDataViewByID(StructDesc.ID, DataView))
+			{
+				return false;
+			}
+
+			if (const FProperty* SingleOutputProperty = GetStructSingleOutputProperty(*StructDesc.Struct))
+			{
+				return CanBindToProperty(SingleOutputProperty, {FBindingChainElement(nullptr, InStructIndex), FBindingChainElement(const_cast<FProperty*>(SingleOutputProperty))});
+			}
 		}
 
+		return true;
+	}
+
+	virtual bool CanAcceptPropertyOrChildrenInternal(const FProperty& SourceProperty, TConstArrayView<FBindingChainElement, int> InBindingChain) override
+	{
 		const int32 SourceStructIndex = InBindingChain[0].ArrayIndex;
-		check(AccessibleStructs.IsValidIndex(SourceStructIndex));
-		const FStateTreeBindableStructDesc& StructDesc = AccessibleStructs[SourceStructIndex];
+		const FStateTreeBindableStructDesc& StructDesc = GetBindableStructDescriptor(SourceStructIndex).Get<FStateTreeBindableStructDesc>();
 
 		if (StructDesc.DataSource == EStateTreeBindableStructSource::PropertyFunction)
 		{
-			IStateTreeEditorPropertyBindingsOwner* BindingOwner = Cast<IStateTreeEditorPropertyBindingsOwner>(WeakOwnerObject.Get());
+			IStateTreeEditorPropertyBindingsOwner* BindingOwner = Cast<IStateTreeEditorPropertyBindingsOwner>(GetOwner());
 			FStateTreeDataView DataView;
 			// If DataView exists, struct is an instance of already bound function.
-			if (BindingOwner == nullptr || BindingOwner->GetDataViewByID(AccessibleStructs[SourceStructIndex].ID, DataView))
+			if (BindingOwner == nullptr || BindingOwner->GetBindingDataViewByID(StructDesc.ID, DataView))
 			{
 				return false;
 			}
 
 			// To avoid duplicates, PropertyFunction struct's children are not allowed to be bound if it contains a compatible single output property.
-			if (const FProperty* SingleOutputProperty = UE::StateTree::GetStructSingleOutputProperty(*StructDesc.Struct))
+			if (const FProperty* SingleOutputProperty = GetStructSingleOutputProperty(*StructDesc.Struct))
 			{
 				if (CanBindToProperty(SingleOutputProperty, {FBindingChainElement(nullptr, SourceStructIndex), FBindingChainElement(const_cast<FProperty*>(SingleOutputProperty))}))
 				{
@@ -1069,360 +320,160 @@ struct FCachedBindingData : public TSharedFromThis<FCachedBindingData>
 			}
 
 			// Binding to non-output PropertyFunctions properties is not allowed.
-			if (InBindingChain.Num() == 1 && UE::StateTree::GetUsageFromMetaData(SourceProperty) != EStateTreePropertyUsage::Output)
+			if (InBindingChain.Num() == 1 && GetUsageFromMetaData(&SourceProperty) != EStateTreePropertyUsage::Output)
 			{
 				return false;
 			}
 		}
 
-		if (UE::StateTree::PropertyRefHelpers::IsPropertyRef(*PropertyHandle->GetProperty()) && !UE::StateTree::PropertyRefHelpers::IsPropertyAccessibleForPropertyRef(*SourceProperty, InBindingChain, StructDesc))
+		if (PropertyRefHelpers::IsPropertyRef(*GetPropertyHandle()->GetProperty()) && !PropertyRefHelpers::IsPropertyAccessibleForPropertyRef(SourceProperty, InBindingChain, StructDesc))
 		{
-			if (!UE::StateTree::PropertyRefHelpers::IsPropertyAccessibleForPropertyRef(*SourceProperty, InBindingChain, StructDesc))
+			if (!PropertyRefHelpers::IsPropertyAccessibleForPropertyRef(SourceProperty, InBindingChain, StructDesc))
 			{
 				return false;
 			}
 		}
 
-		return IsPropertyBindable(*SourceProperty);
-	}
-
-	static bool ArePropertyAndContextStructCompatible(const UStruct* SourceStruct, const FProperty* TargetProperty)
-	{
-		if (const FStructProperty* TargetStructProperty = CastField<FStructProperty>(TargetProperty))
-		{
-			return TargetStructProperty->Struct == SourceStruct;
-		}
-		if (const FObjectProperty* TargetObjectProperty = CastField<FObjectProperty>(TargetProperty))
-		{
-			return SourceStruct != nullptr && SourceStruct->IsChildOf(TargetObjectProperty->PropertyClass);
-		}
-		
-		return false;
+		return true;
 	}
 
-	static bool ArePropertiesCompatible(const FProperty* SourceProperty, const FProperty* TargetProperty, const void* SourcePropertyValue, const void* TargetPropertyValue)
+	virtual bool DeterminePropertiesCompatibilityInternal(
+		const FProperty* InSourceProperty,
+		const FProperty* InTargetProperty,
+		const void* InSourcePropertyValue,
+		const void* InTargetPropertyValue,
+		bool& bOutAreCompatible) const override
 	{
 		// @TODO: Refactor FStateTreePropertyBindings::ResolveCopyType() so that we can use it directly here.
 		
-		bool bCanBind = false;
-
-		const FStructProperty* TargetStructProperty = CastField<FStructProperty>(TargetProperty);
+		const FStructProperty* TargetStructProperty = CastField<FStructProperty>(InTargetProperty);
 		
 		// AnyEnums need special handling.
 		// It is a struct property but we want to treat it as an enum. We need to do this here, instead of 
-		// FStateTreePropertyBindingCompiler::GetPropertyCompatibility() because the treatment depends on the value too.
+		// GetPropertyCompatibility() because the treatment depends on the value too.
 		// Note: AnyEnums will need special handling before they can be used for binding.
 		if (TargetStructProperty && TargetStructProperty->Struct == FStateTreeAnyEnum::StaticStruct())
 		{
 			// If the AnyEnum has AllowAnyBinding, allow to bind to any enum.
-			const bool bAllowAnyBinding = TargetProperty->HasMetaData(UE::StateTree::PropertyBinding::AllowAnyBindingName);
+			const bool bAllowAnyBinding = InTargetProperty->HasMetaData(AllowAnyBindingName);
 
-			check(TargetPropertyValue);
-			const FStateTreeAnyEnum* TargetAnyEnum = static_cast<const FStateTreeAnyEnum*>(TargetPropertyValue);
+			check(InTargetPropertyValue);
+			const FStateTreeAnyEnum* TargetAnyEnum = static_cast<const FStateTreeAnyEnum*>(InTargetPropertyValue);
 
 			// If the enum class is not specified, allow to bind to any enum, if the class is specified allow only that enum.
-			if (const FByteProperty* SourceByteProperty = CastField<FByteProperty>(SourceProperty))
+			if (const FByteProperty* SourceByteProperty = CastField<FByteProperty>(InSourceProperty))
 			{
 				if (UEnum* Enum = SourceByteProperty->GetIntPropertyEnum())
 				{
-					bCanBind = bAllowAnyBinding || TargetAnyEnum->Enum == Enum;
+					bOutAreCompatible = bAllowAnyBinding || TargetAnyEnum->Enum == Enum;
+					return true;
 				}
 			}
-			else if (const FEnumProperty* SourceEnumProperty = CastField<FEnumProperty>(SourceProperty))
+			else if (const FEnumProperty* SourceEnumProperty = CastField<FEnumProperty>(InSourceProperty))
 			{
-				bCanBind = bAllowAnyBinding || TargetAnyEnum->Enum == SourceEnumProperty->GetEnum();
+				bOutAreCompatible = bAllowAnyBinding || TargetAnyEnum->Enum == SourceEnumProperty->GetEnum();
+				return true;
 			}
 		}
 		else if (TargetStructProperty && TargetStructProperty->Struct == FStateTreeStructRef::StaticStruct())
 		{
 			FString BaseStructName;
-			const UScriptStruct* TargetStructRefBaseStruct = UE::StateTree::Compiler::GetBaseStructFromMetaData(TargetProperty, BaseStructName);
+			const UScriptStruct* TargetStructRefBaseStruct = Compiler::GetBaseStructFromMetaData(InTargetProperty, BaseStructName);
 
-			if (const FStructProperty* SourceStructProperty = CastField<FStructProperty>(SourceProperty))
+			if (const FStructProperty* SourceStructProperty = CastField<FStructProperty>(InSourceProperty))
 			{
 				if (SourceStructProperty->Struct == TBaseStructure<FStateTreeStructRef>::Get())
 				{
 					FString SourceBaseStructName;
-					const UScriptStruct* SourceStructRefBaseStruct = UE::StateTree::Compiler::GetBaseStructFromMetaData(SourceStructProperty, SourceBaseStructName);
-					bCanBind = SourceStructRefBaseStruct && SourceStructRefBaseStruct->IsChildOf(TargetStructRefBaseStruct);
+					const UScriptStruct* SourceStructRefBaseStruct = Compiler::GetBaseStructFromMetaData(SourceStructProperty, SourceBaseStructName);
+					bOutAreCompatible = SourceStructRefBaseStruct && SourceStructRefBaseStruct->IsChildOf(TargetStructRefBaseStruct);
+					return true;
 				}
 				else
 				{
-					bCanBind = SourceStructProperty->Struct && SourceStructProperty->Struct->IsChildOf(TargetStructRefBaseStruct);
+					bOutAreCompatible = SourceStructProperty->Struct && SourceStructProperty->Struct->IsChildOf(TargetStructRefBaseStruct);
+					return true;
 				}
 			}
 		}
-		else if (TargetStructProperty && UE::StateTree::PropertyRefHelpers::IsPropertyRef(*TargetStructProperty))
+		else if (TargetStructProperty && PropertyRefHelpers::IsPropertyRef(*TargetStructProperty))
 		{
-			check(TargetPropertyValue);
-			bCanBind = UE::StateTree::PropertyRefHelpers::IsPropertyRefCompatibleWithProperty(*TargetStructProperty, *SourceProperty, TargetPropertyValue, SourcePropertyValue);
+			check(InTargetPropertyValue);
+			bOutAreCompatible = PropertyRefHelpers::IsPropertyRefCompatibleWithProperty(*TargetStructProperty, *InSourceProperty, InTargetPropertyValue, InSourcePropertyValue);
+			return true;
 		}
-		else if (TargetStructProperty && UE::StateTree::PropertyBinding::IsDelegateListenerProperty(*TargetStructProperty))
+		else if (TargetStructProperty && IsDelegateListenerProperty(*TargetStructProperty))
 		{
-			if (const FStructProperty* SourceStructProperty = CastField<FStructProperty>(SourceProperty))
+			if (const FStructProperty* SourceStructProperty = CastField<FStructProperty>(InSourceProperty))
 			{
-				bCanBind = SourceStructProperty->Struct == FStateTreeDelegateDispatcher::StaticStruct();
+				bOutAreCompatible = SourceStructProperty->Struct == FStateTreeDelegateDispatcher::StaticStruct();
+				return true;
 			}
 		}
-		else
-		{
-			// Note: We support type promotion here
-			bCanBind = FStateTreePropertyBindings::GetPropertyCompatibility(SourceProperty, TargetProperty) != EStateTreePropertyAccessCompatibility::Incompatible;
-		}
-
-		return bCanBind;
-	}
-
-	UStruct* ResolveIndirection(TConstArrayView<FBindingChainElement> InBindingChain)
-	{
-		UObject* OwnerObject = WeakOwnerObject.Get();
-		if (!OwnerObject)
-		{
-			return nullptr;
-		}
-		
-		IStateTreeEditorPropertyBindingsOwner* BindingOwner = Cast<IStateTreeEditorPropertyBindingsOwner>(OwnerObject);
-		if (!BindingOwner)
-		{
-			return nullptr;
-		}
-
-		const int32 SourceStructIndex = InBindingChain[0].ArrayIndex;
-		check(SourceStructIndex >= 0 && SourceStructIndex < AccessibleStructs.Num());
-
-		FStateTreeDataView DataView;
-		if (BindingOwner->GetDataViewByID(AccessibleStructs[SourceStructIndex].ID, DataView))
-		{
-			return UE::StateTree::PropertyBinding::ResolveLeafValueStructType(DataView, InBindingChain);
-		}
 
-		return nullptr;
+		return false;
 	}
 
-	FText GetText()
+	virtual bool GetPropertyFunctionText(FConstStructView InPropertyFunctionStructView, FText& OutText) const override
 	{
-		ConditionallyUpdateData();
-
-		// Bound PropertyFunction is allowed to override it's display name.
-		if (IStateTreeEditorPropertyBindingsOwner* BindingOwner = Cast<IStateTreeEditorPropertyBindingsOwner>(WeakOwnerObject.Get()))
+		const FStateTreeEditorNode& EditorNode = InPropertyFunctionStructView.Get<const FStateTreeEditorNode>();
+		if (const FStateTreeNodeBase* Node = EditorNode.Node.GetPtr<const FStateTreeNodeBase>())
 		{
-			if(FStateTreeEditorPropertyBindings* EditorBindings = BindingOwner->GetPropertyEditorBindings())
+			IStateTreeEditorPropertyBindingsOwner* BindingOwner = Cast<IStateTreeEditorPropertyBindingsOwner>(GetOwner());
+			const FText Description = Node->GetDescription(GetSourcePath().GetStructID(), EditorNode.GetInstance(), FStateTreeBindingLookup(BindingOwner), EStateTreeNodeFormatting::Text);
+			if (!Description.IsEmpty())
 			{
-				const FStateTreePropertyPathBinding* CurrentBinding = EditorBindings->GetBindings().FindByPredicate([this](const FStateTreePropertyPathBinding& Binding)
-				{
-					return Binding.GetTargetPath() == TargetPath;
-				});
-
-				if (CurrentBinding)
-				{
-					const FConstStructView PropertyFunctionEditorNodeView = CurrentBinding->GetPropertyFunctionNode();
-					if(PropertyFunctionEditorNodeView.IsValid())
-					{
-						const FStateTreeEditorNode& EditorNode = PropertyFunctionEditorNodeView.Get<const FStateTreeEditorNode>();
-						if (const FStateTreeNodeBase* Node = EditorNode.Node.GetPtr<const FStateTreeNodeBase>())
-						{
-							const FText Description = Node->GetDescription(CachedSourcePath.GetStructID(), EditorNode.GetInstance(), FStateTreeBindingLookup(BindingOwner), EStateTreeNodeFormatting::Text);
-							if (!Description.IsEmpty())
-							{
-								return FText::FormatNamed(FormatableText, TEXT("SourceStruct"), Description);
-							}
-						}
-					}
-				}
+				OutText = FText::FormatNamed(GetFormatableText(), TEXT("SourceStruct"), Description);
+				return true;
 			}
 		}
-
-		return FText::FormatNamed(FormatableText, TEXT("SourceStruct"), SourceStructName);
+		return false;
 	}
-	
-	FText GetTooltipText()
-	{
-		ConditionallyUpdateData();
-
-		// If the source property is a PropertyFunction and it overrides it's display name, it's been used in the tooltip text.
-		if (IStateTreeEditorPropertyBindingsOwner* BindingOwner = Cast<IStateTreeEditorPropertyBindingsOwner>(WeakOwnerObject.Get()))
-		{
-			if(FStateTreeEditorPropertyBindings* EditorBindings = BindingOwner->GetPropertyEditorBindings())
-			{
-				const FStateTreePropertyPathBinding* CurrentBinding = EditorBindings->GetBindings().FindByPredicate([this](const FStateTreePropertyPathBinding& Binding)
-				{
-					return Binding.GetTargetPath() == TargetPath;
-				});
-
-				if (CurrentBinding)
-				{
-					const FConstStructView PropertyFunctionEditorNodeView = CurrentBinding->GetPropertyFunctionNode();
-					if(PropertyFunctionEditorNodeView.IsValid())
-					{
-						const FStateTreeEditorNode& EditorNode = PropertyFunctionEditorNodeView.Get<const FStateTreeEditorNode>();
-						if(const FStateTreeNodeBase* Node = EditorNode.Node.GetPtr<const FStateTreeNodeBase>())
-						{
-							const FText Description = Node->GetDescription(CachedSourcePath.GetStructID(), EditorNode.GetInstance(), FStateTreeBindingLookup(BindingOwner), EStateTreeNodeFormatting::Text);
-							if (!Description.IsEmpty())
-							{
-								return FText::FormatNamed(FormatableTooltipText, TEXT("SourceStruct"), Description);
-							}
-						}
-					}
-				}
-			}
-		}
 
-		return FText::FormatNamed(FormatableTooltipText, TEXT("SourceStruct"), SourceStructName);
-	}
-	
-	FLinearColor GetColor()
+	virtual bool GetPropertyFunctionTooltipText(FConstStructView InPropertyFunctionStructView, FText& OutText) const override
 	{
-		ConditionallyUpdateData();
-
-		// Bound PropertyFunction is allowed to override it's icon color if the binding leads directly into it's single output property.
-		if (CachedSourcePath.NumSegments() == 1)
+		const FStateTreeEditorNode& EditorNode = InPropertyFunctionStructView.Get<const FStateTreeEditorNode>();
+		if (const FStateTreeNodeBase* Node = EditorNode.Node.GetPtr<const FStateTreeNodeBase>())
 		{
-			if (IStateTreeEditorPropertyBindingsOwner* BindingOwner = Cast<IStateTreeEditorPropertyBindingsOwner>(WeakOwnerObject.Get()))
+			IStateTreeEditorPropertyBindingsOwner* BindingOwner = Cast<IStateTreeEditorPropertyBindingsOwner>(GetOwner());
+			const FText Description = Node->GetDescription(GetSourcePath().GetStructID(), EditorNode.GetInstance(), FStateTreeBindingLookup(BindingOwner), EStateTreeNodeFormatting::Text);
+			if (!Description.IsEmpty())
 			{
-				if(FStateTreeEditorPropertyBindings* EditorBindings = BindingOwner->GetPropertyEditorBindings())
-				{
-					const FStateTreePropertyPathBinding* CurrentBinding = EditorBindings->GetBindings().FindByPredicate([this](const FStateTreePropertyPathBinding& Binding)
-					{
-						return Binding.GetTargetPath() == TargetPath;
-					});
-
-					if (CurrentBinding)
-					{
-						const FConstStructView PropertyFunctionEditorNodeView = CurrentBinding->GetPropertyFunctionNode();
-						if (PropertyFunctionEditorNodeView.IsValid())
-						{
-							const FStateTreeEditorNode& EditorNode = PropertyFunctionEditorNodeView.Get<const FStateTreeEditorNode>();
-							if (const FStateTreeNodeBase* Node = EditorNode.Node.GetPtr<const FStateTreeNodeBase>())
-							{
-								if (Node && UE::StateTree::GetStructSingleOutputProperty(*Node->GetInstanceDataType()))
-								{
-									return Node->GetIconColor();
-								}
-					
-							}
-						}
-					}
-				}
+				OutText = FText::FormatNamed(GetFormatableTooltipText(), TEXT("SourceStruct"), Description);
+				return true;
 			}
 		}
-
-		return Color;
+		return false;
 	}
-	
-	const FSlateBrush* GetImage()
-	{
-		ConditionallyUpdateData();
 
-		// Bound PropertyFunction is allowed to override it's icon.
-		if (IStateTreeEditorPropertyBindingsOwner* BindingOwner = Cast<IStateTreeEditorPropertyBindingsOwner>(WeakOwnerObject.Get()))
+	virtual bool GetPropertyFunctionIconColor(FConstStructView InPropertyFunctionStructView, FLinearColor& OutColor) const override
+	{
+		const FStateTreeEditorNode& EditorNode = InPropertyFunctionStructView.Get<const FStateTreeEditorNode>();
+		if (const FStateTreeNodeBase* Node = EditorNode.Node.GetPtr<const FStateTreeNodeBase>())
 		{
-			if(FStateTreeEditorPropertyBindings* EditorBindings = BindingOwner->GetPropertyEditorBindings())
+			if (GetStructSingleOutputProperty(*Node->GetInstanceDataType()))
 			{
-				const FStateTreePropertyPathBinding* CurrentBinding = EditorBindings->GetBindings().FindByPredicate([this](const FStateTreePropertyPathBinding& Binding)
-				{
-					return Binding.GetTargetPath() == TargetPath;
-				});
-
-				if (CurrentBinding)
-				{
-					const FConstStructView PropertyFunctionEditorNodeView = CurrentBinding->GetPropertyFunctionNode();
-					if (PropertyFunctionEditorNodeView.IsValid())
-					{
-						const FStateTreeEditorNode& EditorNode = PropertyFunctionEditorNodeView.Get<const FStateTreeEditorNode>();
-						if (const FStateTreeNodeBase* Node = EditorNode.Node.GetPtr<const FStateTreeNodeBase>())
-						{
-							return UE::StateTreeEditor::EditorNodeUtils::ParseIcon(Node->GetIconName()).GetIcon();
-						}
-					}
-				}
+				OutColor = Node->GetIconColor();
+				return true;
 			}
 		}
-
-		return Image;
+		return false;
 	}
-	
-private:
 
-	void ConditionallyUpdateData()
+	virtual bool GetPropertyFunctionImage(FConstStructView InPropertyFunctionStructView, const FSlateBrush*& OutImage) const override
 	{
-		UObject* OwnerObject = WeakOwnerObject.Get();
-		if (!OwnerObject)
-		{
-			return;
-		}
-		
-		IStateTreeEditorPropertyBindingsOwner* BindingOwner = Cast<IStateTreeEditorPropertyBindingsOwner>(OwnerObject);
-		if (!BindingOwner)
-		{
-			return;
-		}
-
-		FStateTreeEditorPropertyBindings* EditorBindings = BindingOwner->GetPropertyEditorBindings();
-		if (!EditorBindings)
-		{
-			return;
-		}
-
-		const FStateTreePropertyPath* CurrentSourcePath = EditorBindings->GetPropertyBindingSource(TargetPath);
-		bool bPathsIdentical = false;
-		if (CurrentSourcePath)
-		{
-			bPathsIdentical = CachedSourcePath == *CurrentSourcePath;
-		}
-		else
-		{
-			bPathsIdentical = CachedSourcePath.IsPathEmpty();
-		}
-
-		if (!bIsDataCached || !bPathsIdentical)
+		const FStateTreeEditorNode& EditorNode = InPropertyFunctionStructView.Get<const FStateTreeEditorNode>();
+		if (const FStateTreeNodeBase* Node = EditorNode.Node.GetPtr<const FStateTreeNodeBase>())
 		{
-			UpdateData();
+			if (GetStructSingleOutputProperty(*Node->GetInstanceDataType()))
+			{
+				OutImage = UE::StateTreeEditor::EditorNodeUtils::ParseIcon(Node->GetIconName()).GetIcon();
+				return true;
+			}
 		}
-	}
-	
-	TWeakObjectPtr<UObject> WeakOwnerObject = nullptr;
-	FStateTreePropertyPath CachedSourcePath;
-	FStateTreePropertyPath TargetPath;
-	TSharedPtr<const IPropertyHandle> PropertyHandle;
-	TArray<FStateTreeBindableStructDesc> AccessibleStructs;
-
-	/* Default name of the source struct. */
-	FText SourceStructName;
-
-	/* Binding's display name text. Expects it's source struct name to be injected before use. */
-	FText FormatableText;
-
-	/* Binding's tooltip text. Expects it's source struct name to be injected before use. */
-	FText FormatableTooltipText;
-
-	FLinearColor Color = FLinearColor::White;
-	const FSlateBrush* Image = nullptr;
-	
-	bool bIsDataCached = false;
-};
-
-bool IsPropertyBindable(const FProperty& Property)
-{
-	const bool bIsUserEditable = Property.HasAnyPropertyFlags(CPF_Edit);
-	if (!bIsUserEditable)
-	{
-		UE_LOG(LogStateTreeEditor, Verbose, TEXT("Property %s is not bindable because it's not user-settable in the editor"),
-			*Property.GetName());
-		return false;
-	}
-
-	const bool bPrivateOrProtected = !Property.HasAnyPropertyFlags(CPF_NativeAccessSpecifierPrivate | CPF_NativeAccessSpecifierProtected);
-	const bool bPrivateButBlueprintAccessible = Property.GetBoolMetaData(FBlueprintMetadata::MD_AllowPrivateAccess);
-	if (!bPrivateOrProtected && !bPrivateButBlueprintAccessible)
-	{
-		UE_LOG(LogStateTreeEditor, Verbose, TEXT("Property %s is not bindable because it's either private or protected and not private-accessible to blueprints"),
-			*Property.GetName());
 		return false;
 	}
-
-	return true;
-}
+};
 
 /* Provides PropertyFunctionNode instance for a property node. */
 class FStateTreePropertyFunctionNodeProvider : public IStructureDataProvider
@@ -1467,7 +518,7 @@ private:
 		if(IStateTreeEditorPropertyBindingsOwner* Owner = Cast<IStateTreeEditorPropertyBindingsOwner>(RawBindingsOwner))
 		{
 			FStateTreeEditorPropertyBindings* EditorBindings = Owner->GetPropertyEditorBindings();
-			FStateTreePropertyPathBinding* FoundBinding = EditorBindings->GetMutableBindings().FindByPredicate([&InTargetPath](const FStateTreePropertyPathBinding& Binding)
+			FPropertyBindingBinding* FoundBinding = EditorBindings->GetMutableBindings().FindByPredicate([&InTargetPath](const FStateTreePropertyPathBinding& Binding)
 			{
 				return Binding.GetTargetPath() == InTargetPath;
 			});
@@ -1496,350 +547,64 @@ private:
 
 } // UE::StateTree::PropertyBinding
 
-bool FStateTreeBindingExtension::IsPropertyExtendable(const UClass* InObjectClass, const IPropertyHandle& PropertyHandle) const
+TSharedPtr<UE::PropertyBinding::FCachedBindingData> FStateTreeBindingExtension::CreateCachedBindingData(IPropertyBindingBindingCollectionOwner* InBindingsOwner, const FPropertyBindingPath& InTargetPath, const TSharedPtr<IPropertyHandle>& InPropertyHandle, const TConstArrayView<TInstancedStruct<FPropertyBindingBindableStructDescriptor>> InAccessibleStructs) const
 {
-	const FProperty* Property = PropertyHandle.GetProperty();
-	if (Property == nullptr || Property->HasAnyPropertyFlags(CPF_PersistentInstance | CPF_EditorOnly | CPF_Config | CPF_Deprecated))
-	{
-		return false;
-	}
-	
-	FStateTreePropertyPath TargetPath;
-	// Figure out the structs we're editing, and property path relative to current property.
-	const EStateTreePropertyUsage Usage = UE::StateTree::PropertyBinding::MakeStructPropertyPathFromPropertyHandle(PropertyHandle.AsShared(), TargetPath);
-
-	// Binding from Delegate Dispatcher is not allowed.
-	if (const FStructProperty* StructProperty = CastField<FStructProperty>(PropertyHandle.GetProperty()))
-	{
-		if (StructProperty->Struct == FStateTreeDelegateDispatcher::StaticStruct())
-		{
-			return false;
-		}
-	}
+	return MakeShared<UE::StateTree::PropertyBinding::FStateTreeCachedBindingData>(InBindingsOwner, InTargetPath, InPropertyHandle, InAccessibleStructs);
+}
 
+bool FStateTreeBindingExtension::CanBindToProperty(const FPropertyBindingPath& InTargetPath, const IPropertyHandle& InPropertyHandle) const
+{
+	const EStateTreePropertyUsage Usage = UE::StateTree::GetUsageFromMetaData(InPropertyHandle.GetProperty());
 	if (Usage == EStateTreePropertyUsage::Input || Usage == EStateTreePropertyUsage::Context)
 	{
 		// Allow to bind only to the main level on input and context properties.
-		return TargetPath.GetSegments().Num() == 1;
-	}
-	if (Usage == EStateTreePropertyUsage::Parameter)
-	{
-		return true;
+		return InTargetPath.GetSegments().Num() == 1;
 	}
-	
-	return false;
-}
 
+	return (Usage == EStateTreePropertyUsage::Parameter);
+}
 
-void FStateTreeBindingExtension::ExtendWidgetRow(FDetailWidgetRow& InWidgetRow, const IDetailLayoutBuilder& InDetailBuilder, const UClass* InObjectClass, TSharedPtr<IPropertyHandle> InPropertyHandle)
+void FStateTreeBindingExtension::UpdateContextStruct(TConstStructView<FPropertyBindingBindableStructDescriptor> InStructDesc, FBindingContextStruct& InOutContextStruct, TMap<FString, FText>& InOutSectionNames) const
 {
-	if (!IModularFeatures::Get().IsModularFeatureAvailable("PropertyAccessEditor"))
-	{
-		return;
-	}
-
-	IPropertyAccessEditor& PropertyAccessEditor = IModularFeatures::Get().GetModularFeature<IPropertyAccessEditor>("PropertyAccessEditor");
-
-	UObject* OwnerObject = nullptr;
-	FStateTreeEditorPropertyBindings* EditorBindings = nullptr;
-
-	// Array of structs we can bind to.
-	TArray<FBindingContextStruct> BindingContextStructs;
-	TArray<FStateTreeBindableStructDesc> AccessibleStructs;
-
-	// The struct and property where we're binding.
-	FStateTreePropertyPath TargetPath;
-
-	IStateTreeEditorPropertyBindingsOwner* BindingOwner = nullptr;
-
-	TArray<UObject*> OuterObjects;
-	InPropertyHandle->GetOuterObjects(OuterObjects);
-	if (OuterObjects.Num() == 1)
+	const FStateTreeBindableStructDesc& StructDesc = InStructDesc.Get<FStateTreeBindableStructDesc>();
+	// Mare sure same section names get exact same FText representation (binding widget uses IsIdentical() to compare the section names).
+	if (const FText* SectionText = InOutSectionNames.Find(StructDesc.StatePath))
 	{
-		// Only allow to binding when one object is selected.
-		OwnerObject = UE::StateTree::PropertyBinding::FindEditorBindingsOwner(OuterObjects[0]);
-
-		// Figure out the structs we're editing, and property path relative to current property.
-		UE::StateTree::PropertyBinding::MakeStructPropertyPathFromPropertyHandle(InPropertyHandle, TargetPath);
-
-		BindingOwner = Cast<IStateTreeEditorPropertyBindingsOwner>(OwnerObject);
-		if (BindingOwner)
-		{
-			EditorBindings = BindingOwner->GetPropertyEditorBindings();
-			BindingOwner->GetAccessibleStructs(TargetPath.GetStructID(), AccessibleStructs);
-
-			BindingOwner->EnumerateBindablePropertyFunctionNodes([&AccessibleStructs](const UScriptStruct* NodeStruct, const FStateTreeBindableStructDesc& Desc, const FStateTreeDataView Value)
-			{
-				AccessibleStructs.Add(Desc);
-				return EStateTreeVisitor::Continue;
-			});
-
-			TMap<FString, FText> SectionNames;
-			for (FStateTreeBindableStructDesc& StructDesc : AccessibleStructs)
-			{
-				const UStruct* Struct = StructDesc.Struct;
-
-				FBindingContextStruct& ContextStruct = BindingContextStructs.AddDefaulted_GetRef();
-				ContextStruct.DisplayText = FText::FromString(StructDesc.Name.ToString());
-				ContextStruct.Struct = const_cast<UStruct*>(Struct);
-				ContextStruct.Category = StructDesc.Category;
-
-				// Mare sure same section names get exact same FText representation (binding widget uses IsIdentical() to compare the section names).
-				if (const FText* SectionText = SectionNames.Find(StructDesc.StatePath))
-				{
-					ContextStruct.Section = *SectionText;
-				}
-				else
-				{
-					ContextStruct.Section = SectionNames.Add(StructDesc.StatePath, FText::FromString(StructDesc.StatePath));
-				}
-				
-				// PropertyFunction overrides it's struct's icon color.
-				if (StructDesc.DataSource == EStateTreeBindableStructSource::PropertyFunction)
-				{
-					if (const FProperty* OutputProperty = UE::StateTree::GetStructSingleOutputProperty(*StructDesc.Struct))
-					{
-						const UEdGraphSchema_K2* Schema = GetDefault<UEdGraphSchema_K2>();
-						check(Schema);
-
-						FEdGraphPinType PinType;
-						if (Schema->ConvertPropertyToPinType(OutputProperty, PinType))
-						{
-							ContextStruct.Color = Schema->GetPinTypeColor(PinType);
-						}
-					}
-				}
-			}
-		}
+		InOutContextStruct.Section = *SectionText;
 	}
-
-	TSharedPtr<UE::StateTree::PropertyBinding::FCachedBindingData> CachedBindingData = MakeShared<UE::StateTree::PropertyBinding::FCachedBindingData>(OwnerObject, TargetPath, InPropertyHandle, AccessibleStructs);
-
-	// Wrap value widget 
+	else
 	{
-		auto IsValueVisible = TAttribute<EVisibility>::Create([CachedBindingData]() -> EVisibility
-			{
-				return CachedBindingData->HasBinding(FStateTreeEditorPropertyBindings::ESearchMode::Exact) ? EVisibility::Collapsed : EVisibility::Visible;
-			});
-
-		TSharedPtr<SWidget> ValueWidget = InWidgetRow.ValueContent().Widget;
-		InWidgetRow.ValueContent()
-			[
-				SNew(SBox)
-				.Visibility(IsValueVisible)
-				[
-					ValueWidget.ToSharedRef()
-				]
-			];
+		InOutContextStruct.Section = InOutSectionNames.Add(StructDesc.StatePath, FText::FromString(StructDesc.StatePath));
 	}
-	
-	FPropertyBindingWidgetArgs Args;
-	Args.Property = InPropertyHandle->GetProperty();
-
-	Args.OnCanBindPropertyWithBindingChain = FOnCanBindPropertyWithBindingChain::CreateLambda([CachedBindingData](FProperty* InProperty, TConstArrayView<FBindingChainElement> InBindingChain)
-		{
-			return CachedBindingData->CanBindToProperty(InProperty, InBindingChain);
-		});
-
-	Args.OnCanBindToContextStructWithIndex = FOnCanBindToContextStructWithIndex::CreateLambda([CachedBindingData](const UStruct* InStruct, int32 InStructIndex)
-		{
-			return CachedBindingData->CanBindToContextStruct(InStruct, InStructIndex);
-		});
-
-	Args.OnCanAcceptPropertyOrChildrenWithBindingChain = FOnCanAcceptPropertyOrChildrenWithBindingChain::CreateLambda([CachedBindingData](FProperty* InProperty, TConstArrayView<FBindingChainElement> InBindingChain)
-		{
-			return CachedBindingData->CanAcceptPropertyOrChildren(InProperty, InBindingChain);
-		});
-
-	Args.OnCanBindToClass = FOnCanBindToClass::CreateLambda([](UClass* InClass)
-		{
-			return true;
-		});
-
-	Args.OnAddBinding = FOnAddBinding::CreateLambda([CachedBindingData, &InDetailBuilder](FName InPropertyName, TConstArrayView<FBindingChainElement> InBindingChain)
-		{
-			CachedBindingData->AddBinding(InBindingChain);
-			InDetailBuilder.GetPropertyUtilities()->RequestForceRefresh();
-		});
-
-	Args.OnRemoveBinding = FOnRemoveBinding::CreateLambda([CachedBindingData, &InDetailBuilder](FName InPropertyName)
-		{
-			CachedBindingData->RemoveBinding(FStateTreeEditorPropertyBindings::ESearchMode::Exact);
-			InDetailBuilder.GetPropertyUtilities()->RequestForceRefresh();
-		});
-
-	Args.OnCanRemoveBinding = FOnCanRemoveBinding::CreateLambda([CachedBindingData](FName InPropertyName)
-		{
-			return CachedBindingData->HasBinding(FStateTreeEditorPropertyBindings::ESearchMode::Exact);
-		});
 
-	Args.CurrentBindingText = MakeAttributeLambda([CachedBindingData]()
-		{
-			return CachedBindingData->GetText();
-		});
-
-	Args.CurrentBindingToolTipText = MakeAttributeLambda([CachedBindingData]()
-		{
-			return CachedBindingData->GetTooltipText();
-		});
-
-	Args.CurrentBindingImage = MakeAttributeLambda([CachedBindingData]() -> const FSlateBrush*
-		{
-			return CachedBindingData->GetImage();
-		});
-
-	Args.CurrentBindingColor = MakeAttributeLambda([CachedBindingData]() -> FLinearColor
-		{
-			return CachedBindingData->GetColor();
-		});
-
-	if (BindingOwner)
+	// PropertyFunction overrides it's struct's icon color.
+	if (StructDesc.DataSource == EStateTreeBindableStructSource::PropertyFunction)
 	{
-		Args.OnResolveIndirection = FOnResolveIndirection::CreateLambda([CachedBindingData](TConstArrayView<FBindingChainElement> InBindingChain)
+		if (const FProperty* OutputProperty = UE::StateTree::GetStructSingleOutputProperty(*StructDesc.Struct))
 		{
-			return CachedBindingData->ResolveIndirection(InBindingChain);
-		});
-	}
-
-	Args.BindButtonStyle = &FAppStyle::Get().GetWidgetStyle<FButtonStyle>("HoverHintOnly");
-	Args.bAllowNewBindings = false;
-	Args.bAllowArrayElementBindings = false;
-	Args.bAllowUObjectFunctions = false;
-
-	if (CanPromoteToParameter(InPropertyHandle))
-	{
-		Args.MenuExtender = MakeShared<FExtender>();
-		Args.MenuExtender->AddMenuExtension(
-			TEXT("BindingActions"),
-			EExtensionHook::After,
-			nullptr,
-			FMenuExtensionDelegate::CreateLambda([CachedBindingData, AccessibleStructs = MoveTemp(AccessibleStructs), InPropertyHandle](FMenuBuilder& MenuBuilder)
-			{
-				MenuBuilder.AddSubMenu(
-					LOCTEXT("PromoteToParameter", "Promote to Parameter"),
-					LOCTEXT("PromoteToParameterTooltip", "Create a new parameter of the same type as the property, copy value over, and bind the property to the new parameter."),
-					FNewMenuDelegate::CreateLambda([&CachedBindingData, &AccessibleStructs, &InPropertyHandle](FMenuBuilder& InMenuBuilder)
-					{
-						using namespace UE::StateTree::PropertyBinding;
-
-						const FProperty* Property = InPropertyHandle->GetProperty();
-						check(Property);
-						const FName PropertyName = Property->GetFName();
-
-						TSharedRef<FCachedBindingData> CachedBindingDataRef = CachedBindingData.ToSharedRef();
-
-						FSectionHelper SectionHelper(InMenuBuilder);
-						for (const FStateTreeBindableStructDesc& ContextStruct : AccessibleStructs)
-						{
-							TArray<TSharedPtr<const FRefTypeInfo>> RefTypeInfos;
-							if (CachedBindingData->CanCreateParameter(ContextStruct, /*out*/RefTypeInfos))
-							{
-								SectionHelper.SetSection(FText::FromString(ContextStruct.StatePath));
-
-								if (RefTypeInfos.IsEmpty())
-								{
-									InMenuBuilder.AddMenuEntry(FExecuteAction::CreateSP(CachedBindingDataRef, &FCachedBindingData::PromoteToParameter, PropertyName, ContextStruct, TSharedPtr<const FRefTypeInfo>()),
-										MakeContextStructWidget(ContextStruct));
-								}
-								else
-								{
-									InMenuBuilder.AddSubMenu(MakeContextStructWidget(ContextStruct),
-										FNewMenuDelegate::CreateLambda([CachedBindingDataRef, PropertyName, &ContextStruct, RefTypeInfos = MoveTemp(RefTypeInfos)](FMenuBuilder& InSubMenuBuilder)
-										{
-											FSectionHelper SectionHelper(InSubMenuBuilder);
-											SectionHelper.SetSection(LOCTEXT("RefTypeParams", "Reference Types"));
-											for (const TSharedPtr<const FRefTypeInfo>& RefTypeInfo : RefTypeInfos)
-											{
-												InSubMenuBuilder.AddMenuEntry(FExecuteAction::CreateSP(CachedBindingDataRef, &FCachedBindingData::PromoteToParameter, PropertyName, ContextStruct, RefTypeInfo),
-													MakeBindingPropertyInfoWidget(RefTypeInfo->TypeNameText, RefTypeInfo->PinType));
-											}
-										}));
-								}
-							}
-						}
-					})
-				);
-			})
-		);
-	}
+			const UEdGraphSchema_K2* Schema = GetDefault<UEdGraphSchema_K2>();
+			check(Schema);
 
-	// ResetToDefault
-	{
-		InWidgetRow.CustomResetToDefault = FResetToDefaultOverride::Create(
-			MakeAttributeLambda([CachedBindingData, InPropertyHandle]()
+			FEdGraphPinType PinType;
+			if (Schema->ConvertPropertyToPinType(OutputProperty, PinType))
 			{
-				return InPropertyHandle->CanResetToDefault() || CachedBindingData->HasBinding(FStateTreeEditorPropertyBindings::ESearchMode::Includes);
-			}),
-			FSimpleDelegate::CreateLambda([CachedBindingData, &InDetailBuilder, InPropertyHandle]()
-				{
-					if (CachedBindingData->HasBinding(FStateTreeEditorPropertyBindings::ESearchMode::Includes))
-					{
-						CachedBindingData->RemoveBinding(FStateTreeEditorPropertyBindings::ESearchMode::Includes);
-						InDetailBuilder.GetPropertyUtilities()->RequestForceRefresh();
-					}
-					if (InPropertyHandle->CanResetToDefault())
-					{
-						InPropertyHandle->ResetToDefault();
-					}
-				}),
-			false);
+				InOutContextStruct.Color = Schema->GetPinTypeColor(PinType);
+			}
+		}
 	}
-
-	InWidgetRow.ExtensionContent()
-	[
-		PropertyAccessEditor.MakePropertyBindingWidget(BindingContextStructs, Args)
-	];
 }
 
-bool FStateTreeBindingExtension::CanPromoteToParameter(const TSharedPtr<IPropertyHandle>& InPropertyHandle) const
+bool FStateTreeBindingExtension::GetPromotionToParameterOverrideInternal(const FProperty& InProperty, bool& bOutOverride) const
 {
-	const FProperty* Property = InPropertyHandle->GetProperty();
-	if (!Property)
-	{
-		return false;
-	}
-
-	// Property Bag picker only detects Blueprint Types, so only allow properties that are blueprint types
-	// FPropertyBagInstanceDataDetails::OnPropertyNameContent uses SPinTypeSelector to generate the property type picker.
-	// UEdGraphSchema_K2::GetVariableTypeTree (GatherPinsImpl: FindEnums, FindStructs, FindObjectsAndInterfaces) is used there which only allows bp types.
-	// The below behavior mirrors the behavior in the pin gathering but for properties
-
-	if (const FEnumProperty* EnumProperty = CastField<FEnumProperty>(Property))
-	{
-		if (!UEdGraphSchema_K2::IsAllowableBlueprintVariableType(EnumProperty->GetEnum()))
-		{
-			return false;
-		}
-	}
-	else if (const FStructProperty* StructProperty = CastField<FStructProperty>(Property))
+	if (const FStructProperty* StructProperty = CastField<FStructProperty>(&InProperty))
 	{
 		// Support Property Refs as even though these aren't bp types, the actual types that would be added are the ones in the meta-data RefType
 		if (StructProperty->Struct && StructProperty->Struct->IsChildOf(FStateTreePropertyRef::StaticStruct()))
 		{
-			return true;
-		}
-
-		if (!UEdGraphSchema_K2::IsAllowableBlueprintVariableType(StructProperty->Struct))
-		{
-			return false;
-		}
-	}
-	else if (const FObjectPropertyBase* ObjectProperty = CastField<FObjectPropertyBase>(Property))
-	{
-		if (!UEdGraphSchema_K2::IsAllowableBlueprintVariableType(ObjectProperty->PropertyClass))
-		{
-			return false;
-		}
-	}
-	else if (const FInterfaceProperty* InterfaceProperty = CastField<FInterfaceProperty>(Property))
-	{
-		if (!UEdGraphSchema_K2::IsAllowableBlueprintVariableType(InterfaceProperty->InterfaceClass))
-		{
-			return false;
+			bOutOverride = false;
 		}
 	}
-
-	return true;
+	return false;
 }
 
 void FStateTreeBindingsChildrenCustomization::CustomizeChildren(IDetailChildrenBuilder& ChildrenBuilder, TSharedPtr<IPropertyHandle> InPropertyHandle)
@@ -1849,7 +614,7 @@ void FStateTreeBindingsChildrenCustomization::CustomizeChildren(IDetailChildrenB
 	if (OuterObjects.Num() == 1)
 	{
 		FStateTreePropertyPath TargetPath;
-		UE::StateTree::PropertyBinding::MakeStructPropertyPathFromPropertyHandle(InPropertyHandle, TargetPath);
+		UE::PropertyBinding::MakeStructPropertyPathFromPropertyHandle(InPropertyHandle, TargetPath);
 
 		using FStateTreePropertyFunctionNodeProvider = UE::StateTree::PropertyBinding::FStateTreePropertyFunctionNodeProvider;
 		UObject* BindingsOwner = UE::StateTree::PropertyBinding::FindEditorBindingsOwner(OuterObjects[0]);
@@ -1872,13 +637,13 @@ bool FStateTreeBindingsChildrenCustomization::ShouldCustomizeChildren(TSharedRef
 	{
 		// Bound property's children composition gets overridden.
 		FStateTreePropertyPath TargetPath;
-		UE::StateTree::PropertyBinding::MakeStructPropertyPathFromPropertyHandle(InPropertyHandle, TargetPath);
+		UE::PropertyBinding::MakeStructPropertyPathFromPropertyHandle(InPropertyHandle, TargetPath);
 		IStateTreeEditorPropertyBindingsOwner* BindingOwner = Cast<IStateTreeEditorPropertyBindingsOwner>(UE::StateTree::PropertyBinding::FindEditorBindingsOwner(OuterObjects[0]));
 		if (!TargetPath.IsPathEmpty() && BindingOwner)
 		{
-			if (FStateTreeEditorPropertyBindings* EditorBindings = BindingOwner->GetPropertyEditorBindings())
+			if (const FStateTreeEditorPropertyBindings* EditorBindings = BindingOwner->GetPropertyEditorBindings())
 			{
-				return EditorBindings->HasPropertyBinding(TargetPath);
+				return EditorBindings->HasBinding(TargetPath);
 			}
 		}
 	}
diff --git a/Engine/Plugins/Runtime/StateTree/Source/StateTreeEditorModule/Private/Customizations/StateTreeBindingExtension.h b/Engine/Plugins/Runtime/StateTree/Source/StateTreeEditorModule/Private/Customizations/StateTreeBindingExtension.h
--- a/Engine/Plugins/Runtime/StateTree/Source/StateTreeEditorModule/Private/Customizations/StateTreeBindingExtension.h
+++ b/Engine/Plugins/Runtime/StateTree/Source/StateTreeEditorModule/Private/Customizations/StateTreeBindingExtension.h
@@ -4,6 +4,7 @@
 
 #include "IDetailPropertyExtensionHandler.h"
 #include "IDetailPropertyChildrenCustomizationHandler.h"
+#include "PropertyBindingExtension.h"
 
 enum class EStateTreePropertyUsage : uint8;
 
@@ -15,48 +16,23 @@ class FProperty;
 
 namespace UE::StateTree::PropertyBinding
 {
-	extern const FName StateTreeNodeIDName;
-
-	/**
-	 * Get nearest Outer that implements IStateTreeEditorPropertyBindingsOwner.
-	 * @param InObject Object where to start the search.
-	 */
-	UObject* FindEditorBindingsOwner(UObject* InObject);
-
-	/**
-	 * Returns property path for a specific property.
-	 * Walks towards root up until a property with metadata  "StateTreeItemID" is found. 
-	 * The property's metadata "StateTreeItemID" is expected to specify the containing struct ID.
-	 * @param InPropertyHandle Handle to the property to find path for.
-	 * @param OutPath Resulting property path.
-	 * @return usage type extracted from the root-most property.
-	 */
-	EStateTreePropertyUsage MakeStructPropertyPathFromPropertyHandle(TSharedPtr<const IPropertyHandle> InPropertyHandle, FStateTreePropertyPath& OutPath);
-
 	DECLARE_MULTICAST_DELEGATE_TwoParams(FOnStateTreePropertyBindingChanged, const FStateTreePropertyPath& /*SourcePath*/, const FStateTreePropertyPath& /*TargetPath*/);
 	extern STATETREEEDITORMODULE_API FOnStateTreePropertyBindingChanged OnStateTreePropertyBindingChanged;
-
-	/**
-	 * Returns true if provided Property is bindable.
-	 */
-	bool IsPropertyBindable(const FProperty& Property);
 } // UE::StateTree::PropertyBinding
 
-class FStateTreeBindingExtension : public IDetailPropertyExtensionHandler
+class FStateTreeBindingExtension : public FPropertyBindingExtension
 {
-public:
-	// IDetailPropertyExtensionHandler interface
-	virtual bool IsPropertyExtendable(const UClass* InObjectClass, const IPropertyHandle& PropertyHandle) const override;
-	virtual void ExtendWidgetRow(FDetailWidgetRow& InWidgetRow, const IDetailLayoutBuilder& InDetailBuilder, const UClass* InObjectClass, TSharedPtr<IPropertyHandle> PropertyHandle) override;
-
 private:
-	bool CanPromoteToParameter(const TSharedPtr<IPropertyHandle>& InPropertyHandle) const;
+	virtual TSharedPtr<UE::PropertyBinding::FCachedBindingData> CreateCachedBindingData(IPropertyBindingBindingCollectionOwner* InBindingsOwner, const FPropertyBindingPath& InTargetPath, const TSharedPtr<IPropertyHandle>& InPropertyHandle, TConstArrayView<TInstancedStruct<FPropertyBindingBindableStructDescriptor>> InAccessibleStructs) const override;
+	virtual bool CanBindToProperty(const FPropertyBindingPath& InTargetPath, const IPropertyHandle& InPropertyHandle) const override;
+	virtual bool GetPromotionToParameterOverrideInternal(const FProperty& InProperty, bool& bOutOverride) const override;
+	virtual void UpdateContextStruct(TConstStructView<FPropertyBindingBindableStructDescriptor> InStructDesc,FBindingContextStruct& InOutContextStruct, TMap<FString, FText>& InOutSectionNames) const override;
 };
 
 /* Overrides bound property's children composition. */
 class FStateTreeBindingsChildrenCustomization : public IDetailPropertyChildrenCustomizationHandler
 {
 public:
-	bool ShouldCustomizeChildren(TSharedRef<IPropertyHandle> InPropertyHandle) override;
-	void CustomizeChildren(IDetailChildrenBuilder& ChildrenBuilder, TSharedPtr<IPropertyHandle> InPropertyHandle) override;
+	virtual bool ShouldCustomizeChildren(TSharedRef<IPropertyHandle> InPropertyHandle) override;
+	virtual void CustomizeChildren(IDetailChildrenBuilder& ChildrenBuilder, TSharedPtr<IPropertyHandle> InPropertyHandle) override;
 };
diff --git a/Engine/Plugins/Runtime/StateTree/Source/StateTreeEditorModule/Private/Customizations/StateTreeEditorNodeDetails.cpp b/Engine/Plugins/Runtime/StateTree/Source/StateTreeEditorModule/Private/Customizations/StateTreeEditorNodeDetails.cpp
--- a/Engine/Plugins/Runtime/StateTree/Source/StateTreeEditorModule/Private/Customizations/StateTreeEditorNodeDetails.cpp
+++ b/Engine/Plugins/Runtime/StateTree/Source/StateTreeEditorModule/Private/Customizations/StateTreeEditorNodeDetails.cpp
@@ -166,7 +166,7 @@ namespace UE::StateTreeEditor::Internal
 		if (Usage != EStateTreePropertyUsage::Invalid && ID.IsValid())
 		{
 			// Pass the node ID to binding extension. Since the properties are added using AddChildStructure(), we break the hierarchy and cannot access parent.
-			ChildPropHandle->SetInstanceMetaData(UE::StateTree::PropertyBinding::StateTreeNodeIDName, LexToString(ID));
+			ChildPropHandle->SetInstanceMetaData(UE::PropertyBinding::MetaDataStructIDName, LexToString(ID));
 
 			FStateTreePropertyPath Path(ID, *Property->GetFName().ToString());
 			TSharedPtr<SWidget> NameWidget;
@@ -174,7 +174,7 @@ namespace UE::StateTreeEditor::Internal
 			FDetailWidgetRow Row;
 			ChildRow.GetDefaultWidgets(NameWidget, ValueWidget, Row);
 
-			const bool bHasChildPropertyBinding = EditorPropBindings->HasPropertyBinding(Path, FStateTreeEditorPropertyBindings::ESearchMode::Includes);
+			const bool bHasChildPropertyBinding = EditorPropBindings->HasBinding(Path, FPropertyBindingBindingCollection::ESearchMode::Includes);
 			const bool bValidUsage = Usage == EStateTreePropertyUsage::Input || Usage == EStateTreePropertyUsage::Output || Usage == EStateTreePropertyUsage::Context;
 			const bool bIsDelegateDispatcher = UE::StateTreeEditor::Internal::IsDelegateDispatcherProperty(*Property);
 
@@ -188,7 +188,7 @@ namespace UE::StateTreeEditor::Internal
 				{
 					// Use internal type to construct PinType if it's property of PropertyRef type.
 					FStateTreeDataView TargetDataView;
-					if (ensure(EditorData->GetDataViewByID(ID, TargetDataView)))
+					if (ensure(EditorData->GetBindingDataViewByID(ID, TargetDataView)))
 					{
 						TArray<FPropertyBindingPathIndirection> TargetIndirections;
 						if (ensure(Path.ResolveIndirectionsWithValue(TargetDataView, TargetIndirections)))
@@ -205,7 +205,7 @@ namespace UE::StateTreeEditor::Internal
 
 				auto IsValueVisible = TAttribute<EVisibility>::Create([Path, EditorPropBindings]() -> EVisibility
 					{
-						return EditorPropBindings->HasPropertyBinding(Path, FStateTreeEditorPropertyBindings::ESearchMode::Exact) ? EVisibility::Collapsed : EVisibility::Visible;
+						return EditorPropBindings->HasBinding(Path, FPropertyBindingBindingCollection::ESearchMode::Exact) ? EVisibility::Collapsed : EVisibility::Visible;
 					});
 
 				const FSlateBrush* Icon = FBlueprintEditorUtils::GetIconFromPin(PinType, true);
diff --git a/Engine/Plugins/Runtime/StateTree/Source/StateTreeEditorModule/Private/Customizations/StateTreeEventDescDetails.cpp b/Engine/Plugins/Runtime/StateTree/Source/StateTreeEditorModule/Private/Customizations/StateTreeEventDescDetails.cpp
--- a/Engine/Plugins/Runtime/StateTree/Source/StateTreeEditorModule/Private/Customizations/StateTreeEventDescDetails.cpp
+++ b/Engine/Plugins/Runtime/StateTree/Source/StateTreeEditorModule/Private/Customizations/StateTreeEventDescDetails.cpp
@@ -7,6 +7,7 @@
 #include "DetailLayoutBuilder.h"
 #include "IDetailChildrenBuilder.h"
 #include "IPropertyUtilities.h"
+#include "PropertyBindingExtension.h"
 #include "StateTreeBindingExtension.h"
 #include "StateTreeEditorStyle.h"
 #include "Styling/StyleColors.h"
@@ -94,7 +95,7 @@ public:
 		for (TFieldIterator<const FProperty> PropIt(Struct); PropIt; ++PropIt)
 		{
 			const FProperty& Property = **PropIt;
-			if (UE::StateTree::PropertyBinding::IsPropertyBindable(Property))
+			if (UE::PropertyBinding::IsPropertyBindable(Property))
 			{
 				ChildrenBuilder.AddCustomBuilder(MakeShared<FStateTreeEventPayloadDetails>(Property, VisitedStructs));
 			}
diff --git a/Engine/Plugins/Runtime/StateTree/Source/StateTreeEditorModule/Private/Customizations/StateTreeStateParametersDetails.cpp b/Engine/Plugins/Runtime/StateTree/Source/StateTreeEditorModule/Private/Customizations/StateTreeStateParametersDetails.cpp
--- a/Engine/Plugins/Runtime/StateTree/Source/StateTreeEditorModule/Private/Customizations/StateTreeStateParametersDetails.cpp
+++ b/Engine/Plugins/Runtime/StateTree/Source/StateTreeEditorModule/Private/Customizations/StateTreeStateParametersDetails.cpp
@@ -81,7 +81,7 @@ void FStateTreeStateParametersInstanceDataDetails::OnChildRowAdded(IDetailProper
 		// Conditionally control visibility of the value field of bound properties.
 
 		// Pass the node ID to binding extension. Since the properties are added using AddChildStructure(), we break the hierarchy and cannot access parent.
-		ChildPropHandle->SetInstanceMetaData(UE::StateTree::PropertyBinding::StateTreeNodeIDName, LexToString(ID));
+		ChildPropHandle->SetInstanceMetaData(UE::PropertyBinding::MetaDataStructIDName, LexToString(ID));
 
 		FStateTreePropertyPath Path(ID, *Property->GetFName().ToString());
 		
@@ -92,7 +92,7 @@ void FStateTreeStateParametersInstanceDataDetails::OnChildRowAdded(IDetailProper
 				{
 					if (const FStateTreeEditorPropertyBindings* EditorPropBindings = EditorData->GetPropertyEditorBindings())
 					{
-						bHasBinding = EditorPropBindings->HasPropertyBinding(Path); 
+						bHasBinding = EditorPropBindings->HasBinding(Path); 
 					}
 				}
 
diff --git a/Engine/Plugins/Runtime/StateTree/Source/StateTreeEditorModule/Private/Customizations/StateTreeTransitionDetails.cpp b/Engine/Plugins/Runtime/StateTree/Source/StateTreeEditorModule/Private/Customizations/StateTreeTransitionDetails.cpp
--- a/Engine/Plugins/Runtime/StateTree/Source/StateTreeEditorModule/Private/Customizations/StateTreeTransitionDetails.cpp
+++ b/Engine/Plugins/Runtime/StateTree/Source/StateTreeEditorModule/Private/Customizations/StateTreeTransitionDetails.cpp
@@ -199,7 +199,7 @@ void FStateTreeTransitionDetails::CustomizeChildren(TSharedRef<class IPropertyHa
 
 	FGuid ID;
 	UE::StateTree::PropertyHelpers::GetStructValue<FGuid>(IDProperty, ID);
-	DelegateListener->SetInstanceMetaData(UE::StateTree::PropertyBinding::StateTreeNodeIDName, LexToString(ID));
+	DelegateListener->SetInstanceMetaData(UE::PropertyBinding::MetaDataStructIDName, LexToString(ID));
 
 	// State
 	StructBuilder.AddProperty(StateProperty.ToSharedRef());
diff --git a/Engine/Plugins/Runtime/StateTree/Source/StateTreeEditorModule/Private/StateTreeCompiler.cpp b/Engine/Plugins/Runtime/StateTree/Source/StateTreeEditorModule/Private/StateTreeCompiler.cpp
--- a/Engine/Plugins/Runtime/StateTree/Source/StateTreeEditorModule/Private/StateTreeCompiler.cpp
+++ b/Engine/Plugins/Runtime/StateTree/Source/StateTreeEditorModule/Private/StateTreeCompiler.cpp
@@ -15,7 +15,6 @@
 #include "StateTreePropertyRef.h"
 #include "StateTreePropertyRefHelpers.h"
 #include "StateTreePropertyHelpers.h"
-#include "StateTreePropertyFunctionBase.h"
 #include "StateTreeDelegate.h"
 
 
@@ -502,7 +501,7 @@ bool FStateTreeCompiler::CreateConditions(UStateTreeState& State, const FString&
 	{
 		const bool bIsFirst = Index == 0;
 		const FStateTreeEditorNode& CondNode = Conditions[Index];
-		// First operand should be copy as we dont have a previous item to operate on.
+		// First operand should be copied as we dont have a previous item to operate on.
 		const EStateTreeExpressionOperand Operand = bIsFirst ? EStateTreeExpressionOperand::Copy : CondNode.ExpressionOperand;
 		// First indent must be 0 to make the parentheses calculation match.
 		const int32 CurrIndent = bIsFirst ? 0 : FMath::Clamp((int32)CondNode.ExpressionIndent, 0, UE::StateTree::MaxExpressionIndent);
@@ -996,7 +995,7 @@ bool FStateTreeCompiler::CreateStateTransitions()
 			{
 				FStateTreePropertyPath DelegateBindingPath(Transition.ID, GET_MEMBER_NAME_CHECKED(FStateTreeTransition, DelegateListener));
 
-				const FStateTreePropertyPathBinding* Binding = EditorData->EditorBindings.FindPropertyBinding(DelegateBindingPath);
+				const FPropertyBindingBinding* Binding = EditorData->EditorBindings.FindBinding(DelegateBindingPath);
 
 				if (Binding == nullptr)
 				{
@@ -1334,7 +1333,7 @@ bool FStateTreeCompiler::CreateBindingsForStruct(const FStateTreeBindableStructD
 
 bool FStateTreeCompiler::CreatePropertyFunctionsForStruct(FGuid StructID)
 {
-	for (const FStateTreePropertyPathBinding& Binding : EditorData->EditorBindings.GetBindings())
+	for (const FPropertyBindingBinding& Binding : EditorData->EditorBindings.GetBindings())
 	{
 		if (Binding.GetTargetPath().GetStructID() != StructID)
 		{
@@ -2006,7 +2005,7 @@ bool FStateTreeCompiler::GetAndValidateBindings(const FStateTreeBindableStructDe
 
 	bool bSucceeded = true;
 
-	for (FStateTreePropertyPathBinding& Binding : EditorData->EditorBindings.GetMutableBindings())
+	for (FPropertyBindingBinding& Binding : EditorData->EditorBindings.GetMutableBindings())
 	{
 		if (Binding.GetTargetPath().GetStructID() != TargetStruct.ID)
 		{
@@ -2014,7 +2013,7 @@ bool FStateTreeCompiler::GetAndValidateBindings(const FStateTreeBindableStructDe
 			{
 				if (IsPropertyOfTypeOrChild(*FStateTreeDelegateDispatcher::StaticStruct(), TargetStruct, Binding.GetSourcePath()))
 				{
-					FStateTreePropertyPathBinding& BindingCopy = OutValidatedBindings.DelegateSourceBindings.Add_GetRef(Binding);
+					FStateTreePropertyPathBinding& BindingCopy = OutValidatedBindings.DelegateSourceBindings.Add_GetRef(FStateTreePropertyPathBinding(Binding));
 					BindingCopy.SetSourceDataHandle(TargetStruct.DataHandle);
 				}
 			}
@@ -2063,12 +2062,12 @@ bool FStateTreeCompiler::GetAndValidateBindings(const FStateTreeBindableStructDe
 		}
 		
 		// Source must be accessible by the target struct via all execution paths.
-		TArray<FStateTreeBindableStructDesc> AccessibleStructs;
-		EditorData->GetAccessibleStructs(Binding.GetTargetPath().GetStructID(), AccessibleStructs);
+		TArray<TInstancedStruct<FPropertyBindingBindableStructDescriptor>> AccessibleStructs;
+		EditorData->GetBindableStructs(Binding.GetTargetPath().GetStructID(), AccessibleStructs);
 
-		const bool bSourceAccessible = AccessibleStructs.ContainsByPredicate([SourceStructID](const FStateTreeBindableStructDesc& Structs)
+		const bool bSourceAccessible = AccessibleStructs.ContainsByPredicate([SourceStructID](TConstStructView<FPropertyBindingBindableStructDescriptor> Struct)
 			{
-				return (Structs.ID == SourceStructID);
+				return (Struct.Get().ID == SourceStructID);
 			});
 
 		if (!bSourceAccessible)
@@ -2114,7 +2113,7 @@ bool FStateTreeCompiler::GetAndValidateBindings(const FStateTreeBindableStructDe
 			continue;
 		}
 		
-		FStateTreePropertyPathBinding& BindingCopy = OutputBindings->Add_GetRef(Binding);
+		FStateTreePropertyPathBinding& BindingCopy = OutputBindings->Add_GetRef(FStateTreePropertyPathBinding(Binding));
 		BindingCopy.SetSourceDataHandle(SourceStruct->DataHandle);
 
 		// Special case for AnyEnum. StateTreeBindingExtension allows AnyEnums to bind to other enum types.
@@ -2133,7 +2132,7 @@ bool FStateTreeCompiler::GetAndValidateBindings(const FStateTreeBindableStructDe
 				BindingCopy.GetMutableTargetPath().AddPathSegment(GET_MEMBER_NAME_STRING_CHECKED(FStateTreeAnyEnum, Value));
 			}
 		}
-		
+
 		// Check if the bindings is for struct ref and validate the types.
 		if (!ValidateStructRef(*SourceStruct, Binding.GetSourcePath(), TargetStruct, Binding.GetTargetPath()))
 		{
@@ -2209,7 +2208,7 @@ bool FStateTreeCompiler::GetAndValidateBindings(const FStateTreeBindableStructDe
 			}
 			else if (Usage == EStateTreePropertyUsage::Context)
 			{
-				// Make sure that an Context property is manually or automatically bound. 
+				// Make sure a Context property is manually or automatically bound.
 				const UStruct* ContextObjectType = nullptr; 
 				if (const FStructProperty* StructProperty = CastField<FStructProperty>(Property))
 				{
diff --git a/Engine/Plugins/Runtime/StateTree/Source/StateTreeEditorModule/Private/StateTreeEditingSubsystem.cpp b/Engine/Plugins/Runtime/StateTree/Source/StateTreeEditorModule/Private/StateTreeEditingSubsystem.cpp
--- a/Engine/Plugins/Runtime/StateTree/Source/StateTreeEditorModule/Private/StateTreeEditingSubsystem.cpp
+++ b/Engine/Plugins/Runtime/StateTree/Source/StateTreeEditorModule/Private/StateTreeEditingSubsystem.cpp
@@ -278,11 +278,11 @@ void UStateTreeEditingSubsystem::ValidateStateTree(TNotNull<const UStateTree*> I
 			return;
 		}
 
-		TMap<FGuid, const FStateTreeDataView> AllStructValues;
+		TMap<FGuid, const FPropertyBindingDataView> AllStructValues;
 		TreeData->GetAllStructValues(AllStructValues);
 		constexpr bool bMarkDirty = false;
 		TreeData->Modify(bMarkDirty);
-		TreeData->GetPropertyEditorBindings()->RemoveUnusedBindings(AllStructValues);
+		TreeData->GetPropertyEditorBindings()->RemoveInvalidBindings(AllStructValues);
 	};
 
 	auto UpdateLinkedStateParameters = [](const UStateTree& StateTree)
