[PropertyBindingsUtils]

diff --git a/Engine/Plugins/Runtime/StateTree/Source/StateTreeEditorModule/Private/StateTreeEditorData.cpp b/Engine/Plugins/Runtime/StateTree/Source/StateTreeEditorModule/Private/StateTreeEditorData.cpp
--- a/Engine/Plugins/Runtime/StateTree/Source/StateTreeEditorModule/Private/StateTreeEditorData.cpp
+++ b/Engine/Plugins/Runtime/StateTree/Source/StateTreeEditorModule/Private/StateTreeEditorData.cpp
@@ -12,6 +12,7 @@
 #include "StateTreePropertyHelpers.h"
 #include "StateTreeTaskBase.h"
 #include "Algo/LevenshteinDistance.h"
+#include "Customizations/StateTreeBindingExtension.h"
 #include "Customizations/StateTreeEditorNodeUtils.h"
 #include "Modules/ModuleManager.h"
 
@@ -32,9 +33,9 @@ namespace UE::StateTree::Editor
 
 	bool IsPropertyFunctionOwnedByNode(FGuid NodeID, FGuid PropertyFuncID, const FStateTreeEditorPropertyBindings& EditorBindings)
 	{
-		for (const FStateTreePropertyPathBinding& Binding : EditorBindings.GetBindings())
+		for (const FPropertyBindingBinding& Binding : EditorBindings.GetBindings())
 		{
-			const FGuid TargetID = Binding.GetTargetPath().GetStructID();			
+			const FGuid TargetID = Binding.GetTargetPath().GetStructID();
 			if (TargetID == NodeID)
 			{
 				return true;
@@ -276,10 +277,10 @@ void UStateTreeEditorData::PostEditChangeChainProperty(FPropertyChangedChainEven
 			if (MemberProperty->GetFName() == GET_MEMBER_NAME_CHECKED(UStateTreeEditorData, Evaluators)
 				|| MemberProperty->GetFName() == GET_MEMBER_NAME_CHECKED(UStateTreeEditorData, GlobalTasks))
 			{
-				TMap<FGuid, const FStateTreeDataView> AllStructValues;
+				TMap<FGuid, const FPropertyBindingDataView> AllStructValues;
 				GetAllStructValues(AllStructValues);
 				Modify();
-				EditorBindings.RemoveUnusedBindings(AllStructValues);
+				EditorBindings.RemoveInvalidBindings(AllStructValues);
 			}
 		}
 
@@ -322,7 +323,7 @@ void UStateTreeEditorData::PostDuplicate(EDuplicateMode::Type DuplicateMode)
 }
 #endif // WITH_EDITOR
 
-void UStateTreeEditorData::GetAccessibleStructs(const FGuid TargetStructID, TArray<FStateTreeBindableStructDesc>& OutStructDescs) const
+void UStateTreeEditorData::GetBindableStructs(const FGuid TargetStructID, TArray<TInstancedStruct<FPropertyBindingBindableStructDescriptor>>& OutStructDescs) const
 {
 	// Find the states that are updated before the current state.
 	TArray<const UStateTreeState*> Path;
@@ -340,19 +341,19 @@ void UStateTreeEditorData::GetAccessibleStructs(const FGuid TargetStructID, TArr
 		State = State->Parent;
 	}
 	
-	GetAccessibleStructs(Path, TargetStructID, OutStructDescs);
+	GetAccessibleStructsInExecutionPath(Path, TargetStructID, OutStructDescs);
 }
 
-void UStateTreeEditorData::GetAccessibleStructs(const TConstArrayView<const UStateTreeState*> Path, const FGuid TargetStructID, TArray<FStateTreeBindableStructDesc>& OutStructDescs) const
+void UStateTreeEditorData::GetAccessibleStructsInExecutionPath(const TConstArrayView<const UStateTreeState*> Path, const FGuid TargetStructID, TArray<TInstancedStruct<FPropertyBindingBindableStructDescriptor>>& OutStructDescs) const
 {
 	const UStateTree* StateTree = GetTypedOuter<UStateTree>();
 	checkf(StateTree, TEXT("UStateTreeEditorData should only be allocated within a UStateTree"));
 
-	FStateTreeBindableStructDesc TargetStructDesc;
+	TInstancedStruct<FPropertyBindingBindableStructDescriptor> TargetStructDesc;
 	bool bIsTargetPropertyFunction = false;
-	if (GetStructByID(TargetStructID, TargetStructDesc))
+	if (GetBindableStructByID(TargetStructID, TargetStructDesc))
 	{
-		bIsTargetPropertyFunction = TargetStructDesc.DataSource == EStateTreeBindableStructSource::PropertyFunction;
+		bIsTargetPropertyFunction = TargetStructDesc.Get<FStateTreeBindableStructDesc>().DataSource == EStateTreeBindableStructSource::PropertyFunction;
 	}
 
 	EStateTreeVisitor BaseProgress = VisitGlobalNodes([&OutStructDescs, TargetStructID]
@@ -363,14 +364,14 @@ void UStateTreeEditorData::GetAccessibleStructs(const TConstArrayView<const USta
 			return EStateTreeVisitor::Break;
 		}
 		
-		OutStructDescs.Add(Desc);
+		OutStructDescs.Add(TInstancedStruct<FStateTreeBindableStructDesc>::Make(Desc));
 		
 		return EStateTreeVisitor::Continue;
 	});
 
 	if (BaseProgress == EStateTreeVisitor::Continue)
 	{
-		TArray<FStateTreeBindableStructDesc, TInlineAllocator<32>> BindableDescs;
+		TArray<TInstancedStruct<FStateTreeBindableStructDesc>, TInlineAllocator<32>> BindableDescs;
 
 		for (const UStateTreeState* State : Path)
 		{
@@ -395,7 +396,7 @@ void UStateTreeEditorData::GetAccessibleStructs(const TConstArrayView<const USta
 						case EStateTreeBindableStructSource::StateParameter:
 						case EStateTreeBindableStructSource::Task:
 						case EStateTreeBindableStructSource::StateEvent:
-							BindableDescs.Add(Desc);
+							BindableDescs.Add(TInstancedStruct<FStateTreeBindableStructDesc>::Make(Desc));
 							break;
 
 						case EStateTreeBindableStructSource::TransitionEvent:
@@ -413,7 +414,7 @@ void UStateTreeEditorData::GetAccessibleStructs(const TConstArrayView<const USta
 										{
 											if (Transition.GetEventID() == Desc.ID)
 											{
-												BindableDescs.Add(Desc);
+												BindableDescs.Add(TInstancedStruct<FStateTreeBindableStructDesc>::Make(Desc));
 											}
 
 											bFoundOwningTransition = true;
@@ -436,7 +437,7 @@ void UStateTreeEditorData::GetAccessibleStructs(const TConstArrayView<const USta
 							{
 								if (UE::StateTree::Editor::IsPropertyFunctionOwnedByNode(TargetStructID, Desc.ID, EditorBindings))
 								{
-									BindableDescs.Add(Desc);
+									BindableDescs.Add(TInstancedStruct<FStateTreeBindableStructDesc>::Make(Desc));
 								}
 							}
 						}
@@ -523,7 +524,7 @@ EStateTreeVisitor UStateTreeEditorData::EnumerateBindablePropertyFunctionNodes(T
 	FStateTreeNodeClassCache* ClassCache = EditorModule.GetNodeClassCache().Get();
 	check(ClassCache);
 
-	TArray<TSharedPtr<FStateTreeNodeClassData>> StructNodes;	
+	TArray<TSharedPtr<FStateTreeNodeClassData>> StructNodes;
 	ClassCache->GetStructs(FStateTreePropertyFunctionBase::StaticStruct(), StructNodes);
 	for (const TSharedPtr<FStateTreeNodeClassData>& NodeClassData : StructNodes)
 	{
@@ -558,6 +559,15 @@ EStateTreeVisitor UStateTreeEditorData::EnumerateBindablePropertyFunctionNodes(T
 	return EStateTreeVisitor::Continue;
 }
 
+void UStateTreeEditorData::AppendBindablePropertyFunctionStructs(TArray<TInstancedStruct<FPropertyBindingBindableStructDescriptor>>& InOutStructs) const
+{
+	EnumerateBindablePropertyFunctionNodes([&InOutStructs](const UScriptStruct*, const FStateTreeBindableStructDesc& Desc, const FPropertyBindingDataView)
+		{
+			InOutStructs.Add(TInstancedStruct<FStateTreeBindableStructDesc>::Make(Desc));
+			return EStateTreeVisitor::Continue;
+		});
+}
+
 bool UStateTreeEditorData::CanCreateParameter(const FGuid StructID) const
 {
 	if (RootParameters.ID == StructID)
@@ -580,57 +590,28 @@ bool UStateTreeEditorData::CanCreateParameter(const FGuid StructID) const
 	return bFoundStructID;
 }
 
-void UStateTreeEditorData::CreateParameters(const FGuid StructID, TArrayView<FStateTreeEditorPropertyCreationDesc> InOutCreationDescs)
+void UStateTreeEditorData::CreateParametersForStruct(const FGuid StructID, TArrayView<UE::PropertyBinding::FPropertyCreationDescriptor> InOutCreationDescs)
 {
 	if (InOutCreationDescs.IsEmpty())
 	{
 		return;
 	}
 
-	auto CreateProperties = [&InOutCreationDescs](FInstancedPropertyBag& PropertyBag)
-	{
-		TArray<FPropertyBagPropertyDesc, TInlineAllocator<1>> PropertyDescs;
-		PropertyDescs.Reserve(InOutCreationDescs.Num());
-
-		// Generate unique names for the incoming property descs to avoid changing the existing properties in the bag
-		for (FStateTreeEditorPropertyCreationDesc& CreationDesc : InOutCreationDescs)
-		{
-			int32 Index = CreationDesc.PropertyDesc.Name.GetNumber();
-			while (PropertyBag.FindPropertyDescByName(CreationDesc.PropertyDesc.Name))
-			{
-				CreationDesc.PropertyDesc.Name = FName(CreationDesc.PropertyDesc.Name, Index++);
-			}
-			PropertyDescs.Add(CreationDesc.PropertyDesc);
-		}
-
-		PropertyBag.AddProperties(PropertyDescs);
-
-		for (const FStateTreeEditorPropertyCreationDesc& CreationDesc : InOutCreationDescs)
-		{
-			// Attempt to copy the value from the Source Property / Addr to Property Desc
-			// There could be Type Mismatches if the Descs don't match the Source Property, but attempt to do it on all property descs
-			if (CreationDesc.SourceProperty && CreationDesc.SourceContainerAddress)
-			{
-				PropertyBag.SetValue(CreationDesc.PropertyDesc.Name, CreationDesc.SourceProperty, CreationDesc.SourceContainerAddress);
-			}
-		}
-	};
-
 	const UStateTree* StateTree = GetTypedOuter<UStateTree>();
 	checkf(StateTree, TEXT("UStateTreeEditorData should only be allocated within a UStateTree"));
 
 	if (RootParameters.ID == StructID)
 	{
-		CreateProperties(RootParameters.Parameters);
+		UE::PropertyBinding::CreateUniquelyNamedPropertiesInPropertyBag(InOutCreationDescs, RootParameters.Parameters);
 		UE::StateTree::Delegates::OnParametersChanged.Broadcast(*StateTree);
 		return;
 	}
 
-	VisitHierarchy([&StructID, &CreateProperties, StateTree](UStateTreeState& State, UStateTreeState* ParentState)->EStateTreeVisitor
+	VisitHierarchy([&StructID, StateTree, &InOutCreationDescs](UStateTreeState& State, UStateTreeState* ParentState)->EStateTreeVisitor
 	{
 		if (State.Parameters.ID == StructID)
 		{
-			CreateProperties(State.Parameters.Parameters);
+			UE::PropertyBinding::CreateUniquelyNamedPropertiesInPropertyBag(InOutCreationDescs,State.Parameters.Parameters);
 			UE::StateTree::Delegates::OnStateParametersChanged.Broadcast(*StateTree, State.ID);
 			return EStateTreeVisitor::Break;
 		}
@@ -638,13 +619,18 @@ void UStateTreeEditorData::CreateParameters(const FGuid StructID, TArrayView<FSt
 	});
 }
 
-bool UStateTreeEditorData::GetStructByID(const FGuid StructID, FStateTreeBindableStructDesc& OutStructDesc) const
+void UStateTreeEditorData::OnPropertyBindingChanged(const FPropertyBindingPath& InSourcePath, const FPropertyBindingPath& InTargetPath)
+{
+	UE::StateTree::PropertyBinding::OnStateTreePropertyBindingChanged.Broadcast(InSourcePath, InTargetPath);
+}
+
+bool UStateTreeEditorData::GetBindableStructByID(const FGuid StructID, TInstancedStruct<FPropertyBindingBindableStructDescriptor>& OutStructDesc) const
 {
 	VisitAllNodes([&OutStructDesc, StructID](const UStateTreeState* State, const FStateTreeBindableStructDesc& Desc, const FStateTreeDataView Value)
 	{
 		if (Desc.ID == StructID)
 		{
-			OutStructDesc = Desc;
+			OutStructDesc = TInstancedStruct<FStateTreeBindableStructDesc>::Make(Desc);
 			return EStateTreeVisitor::Break;
 		}
 		return EStateTreeVisitor::Continue;
@@ -653,7 +639,7 @@ bool UStateTreeEditorData::GetStructByID(const FGuid StructID, FStateTreeBindabl
 	return OutStructDesc.IsValid();
 }
 
-bool UStateTreeEditorData::GetDataViewByID(const FGuid StructID, FStateTreeDataView& OutDataView) const
+bool UStateTreeEditorData::GetBindingDataViewByID(const FGuid StructID, FPropertyBindingDataView& OutDataView) const
 {
 	bool bFound = false;
 	VisitAllNodes([&OutDataView, &bFound, StructID](const UStateTreeState* State, const FStateTreeBindableStructDesc& Desc, const FStateTreeDataView Value)
@@ -711,16 +697,30 @@ UStateTreeState* UStateTreeEditorData::GetMutableStateByID(const FGuid StateID)
 	return Result;
 }
 
-void UStateTreeEditorData::GetAllStructValues(TMap<FGuid, const FStateTreeDataView>& AllValues) const
+void UStateTreeEditorData::GetAllStructValues(TMap<FGuid, const FPropertyBindingDataView>& OutAllValues) const
 {
-	AllValues.Reset();
+	OutAllValues.Reset();
 
 	const UStateTree* StateTree = GetTypedOuter<UStateTree>();
 	checkf(StateTree, TEXT("UStateTreeEditorData should only be allocated within a UStateTree"));
 
-	VisitAllNodes([&AllValues](const UStateTreeState* State, const FStateTreeBindableStructDesc& Desc, const FStateTreeDataView Value)
+	VisitAllNodes([&OutAllValues](const UStateTreeState* State, const FStateTreeBindableStructDesc& Desc, const FStateTreeDataView Value)
 		{
-			AllValues.Emplace(Desc.ID, Value);
+			OutAllValues.Emplace(Desc.ID, Value);
+			return EStateTreeVisitor::Continue;
+		});
+}
+
+void UStateTreeEditorData::GetAllStructValues(TMap<FGuid, const FStateTreeDataView>& OutAllValues) const
+{
+	OutAllValues.Reset();
+
+	const UStateTree* StateTree = GetTypedOuter<UStateTree>();
+	checkf(StateTree, TEXT("UStateTreeEditorData should only be allocated within a UStateTree"));
+
+	VisitAllNodes([&OutAllValues](const UStateTreeState* State, const FStateTreeBindableStructDesc& Desc, const FStateTreeDataView Value)
+		{
+			OutAllValues.Emplace(Desc.ID, Value);
 			return EStateTreeVisitor::Continue;
 		});
 }
@@ -979,9 +979,9 @@ void UStateTreeEditorData::FixDuplicateIDs()
 	});
 
 	// It is possible that the user has changed the node type so some of the bindings might not make sense anymore, clean them up.
-	TMap<FGuid, const FStateTreeDataView> AllValues;
+	TMap<FGuid, const FPropertyBindingDataView> AllValues;
 	GetAllStructValues(AllValues);
-	EditorBindings.RemoveUnusedBindings(AllValues);
+	EditorBindings.RemoveInvalidBindings(AllValues);
 }
 
 void UStateTreeEditorData::DuplicateIDs()
@@ -1116,7 +1116,7 @@ void UStateTreeEditorData::UpdateBindingsInstanceStructs()
 {
 	TMap<FGuid, const FStateTreeDataView> AllValues;
 	GetAllStructValues(AllValues);
-	for (FStateTreePropertyPathBinding& Binding : EditorBindings.GetMutableBindings())
+	for (FPropertyBindingBinding& Binding : EditorBindings.GetMutableBindings())
 	{
 		if (AllValues.Contains(Binding.GetSourcePath().GetStructID()))
 		{
@@ -1389,10 +1389,10 @@ EStateTreeVisitor UStateTreeEditorData::VisitStateNodes(const UStateTreeState& S
 
 EStateTreeVisitor UStateTreeEditorData::VisitStructBoundPropertyFunctions(FGuid StructID, const FString& StatePath, TFunctionRef<EStateTreeVisitor(const FStateTreeEditorNode& EditorNode, const FStateTreeBindableStructDesc& Desc, const FStateTreeDataView Value)> InFunc) const
 {
-	TArray<const FStateTreePropertyPathBinding*> Bindings;
-	EditorBindings.GetPropertyBindingsFor(StructID, Bindings);
+	TArray<const FPropertyBindingBinding*> Bindings;
+	EditorBindings.FPropertyBindingBindingCollection::GetBindingsFor(StructID, Bindings);
 
-	for (const FStateTreePropertyPathBinding* Binding : Bindings)
+	for (const FPropertyBindingBinding* Binding : Bindings)
 	{
 		const FConstStructView FunctionNodeView = Binding->GetPropertyFunctionNode();
 		if (const FStateTreeEditorNode* FunctionNode = FunctionNodeView.GetPtr<const FStateTreeEditorNode>())
diff --git a/Engine/Plugins/Runtime/StateTree/Source/StateTreeEditorModule/Private/StateTreeEditorPropertyBindings.cpp b/Engine/Plugins/Runtime/StateTree/Source/StateTreeEditorModule/Private/StateTreeEditorPropertyBindings.cpp
--- a/Engine/Plugins/Runtime/StateTree/Source/StateTreeEditorModule/Private/StateTreeEditorPropertyBindings.cpp
+++ b/Engine/Plugins/Runtime/StateTree/Source/StateTreeEditorModule/Private/StateTreeEditorPropertyBindings.cpp
@@ -19,240 +19,102 @@ UStateTreeEditorPropertyBindingsOwner::UStateTreeEditorPropertyBindingsOwner(con
 
 //////////////////////////////////////////////////////////////////////////
 
-void FStateTreeEditorPropertyBindings::SetBindingsOwner(gsl::not_null<TScriptInterface<IStateTreeEditorPropertyBindingsOwner>> InBindingsOwner)
+FPropertyBindingPath FStateTreeEditorPropertyBindings::AddFunctionBinding(const UScriptStruct* InPropertyFunctionNodeStruct, TConstArrayView<FPropertyBindingPathSegment> InSourcePathSegments, const FPropertyBindingPath& TargetPath)
 {
-	BindingsOwner = InBindingsOwner;
-}
-
-void FStateTreeEditorPropertyBindings::AddPropertyBinding(const FStateTreePropertyPath& SourcePath, const FStateTreePropertyPath& TargetPath)
-{
-	RemovePropertyBindings(TargetPath, ESearchMode::Exact);
-
-	FStateTreePropertyPathBinding Binding(SourcePath, TargetPath);
-
-	// If we have bindings owner, update property path segments to capture property IDs, etc.
-	if (BindingsOwner)
-	{
-		FStateTreeDataView SourceDataView;
-		if (BindingsOwner->GetDataViewByID(Binding.GetSourcePath().GetStructID(), SourceDataView))
-		{
-			Binding.GetMutableSourcePath().UpdateSegmentsFromValue(SourceDataView);
-		}
-		FStateTreeDataView TargetDataView;
-		if (BindingsOwner->GetDataViewByID(Binding.GetTargetPath().GetStructID(), TargetDataView))
-		{
-			Binding.GetMutableTargetPath().UpdateSegmentsFromValue(TargetDataView);
-		}
-	}
-	
-	PropertyBindings.Add(Binding);
-}
-
-FStateTreePropertyPath FStateTreeEditorPropertyBindings::AddFunctionPropertyBinding(const UScriptStruct* PropertyFunctionNodeStruct, TConstArrayView<FPropertyBindingPathSegment> SourcePathSegments, const FStateTreePropertyPath& TargetPath)
-{
-	check(PropertyFunctionNodeStruct->IsChildOf<FStateTreePropertyFunctionBase>());
+	check(InPropertyFunctionNodeStruct->IsChildOf<FStateTreePropertyFunctionBase>());
 
 	FInstancedStruct PropertyFunctionNode(FStateTreeEditorNode::StaticStruct());
 	FStateTreeEditorNode& PropertyFunction = PropertyFunctionNode.GetMutable<FStateTreeEditorNode>();
 	const FGuid NodeID = FGuid::NewGuid();
 	PropertyFunction.ID = NodeID;
-	PropertyFunction.Node.InitializeAs(PropertyFunctionNodeStruct);
+	PropertyFunction.Node.InitializeAs(InPropertyFunctionNodeStruct);
 	const FStateTreePropertyFunctionBase& Function = PropertyFunction.Node.Get<FStateTreePropertyFunctionBase>();
 	if (const UScriptStruct* InstanceType = Cast<const UScriptStruct>(Function.GetInstanceDataType()))
 	{
 		PropertyFunction.Instance.InitializeAs(InstanceType);
 	}
 
-	RemovePropertyBindings(TargetPath, ESearchMode::Exact);
-	FStateTreePropertyPath SourcePath = FStateTreePropertyPath(NodeID, SourcePathSegments);
+	Super::RemoveBindings(TargetPath, ESearchMode::Exact);
+	FPropertyBindingPath SourcePath = FPropertyBindingPath(NodeID, InSourcePathSegments);
 	PropertyBindings.Emplace(MoveTemp(PropertyFunctionNode), SourcePath, TargetPath);
 	return SourcePath;
 }
 
-void FStateTreeEditorPropertyBindings::AddPropertyBinding(const FStateTreePropertyPathBinding& Binding)
+FPropertyBindingBinding* FStateTreeEditorPropertyBindings::AddBindingInternal(const FPropertyBindingPath& InSourcePath
+	, const FPropertyBindingPath& InTargetPath)
 {
-	RemovePropertyBindings(Binding.GetTargetPath(), ESearchMode::Exact);
-	PropertyBindings.Add(Binding);
+	return &PropertyBindings.Emplace_GetRef(InSourcePath, InTargetPath);
 }
 
-void FStateTreeEditorPropertyBindings::RemovePropertyBindings(const FStateTreePropertyPath& TargetPath, ESearchMode SearchMode)
+void FStateTreeEditorPropertyBindings::RemoveBindingsInternal(TFunctionRef<bool(const FPropertyBindingBinding&)> InPredicate)
 {
-	if (SearchMode == ESearchMode::Exact)
-	{
-		PropertyBindings.RemoveAllSwap([&TargetPath](const FStateTreePropertyPathBinding& Binding)
-			{
-				return Binding.GetTargetPath() == TargetPath;
-			});
-	}
-	else
-	{
-		PropertyBindings.RemoveAllSwap([&TargetPath](const FStateTreePropertyPathBinding& Binding)
-			{
-				return Binding.GetTargetPath().Includes(TargetPath);
-			});
-	}
+	PropertyBindings.RemoveAllSwap(InPredicate);
 }
 
-void FStateTreeEditorPropertyBindings::CopyBindings(const FGuid FromStructID, const FGuid ToStructID)
+bool FStateTreeEditorPropertyBindings::HasBindingInternal(TFunctionRef<bool(const FPropertyBindingBinding&)> InPredicate) const
 {
-	// Copy all bindings that target "FromStructID" and retarget them to "ToStructID".
-	TArray<FStateTreePropertyPathBinding> NewBindings;
-	for (const FStateTreePropertyPathBinding& Binding : PropertyBindings)
-	{
-		if (Binding.GetTargetPath().GetStructID() == FromStructID)
-		{
-			NewBindings.Emplace(Binding.GetSourcePath(), FStateTreePropertyPath(ToStructID, Binding.GetTargetPath().GetSegments()));
-		}
-	}
+	return PropertyBindings.ContainsByPredicate(InPredicate);
+}
 
-	PropertyBindings.Append(NewBindings);
+const FPropertyBindingBinding* FStateTreeEditorPropertyBindings::FindBindingInternal(TFunctionRef<bool(const FPropertyBindingBinding&)> InPredicate) const
+{
+	return PropertyBindings.FindByPredicate(InPredicate);
 }
 
-bool FStateTreeEditorPropertyBindings::HasPropertyBinding(const FStateTreePropertyPath& TargetPath, ESearchMode RemoveMode) const
+void FStateTreeEditorPropertyBindings::AddBindingInternal(const FPropertyBindingBinding& InBinding)
 {
-	if (RemoveMode == ESearchMode::Exact)
-	{
-		return PropertyBindings.ContainsByPredicate([&TargetPath](const FStateTreePropertyPathBinding& Binding)
-			{
-				return Binding.GetTargetPath() == TargetPath;
-			});
-	}
-	else
-	{
-		return PropertyBindings.ContainsByPredicate([&TargetPath](const FStateTreePropertyPathBinding& Binding)
-			{
-				return Binding.GetTargetPath().Includes(TargetPath);
-			});
-	}
+	PropertyBindings.Add(FStateTreePropertyPathBinding(InBinding));
 }
 
-const FStateTreePropertyPathBinding* FStateTreeEditorPropertyBindings::FindPropertyBinding(const FStateTreePropertyPath& TargetPath, ESearchMode RemoveMode) const
+void FStateTreeEditorPropertyBindings::ForEachBinding(TFunctionRef<void(const FPropertyBindingBinding& Binding)> InFunction) const
 {
-	if (RemoveMode == ESearchMode::Exact)
-	{
-		return PropertyBindings.FindByPredicate([&TargetPath](const FStateTreePropertyPathBinding& Binding)
-			{
-				return Binding.GetTargetPath() == TargetPath;
-			});
-	}
-	else
+	for (const FStateTreePropertyPathBinding& Binding : PropertyBindings)
 	{
-		return PropertyBindings.FindByPredicate([&TargetPath](const FStateTreePropertyPathBinding& Binding)
-			{
-				return Binding.GetTargetPath().Includes(TargetPath);
-			});
+		InFunction(Binding);
 	}
 }
 
-const FStateTreePropertyPath* FStateTreeEditorPropertyBindings::GetPropertyBindingSource(const FStateTreePropertyPath& TargetPath) const
+void FStateTreeEditorPropertyBindings::ForEachMutableBinding(TFunctionRef<void(FPropertyBindingBinding& Binding)> InFunction)
 {
-	const FStateTreePropertyPathBinding* Binding = PropertyBindings.FindByPredicate([&TargetPath](const FStateTreePropertyPathBinding& Binding)
-		{
-			return Binding.GetTargetPath() == TargetPath;
-		});
-	return Binding ? &Binding->GetSourcePath() : nullptr;
+	for (FStateTreePropertyPathBinding& Binding : PropertyBindings)
+	{
+		InFunction(Binding);
+	}
 }
 
-void FStateTreeEditorPropertyBindings::GetPropertyBindingsFor(const FGuid StructID, TArray<const FStateTreePropertyPathBinding*>& OutBindings) const
+void FStateTreeEditorPropertyBindings::VisitBindings(TFunctionRef<EVisitResult(const FPropertyBindingBinding& Binding)> InFunction) const
 {
 	for (const FStateTreePropertyPathBinding& Binding : PropertyBindings)
 	{
-		if (Binding.GetSourcePath().GetStructID().IsValid() && Binding.GetTargetPath().GetStructID() == StructID)
+		if (InFunction(Binding) == EVisitResult::Break)
 		{
-			OutBindings.Add(&Binding);
+			break;
 		}
 	}
 }
 
-void FStateTreeEditorPropertyBindings::RemoveUnusedBindings(const TMap<FGuid, const FStateTreeDataView>& ValidStructs)
+void FStateTreeEditorPropertyBindings::VisitMutableBindings(TFunctionRef<EVisitResult(FPropertyBindingBinding& Binding)> InFunction)
 {
-	PropertyBindings.RemoveAllSwap([ValidStructs](FStateTreePropertyPathBinding& Binding)
-		{
-			// Remove binding if it's target struct has been removed
-			if (!ValidStructs.Contains(Binding.GetTargetPath().GetStructID()))
-			{
-				// Remove
-				return true;
-			}
-
-			// Target path should always have at least one segment (copy bind directly on a target struct/object). 
-			if (Binding.GetTargetPath().IsPathEmpty())
-			{
-				return true;
-			}
-
-			// Remove binding if path containing instanced indirections (e.g. instance struct or object) cannot be resolved.
-			// TODO: Try to use core redirect to find new name.
-			{
-				const FStateTreeDataView* SourceValue = ValidStructs.Find(Binding.GetSourcePath().GetStructID());
-				if (SourceValue && SourceValue->IsValid())
-				{
-					FString Error;
-					TArray<FPropertyBindingPathIndirection> Indirections;
-					if (!Binding.GetSourcePath().ResolveIndirectionsWithValue(*SourceValue, Indirections, &Error))
-					{
-						UE_LOG(LogStateTree, Verbose, TEXT("Removing binding to %s because binding source path cannot be resolved: %s"),
-							*Binding.GetSourcePath().ToString(), *Error); // Error contains the target path.
-
-						// Remove
-						return true;
-					}
-				}
-			}
-			
-			{
-				const FStateTreeDataView TargetValue = ValidStructs.FindChecked(Binding.GetTargetPath().GetStructID());
-				FString Error;
-				TArray<FPropertyBindingPathIndirection> Indirections;
-				if (!Binding.GetTargetPath().ResolveIndirectionsWithValue(TargetValue, Indirections, &Error))
-				{
-					UE_LOG(LogStateTree, Verbose, TEXT("Removing binding to %s because binding target path cannot be resolved: %s"),
-						*Binding.GetSourcePath().ToString(), *Error); // Error contains the target path.
-
-					// Remove
-					return true;
-				}
-			}
-
-			return false;
-		});
-}
-
-bool FStateTreeEditorPropertyBindings::ContainsAnyStruct(const TSet<const UStruct*>& Structs)
-{
-	auto PathContainsStruct = [&Structs](const FStateTreePropertyPath& PropertyPath)
+	for (FStateTreePropertyPathBinding& Binding : PropertyBindings)
 	{
-		for (const FPropertyBindingPathSegment& Segment : PropertyPath.GetSegments())
+		if (InFunction(Binding) == EVisitResult::Break)
 		{
-			if (Structs.Contains(Segment.GetInstanceStruct()))
-			{
-				return true;
-			}
-		}
-		return false;
-	};
-	
-	for (FStateTreePropertyPathBinding& PropertyPathBinding : PropertyBindings)
-	{
-		if (PathContainsStruct(PropertyPathBinding.GetSourcePath()))
-		{
-			return true;
-		}
-		if (PathContainsStruct(PropertyPathBinding.GetTargetPath()))
-		{
-			return true;
+			break;
 		}
 	}
-	
-	return false;
+}
+
+int32 FStateTreeEditorPropertyBindings::GetNumBindings() const
+{
+	return PropertyBindings.Num();
 }
 
 PRAGMA_DISABLE_DEPRECATION_WARNINGS
 void FStateTreeEditorPropertyBindings::GetPropertyBindingsFor(const FGuid StructID, TArray<FStateTreePropertyPathBinding>& OutBindings) const
 {
-	TArray<const FStateTreePropertyPathBinding*> NodeBindings;
-	GetPropertyBindingsFor(StructID, NodeBindings);
-	Algo::Transform(NodeBindings, OutBindings, [](const FStateTreePropertyPathBinding* BindingPtr) { return *BindingPtr; });
+	TArray<const FPropertyBindingBinding*> NodeBindings;
+	FPropertyBindingBindingCollection::GetBindingsFor(StructID, NodeBindings);
+	Algo::Transform(NodeBindings, OutBindings, [](const FPropertyBindingBinding* BindingPtr) { return *static_cast<const FStateTreePropertyPathBinding*>(BindingPtr); });
 }
 PRAGMA_ENABLE_DEPRECATION_WARNINGS
 
@@ -270,7 +132,7 @@ const FStateTreePropertyPath* FStateTreeBindingLookup::GetPropertyBindingSource(
 	const FStateTreeEditorPropertyBindings* EditorBindings = BindingOwner->GetPropertyEditorBindings();
 	check(EditorBindings);
 
-	return EditorBindings->GetPropertyBindingSource(InTargetPath);
+	return static_cast<const FStateTreePropertyPath*>(EditorBindings->FPropertyBindingBindingCollection::GetBindingSource(InTargetPath));
 }
 
 FText FStateTreeBindingLookup::GetPropertyPathDisplayName(const FStateTreePropertyPath& InPath, EStateTreeNodeFormatting Formatting) const
@@ -284,7 +146,7 @@ FText FStateTreeBindingLookup::GetPropertyPathDisplayName(const FStateTreeProper
 
 	// If path's struct is a PropertyFunction, let it override a display name.
 	{
-		const FStateTreePropertyPathBinding* BindingToPath = EditorBindings->GetBindings().FindByPredicate([&InPath](const FStateTreePropertyPathBinding& Binding)
+		const FPropertyBindingBinding* BindingToPath = EditorBindings->GetBindings().FindByPredicate([&InPath](const FPropertyBindingBinding& Binding)
 		{
 			return Binding.GetSourcePath() == InPath;
 		});
@@ -316,10 +178,10 @@ FText FStateTreeBindingLookup::GetPropertyPathDisplayName(const FStateTreeProper
 	
 		if (StructName.IsEmpty())
 		{
-			FStateTreeBindableStructDesc Struct;
-			if (BindingOwner->GetStructByID(InPath.GetStructID(), Struct))
+			TInstancedStruct<FPropertyBindingBindableStructDescriptor> Struct;
+			if (static_cast<const IPropertyBindingBindingCollectionOwner*>(BindingOwner)->GetBindableStructByID(InPath.GetStructID(), Struct))
 			{
-				StructName = Struct.Name.ToString();
+				StructName = Struct.Get().Name.ToString();
 			}
 		}
 	}
@@ -340,7 +202,7 @@ FText FStateTreeBindingLookup::GetBindingSourceDisplayName(const FStateTreePrope
 	check(EditorBindings);
 
 	// Check if the target property is bound, if so, return binding description.
-	if (const FStateTreePropertyPath* SourcePath = GetPropertyBindingSource(InTargetPath))
+	if (const FPropertyBindingPath* SourcePath = GetPropertyBindingSource(InTargetPath))
 	{
 		return GetPropertyPathDisplayName(*SourcePath, Formatting);
 	}
@@ -350,11 +212,11 @@ FText FStateTreeBindingLookup::GetBindingSourceDisplayName(const FStateTreePrope
 	const FProperty* TargetProperty = nullptr;
 	EStateTreePropertyUsage Usage = EStateTreePropertyUsage::Invalid;
 	
-	FStateTreeBindableStructDesc TargetStructDesc;
-	if (BindingOwner->GetStructByID(InTargetPath.GetStructID(), TargetStructDesc))
+	TInstancedStruct<FPropertyBindingBindableStructDescriptor> TargetStructDesc;
+	if (static_cast<const IPropertyBindingBindingCollectionOwner*>(BindingOwner)->GetBindableStructByID(InTargetPath.GetStructID(), TargetStructDesc))
 	{
 		TArray<FPropertyBindingPathIndirection> Indirection;
-		if (InTargetPath.ResolveIndirections(TargetStructDesc.Struct, Indirection)
+		if (InTargetPath.ResolveIndirections(TargetStructDesc.Get().Struct, Indirection)
 			&& Indirection.Num() > 0)
 		{
 			const FPropertyBindingPathIndirection& Leaf = Indirection.Last();
@@ -399,11 +261,11 @@ const FProperty* FStateTreeBindingLookup::GetPropertyPathLeafProperty(const FSta
 	check(EditorBindings);
 
 	const FProperty* Result = nullptr;
-	FStateTreeBindableStructDesc Struct;
-	if (BindingOwner->GetStructByID(InPath.GetStructID(), Struct))
+	TInstancedStruct<FPropertyBindingBindableStructDescriptor> Struct;
+	if (static_cast<const IPropertyBindingBindingCollectionOwner*>(BindingOwner)->GetBindableStructByID(InPath.GetStructID(), Struct))
 	{
 		TArray<FPropertyBindingPathIndirection> Indirection;
-		if (InPath.ResolveIndirections(Struct.Struct, Indirection) && Indirection.Num() > 0)
+		if (InPath.ResolveIndirections(Struct.Get().Struct, Indirection) && Indirection.Num() > 0)
 		{
 			return Indirection.Last().GetProperty();
 		}
diff --git a/Engine/Plugins/Runtime/StateTree/Source/StateTreeEditorModule/Private/StateTreePropertyBindingCompiler.cpp b/Engine/Plugins/Runtime/StateTree/Source/StateTreeEditorModule/Private/StateTreePropertyBindingCompiler.cpp
--- a/Engine/Plugins/Runtime/StateTree/Source/StateTreeEditorModule/Private/StateTreePropertyBindingCompiler.cpp
+++ b/Engine/Plugins/Runtime/StateTree/Source/StateTreeEditorModule/Private/StateTreePropertyBindingCompiler.cpp
@@ -2,7 +2,6 @@
 #include "StateTreePropertyBindingCompiler.h"
 #include "IPropertyAccessEditor.h"
 #include "PropertyPathHelpers.h"
-#include "StateTreeCompiler.h"
 #include "StateTreeCompilerLog.h"
 #include "StateTreeEditorPropertyBindings.h"
 #include "Misc/EnumerateRange.h"
diff --git a/Engine/Plugins/Runtime/StateTree/Source/StateTreeEditorModule/Private/StateTreePropertyHelpers.cpp b/Engine/Plugins/Runtime/StateTree/Source/StateTreeEditorModule/Private/StateTreePropertyHelpers.cpp
--- a/Engine/Plugins/Runtime/StateTree/Source/StateTreeEditorModule/Private/StateTreePropertyHelpers.cpp
+++ b/Engine/Plugins/Runtime/StateTree/Source/StateTreeEditorModule/Private/StateTreePropertyHelpers.cpp
@@ -131,7 +131,7 @@ void DispatchPostEditToNodes(UObject& Owner, FPropertyChangedChainEvent& InPrope
 			{
 				if (TargetPath.GetStructID().IsValid())
 				{
-					FStateTreePropertyPathBinding* FoundBinding = EditorData.GetPropertyEditorBindings()->GetMutableBindings().FindByPredicate([&TargetPath](const FStateTreePropertyPathBinding& Binding)
+					FPropertyBindingBinding* FoundBinding = EditorData.GetPropertyEditorBindings()->GetMutableBindings().FindByPredicate([&TargetPath](const FPropertyBindingBinding& Binding)
 					{
 						return TargetPath == Binding.GetTargetPath();
 					});
diff --git a/Engine/Plugins/Runtime/StateTree/Source/StateTreeEditorModule/Private/StateTreeState.cpp b/Engine/Plugins/Runtime/StateTree/Source/StateTreeEditorModule/Private/StateTreeState.cpp
--- a/Engine/Plugins/Runtime/StateTree/Source/StateTreeEditorModule/Private/StateTreeState.cpp
+++ b/Engine/Plugins/Runtime/StateTree/Source/StateTreeEditorModule/Private/StateTreeState.cpp
@@ -356,9 +356,9 @@ void UStateTreeState::PostEditChangeChainProperty(FPropertyChangedChainEvent& Pr
 				TreeData->Modify();
 				FStateTreeEditorPropertyBindings* Bindings = TreeData->GetPropertyEditorBindings();
 				check(Bindings);
-				TMap<FGuid, const FStateTreeDataView> AllStructValues;
+				TMap<FGuid, const FPropertyBindingDataView> AllStructValues;
 				TreeData->GetAllStructValues(AllStructValues);
-				Bindings->RemoveUnusedBindings(AllStructValues);
+				Bindings->RemoveInvalidBindings(AllStructValues);
 			}
 		}
 	}
@@ -496,7 +496,7 @@ void UStateTreeState::SetParametersPropertyOverridden(const FGuid PropertyID, co
 						EditorData->Modify();
 
 						FStateTreePropertyPath Path(Parameters.ID, Desc->CachedProperty->GetFName());
-						Bindings->RemovePropertyBindings(Path);
+						Bindings->RemoveBindings(Path);
 					}
 				}
 			}
diff --git a/Engine/Plugins/Runtime/StateTree/Source/StateTreeEditorModule/Private/StateTreeViewModel.cpp b/Engine/Plugins/Runtime/StateTree/Source/StateTreeEditorModule/Private/StateTreeViewModel.cpp
--- a/Engine/Plugins/Runtime/StateTree/Source/StateTreeEditorModule/Private/StateTreeViewModel.cpp
+++ b/Engine/Plugins/Runtime/StateTree/Source/StateTreeEditorModule/Private/StateTreeViewModel.cpp
@@ -57,10 +57,10 @@ namespace UE::StateTree::Editor
 		
 		TreeData->VisitStateNodes(*State, [TreeData, &AllBindings](const UStateTreeState* State, const FStateTreeBindableStructDesc& Desc, const FStateTreeDataView Value)
 		{
-			TArray<const FStateTreePropertyPathBinding*> NodeBindings;
-			TreeData->GetPropertyEditorBindings()->GetPropertyBindingsFor(Desc.ID, NodeBindings);
-			Algo::Transform(NodeBindings, AllBindings, [](const FStateTreePropertyPathBinding* BindingPtr) { return *BindingPtr; });
-			return EStateTreeVisitor::Continue;				
+			TArray<const FPropertyBindingBinding*> NodeBindings;
+			TreeData->GetPropertyEditorBindings()->FPropertyBindingBindingCollection::GetBindingsFor(Desc.ID, NodeBindings);
+			Algo::Transform(NodeBindings, AllBindings, [](const FPropertyBindingBinding* BindingPtr) { return FStateTreePropertyPathBinding(*BindingPtr); });
+			return EStateTreeVisitor::Continue;
 		});
 
 		for (UStateTreeState* ChildState : State->Children)
@@ -1068,7 +1068,7 @@ void FStateTreeViewModel::PasteStatesAsChildrenFromText(const FString& TextToImp
 						Binding.GetMutableSourcePath().SetStructID(*NewSourceID);
 					}
 					
-					TreeData->GetPropertyEditorBindings()->AddPropertyBinding(MoveTemp(Binding));
+					TreeData->GetPropertyEditorBindings()->AddBinding(MoveTemp(Binding));
 				}
 			}
 		}
diff --git a/Engine/Plugins/Runtime/StateTree/Source/StateTreeEditorModule/Public/StateTreeEditorData.h b/Engine/Plugins/Runtime/StateTree/Source/StateTreeEditorModule/Public/StateTreeEditorData.h
--- a/Engine/Plugins/Runtime/StateTree/Source/StateTreeEditorModule/Public/StateTreeEditorData.h
+++ b/Engine/Plugins/Runtime/StateTree/Source/StateTreeEditorModule/Public/StateTreeEditorData.h
@@ -63,19 +63,37 @@ public:
 
 	virtual void PostInitProperties() override;
 	
-	// IStateTreeEditorPropertyBindingsOwner
-	virtual void GetAccessibleStructs(const FGuid TargetStructID, TArray<FStateTreeBindableStructDesc>& OutStructDescs) const override;
-	virtual bool GetStructByID(const FGuid StructID, FStateTreeBindableStructDesc& OutStructDesc) const override;
-	virtual bool GetDataViewByID(const FGuid StructID, FStateTreeDataView& OutDataView) const override;
-	virtual FStateTreeEditorPropertyBindings* GetPropertyEditorBindings() override { return &EditorBindings; }
-	virtual const FStateTreeEditorPropertyBindings* GetPropertyEditorBindings() const override { return &EditorBindings; }
-	virtual FStateTreeBindableStructDesc FindContextData(const UStruct* ObjectType, const FString ObjectNameHint) const override;
+	//~ Begin IStateTreeEditorPropertyBindingsOwner interface
+	virtual void GetBindableStructs(const FGuid TargetStructID, TArray<TInstancedStruct<FPropertyBindingBindableStructDescriptor>>& OutStructDescs) const override;
+	virtual bool GetBindableStructByID(const FGuid StructID, TInstancedStruct<FPropertyBindingBindableStructDescriptor>& OutStructDesc) const override;
+	virtual bool GetBindingDataViewByID(const FGuid StructID, FPropertyBindingDataView& OutDataView) const override;
+	virtual const FStateTreeEditorPropertyBindings* GetPropertyEditorBindings() const override
+	{
+		return &EditorBindings;
+	}
+
+	virtual FPropertyBindingBindingCollection* GetEditorPropertyBindings() override
+	{
+		return &EditorBindings;
+	}
 
-	virtual EStateTreeVisitor EnumerateBindablePropertyFunctionNodes(TFunctionRef<EStateTreeVisitor(const UScriptStruct* NodeStruct, const FStateTreeBindableStructDesc& Desc, const FStateTreeDataView Value)> InFunc) const;
+	virtual const FPropertyBindingBindingCollection* GetEditorPropertyBindings() const override
+	{
+		return &EditorBindings;
+	}
 
-	virtual bool CanCreateParameter(const FGuid StructID) const override;
-	virtual void CreateParameters(const FGuid StructID, TArrayView<FStateTreeEditorPropertyCreationDesc> InOutCreationDescs) override;
-	// ~IStateTreeEditorPropertyBindingsOwner
+	virtual FStateTreeEditorPropertyBindings* GetPropertyEditorBindings() override
+	{
+		return &EditorBindings;
+	}
+
+	virtual FStateTreeBindableStructDesc FindContextData(const UStruct* ObjectType, const FString ObjectNameHint) const override;
+
+	virtual bool CanCreateParameter(const FGuid InStructID) const override;
+	virtual void CreateParametersForStruct(const FGuid InStructID, TArrayView<UE::PropertyBinding::FPropertyCreationDescriptor> InOutCreationDescs) override;
+	virtual void OnPropertyBindingChanged(const FPropertyBindingPath& InSourcePath, const FPropertyBindingPath& InTargetPath) override;
+	virtual void AppendBindablePropertyFunctionStructs(TArray<TInstancedStruct<FPropertyBindingBindableStructDescriptor>>& InOutStructs) const override;
+	//~ End IStateTreeEditorPropertyBindingsOwner interface
 
 	/**
 	 * Returns the description for the node for UI.
@@ -106,7 +124,10 @@ public:
 	UStateTreeState* GetMutableStateByID(const FGuid StateID);
 
 	/** @returns the IDs and instance values of all bindable structs in the StateTree. */
-	void GetAllStructValues(TMap<FGuid, const FStateTreeDataView>& AllValues) const;
+	void GetAllStructValues(TMap<FGuid, const FStateTreeDataView>& OutAllValues) const;
+
+	/** @returns the IDs and instance values of all bindable structs in the StateTree. */
+	void GetAllStructValues(TMap<FGuid, const FPropertyBindingDataView>& OutAllValues) const;
 
 	/**
 	* Iterates over all structs that are related to binding
@@ -148,9 +169,15 @@ public:
 	 * Returns array of nodes along the execution path, up to the TargetStruct.
 	 * @param Path The states to visit during the check
 	 * @param TargetStructID The ID of the node where to stop.
-	 * @param OutStructDescs Array of nodes accessible on the given path.  
+	 * @param OutStructDescs Array of nodes accessible on the given path.
 	 */
-	void GetAccessibleStructs(const TConstArrayView<const UStateTreeState*> Path, const FGuid TargetStructID, TArray<FStateTreeBindableStructDesc>& OutStructDescs) const;
+	void GetAccessibleStructsInExecutionPath(const TConstArrayView<const UStateTreeState*> Path, const FGuid TargetStructID, TArray<TInstancedStruct<FPropertyBindingBindableStructDescriptor>>& OutStructDescs) const;
+
+	UE_DEPRECATED(5.6, "Use GetAccessibleStructsInExecutionPath instead")
+	void GetAccessibleStruct(const TConstArrayView<const UStateTreeState*> Path, const FGuid TargetStructID, TArray<TInstancedStruct<FPropertyBindingBindableStructDescriptor>>& OutStructDescs) const
+	{
+		GetAccessibleStructsInExecutionPath(Path, TargetStructID, OutStructDescs);
+	}
 
 	void ReparentStates();
 	
@@ -219,7 +246,7 @@ public:
 	 */
 	void AddPropertyBinding(const FStateTreePropertyPath& SourcePath, const FStateTreePropertyPath& TargetPath)
 	{
-		EditorBindings.AddPropertyBinding(SourcePath, TargetPath);
+		EditorBindings.AddBinding(SourcePath, TargetPath);
 	}
 
 	/**
@@ -227,7 +254,7 @@ public:
 	 */
 	void AddPropertyBinding(const UScriptStruct* PropertyFunctionNodeStruct, TConstArrayView<FPropertyBindingPathSegment> SourcePathSegments, const FStateTreePropertyPath& TargetPath)
 	{
-		EditorBindings.AddFunctionPropertyBinding(PropertyFunctionNodeStruct, SourcePathSegments, TargetPath);
+		EditorBindings.AddFunctionBinding(PropertyFunctionNodeStruct, SourcePathSegments, TargetPath);
 	}
 
 	/**
@@ -241,7 +268,7 @@ public:
 		TargetPath.SetStructID(TargetNode.ID);
 		if (SourcePath.FromString(SourcePathStr) && TargetPath.FromString(TargetPathStr))
 		{
-			EditorBindings.AddPropertyBinding(SourcePath, TargetPath);
+			EditorBindings.AddBinding(SourcePath, TargetPath);
 			return true;
 		}
 		return false;
@@ -266,6 +293,8 @@ public:
 	}
 
 private:
+	virtual EStateTreeVisitor EnumerateBindablePropertyFunctionNodes(TFunctionRef<EStateTreeVisitor(const UScriptStruct* NodeStruct, const FStateTreeBindableStructDesc& Desc, const FStateTreeDataView Value)> InFunc) const override;
+
 	void FixObjectInstance(TSet<UObject*>& SeenObjects, UObject& Outer, FStateTreeEditorNode& Node);
 	void FixObjectNodes();
 	void FixDuplicateIDs();
@@ -310,7 +339,7 @@ public:
 	 * Transient list of breakpoints added in the debugging session.
 	 * These will be lost if the asset gets reloaded.
 	 * If there is eventually a change to make those persist with the asset
-	 * we need to prune all dangling breakpoints after states/tasks got removed. (see RemoveUnusedBindings)
+	 * we need to prune all dangling breakpoints after states/tasks got removed.
 	 */
 	UPROPERTY(Transient)
 	TArray<FStateTreeEditorBreakpoint> Breakpoints;
diff --git a/Engine/Plugins/Runtime/StateTree/Source/StateTreeEditorModule/Public/StateTreeEditorPropertyBindings.h b/Engine/Plugins/Runtime/StateTree/Source/StateTreeEditorModule/Public/StateTreeEditorPropertyBindings.h
--- a/Engine/Plugins/Runtime/StateTree/Source/StateTreeEditorModule/Public/StateTreeEditorPropertyBindings.h
+++ b/Engine/Plugins/Runtime/StateTree/Source/StateTreeEditorModule/Public/StateTreeEditorPropertyBindings.h
@@ -3,6 +3,9 @@
 #pragma once
 
 #include "UObject/Interface.h"
+#include "PropertyBindingExtension.h"
+#include "PropertyBindingBindingCollectionOwner.h"
+#include "PropertyBindingBindingCollectionOwner.h"
 #include "StateTreePropertyBindings.h"
 #include "StateTreeEditorNode.h"
 #include "StateTreeEditorPropertyBindings.generated.h"
@@ -12,123 +15,143 @@ class IStateTreeEditorPropertyBindingsOwner;
 enum class EStateTreeVisitor : uint8;
 
 /**
- * Editor representation of a all property bindings in a StateTree
+ * Editor representation of all property bindings in a StateTree
  */
 USTRUCT()
-struct STATETREEEDITORMODULE_API FStateTreeEditorPropertyBindings
+struct STATETREEEDITORMODULE_API FStateTreeEditorPropertyBindings : public FPropertyBindingBindingCollection
 {
 	GENERATED_BODY()
 
-	/** Sets associated bindings owner, used to validate added property paths. */
-	void SetBindingsOwner(gsl::not_null<TScriptInterface<IStateTreeEditorPropertyBindingsOwner>> BindingsOwner);
-
 	/**
 	 * Adds binding between source and destination paths. Removes any bindings to TargetPath before adding the new one.
 	 * @param SourcePath Binding source property path.
 	 * @param TargetPath Binding target property path.
 	 */
-	void AddPropertyBinding(const FStateTreePropertyPath& SourcePath, const FStateTreePropertyPath& TargetPath);
+	UE_DEPRECATED(5.6, "Use AddBinding taking FPropertyBindingPath instead")
+	void AddPropertyBinding(const FStateTreePropertyPath& SourcePath, const FStateTreePropertyPath& TargetPath)
+	{
+		AddBinding(SourcePath, TargetPath);
+	}
 	
 	/**
 	 * Adds binding.
 	 * @param Binding Binding to be added.
 	*/
-	void AddPropertyBinding(const FStateTreePropertyPathBinding& Binding);
+	UE_DEPRECATED(5.6, "Use the version taking FPropertyBindingPath instead")
+	void AddPropertyBinding(const FStateTreePropertyPathBinding& Binding)
+	{
+	}
 
 	/**
 	 * Adds binding between PropertyFunction of the provided type and destination path.
-	 * @param PropertyFunctionNodeStruct Struct of PropertyFunction.
-	 * @param SourcePathSegments Binding source property path segments.
-	 * @param TargetPath Binding target property path.
+	 * @param InPropertyFunctionNodeStruct Struct of PropertyFunction.
+	 * @param InSourcePathSegments Binding source property path segments.
+	 * @param InTargetPath Binding target property path.
 	 * @return Constructed binding source property path.
 	 */
-	FStateTreePropertyPath AddFunctionPropertyBinding(const UScriptStruct* PropertyFunctionNodeStruct, TConstArrayView<FPropertyBindingPathSegment> SourcePathSegments, const FStateTreePropertyPath& TargetPath);
-
-	enum class ESearchMode
+	FPropertyBindingPath AddFunctionBinding(const UScriptStruct* InPropertyFunctionNodeStruct, TConstArrayView<FPropertyBindingPathSegment> InSourcePathSegments, const FPropertyBindingPath& InTargetPath);
+	UE_DEPRECATED(5.6, "Use AddFunctionBinding taking FPropertyBindingPath instead")
+	FStateTreePropertyPath AddFunctionPropertyBinding(const UScriptStruct* InPropertyFunctionNodeStruct, TConstArrayView<FPropertyBindingPathSegment> InSourcePathSegments, const FStateTreePropertyPath& InTargetPath)
 	{
-		Exact,				// Binding with exact matching path.
-		Includes,			// Binding with path that matches but the binding path can be longer.
-	};
+		return AddFunctionBinding(InPropertyFunctionNodeStruct, InSourcePathSegments, InTargetPath);
+	}
 
 	/**
 	 * Removes all bindings to target path.
 	 * @param TargetPath Target property path.
 	 */ 
-	void RemovePropertyBindings(const FStateTreePropertyPath& TargetPath, ESearchMode SearchMode = ESearchMode::Exact);
+	UE_DEPRECATED(5.6, "Use RemoveBinding taking FPropertyBindingPath instead")
+	void RemovePropertyBindings(const FStateTreePropertyPath& TargetPath, ESearchMode SearchMode = ESearchMode::Exact)
+	{
+		RemoveBindings(TargetPath, SearchMode);
+	}
 	
 	/**
 	 * Has any binding to the target path.
 	 * @return True of the target path has any bindings.
 	 */
-	bool HasPropertyBinding(const FStateTreePropertyPath& TargetPath, ESearchMode SearchMode = ESearchMode::Exact) const;
+	UE_DEPRECATED(5.6, "Use HasBinding taking FPropertyBindingPath instead")
+	bool HasPropertyBinding(const FStateTreePropertyPath& TargetPath, ESearchMode SearchMode = ESearchMode::Exact) const
+	{
+		return HasBinding(TargetPath, SearchMode);
+	}
 
 	/**
 	 * @return binding to the target path.
 	 */
 	const FStateTreePropertyPathBinding* FindPropertyBinding(const FStateTreePropertyPath& TargetPath, ESearchMode SearchMode = ESearchMode::Exact) const;
 
-	/**
-	 * Copies property bindings from an existing struct to another.
-	 * Overrides a binding to a specific property if it already exists in ToStructID.
-	 * @param FromStructID ID of the struct to copy from.
-	 * @param ToStructID ID of the struct to copy to.
-	 */
-	void CopyBindings(const FGuid FromStructID, const FGuid ToStructID);
-	
 	/**
 	 * @return Source path for given target path, or null if binding does not exists.
 	 */
-	const FStateTreePropertyPath* GetPropertyBindingSource(const FStateTreePropertyPath& TargetPath) const;
-	
+	UE_DEPRECATED(5.6, "Use GetBindingSource taking FPropertyBindingPath instead")
+	const FStateTreePropertyPath* GetPropertyBindingSource(const FStateTreePropertyPath& TargetPath) const
+	{
+		return nullptr;
+	}
+
 	/**
 	 * Returns all pointers to bindings for a specified structs based in struct ID.
 	 * @param StructID ID of the struct to find bindings for.
 	 * @param OutBindings Bindings for specified struct.
 	 */
-	void GetPropertyBindingsFor(const FGuid StructID, TArray<const FStateTreePropertyPathBinding*>& OutBindings) const;
-	
+	UE_DEPRECATED(5.6, "Use GetBindingsFor taking FPropertyBindingPath instead")
+	void GetPropertyBindingsFor(const FGuid StructID, TArray<const FStateTreePropertyPathBinding*>& OutBindings) const
+	{
+	}
+
 	/**
 	 * Removes bindings which do not point to valid structs IDs.
 	 * @param ValidStructs Set of struct IDs that are currently valid.
 	 */
-	void RemoveUnusedBindings(const TMap<FGuid, const FStateTreeDataView>& ValidStructs);
-
-	/** @return true if any of the bindings references any of the Structs. */
-	bool ContainsAnyStruct(const TSet<const UStruct*>& Structs);
+	UE_DEPRECATED(5.6, "Use RemoveInvalidBindings taking FPropertyBindingDataView instead")
+	void RemoveUnusedBindings(const TMap<FGuid, const FStateTreeDataView>& ValidStructs)
+	{
+	}
 
-	/** @return array view to all bindings. */
-	const TConstArrayView<FStateTreePropertyPathBinding> GetBindings() const
+	/** @return const array view to all bindings. */
+	TConstArrayView<FStateTreePropertyPathBinding> GetBindings() const
 	{
 		return PropertyBindings;
 	}
 
-	const TArrayView<FStateTreePropertyPathBinding> GetMutableBindings()
+	/** @return array view to all bindings. */
+	TArrayView<FStateTreePropertyPathBinding> GetMutableBindings()
 	{
 		return PropertyBindings;
 	}
 
+	virtual int32 GetNumBindings() const override;
+	virtual void ForEachBinding(TFunctionRef<void(const FPropertyBindingBinding& Binding)> InFunction) const override;
+	virtual void ForEachMutableBinding(TFunctionRef<void(FPropertyBindingBinding& Binding)> InFunction) override;
+	virtual void VisitBindings(TFunctionRef<EVisitResult(const FPropertyBindingBinding& Binding)> InFunction) const override;
+	virtual void VisitMutableBindings(TFunctionRef<EVisitResult(FPropertyBindingBinding& Binding)> InFunction) override;
+
 PRAGMA_DISABLE_DEPRECATION_WARNINGS
 	UE_DEPRECATED(5.5, "Use GetPropertyBindingsFor returning pointers instead.")
 	void GetPropertyBindingsFor(const FGuid StructID, TArray<FStateTreePropertyPathBinding>& OutBindings) const;
 PRAGMA_ENABLE_DEPRECATION_WARNINGS
 
 private:
+	virtual FPropertyBindingBinding* AddBindingInternal(const FPropertyBindingPath& InSourcePath, const FPropertyBindingPath& InTargetPath) override;
+	virtual void AddBindingInternal(const FPropertyBindingBinding& InBinding) override;
+	virtual void RemoveBindingsInternal(TFunctionRef<bool(const FPropertyBindingBinding&)> InPredicate) override;
+	virtual bool HasBindingInternal(TFunctionRef<bool(const FPropertyBindingBinding&)> InPredicate) const override;
+	virtual const FPropertyBindingBinding* FindBindingInternal(TFunctionRef<bool(const FPropertyBindingBinding&)> InPredicate) const override;
+
 	UPROPERTY()
 	TArray<FStateTreePropertyPathBinding> PropertyBindings;
-
-	UPROPERTY(Transient)
-	TScriptInterface<IStateTreeEditorPropertyBindingsOwner> BindingsOwner = nullptr; 
 };
 
 
 UINTERFACE(MinimalAPI, meta = (CannotImplementInterfaceInBlueprint))
-class UStateTreeEditorPropertyBindingsOwner : public UInterface
+class UStateTreeEditorPropertyBindingsOwner : public UPropertyBindingBindingCollectionOwner
 {
 	GENERATED_UINTERFACE_BODY()
 };
 
 /** Struct of Parameters used to Create a Property */
-struct FStateTreeEditorPropertyCreationDesc
+struct UE_DEPRECATED(5.6, "Use FPropertyCreationDesc instead") FStateTreeEditorPropertyCreationDesc
 {
 	/** Property Bag Description of the Property to Create */
 	FPropertyBagPropertyDesc PropertyDesc;
@@ -140,16 +163,21 @@ struct FStateTreeEditorPropertyCreationDesc
 	const void* SourceContainerAddress = nullptr;
 };
 
-class STATETREEEDITORMODULE_API IStateTreeEditorPropertyBindingsOwner
+class STATETREEEDITORMODULE_API IStateTreeEditorPropertyBindingsOwner : public IPropertyBindingBindingCollectionOwner
 {
-	GENERATED_IINTERFACE_BODY()
+	GENERATED_BODY()
+
+public:
 
 	/**
 	 * Returns structs within the owner that are visible for target struct.
 	 * @param TargetStructID Target struct ID
 	 * @param OutStructDescs Result descriptors of the visible structs.
 	 */
-	virtual void GetAccessibleStructs(const FGuid TargetStructID, TArray<FStateTreeBindableStructDesc>& OutStructDescs) const PURE_VIRTUAL(IStateTreeEditorPropertyBindingsOwner::GetAccessibleStructs, return; );
+	UE_DEPRECATED(5.6, "Use version taking FPropertyBindingBindableStructDescriptor instead")
+	virtual void GetAccessibleStructs(const FGuid TargetStructID, TArray<FStateTreeBindableStructDesc>& OutStructDescs) const final
+	{
+	}
 
 	/**
 	 * Returns struct descriptor based on struct ID.
@@ -157,7 +185,11 @@ class STATETREEEDITORMODULE_API IStateTreeEditorPropertyBindingsOwner
 	 * @param OutStructDesc Result descriptor.
 	 * @return True if struct found.
 	 */
-	virtual bool GetStructByID(const FGuid StructID, FStateTreeBindableStructDesc& OutStructDesc) const PURE_VIRTUAL(IStateTreeEditorPropertyBindingsOwner::GetStructByID, return false; );
+	UE_DEPRECATED(5.6, "Use version taking FPropertyBindingBindableStructDescriptor instead")
+	virtual bool GetStructByID(const FGuid StructID, FStateTreeBindableStructDesc& OutStructDesc) const final
+	{
+		return false;
+	}
 
 	/**
 	 * Finds a bindable context struct based on name and type.
@@ -165,14 +197,18 @@ class STATETREEEDITORMODULE_API IStateTreeEditorPropertyBindingsOwner
 	 * @param ObjectNameHint Name to use if multiple context objects of same type are found. 
 	 */
 	virtual FStateTreeBindableStructDesc FindContextData(const UStruct* ObjectType, const FString ObjectNameHint) const PURE_VIRTUAL(IStateTreeEditorPropertyBindingsOwner::FindContextData, return {}; );
-	
+
 	/**
 	 * Returns data view based on struct ID.
 	 * @param StructID Target struct ID
 	 * @param OutDataView Result data view.
 	 * @return True if struct found.
 	 */
-	virtual bool GetDataViewByID(const FGuid StructID, FStateTreeDataView& OutDataView) const PURE_VIRTUAL(IStateTreeEditorPropertyBindingsOwner::GetDataViewByID, return false; );
+	UE_DEPRECATED(5.6, "Use version taking FPropertyBindingDataView instead")
+	virtual bool GetDataViewByID(const FGuid StructID, FStateTreeDataView& OutDataView) const final
+	{
+		return false;
+	}
 
 	/** @return Pointer to editor property bindings. */
 	virtual FStateTreeEditorPropertyBindings* GetPropertyEditorBindings() PURE_VIRTUAL(IStateTreeEditorPropertyBindingsOwner::GetPropertyEditorBindings, return nullptr; );
@@ -182,19 +218,17 @@ class STATETREEEDITORMODULE_API IStateTreeEditorPropertyBindingsOwner
 
 	virtual EStateTreeVisitor EnumerateBindablePropertyFunctionNodes(TFunctionRef<EStateTreeVisitor(const UScriptStruct* NodeStruct, const FStateTreeBindableStructDesc& Desc, const FStateTreeDataView Value)> InFunc) const PURE_VIRTUAL(IStateTreeEditorPropertyBindingsOwner::EnumerateBindablePropertyFunctionNodes, return static_cast<EStateTreeVisitor>(0); );
 
-	/**
-	 * Determines whether the struct matching the given struct id is capable of adding new properties
-	 * @param StructID Target struct ID
-	 * @return True if struct supports adding new properties
-	 */
-	virtual bool CanCreateParameter(const FGuid StructID) const PURE_VIRTUAL(IStateTreeEditorPropertyBindingsOwner::CanCreateProperty, return false;);
-
+PRAGMA_DISABLE_DEPRECATION_WARNINGS
 	/**
 	 * Creates the given properties in the property bag of the struct matching the given struct ID
 	 * @param StructID Target struct ID
 	 * @param InOutCreationDescs the descriptions of the properties to create. This is modified to update the property names that actually got created
 	 */
-	virtual void CreateParameters(const FGuid StructID, TArrayView<FStateTreeEditorPropertyCreationDesc> InOutCreationDescs) PURE_VIRTUAL(IStateTreeEditorPropertyBindingsOwner::CreateProperties);
+	UE_DEPRECATED(5.6, "Use version taking FPropertyCreationDesc instead")
+	virtual void CreateParameters(const FGuid StructID, TArrayView<FStateTreeEditorPropertyCreationDesc> InOutCreationDescs) final
+	{
+	}
+PRAGMA_ENABLE_DEPRECATION_WARNINGS
 };
 
 // TODO: We should merge this with IStateTreeEditorPropertyBindingsOwner and FStateTreeEditorPropertyBindings.
diff --git a/Engine/Plugins/Runtime/StateTree/Source/StateTreeEditorModule/StateTreeEditorModule.Build.cs b/Engine/Plugins/Runtime/StateTree/Source/StateTreeEditorModule/StateTreeEditorModule.Build.cs
--- a/Engine/Plugins/Runtime/StateTree/Source/StateTreeEditorModule/StateTreeEditorModule.Build.cs
+++ b/Engine/Plugins/Runtime/StateTree/Source/StateTreeEditorModule/StateTreeEditorModule.Build.cs
@@ -31,6 +31,7 @@ namespace UnrealBuildTool.Rules
 				"Projects",
 				"BlueprintGraph",
 				"PropertyBindingUtils",
+				"PropertyBindingUtilsEditor",
 				"PropertyAccessEditor",
 				"StructUtilsEditor",
 				"GameplayTags",
diff --git a/Engine/Plugins/Runtime/StateTree/Source/StateTreeModule/Private/StateTree.cpp b/Engine/Plugins/Runtime/StateTree/Source/StateTreeModule/Private/StateTree.cpp
--- a/Engine/Plugins/Runtime/StateTree/Source/StateTreeModule/Private/StateTree.cpp
+++ b/Engine/Plugins/Runtime/StateTree/Source/StateTreeModule/Private/StateTree.cpp
@@ -20,7 +20,6 @@
 #include "Editor.h"
 #include "StateTreeModuleImpl.h"
 #include "StructUtils/UserDefinedStruct.h"
-#include "StructUtilsDelegates.h"
 #include "Templates/GuardValueAccessors.h"
 #include "UObject/LinkerLoad.h"
 #endif
@@ -439,7 +438,6 @@ void UStateTree::PostLoad()
 		}
 	}
 
-	const int32 CurrentVersion = GetLinkerCustomVersion(FStateTreeCustomVersion::GUID);
 #if WITH_EDITOR
 	if (EditorData)
 	{
@@ -451,6 +449,7 @@ void UStateTree::PostLoad()
 		Compile();
 	}
 #else
+	const int32 CurrentVersion = GetLinkerCustomVersion(FStateTreeCustomVersion::GUID);
 	if (CurrentVersion < FStateTreeCustomVersion::LatestVersion)
 	{		
 		UE_LOG(LogStateTree, Error, TEXT("%s: compiled data is in older format. Please recompile the StateTree asset."), *GetFullName());
@@ -778,7 +777,7 @@ bool UStateTree::PatchBindings()
 		}
 	}
 
-	auto PatchPropertyPath = [](FStateTreePropertyPath& PropertyPath)
+	auto PatchPropertyPath = [](FPropertyBindingPath& PropertyPath)
 	{
 		for (FPropertyBindingPathSegment& Segment : PropertyPath.GetMutableSegments())
 		{
diff --git a/Engine/Plugins/Runtime/StateTree/Source/StateTreeModule/Private/StateTreePropertyBindings.cpp b/Engine/Plugins/Runtime/StateTree/Source/StateTreeModule/Private/StateTreePropertyBindings.cpp
--- a/Engine/Plugins/Runtime/StateTree/Source/StateTreeModule/Private/StateTreePropertyBindings.cpp
+++ b/Engine/Plugins/Runtime/StateTree/Source/StateTreeModule/Private/StateTreePropertyBindings.cpp
@@ -326,7 +326,7 @@ bool FStateTreePropertyBindings::ResolvePath(const UStruct* Struct, const FState
 		Indirection.Offset = static_cast<uint16>(PathIndirection.GetPropertyOffset());
 		Indirection.Type = PathIndirection.GetAccessType();
 
-		if (Indirection.Type == EPropertyBindingAccessType::IndexArray)
+		if (Indirection.Type == EPropertyBindingPropertyAccessType::IndexArray)
 		{
 			if (const FArrayProperty* ArrayProperty = CastField<FArrayProperty>(PathIndirection.GetProperty()))
 			{
@@ -346,8 +346,8 @@ bool FStateTreePropertyBindings::ResolvePath(const UStruct* Struct, const FState
 				return false;
 			}
 		}
-		else if (Indirection.Type == EPropertyBindingAccessType::StructInstance
-				|| Indirection.Type == EPropertyBindingAccessType::ObjectInstance)
+		else if (Indirection.Type == EPropertyBindingPropertyAccessType::StructInstance
+				|| Indirection.Type == EPropertyBindingPropertyAccessType::ObjectInstance)
 		{
 			if (PathIndirection.GetInstanceStruct())
 			{
@@ -370,33 +370,33 @@ bool FStateTreePropertyBindings::ResolvePath(const UStruct* Struct, const FState
 			if ((Index + 1) < TempIndirections.Num())
 			{
 				const FStateTreePropertyIndirection& NextIndirection = TempIndirections[Index + 1];
-				if (Indirection.Type == EPropertyBindingAccessType::Offset
-					&& NextIndirection.Type == EPropertyBindingAccessType::Offset)
+				if (Indirection.Type == EPropertyBindingPropertyAccessType::Offset
+					&& NextIndirection.Type == EPropertyBindingPropertyAccessType::Offset)
 				{
 					// Collapse adjacent offset indirections
 					Indirection.Offset += NextIndirection.Offset;
 					TempIndirections.RemoveAt(Index + 1);
 					Index--;
 				}
-				else if (Indirection.Type == EPropertyBindingAccessType::IndexArray
-					&& NextIndirection.Type == EPropertyBindingAccessType::Offset
+				else if (Indirection.Type == EPropertyBindingPropertyAccessType::IndexArray
+					&& NextIndirection.Type == EPropertyBindingPropertyAccessType::Offset
 					&& NextIndirection.Offset == 0)
 				{
 					// Remove empty offset after array indexing.
 					TempIndirections.RemoveAt(Index + 1);
 					Index--;
 				}
-				else if (Indirection.Type == EPropertyBindingAccessType::StructInstance
-					&& NextIndirection.Type == EPropertyBindingAccessType::Offset
+				else if (Indirection.Type == EPropertyBindingPropertyAccessType::StructInstance
+					&& NextIndirection.Type == EPropertyBindingPropertyAccessType::Offset
 					&& NextIndirection.Offset == 0)
 				{
 					// Remove empty offset after struct indirection.
 					TempIndirections.RemoveAt(Index + 1);
 					Index--;
 				}
-				else if ((Indirection.Type == EPropertyBindingAccessType::Object
-						|| Indirection.Type == EPropertyBindingAccessType::ObjectInstance)
-					&& NextIndirection.Type == EPropertyBindingAccessType::Offset
+				else if ((Indirection.Type == EPropertyBindingPropertyAccessType::Object
+						|| Indirection.Type == EPropertyBindingPropertyAccessType::ObjectInstance)
+					&& NextIndirection.Type == EPropertyBindingPropertyAccessType::Offset
 					&& NextIndirection.Offset == 0)
 				{
 					// Remove empty offset after object indirection.
@@ -424,7 +424,7 @@ bool FStateTreePropertyBindings::ResolvePath(const UStruct* Struct, const FState
 		// Indirections can be empty in case we're directly binding to source structs.
 		// Zero offset will return the struct itself.
 		OutFirstIndirection.Offset = 0;
-		OutFirstIndirection.Type = EPropertyBindingAccessType::Offset;
+		OutFirstIndirection.Type = EPropertyBindingPropertyAccessType::Offset;
 
 		OutLeafIndirection = FPropertyBindingPathIndirection(Struct);
 	}
@@ -499,7 +499,7 @@ bool FStateTreePropertyBindings::ResolveCopyType(const FPropertyBindingPathIndir
 		}
 	}
 
-	const EStateTreePropertyAccessCompatibility Compatibility = FStateTreePropertyBindings::GetPropertyCompatibility(SourceProperty, TargetProperty);
+	const UE::PropertyBinding::EPropertyCompatibility Compatibility = UE::PropertyBinding::GetPropertyCompatibility(SourceProperty, TargetProperty);
 
 	// Extract underlying types for enums
 	if (const FEnumProperty* EnumPropertyA = CastField<const FEnumProperty>(SourceProperty))
@@ -512,7 +512,7 @@ bool FStateTreePropertyBindings::ResolveCopyType(const FPropertyBindingPathIndir
 		TargetProperty = EnumPropertyB->GetUnderlyingProperty();
 	}
 
-	if (Compatibility == EStateTreePropertyAccessCompatibility::Compatible)
+	if (Compatibility == UE::PropertyBinding::EPropertyCompatibility::Compatible)
 	{
 		if (CastField<FNameProperty>(TargetProperty))
 		{
@@ -562,7 +562,7 @@ bool FStateTreePropertyBindings::ResolveCopyType(const FPropertyBindingPathIndir
 			return true;
 		}
 	}
-	else if (Compatibility == EStateTreePropertyAccessCompatibility::Promotable)
+	else if (Compatibility == UE::PropertyBinding::EPropertyCompatibility::Promotable)
 	{
 		if (SourceProperty->IsA<FBoolProperty>())
 		{
@@ -702,123 +702,16 @@ bool FStateTreePropertyBindings::ResolveCopyType(const FPropertyBindingPathIndir
 	return false;
 }
 
+PRAGMA_DISABLE_DEPRECATION_WARNINGS
 EStateTreePropertyAccessCompatibility FStateTreePropertyBindings::GetPropertyCompatibility(const FProperty* FromProperty, const FProperty* ToProperty)
 {
-	if (FromProperty == ToProperty)
-	{
-		return EStateTreePropertyAccessCompatibility::Compatible;
-	}
-
-	if (FromProperty == nullptr || ToProperty == nullptr)
-	{
-		return EStateTreePropertyAccessCompatibility::Incompatible;
-	}
-
-	// Special case for object properties since InPropertyA->SameType(InPropertyB) requires both properties to be of the exact same class.
-	// In our case we want to be able to bind a source property if its class is a child of the target property class.
-	if (FromProperty->IsA<FObjectPropertyBase>() && ToProperty->IsA<FObjectPropertyBase>())
-	{
-		const FObjectPropertyBase* SourceProperty = CastField<FObjectPropertyBase>(FromProperty);
-		const FObjectPropertyBase* TargetProperty = CastField<FObjectPropertyBase>(ToProperty);
-		return (SourceProperty->PropertyClass->IsChildOf(TargetProperty->PropertyClass)) ? EStateTreePropertyAccessCompatibility::Compatible : EStateTreePropertyAccessCompatibility::Incompatible;
-	}
-
-	// When copying to an enum property, expect FromProperty to be the same enum.
-	auto GetPropertyEnum = [](const FProperty* Property) -> const UEnum*
-	{
-		if (const FByteProperty* ByteProperty = CastField<FByteProperty>(Property))
-		{
-			return ByteProperty->GetIntPropertyEnum();
-		}
-		if (const FEnumProperty* EnumProperty = CastField<FEnumProperty>(Property))
-		{
-			return EnumProperty->GetEnum();
-		}
-		return nullptr;
-	};
-	
-	if (const UEnum* ToPropertyEnum = GetPropertyEnum(ToProperty))
-	{
-		const UEnum* FromPropertyEnum = GetPropertyEnum(FromProperty);
-		return (ToPropertyEnum == FromPropertyEnum) ? EStateTreePropertyAccessCompatibility::Compatible : EStateTreePropertyAccessCompatibility::Incompatible;
-	}
-	
-	// Allow source enums to be promoted to numbers.
-	if (const FEnumProperty* EnumPropertyA = CastField<const FEnumProperty>(FromProperty))
-	{
-		FromProperty = EnumPropertyA->GetUnderlyingProperty();
-	}
-
-	if (FromProperty->SameType(ToProperty))
-	{
-		return EStateTreePropertyAccessCompatibility::Compatible;
-	}
-	else
-	{
-		// Not directly compatible, check for promotions
-		if (FromProperty->IsA<FBoolProperty>())
-		{
-			if (ToProperty->IsA<FByteProperty>()
-				|| ToProperty->IsA<FIntProperty>()
-				|| ToProperty->IsA<FUInt32Property>()
-				|| ToProperty->IsA<FInt64Property>()
-				|| ToProperty->IsA<FFloatProperty>()
-				|| ToProperty->IsA<FDoubleProperty>())
-			{
-				return EStateTreePropertyAccessCompatibility::Promotable;
-			}
-		}
-		else if (FromProperty->IsA<FByteProperty>())
-		{
-			if (ToProperty->IsA<FIntProperty>()
-				|| ToProperty->IsA<FUInt32Property>()
-				|| ToProperty->IsA<FInt64Property>()
-				|| ToProperty->IsA<FFloatProperty>()
-				|| ToProperty->IsA<FDoubleProperty>())
-			{
-				return EStateTreePropertyAccessCompatibility::Promotable;
-			}
-		}
-		else if (FromProperty->IsA<FIntProperty>())
-		{
-			if (ToProperty->IsA<FInt64Property>()
-				|| ToProperty->IsA<FFloatProperty>()
-				|| ToProperty->IsA<FDoubleProperty>())
-			{
-				return EStateTreePropertyAccessCompatibility::Promotable;
-			}
-		}
-		else if (FromProperty->IsA<FUInt32Property>())
-		{
-			if (ToProperty->IsA<FInt64Property>()
-				|| ToProperty->IsA<FFloatProperty>()
-				|| ToProperty->IsA<FDoubleProperty>())
-			{
-				return EStateTreePropertyAccessCompatibility::Promotable;
-			}
-		}
-		else if (FromProperty->IsA<FFloatProperty>())
-		{
-			if (ToProperty->IsA<FIntProperty>()
-				|| ToProperty->IsA<FInt64Property>()
-				|| ToProperty->IsA<FDoubleProperty>())
-			{
-				return EStateTreePropertyAccessCompatibility::Promotable;
-			}
-		}
-		else if (FromProperty->IsA<FDoubleProperty>())
-		{
-			if (ToProperty->IsA<FIntProperty>()
-				|| ToProperty->IsA<FInt64Property>()
-				|| ToProperty->IsA<FFloatProperty>())
-			{
-				return EStateTreePropertyAccessCompatibility::Promotable;
-			}
-		}
-	}
-
-	return EStateTreePropertyAccessCompatibility::Incompatible;
+	// Temporary using cast as this method is no longer used and will be removed so we don't care about matching values in both enums
+	//const UE::PropertyBinding::EPropertyCompatibility Value = UE::PropertyBinding::GetPropertyCompatibility(FromProperty, ToProperty);
+	//check(UEnum::GetValueAsName<EStateTreePropertyAccessCompatibility>((EStateTreePropertyAccessCompatibility)Value) ==
+	//	  UEnum::GetValueAsName<UE::PropertyBinding::EPropertyCompatibility>(Value));
+	return static_cast<EStateTreePropertyAccessCompatibility>(UE::PropertyBinding::GetPropertyCompatibility(FromProperty, ToProperty));
 }
+PRAGMA_ENABLE_DEPRECATION_WARNINGS
 
 uint8* FStateTreePropertyBindings::GetAddress(FStateTreeDataView InStructView, const FStateTreePropertyIndirection& FirstIndirection, const FProperty* LeafProperty) const
 {
@@ -835,32 +728,32 @@ uint8* FStateTreePropertyBindings::GetAddress(FStateTreeDataView InStructView, c
 	{
 		switch (Indirection->Type)
 		{
-		case EPropertyBindingAccessType::Offset:
+		case EPropertyBindingPropertyAccessType::Offset:
 		{
 			Address = Address + Indirection->Offset;
 			break;
 		}
-		case EPropertyBindingAccessType::Object:
+		case EPropertyBindingPropertyAccessType::Object:
 		{
 			UObject* Object = *reinterpret_cast<UObject**>(Address + Indirection->Offset);
 			Address = reinterpret_cast<uint8*>(Object);
 			break;
 		}
-		case EPropertyBindingAccessType::WeakObject:
+		case EPropertyBindingPropertyAccessType::WeakObject:
 		{
 			TWeakObjectPtr<UObject>& WeakObjectPtr = *reinterpret_cast<TWeakObjectPtr<UObject>*>(Address + Indirection->Offset);
 			UObject* Object = WeakObjectPtr.Get();
 			Address = reinterpret_cast<uint8*>(Object);
 			break;
 		}
-		case EPropertyBindingAccessType::SoftObject:
+		case EPropertyBindingPropertyAccessType::SoftObject:
 		{
 			FSoftObjectPtr& SoftObjectPtr = *reinterpret_cast<FSoftObjectPtr*>(Address + Indirection->Offset);
 			UObject* Object = SoftObjectPtr.Get();
 			Address = reinterpret_cast<uint8*>(Object);
 			break;
 		}
-		case EPropertyBindingAccessType::ObjectInstance:
+		case EPropertyBindingPropertyAccessType::ObjectInstance:
 		{
 			check(Indirection->InstanceStruct);
 			UObject* Object = *reinterpret_cast<UObject**>(Address + Indirection->Offset);
@@ -876,7 +769,7 @@ uint8* FStateTreePropertyBindings::GetAddress(FStateTreeDataView InStructView, c
 			}
 			break;
 		}
-		case EPropertyBindingAccessType::StructInstance:
+		case EPropertyBindingPropertyAccessType::StructInstance:
 		{
 			check(Indirection->InstanceStruct);
 			FInstancedStruct& InstancedStruct = *reinterpret_cast<FInstancedStruct*>(Address + Indirection->Offset);
@@ -893,7 +786,7 @@ uint8* FStateTreePropertyBindings::GetAddress(FStateTreeDataView InStructView, c
 			}
 			break;
 		}
-		case EPropertyBindingAccessType::IndexArray:
+		case EPropertyBindingPropertyAccessType::IndexArray:
 		{
 			check(Indirection->ArrayProperty);
 			FScriptArrayHelper Helper(Indirection->ArrayProperty, Address + Indirection->Offset);
@@ -910,7 +803,7 @@ uint8* FStateTreePropertyBindings::GetAddress(FStateTreeDataView InStructView, c
 		}
 		default:
 			ensureMsgf(false, TEXT("FStateTreePropertyBindings::GetAddress: Unhandled indirection type %s for '%s'"),
-				*StaticEnum<EPropertyBindingAccessType>()->GetValueAsString(Indirection->Type), *LeafProperty->GetNameCPP());
+				*StaticEnum<EPropertyBindingPropertyAccessType>()->GetValueAsString(Indirection->Type), *LeafProperty->GetNameCPP());
 		}
 
 		Indirection = Indirection->NextIndex.IsValid() ? &PropertyIndirections[Indirection->NextIndex.Get()] : nullptr;
diff --git a/Engine/Plugins/Runtime/StateTree/Source/StateTreeModule/Public/StateTreePropertyBindings.h b/Engine/Plugins/Runtime/StateTree/Source/StateTreeModule/Public/StateTreePropertyBindings.h
--- a/Engine/Plugins/Runtime/StateTree/Source/StateTreeModule/Public/StateTreePropertyBindings.h
+++ b/Engine/Plugins/Runtime/StateTree/Source/StateTreeModule/Public/StateTreePropertyBindings.h
@@ -2,7 +2,9 @@
 
 #pragma once
 
+#include "PropertyBindingBindableStructDescriptor.h"
 #include "PropertyBindingPath.h"
+#include "PropertyBindingBinding.h"
 #include "StateTreeNodeBase.h"
 #include "StateTreeTypes.h"
 #include "StateTreePropertyRefHelpers.h"
@@ -103,7 +105,7 @@ enum class EStateTreePropertyCopyType : uint8
 };
 
 /** Enum describing property compatibility */
-enum class EStateTreePropertyAccessCompatibility
+enum class UE_DEPRECATED(5.6, "Use UE::PropertyBinding::EPropertyCompatibility instead") EStateTreePropertyAccessCompatibility : uint8
 {
 	/** Properties are incompatible */
 	Incompatible,
@@ -118,7 +120,7 @@ enum class EStateTreePropertyAccessCompatibility
  * Each struct has unique identifier, which is used to distinguish them, and name that is mostly for debugging and UI.
  */
 USTRUCT()
-struct STATETREEMODULE_API FStateTreeBindableStructDesc
+struct STATETREEMODULE_API FStateTreeBindableStructDesc : public FPropertyBindingBindableStructDescriptor
 {
 	GENERATED_BODY()
 
@@ -126,42 +128,28 @@ struct STATETREEMODULE_API FStateTreeBindableStructDesc
 
 #if WITH_EDITORONLY_DATA
 	FStateTreeBindableStructDesc(const FString& InStatePath, const FName InName, const UStruct* InStruct, const FStateTreeDataHandle InDataHandle, const EStateTreeBindableStructSource InDataSource, const FGuid InGuid)
-		: Struct(InStruct)
-		, Name(InName)
+		: FPropertyBindingBindableStructDescriptor(InName, InStruct, InGuid)
 		, DataHandle(InDataHandle)
 		, DataSource(InDataSource)
-		, ID(InGuid)
 		, StatePath(InStatePath)
 	{
 	}
 
 	UE_DEPRECATED(5.5, "Use constructor with StatePath instead.")
 	FStateTreeBindableStructDesc(const FName InName, const UStruct* InStruct, const FStateTreeDataHandle InDataHandle, const EStateTreeBindableStructSource InDataSource, const FGuid InGuid)
-		: Struct(InStruct)
-		, Name(InName)
+		: FPropertyBindingBindableStructDescriptor(InName, InStruct, InGuid)
 		, DataHandle(InDataHandle)
 		, DataSource(InDataSource)
-		, ID(InGuid)
 	{
 	}
 
-	bool operator==(const FStateTreeBindableStructDesc& RHS) const
+	virtual FString GetSection() const override
 	{
-		return ID == RHS.ID && Struct == RHS.Struct; // Not checking name, it's cosmetic.
+		return StatePath;
 	}
-#endif
-
-	bool IsValid() const { return Struct != nullptr; }
-
-	FString ToString() const;
-	
-	/** The type of the struct or class. */
-	UPROPERTY()
-	TObjectPtr<const UStruct> Struct = nullptr;
+#endif // WITH_EDITORONLY_DATA
 
-	/** Name of the container (cosmetic). */
-	UPROPERTY()
-	FName Name;
+	virtual FString ToString() const override;
 
 	/** Runtime data the struct represents. */
 	UPROPERTY()
@@ -172,17 +160,9 @@ struct STATETREEMODULE_API FStateTreeBindableStructDesc
 	EStateTreeBindableStructSource DataSource = EStateTreeBindableStructSource::Context;
 
 #if WITH_EDITORONLY_DATA
-	/** Unique identifier of the struct. */
-	UPROPERTY()
-	FGuid ID;
-
-	/** In Editor path to State containting the data. */
+	/** In Editor path to State containing the data. */
 	UPROPERTY(Transient)
 	FString StatePath;
-
-	/** Category of the bindable struct. Can be used to display the category in a menu. */
-	UPROPERTY()
-	FString Category;
 #endif
 };
 
@@ -194,11 +174,17 @@ struct UE_DEPRECATED(5.6, "Use FPropertyBindingPathIndirection instead") FStateT
  * Representation of a property path used for property binding in StateTree.
  */
 USTRUCT()
-struct STATETREEMODULE_API FStateTreePropertyPath : public FPropertyBindingPath
+struct STATETREEMODULE_API FStateTreePropertyPath final : public FPropertyBindingPath
 {
 	GENERATED_BODY()
 
 	using FPropertyBindingPath::FPropertyBindingPath;
+
+	//~ Intentionally not explicit to facilitate introduction of base class FPropertyBindingPath
+	FStateTreePropertyPath(const FPropertyBindingPath& Other)
+		: FPropertyBindingPath(Other)
+	{
+	}
 };
 
 
@@ -227,34 +213,35 @@ struct UE_DEPRECATED(all, "Use FStateTreePropertyPath instead.") STATETREEMODULE
  * Representation of a property binding in StateTree
  */
 USTRUCT()
-struct STATETREEMODULE_API FStateTreePropertyPathBinding
+struct STATETREEMODULE_API FStateTreePropertyPathBinding : public FPropertyBindingBinding
 {
 	GENERATED_BODY()
 
-PRAGMA_DISABLE_DEPRECATION_WARNINGS	
+PRAGMA_DISABLE_DEPRECATION_WARNINGS
 	FStateTreePropertyPathBinding() = default;
-	
+
 	FStateTreePropertyPathBinding(const FStateTreePropertyPath& InSourcePath, const FStateTreePropertyPath& InTargetPath)
-		: SourcePropertyPath(InSourcePath)
-		, TargetPropertyPath(InTargetPath)
+		: FPropertyBindingBinding(InSourcePath, InTargetPath)
 	{
 	}
 
 	FStateTreePropertyPathBinding(const FStateTreeDataHandle InSourceDataHandle, const FStateTreePropertyPath& InSourcePath, const FStateTreePropertyPath& InTargetPath)
-		: SourcePropertyPath(InSourcePath)
-		, TargetPropertyPath(InTargetPath)
+		: FPropertyBindingBinding(InSourcePath, InTargetPath)
 		, SourceDataHandle(InSourceDataHandle)
 	{
 	}
 
 #if WITH_EDITOR
-FStateTreePropertyPathBinding(FConstStructView InFunctionNodeStruct, const FStateTreePropertyPath& InSourcePath, const FStateTreePropertyPath& InTargetPath)
-		: SourcePropertyPath(InSourcePath)
-		, TargetPropertyPath(InTargetPath)
-		, PropertyFunctionNode(InFunctionNodeStruct)
+	FStateTreePropertyPathBinding(FConstStructView InFunctionNodeStruct, const FStateTreePropertyPath& InSourcePath, const FStateTreePropertyPath& InTargetPath)
+		: FPropertyBindingBinding(InFunctionNodeStruct, InSourcePath, InTargetPath)
+	{
+	}
+#endif // WITH_EDITOR
+
+	explicit FStateTreePropertyPathBinding(const FPropertyBindingBinding& Other)
+		: FPropertyBindingBinding(Other)
 	{
 	}
-#endif
 
 PRAGMA_ENABLE_DEPRECATION_WARNINGS
 
@@ -262,12 +249,6 @@ PRAGMA_ENABLE_DEPRECATION_WARNINGS
 	void PostSerialize(const FArchive& Ar);
 #endif // WITH_EDITORONLY_DATA
 
-	const FStateTreePropertyPath& GetSourcePath() const { return SourcePropertyPath; }
-	const FStateTreePropertyPath& GetTargetPath() const { return TargetPropertyPath; }
-
-	FStateTreePropertyPath& GetMutableSourcePath() { return SourcePropertyPath; }
-	FStateTreePropertyPath& GetMutableTargetPath() { return TargetPropertyPath; }
-
 	void SetSourceDataHandle(const FStateTreeDataHandle NewSourceDataHandle) { SourceDataHandle = NewSourceDataHandle; }
 	FStateTreeDataHandle GetSourceDataHandle() const { return SourceDataHandle; }
 
@@ -277,23 +258,11 @@ PRAGMA_ENABLE_DEPRECATION_WARNINGS
 #endif
 
 private:
-	/** Source property path of the binding */
-	UPROPERTY()
-	FStateTreePropertyPath SourcePropertyPath;
-
-	/** Target property path of the binding */
-	UPROPERTY()
-	FStateTreePropertyPath TargetPropertyPath;
-
 	/** Describes how to get the source data pointer for the binding. */
 	UPROPERTY()
 	FStateTreeDataHandle SourceDataHandle = FStateTreeDataHandle::Invalid;
 
 #if WITH_EDITORONLY_DATA
-	/**	Instance of bound PropertyFunction. */
-	UPROPERTY()
-	FInstancedStruct PropertyFunctionNode;
-
 PRAGMA_DISABLE_DEPRECATION_WARNINGS
 	UPROPERTY()
 	FStateTreeEditorPropertyPath SourcePath_DEPRECATED;
@@ -372,7 +341,7 @@ struct STATETREEMODULE_API FStateTreePropertyIndirection
 
 	/** Type of access/indirection. */
 	UPROPERTY()
-	EPropertyBindingAccessType Type = EPropertyBindingAccessType::Offset;
+	EPropertyBindingPropertyAccessType Type = EPropertyBindingPropertyAccessType::Offset;
 
 	/** Type of the struct or object instance in case the segment is pointing into an instanced data. */
 	UPROPERTY()
@@ -586,9 +555,6 @@ struct STATETREEMODULE_API FStateTreePropertyBindings
 	 */
 	bool ContainsAnyStruct(const TSet<const UStruct*>& Structs);
 
-	/** @return how properties are compatible for copying. */
-	static EStateTreePropertyAccessCompatibility GetPropertyCompatibility(const FProperty* FromProperty, const FProperty* ToProperty);
-
 	/**
 	 * Resolves what kind of copy type to use between specified property indirections.
 	 * @param SourceIndirection Property path indirections of the copy source,
@@ -599,6 +565,9 @@ struct STATETREEMODULE_API FStateTreePropertyBindings
 	[[nodiscard]] static bool ResolveCopyType(const FPropertyBindingPathIndirection& SourceIndirection, const FPropertyBindingPathIndirection& TargetIndirection, FStateTreePropertyCopy& OutCopy);
 
 PRAGMA_DISABLE_DEPRECATION_WARNINGS
+	UE_DEPRECATED(5.6, "Use UE::PropertyBinding::GetPropertyCompatibility instead")
+	static EStateTreePropertyAccessCompatibility GetPropertyCompatibility(const FProperty* FromProperty, const FProperty* ToProperty);
+
 	UE_DEPRECATED(5.6, "Use the overload taking FPropertyBindingPathIndirection instead")
 	[[nodiscard]] static bool ResolveCopyType(const FStateTreePropertyPathIndirection& SourceIndirection, const FStateTreePropertyPathIndirection& TargetIndirection, FStateTreePropertyCopy& OutCopy);
 
diff --git a/Engine/Plugins/Runtime/StateTree/Source/StateTreeModule/Public/StateTreeTypes.h b/Engine/Plugins/Runtime/StateTree/Source/StateTreeModule/Public/StateTreeTypes.h
--- a/Engine/Plugins/Runtime/StateTree/Source/StateTreeModule/Public/StateTreeTypes.h
+++ b/Engine/Plugins/Runtime/StateTree/Source/StateTreeModule/Public/StateTreeTypes.h
@@ -2,7 +2,11 @@
 
 #pragma once
 
+#if UE_ENABLE_INCLUDE_ORDER_DEPRECATED_IN_5_6
 #include "PropertyBindingPath.h"
+#endif // UE_ENABLE_INCLUDE_ORDER_DEPRECATED_IN_5_6
+
+#include "PropertyBindingDataView.h"
 #include "StructUtils/PropertyBag.h"
 #include "GameplayTagContainer.h"
 #include "StateTreeIndexTypes.h"
diff --git a/Engine/Plugins/Runtime/StateTree/Source/StateTreeTestSuite/Private/StateTreePropertyPathTest.cpp b/Engine/Plugins/Runtime/StateTree/Source/StateTreeTestSuite/Private/StateTreePropertyPathTest.cpp
--- a/Engine/Plugins/Runtime/StateTree/Source/StateTreeTestSuite/Private/StateTreePropertyPathTest.cpp
+++ b/Engine/Plugins/Runtime/StateTree/Source/StateTreeTestSuite/Private/StateTreePropertyPathTest.cpp
@@ -32,8 +32,8 @@ struct FStateTreeTest_PropertyPathOffset : FStateTreeTestBase
 		AITEST_EQUAL("Should have no resolve errors", ResolveErrors.Len(), 0);
 		
 		AITEST_EQUAL("Should have 2 indirections", Indirections.Num(), 2);
-		AITEST_EQUAL("Indirection 0 should be Offset type", Indirections[0].GetAccessType(), EPropertyBindingAccessType::Offset);
-		AITEST_EQUAL("Indirection 1 should be Offset type", Indirections[1].GetAccessType(), EPropertyBindingAccessType::Offset);
+		AITEST_EQUAL("Indirection 0 should be Offset type", Indirections[0].GetAccessType(), EPropertyBindingPropertyAccessType::Offset);
+		AITEST_EQUAL("Indirection 1 should be Offset type", Indirections[1].GetAccessType(), EPropertyBindingPropertyAccessType::Offset);
 
 		return true;
 	}
@@ -147,7 +147,7 @@ struct FStateTreeTest_PropertyPathWrongObject : FStateTreeTestBase
 
 			AITEST_TRUE("Resolve path should succeeed", bResolveResult);
 			AITEST_EQUAL("Should have 2 indirections", Indirections.Num(), 2);
-			AITEST_TRUE("Object ", Indirections[0].GetAccessType() == EPropertyBindingAccessType::ObjectInstance);
+			AITEST_TRUE("Object ", Indirections[0].GetAccessType() == EPropertyBindingPropertyAccessType::ObjectInstance);
 			AITEST_TRUE("Object ", Indirections[0].GetContainerStruct() == Object->GetClass());
 			AITEST_TRUE("Object ", Indirections[0].GetInstanceStruct() == UStateTreeTest_PropertyObjectInstancedWithB::StaticClass());
 			AITEST_EQUAL("Should not have error", ResolveErrors.Len(), 0);
@@ -190,8 +190,8 @@ struct FStateTreeTest_PropertyPathArray : FStateTreeTestBase
 		AITEST_TRUE("Resolve path should succeeed", bResolveResult);
 		AITEST_EQUAL("Should have no resolve errors", ResolveErrors.Len(), 0);
 		AITEST_EQUAL("Should have 2 indirections", Indirections.Num(), 2);
-		AITEST_EQUAL("Indirection 0 should be IndexArray type", Indirections[0].GetAccessType(), EPropertyBindingAccessType::IndexArray);
-		AITEST_EQUAL("Indirection 1 should be Offset type", Indirections[1].GetAccessType(), EPropertyBindingAccessType::Offset);
+		AITEST_EQUAL("Indirection 0 should be IndexArray type", Indirections[0].GetAccessType(), EPropertyBindingPropertyAccessType::IndexArray);
+		AITEST_EQUAL("Indirection 1 should be Offset type", Indirections[1].GetAccessType(), EPropertyBindingPropertyAccessType::Offset);
 
 		const int32 Value = *reinterpret_cast<const int32*>(Indirections[1].GetPropertyAddress());
 		AITEST_EQUAL("Value should be 123", Value, 123);
@@ -249,9 +249,9 @@ struct FStateTreeTest_PropertyPathArrayOfStructs : FStateTreeTestBase
 			AITEST_TRUE("Resolve path1 should succeeed", bResolveResult);
 			AITEST_EQUAL("Should have no resolve errors", ResolveErrors.Len(), 0);
 			AITEST_EQUAL("Should have 3 indirections", Indirections.Num(), 3);
-			AITEST_EQUAL("Indirection 0 should be ArrayIndex type", Indirections[0].GetAccessType(), EPropertyBindingAccessType::IndexArray);
-			AITEST_EQUAL("Indirection 1 should be Offset type", Indirections[1].GetAccessType(), EPropertyBindingAccessType::Offset);
-			AITEST_EQUAL("Indirection 2 should be Offset type", Indirections[2].GetAccessType(), EPropertyBindingAccessType::Offset);
+			AITEST_EQUAL("Indirection 0 should be ArrayIndex type", Indirections[0].GetAccessType(), EPropertyBindingPropertyAccessType::IndexArray);
+			AITEST_EQUAL("Indirection 1 should be Offset type", Indirections[1].GetAccessType(), EPropertyBindingPropertyAccessType::Offset);
+			AITEST_EQUAL("Indirection 2 should be Offset type", Indirections[2].GetAccessType(), EPropertyBindingPropertyAccessType::Offset);
 
 			const int32 Value = *reinterpret_cast<const int32*>(Indirections[2].GetPropertyAddress());
 			AITEST_EQUAL("Value should be 3", Value, 3);
@@ -265,10 +265,10 @@ struct FStateTreeTest_PropertyPathArrayOfStructs : FStateTreeTestBase
 			AITEST_TRUE("Resolve path2 should succeeed", bResolveResult);
 			AITEST_EQUAL("Should have no resolve errors", ResolveErrors.Len(), 0);
 			AITEST_EQUAL("Should have 4 indirections", Indirections.Num(), 4);
-			AITEST_EQUAL("Indirection 0 should be ArrayIndex type", Indirections[0].GetAccessType(), EPropertyBindingAccessType::IndexArray);
-			AITEST_EQUAL("Indirection 1 should be Offset type", Indirections[1].GetAccessType(), EPropertyBindingAccessType::Offset);
-			AITEST_EQUAL("Indirection 2 should be Offset type", Indirections[2].GetAccessType(), EPropertyBindingAccessType::Offset);
-			AITEST_EQUAL("Indirection 3 should be Offset type", Indirections[3].GetAccessType(), EPropertyBindingAccessType::Offset);
+			AITEST_EQUAL("Indirection 0 should be ArrayIndex type", Indirections[0].GetAccessType(), EPropertyBindingPropertyAccessType::IndexArray);
+			AITEST_EQUAL("Indirection 1 should be Offset type", Indirections[1].GetAccessType(), EPropertyBindingPropertyAccessType::Offset);
+			AITEST_EQUAL("Indirection 2 should be Offset type", Indirections[2].GetAccessType(), EPropertyBindingPropertyAccessType::Offset);
+			AITEST_EQUAL("Indirection 3 should be Offset type", Indirections[3].GetAccessType(), EPropertyBindingPropertyAccessType::Offset);
 
 			const int32 Value = *reinterpret_cast<const int32*>(Indirections[3].GetPropertyAddress());
 			AITEST_EQUAL("Value should be 42", Value, 42);
@@ -305,9 +305,9 @@ struct FStateTreeTest_PropertyPathArrayOfInstancedObjects : FStateTreeTestBase
 			AITEST_TRUE("Resolve path should succeeed", bResolveResult);
 			AITEST_EQUAL("Should have no resolve errors", ResolveErrors.Len(), 0);
 			AITEST_EQUAL("Should have 3 indirections", Indirections.Num(), 3);
-			AITEST_EQUAL("Indirection 0 should be ArrayIndex type", Indirections[0].GetAccessType(), EPropertyBindingAccessType::IndexArray);
-			AITEST_EQUAL("Indirection 1 should be StructInstance type", Indirections[1].GetAccessType(), EPropertyBindingAccessType::StructInstance);
-			AITEST_EQUAL("Indirection 2 should be Offset type", Indirections[2].GetAccessType(), EPropertyBindingAccessType::Offset);
+			AITEST_EQUAL("Indirection 0 should be ArrayIndex type", Indirections[0].GetAccessType(), EPropertyBindingPropertyAccessType::IndexArray);
+			AITEST_EQUAL("Indirection 1 should be StructInstance type", Indirections[1].GetAccessType(), EPropertyBindingPropertyAccessType::StructInstance);
+			AITEST_EQUAL("Indirection 2 should be Offset type", Indirections[2].GetAccessType(), EPropertyBindingPropertyAccessType::Offset);
 		}
 
 		{
@@ -318,9 +318,9 @@ struct FStateTreeTest_PropertyPathArrayOfInstancedObjects : FStateTreeTestBase
 			AITEST_TRUE("Resolve path should succeeed", bResolveResult);
 			AITEST_EQUAL("Should have no resolve errors", ResolveErrors.Len(), 0);
 			AITEST_EQUAL("Should have 3 indirections", Indirections.Num(), 3);
-			AITEST_EQUAL("Indirection 0 should be ArrayIndex type", Indirections[0].GetAccessType(), EPropertyBindingAccessType::IndexArray);
-			AITEST_EQUAL("Indirection 1 should be StructInstance type", Indirections[1].GetAccessType(), EPropertyBindingAccessType::StructInstance);
-			AITEST_EQUAL("Indirection 2 should be Offset type", Indirections[2].GetAccessType(), EPropertyBindingAccessType::Offset);
+			AITEST_EQUAL("Indirection 0 should be ArrayIndex type", Indirections[0].GetAccessType(), EPropertyBindingPropertyAccessType::IndexArray);
+			AITEST_EQUAL("Indirection 1 should be StructInstance type", Indirections[1].GetAccessType(), EPropertyBindingPropertyAccessType::StructInstance);
+			AITEST_EQUAL("Indirection 2 should be Offset type", Indirections[2].GetAccessType(), EPropertyBindingPropertyAccessType::Offset);
 
 			const int32 Value = *reinterpret_cast<const int32*>(Indirections[2].GetPropertyAddress());
 			AITEST_EQUAL("Value should be 123", Value, 123);
diff --git a/Engine/Plugins/Runtime/StateTree/Source/StateTreeTestSuite/StateTreeTestSuite.Build.cs b/Engine/Plugins/Runtime/StateTree/Source/StateTreeTestSuite/StateTreeTestSuite.Build.cs
--- a/Engine/Plugins/Runtime/StateTree/Source/StateTreeTestSuite/StateTreeTestSuite.Build.cs
+++ b/Engine/Plugins/Runtime/StateTree/Source/StateTreeTestSuite/StateTreeTestSuite.Build.cs
@@ -25,7 +25,8 @@ namespace UnrealBuildTool.Rules
 					"StateTreeEditorModule",
 					"AITestSuite",
 					"GameplayTags",
-					"PropertyBindingUtils"
+					"PropertyBindingUtils",
+					"PropertyBindingUtilsEditor"
 				}
 			);
 
