StateTree: Fix the crash when a tree is empty and used as a linked asset. This doesn't occur in the editor because empty state trees are prevented. However, it can be reproduced by creating a state tree programmatically, like with a unit test.

#rnx

diff --git a/Engine/Plugins/Runtime/StateTree/Source/StateTreeModule/Private/StateTreeExecutionContext.cpp b/Engine/Plugins/Runtime/StateTree/Source/StateTreeModule/Private/StateTreeExecutionContext.cpp
--- a/Engine/Plugins/Runtime/StateTree/Source/StateTreeModule/Private/StateTreeExecutionContext.cpp
+++ b/Engine/Plugins/Runtime/StateTree/Source/StateTreeModule/Private/StateTreeExecutionContext.cpp
@@ -4805,7 +4805,7 @@ bool FStateTreeExecutionContext::SelectStateInternal(
 		}
 		else if (NextState.Type == EStateTreeStateType::LinkedAsset)
 		{
-			if (NextLinkedStateAsset == nullptr)
+			if (NextLinkedStateAsset == nullptr || NextLinkedStateAsset->States.Num() == 0)
 			{
 				break;
 			}
