StateTree: Temporary fix crash with invalid delegate frame.

diff --git a/Engine/Plugins/Runtime/StateTree/Source/StateTreeModule/Private/StateTreeExecutionContext.cpp b/Engine/Plugins/Runtime/StateTree/Source/StateTreeModule/Private/StateTreeExecutionContext.cpp
--- a/Engine/Plugins/Runtime/StateTree/Source/StateTreeModule/Private/StateTreeExecutionContext.cpp
+++ b/Engine/Plugins/Runtime/StateTree/Source/StateTreeModule/Private/StateTreeExecutionContext.cpp
@@ -2656,11 +2656,8 @@ void FStateTreeExecutionContext::ExitState(const FStateTreeTransitionResult& Tra
 		{
 			RemoveAllDelegateListeners();
 }
-		else
+		else if (ensure(FirstModifiedFrame != INDEX_NONE && FirstModifiedState != INDEX_NONE))
 		{
-			check(FirstModifiedFrame != INDEX_NONE);
-			check(FirstModifiedState != INDEX_NONE);
-
 			const FStateTreeExecutionFrame* ChangedFrame = &Exec.ActiveFrames[FirstModifiedFrame];
 			FStateTreeStateHandle ChangedState = ChangedFrame->ActiveStates[FirstModifiedState];
 			TConstArrayView<FStateTreeExecutionFrame> DeactivatedFrames;
