StateTree: Use a unique id to identify each execution frame and each active state. Introduce ActiveStatePath, a list of all the active states of all the active frames.

diff --git a/Engine/Plugins/Runtime/StateTree/Source/StateTreeModule/Private/StateTreeExecutionContext.cpp b/Engine/Plugins/Runtime/StateTree/Source/StateTreeModule/Private/StateTreeExecutionContext.cpp
--- a/Engine/Plugins/Runtime/StateTree/Source/StateTreeModule/Private/StateTreeExecutionContext.cpp
+++ b/Engine/Plugins/Runtime/StateTree/Source/StateTreeModule/Private/StateTreeExecutionContext.cpp
@@ -602,6 +602,7 @@ EStateTreeRunStatus FStateTreeExecutionContext::Start(FStartParameters Parameter
 	// Initialize for the init frame.
 	ensure(Exec.ActiveFrames.Num() == 0);
 	FStateTreeExecutionFrame& InitFrame = Exec.ActiveFrames.AddDefaulted_GetRef();
+	InitFrame.FrameId = UE::StateTree::FActiveFrameId(InstanceDataStorage->GenerateUniqueId());
 	InitFrame.StateTree = &RootStateTree;
 	InitFrame.RootState = FStateTreeStateHandle::Root;
 	InitFrame.ActiveStates = {};
@@ -2210,7 +2211,7 @@ EStateTreeRunStatus FStateTreeExecutionContext::EnterState(FStateTreeTransitionR
 			const FCompactStateTreeState& State = CurrentStateTree->States[CurrentHandle.Index];
 
 			// Add only enabled States to the list of active States
-			if (State.bEnabled && !CurrentFrame.ActiveStates.Push(CurrentHandle))
+			if (State.bEnabled && !CurrentFrame.ActiveStates.Push(CurrentHandle, NextFrame.ActiveStates.StateIds[StateIndex]))
 			{
 				STATETREE_LOG(Error, TEXT("%hs: Reached max execution depth when trying to enter state '%s'.  '%s' using StateTree '%s'."),
 					__FUNCTION__, *GetStateStatusString(Exec), *GetNameSafe(&Owner), *GetFullNameSafe(&RootStateTree));
@@ -3793,7 +3794,7 @@ void FStateTreeExecutionContext::SetupNextTransition(const FStateTreeExecutionFr
 	}
 	else
 	{
-		NewFrame.ActiveStates = FStateTreeActiveStates(NextState);
+		NewFrame.ActiveStates = FStateTreeActiveStates(NextState, UE::StateTree::FActiveStateId::Invalid);
 	}
 }
 
@@ -4680,7 +4681,8 @@ bool FStateTreeExecutionContext::SelectStateInternal(
 		StateSelectionEvents.Emplace();
 	}
 
-	if (!CurrentFrame.ActiveStates.Push(NextStateHandle))
+	// Activate/Push a new state
+	if (!CurrentFrame.ActiveStates.Push(NextStateHandle, UE::StateTree::FActiveStateId(InstanceDataStorage->GenerateUniqueId())))
 	{
 		STATETREE_LOG(Error, TEXT("%hs: Reached max execution depth when trying to select state %s from '%s'.  '%s' using StateTree '%s'."),
 			__FUNCTION__, *GetSafeStateName(CurrentFrame, NextStateHandle), *GetStateStatusString(Exec), *GetNameSafe(&Owner), *GetFullNameSafe(CurrentFrame.StateTree));
@@ -4746,6 +4748,7 @@ bool FStateTreeExecutionContext::SelectStateInternal(
 				FStateTreeExecutionFrame NewFrame;
 				NewFrame.StateTree = CurrentFrame.StateTree;
 				NewFrame.RootState = NextState.LinkedState;
+				NewFrame.FrameId = UE::StateTree::FActiveFrameId(InstanceDataStorage->GenerateUniqueId());
 				NewFrame.ExternalDataBaseIndex = CurrentFrame.ExternalDataBaseIndex;
 
 				// Check and prevent recursion.
@@ -4768,6 +4771,7 @@ bool FStateTreeExecutionContext::SelectStateInternal(
 					});
 				if (ExistingFrame)
 				{
+					// Note. Do not copy the FrameId. This is a new frame and the enter/exit state should be executed.
 					NewFrame.ActiveInstanceIndexBase = ExistingFrame->ActiveInstanceIndexBase;
 					NewFrame.GlobalInstanceIndexBase = ExistingFrame->GlobalInstanceIndexBase;
 					NewFrame.StateParameterDataHandle = ExistingFrame->StateParameterDataHandle;
@@ -4823,6 +4827,7 @@ bool FStateTreeExecutionContext::SelectStateInternal(
 			}
 				
 			FStateTreeExecutionFrame NewFrame;
+			NewFrame.FrameId = UE::StateTree::FActiveFrameId(InstanceDataStorage->GenerateUniqueId());
 			NewFrame.StateTree = NextLinkedStateAsset;
 			NewFrame.RootState = FStateTreeStateHandle::Root;
 			NewFrame.bIsGlobalFrame = true;
@@ -4849,6 +4854,7 @@ bool FStateTreeExecutionContext::SelectStateInternal(
 			bool bStartedTemporaryEvaluatorsAndGlobalTasks = false;
 			if (ExistingFrame)
 			{
+				// Note. Do not copy the FrameId. This is a new frame and the enter/exit state should be executed.
 				NewFrame.ActiveInstanceIndexBase = ExistingFrame->ActiveInstanceIndexBase;
 				NewFrame.GlobalInstanceIndexBase = ExistingFrame->GlobalInstanceIndexBase;
 				NewFrame.StateParameterDataHandle = ExistingFrame->StateParameterDataHandle;
diff --git a/Engine/Plugins/Runtime/StateTree/Source/StateTreeModule/Private/StateTreeExecutionTypes.cpp b/Engine/Plugins/Runtime/StateTree/Source/StateTreeModule/Private/StateTreeExecutionTypes.cpp
--- a/Engine/Plugins/Runtime/StateTree/Source/StateTreeModule/Private/StateTreeExecutionTypes.cpp
+++ b/Engine/Plugins/Runtime/StateTree/Source/StateTreeModule/Private/StateTreeExecutionTypes.cpp
@@ -91,6 +91,37 @@ FRecordedStateTreeTransitionResult::FRecordedStateTreeTransitionResult(const FSt
 	SourceRootState = Transition.SourceRootState;
 }
 
+
+//----------------------------------------------------------------------//
+// FStateTreeExecutionState
+//----------------------------------------------------------------------//
+UE::StateTree::FActiveStatePath FStateTreeExecutionState::GetActiveStatePath() const
+{
+	int32 NewNum = 0;
+	for (const FStateTreeExecutionFrame& Frame : ActiveFrames)
+	{
+		NewNum += Frame.ActiveStates.Num();
+	}
+
+	if (NewNum == 0 || ActiveFrames[0].StateTree == nullptr)
+	{
+		return UE::StateTree::FActiveStatePath();
+	}
+
+	TArray<UE::StateTree::FActiveStatePath::FElement> Elements;
+	Elements.Reserve(NewNum);
+
+	for (const FStateTreeExecutionFrame& Frame : ActiveFrames)
+	{
+		for (int32 StateIndex = 0; StateIndex < Frame.ActiveStates.Num(); ++StateIndex)
+		{
+			Elements.Emplace(Frame.ActiveStates.States[StateIndex], Frame.ActiveStates.StateIds[StateIndex]);
+		}
+	}
+
+	return UE::StateTree::FActiveStatePath(ActiveFrames[0].StateTree, MoveTemp(Elements));
+}
+
 //----------------------------------------------------------------------//
 // FStateTreeExecutionFrame
 //----------------------------------------------------------------------//
@@ -166,6 +197,10 @@ float FStateTreeScheduledTick::GetTickRate() const
 	return NextDeltaTime;
 }
 
+
+//----------------------------------------------------------------------//
+// FStateTreeDelegateActiveListeners
+//----------------------------------------------------------------------//
 FStateTreeDelegateActiveListeners::FActiveListener::FActiveListener(FStateTreeDelegateListener InListener, FSimpleDelegate InDelegate, TNotNull<const UStateTree*> InFrameStateTree, FStateTreeStateHandle InFrameRootState, FStateTreeStateHandle InOwningState, FStateTreeIndex16 InOwningNodeIndex)
 	: FrameStateTree(InFrameStateTree)
 	, Delegate(MoveTemp(InDelegate))
diff --git a/Engine/Plugins/Runtime/StateTree/Source/StateTreeModule/Private/StateTreeInstanceData.cpp b/Engine/Plugins/Runtime/StateTree/Source/StateTreeModule/Private/StateTreeInstanceData.cpp
--- a/Engine/Plugins/Runtime/StateTree/Source/StateTreeModule/Private/StateTreeInstanceData.cpp
+++ b/Engine/Plugins/Runtime/StateTree/Source/StateTreeModule/Private/StateTreeInstanceData.cpp
@@ -328,6 +328,23 @@ void FStateTreeInstanceStorage::SetGlobalParameters(const FInstancedPropertyBag&
 	GlobalParameters = Parameters;
 }
 
+uint32 FStateTreeInstanceStorage::GenerateUniqueId()
+{
+	uint32 NewId = ++UniqueIdGenerator;
+	if (NewId == 0)
+	{
+#if WITH_STATETREE_TRACE && DO_ENSURE
+		ensureAlwaysMsgf(false, TEXT("The unique id overflow. Id:%d Serial:%d"), ExecutionState.InstanceDebugId.Id, ExecutionState.InstanceDebugId.SerialNumber);
+#elif WITH_STATETREE_TRACE
+	UE_LOG(LogStateTree, Error, TEXT("The unique id overflow. Id:%d Serial:%d"), ExecutionState.InstanceDebugId.Id, ExecutionState.InstanceDebugId.SerialNumber);
+#else
+	UE_LOG(LogStateTree, Error, TEXT("The unique id overflow."));
+#endif
+		NewId = ++UniqueIdGenerator;
+	}
+	return NewId;
+}
+
 void FStateTreeInstanceStorage::AddStructReferencedObjects(FReferenceCollector& Collector)
 {
 	Collector.AddPropertyReferencesWithStructARO(TBaseStructure<FStateTreeInstanceStorage>::Get(), this);
diff --git a/Engine/Plugins/Runtime/StateTree/Source/StateTreeModule/Private/StateTreeStatePath.cpp b/Engine/Plugins/Runtime/StateTree/Source/StateTreeModule/Private/StateTreeStatePath.cpp
new file mode 100644
--- /dev/null
+++ b/Engine/Plugins/Runtime/StateTree/Source/StateTreeModule/Private/StateTreeStatePath.cpp
@@ -0,0 +1,173 @@
+// Copyright Epic Games, Inc. All Rights Reserved.
+
+#include "StateTreeStatePath.h"
+
+#include "StateTree.h"
+#include "StateTreeTypes.h"
+
+namespace UE::StateTree
+{
+	const FActiveFrameId FActiveFrameId::Invalid = FActiveFrameId();
+	const FActiveStateId FActiveStateId::Invalid = FActiveStateId();
+
+	FActiveStatePath::FActiveStatePath(TNotNull<const UStateTree*> InStateTree, const TArrayView<const FElement> InElements)
+		: StateTree(InStateTree)
+		, States(InElements)
+	{
+#if WITH_STATETREE_DEBUG
+		bool bIsValid = true;
+		for (const FElement& Element : InElements)
+		{
+			if (!Element.GetHandle().IsValid() || Element.GetHandle().IsCompletionState())
+			{
+				bIsValid = false;
+				break;
+			}
+		}
+		checkf(bIsValid, TEXT("Contains an invalid handle."));
+#endif
+	}
+
+	FActiveStatePath::FActiveStatePath(TNotNull<const UStateTree*> InStateTree, TArray<FElement> InElements)
+		: StateTree(InStateTree)
+		, States(MoveTemp(InElements))
+	{
+#if WITH_STATETREE_DEBUG
+		bool bIsValid = true;
+		for (const FElement& Element : InElements)
+		{
+			if (!Element.GetHandle().IsValid() || Element.GetHandle().IsCompletionState())
+			{
+				bIsValid = false;
+				break;
+			}
+		}
+		checkf(bIsValid, TEXT("Contains an invalid handle."));
+#endif
+	}
+
+	bool FActiveStatePath::Matches(const FActiveStatePath& Other) const
+	{
+		return States.Num() == Other.States.Num() && Intersect(Other).Num() == States.Num();
+	}
+
+	bool FActiveStatePath::Matches(FElement Other) const
+	{
+		return States.Num() > 0 && States.Last() == Other;
+	}
+	
+	bool FActiveStatePath::Matches(FActiveStateId Other) const
+	{
+		return States.Num() > 0 && States.Last().GetActiveStateId() == Other;
+	}
+
+	const TArrayView<const FActiveStatePath::FElement> FActiveStatePath::Intersect(const FActiveStatePath& Other) const
+	{
+		// If the id is the same, then are equal.
+		if (States.Num() == Other.States.Num()
+			&& States.Num() > 0
+			&& States.Last() == Other.States.Last())
+		{
+			return States;
+		}
+
+		int32 MatchIndex = 0;
+		const int32 MinNum = FMath::Min(States.Num(), Other.States.Num());
+		for (MatchIndex = 0; MatchIndex < MinNum; ++MatchIndex)
+		{
+			if (States[MatchIndex] != Other.States[MatchIndex])
+			{
+				break;
+			}
+		}
+		return MakeConstArrayView(States.GetData(), MatchIndex);
+	}
+
+	bool FActiveStatePath::StartsWith(const FActiveStatePath& Other) const
+	{
+		// No need to check all the element. If the first and last match then they all matches.
+		return Other.States.Num() > 0 ? Contains(Other.States.Last()) && States[0] == Other.States[0] : false;
+	}
+
+	bool FActiveStatePath::Contains(FElement Other) const
+	{
+		for (int32 Index = States.Num() - 1; Index >= 0; --Index)
+		{
+			if (States[Index] == Other)
+			{
+				return true;
+			}
+		}
+		return false;
+	}
+
+	bool FActiveStatePath::Contains(FActiveStateId Other) const
+	{
+		if (Other.IsValid())
+		{
+			for (int32 Index = States.Num() - 1; Index >= 0; --Index)
+			{
+				if (States[Index].GetActiveStateId() == Other)
+				{
+					return true;
+				}
+			}
+		}
+		return false;
+	}
+
+	TValueOrError<FString, void> FActiveStatePath::Describe() const
+	{
+		const UStateTree* CurrentTree = StateTree.Get();
+		if (CurrentTree == nullptr)
+		{
+			return MakeError();
+		}
+
+		TStringBuilder<256> Builder;
+		bool bNewTree = true;
+		for (const FElement& Element : States)
+		{
+			if (!Element.GetHandle().IsValid() || Element.GetHandle().IsCompletionState())
+			{
+				return MakeError();
+			}
+
+			{
+				if (Builder.Len() > 0)
+				{
+					Builder << TEXT("; ");
+				}
+				if (bNewTree)
+				{
+					Builder << TEXT('{');
+					CurrentTree->GetPathName(nullptr, Builder);
+					Builder << TEXT('}');
+					bNewTree = false;
+				}
+
+				if (CurrentTree->GetStates().IsValidIndex(Element.GetHandle().Index))
+				{
+					return MakeError();
+				}
+
+				const FCompactStateTreeState& State = CurrentTree->GetStates()[Element.GetHandle().Index];
+				Builder << State.Name;
+				Builder << TEXT('(');
+				Builder << Element.GetHandle().Index;
+				Builder << TEXT(')');
+
+				if (State.Type == EStateTreeStateType::LinkedAsset)
+				{
+					CurrentTree = State.LinkedAsset;
+					if (CurrentTree == nullptr)
+					{
+						return MakeError();
+					}
+				}
+			}
+		}
+		return MakeValue(Builder.ToString());
+	}
+
+}
\ No newline at end of file
diff --git a/Engine/Plugins/Runtime/StateTree/Source/StateTreeModule/Public/StateTreeExecutionTypes.h b/Engine/Plugins/Runtime/StateTree/Source/StateTreeModule/Public/StateTreeExecutionTypes.h
--- a/Engine/Plugins/Runtime/StateTree/Source/StateTreeModule/Public/StateTreeExecutionTypes.h
+++ b/Engine/Plugins/Runtime/StateTree/Source/StateTreeModule/Public/StateTreeExecutionTypes.h
@@ -4,6 +4,7 @@
 #include "StateTreeTypes.h"
 #include "StateTreeEvents.h"
 #include "StateTreeExecutionExtension.h"
+#include "StateTreeStatePath.h"
 #include "UObject/ObjectKey.h"
 
 #include "StateTreeExecutionTypes.generated.h"
@@ -298,9 +299,14 @@ struct STATETREEMODULE_API FStateTreeActiveStates
 
 	FStateTreeActiveStates() = default;
 	
+	UE_DEPRECATED(5.6, "Use the constructor with FActiveStateId.")
 	explicit FStateTreeActiveStates(const FStateTreeStateHandle StateHandle)
 	{
-		Push(StateHandle);
+		Push(StateHandle, UE::StateTree::FActiveStateId::Invalid);
+	}
+	explicit FStateTreeActiveStates(const FStateTreeStateHandle StateHandle, UE::StateTree::FActiveStateId StateId)
+	{
+		Push(StateHandle, StateId);
 	}
 	
 	/** Resets the active state array to empty. */
@@ -309,21 +315,37 @@ struct STATETREEMODULE_API FStateTreeActiveStates
 		NumStates = 0;
 	}
 
+	UE_DEPRECATED(5.6, "Use the Push override with the FActiveStateId.")
 	/** Pushes new state at the back of the array and returns true if there was enough space. */
 	bool Push(const FStateTreeStateHandle StateHandle)
+	{
+		return Push(StateHandle, UE::StateTree::FActiveStateId::Invalid);
+	}
+
+	/** Pushes new state at the back of the array and returns true if there was enough space. */
+	bool Push(const FStateTreeStateHandle StateHandle, UE::StateTree::FActiveStateId StateId)
 	{
 		if ((NumStates + 1) > MaxStates)
 		{
 			return false;
 		}
 		
-		States[NumStates++] = StateHandle;
+		States[NumStates] = StateHandle;
+		StateIds[NumStates] = StateId;
+		++NumStates;
 		
 		return true;
 	}
 
+	UE_DEPRECATED(5.6, "Use PushFront override with FActiveStateId.")
 	/** Pushes new state at the front of the array and returns true if there was enough space. */
 	bool PushFront(const FStateTreeStateHandle StateHandle)
+	{
+		return PushFront(StateHandle, UE::StateTree::FActiveStateId::Invalid);
+	}
+
+	/** Pushes new state at the front of the array and returns true if there was enough space. */
+	bool PushFront(const FStateTreeStateHandle StateHandle, UE::StateTree::FActiveStateId StateId)
 	{
 		if ((NumStates + 1) > MaxStates)
 		{
@@ -334,8 +356,10 @@ struct STATETREEMODULE_API FStateTreeActiveStates
 		for (int32 Index = (int32)NumStates - 1; Index > 0; Index--)
 		{
 			States[Index] = States[Index - 1];
+			StateIds[Index] = StateIds[Index - 1];
 		}
 		States[0] = StateHandle;
+		StateIds[0] = StateId;
 		
 		return true;
 	}
@@ -362,6 +386,7 @@ struct STATETREEMODULE_API FStateTreeActiveStates
 			for (int32 Index = NumStates; Index < NewNum; Index++)
 			{
 				States[Index] = FStateTreeStateHandle::Invalid;
+				StateIds[Index] = UE::StateTree::FActiveStateId::Invalid;
 			}
 		}
 		NumStates = static_cast<uint8>(NewNum);
@@ -446,6 +471,7 @@ struct STATETREEMODULE_API FStateTreeActiveStates
 	FORCEINLINE const FStateTreeStateHandle* begin() const { return &States[0]; }
 	FORCEINLINE const FStateTreeStateHandle* end  () const { return &States[0] + Num(); }
 
+	UE::StateTree::FActiveStateId StateIds[MaxStates];
 
 	UPROPERTY(VisibleDefaultsOnly, Category = Default)
 	FStateTreeStateHandle States[MaxStates];
@@ -684,6 +710,9 @@ struct STATETREEMODULE_API FStateTreeExecutionFrame
 	UPROPERTY()
 	FStateTreeActiveStates ActiveStates;
 
+	/** Unique frame id for this frame. Can be used to identify the frame. */
+	UE::StateTree::FActiveFrameId FrameId;
+
 	/** First index of the external data for this frame. */
 	UPROPERTY()
 	FStateTreeIndex16 ExternalDataBaseIndex = FStateTreeIndex16::Invalid;
@@ -739,6 +768,9 @@ PRAGMA_ENABLE_DEPRECATION_WARNINGS
 		*this = FStateTreeExecutionState();
 	}
 
+	/** @return the unique path of all the active states of all the active execution frames. */
+	UE::StateTree::FActiveStatePath GetActiveStatePath() const;
+
 	UE_DEPRECATED(5.5, "Use FindAndRemoveExpiredDelayedTransitions() instead or search through DelayedTransitions directly.")
 	FStateTreeTransitionDelayedState* FindDelayedTransition(const UStateTree* OwnerStateTree, const FStateTreeIndex16 TransitionIndex) { return nullptr; }
 
diff --git a/Engine/Plugins/Runtime/StateTree/Source/StateTreeModule/Public/StateTreeInstanceData.h b/Engine/Plugins/Runtime/StateTree/Source/StateTreeModule/Public/StateTreeInstanceData.h
--- a/Engine/Plugins/Runtime/StateTree/Source/StateTreeModule/Public/StateTreeInstanceData.h
+++ b/Engine/Plugins/Runtime/StateTree/Source/StateTreeModule/Public/StateTreeInstanceData.h
@@ -272,6 +272,9 @@ struct STATETREEMODULE_API FStateTreeInstanceStorage
 		return GlobalParameters.GetMutableValue();
 	}
 
+	/** @return a unique number used to make active frame id and active state id. */
+	uint32 GenerateUniqueId();
+
 	/** Note, called by FStateTreeInstanceData. */
 	void AddStructReferencedObjects(FReferenceCollector& Collector);
 
@@ -314,6 +317,10 @@ protected:
 	UPROPERTY(Transient)
 	FInstancedPropertyBag GlobalParameters;
 
+	/** Unique id.  */
+	UPROPERTY(Transient)
+	uint32 UniqueIdGenerator = 0;
+
 	/** Used to detect if we are using the instance data on multiple threads at the same time. */
 	UE_MT_DECLARE_RW_RECURSIVE_ACCESS_DETECTOR(AccessDetector)
 
diff --git a/Engine/Plugins/Runtime/StateTree/Source/StateTreeModule/Public/StateTreeStatePath.h b/Engine/Plugins/Runtime/StateTree/Source/StateTreeModule/Public/StateTreeStatePath.h
new file mode 100644
--- /dev/null
+++ b/Engine/Plugins/Runtime/StateTree/Source/StateTreeModule/Public/StateTreeStatePath.h
@@ -0,0 +1,170 @@
+// Copyright Epic Games, Inc. All Rights Reserved.
+
+#pragma once
+
+#include "CoreMinimal.h"
+#include "Containers/ArrayView.h"
+#include "Misc/NotNull.h"
+#include "StateTreeTypes.h"
+
+class UStateTree;
+struct FStateTreeExecutionFrame;
+
+namespace UE::StateTree
+{
+	/** A unique id for FStateTreeExecutionFrame. */
+	struct FActiveFrameId
+	{
+	private:
+		static constexpr uint32 InvalidValue = 0;
+
+	public:
+		static const FActiveFrameId Invalid;
+
+		explicit FActiveFrameId() = default;		
+		explicit FActiveFrameId(uint32 NewId)
+			: Value(NewId)
+		{
+			check(NewId != InvalidValue);
+		}
+
+		[[nodiscard]] bool operator==(const FActiveFrameId&) const = default;
+		[[nodiscard]] bool operator!=(const FActiveFrameId&) const = default;
+
+		[[nodiscard]] bool IsValid() const
+		{
+			return Value != InvalidValue;
+		}
+		
+	private:
+		uint32 Value = InvalidValue;
+	};
+	
+	/** A unique id for FStateTreeActiveStates state. */
+	struct FActiveStateId
+	{
+	private:
+		static constexpr uint32 InvalidValue = 0;
+
+	public:
+		static const FActiveStateId Invalid;
+
+		explicit FActiveStateId() = default;
+		explicit FActiveStateId(uint32 NewId)
+			: Value(NewId)
+		{
+			check(NewId != InvalidValue);
+		}
+		
+		[[nodiscard]] bool operator==(const FActiveStateId&) const = default;
+		[[nodiscard]] bool operator!=(const FActiveStateId&) const = default;
+
+		[[nodiscard]] bool IsValid() const
+		{
+			return Value != InvalidValue;
+		}
+	
+	private:
+		uint32 Value = InvalidValue;
+	};
+
+	/**
+	 * Describe the state list used to get to a specific state.
+	 * Since a state can be a subtree and can be linked to other states, the path to activate the subtree can be different.
+	 * A state can enter, then exit, then reenter. It is considered the same state path but not the same unique state path.
+	 * For a tree like
+	 * RootA(0)
+	 *   StateA(1)
+	 *   StateB(2)
+	 *     LinkedStateA(3)
+	 *   StateC(4)
+	 *     LinkedStateA(5)
+	 * SubTreeA(6)
+	 *   StateD(7)
+	 * The path to get to StateD can be RootA.StateB.LinkedStateA.SubTreeA.StateD(0.2.3.6.7)
+	 */
+	class STATETREEMODULE_API FActiveStatePath
+	{
+	public:
+		/** An state entry in the state path. */
+		struct FElement
+		{
+			explicit FElement(FStateTreeStateHandle InHandle, FActiveStateId InStateId)
+				: ActiveStateId(InStateId)
+				, TreeStateIndex(InHandle)
+			{}
+
+			[[nodiscard]] bool operator==(const FElement&) const = default;
+			[[nodiscard]] bool operator!=(const FElement&) const = default;
+
+			[[nodiscard]] bool IsValid() const
+			{
+				return TreeStateIndex.IsValid();
+			}
+
+			[[nodiscard]] FActiveStateId GetActiveStateId() const
+			{
+				return ActiveStateId;
+			}
+
+			[[nodiscard]] FStateTreeStateHandle GetHandle() const
+			{
+				return TreeStateIndex;
+			}
+
+		private:
+			/** The unique id of the state for this instance. */
+			FActiveStateId ActiveStateId;
+			/** The index of the state inside the state tree asset. */
+			FStateTreeStateHandle TreeStateIndex;
+		};
+
+	public:
+		explicit FActiveStatePath() = default;
+		explicit FActiveStatePath(TNotNull<const UStateTree*> InStateTree, const TArrayView<const FElement> InElements);
+		explicit FActiveStatePath(TNotNull<const UStateTree*> InStateTree, TArray<FElement> InElements);
+
+		/** @return the number of elements in the path. */
+		[[nodiscard]] int32 Num() const
+		{
+			return States.Num();
+		}
+
+		/** @return a view of the states in the path. */
+		[[nodiscard]] const TArrayView<const FElement> GetView() const
+		{
+			return States;
+		}
+
+		/** @return true if both paths match exactly. */
+		[[nodiscard]] bool Matches(const FActiveStatePath& Other) const;
+
+		/** @return true if the last element in the path matches. */
+		[[nodiscard]] bool Matches(FElement Other) const;
+
+		/** @return true if the last element in the path matches. */
+		[[nodiscard]] bool Matches(FActiveStateId Other) const;
+
+		/** @return the path where the 2 paths intersect. */
+		[[nodiscard]] const TArrayView<const FElement> Intersect(const FActiveStatePath& Other) const;
+
+		/** @return the path that is not in the other path. */
+		[[nodiscard]] const TArrayView<const FElement> Difference(const FActiveStatePath& Other) const;
+
+		/** @return true if this path starts with and contains the other path. */
+		[[nodiscard]] bool StartsWith(const FActiveStatePath& Other) const;
+
+		/** @return true if this path contains anywhere the other element. */
+		[[nodiscard]] bool Contains(FElement Other) const;
+
+		/** @return true if this path contains anywhere the other element. */
+		[[nodiscard]] bool Contains(FActiveStateId Other) const;
+
+		/** @return debugging information about the path. */
+		[[nodiscard]] TValueOrError<FString, void> Describe() const;
+
+	private:
+		TWeakObjectPtr<const UStateTree> StateTree;
+		TArray<FElement> States;
+	};
+}
diff --git a/Engine/Plugins/Runtime/StateTree/Source/StateTreeTestSuite/Private/StateTreeStatePathTest.cpp b/Engine/Plugins/Runtime/StateTree/Source/StateTreeTestSuite/Private/StateTreeStatePathTest.cpp
new file mode 100644
--- /dev/null
+++ b/Engine/Plugins/Runtime/StateTree/Source/StateTreeTestSuite/Private/StateTreeStatePathTest.cpp
@@ -0,0 +1,446 @@
+// Copyright Epic Games, Inc. All Rights Reserved.
+
+#include "StateTreeTest.h"
+#include "StateTreeTestBase.h"
+#include "StateTreeTestTypes.h"
+
+#include "StateTreeCompilerLog.h"
+#include "StateTreeEditorData.h"
+#include "StateTreeCompiler.h"
+#include "Conditions/StateTreeCommonConditions.h"
+
+#define LOCTEXT_NAMESPACE "AITestSuite_StateTreeTest"
+
+UE_DISABLE_OPTIMIZATION_SHIP
+
+namespace UE::StateTree::Tests
+{
+struct FStateTreeTest_StatePath_LinkStates : FStateTreeTestBase
+{
+	virtual bool InstantTest() override
+	{
+		//Tree 1
+		//	Root
+		//		RootState1 -> Next
+		//			StateLinkedTree (Sub1) -> Next
+		//			StateLinkedTree (Sub1) -> Next
+		//		RootStateLinkedTree (Sub2) -> Next
+		//		RootStateLinkedTree (Sub2) -> Next
+		//		RootStateLinkedTree (Tree2) -> Next
+		//		RootStateLinkedTree (Tree3) -> Next			# Tree3 fails
+		//		RootStateLinkedTree (Tree2) -> Next
+		//		RootState2 -> Root
+		//	Sub1
+		//		StateLinkedTree (Tree2) -> Next
+		//		StateLinkedTree (Tree2) -> Next
+		//		StateLinkedTree (Sub2) -> Next
+		//		StateLinkedTree (Sub2) -> Next
+		//		Sub1State1 -> Success
+		//	Sub2
+		//		Sub2State1 -> Next
+		//		Sub2State2 -> Success
+		//Tree 2
+		//	Root
+		//		RootStateLinkedTree (Sub1) -> Next
+		//		RootStateLinkedTree (Sub1) -> Next
+		//		RootState1 -> Success
+		//	Sub1
+		//		Sub1State1 -> Next
+		//		Sub1State2 -> Success
+		//Tree 3
+		//	Root
+		//		RootStateLinkedTree (Sub1) -> Success
+		//	Sub1
+		//		Sub1State1 -> cond fail -> success
+		UStateTree& StateTree1 = NewStateTree();
+		UStateTree& StateTree2 = NewStateTree();
+		UStateTree& StateTree3 = NewStateTree();
+		// Tree1
+		{
+			UStateTreeEditorData& EditorData = *Cast<UStateTreeEditorData>(StateTree1.EditorData);
+			UStateTreeState& Root = EditorData.AddSubTree("Tree1Root");
+			UStateTreeState& Sub1 = EditorData.AddSubTree("Tree1Sub1");
+			Sub1.Type = EStateTreeStateType::Subtree;
+			UStateTreeState& Sub2 = EditorData.AddSubTree("Tree1Sub2");
+			Sub2.Type = EStateTreeStateType::Subtree;
+			// Root
+			{
+				UStateTreeState& Tree1RootState1 = Root.AddChildState("Tree1RootState1", EStateTreeStateType::State);
+				UStateTreeState& Tree1RootStateSub2A = Root.AddChildState("Tree1RootStateSub2A", EStateTreeStateType::Linked);
+				//Tree1RootState1
+				{
+					{
+						UStateTreeState& ChildState = Tree1RootState1.AddChildState("Tree1State1StateSub1A", EStateTreeStateType::Linked);
+						ChildState.LinkedSubtree = Sub1.GetLinkToState();
+						ChildState.AddTransition(EStateTreeTransitionTrigger::OnStateSucceeded, EStateTreeTransitionType::NextState);
+					}
+					{
+						UStateTreeState& ChildState = Tree1RootState1.AddChildState("Tree1State1StateSub1B", EStateTreeStateType::Linked);
+						ChildState.LinkedSubtree = Sub1.GetLinkToState();
+						FStateTreeTransition& Transition = ChildState.AddTransition(EStateTreeTransitionTrigger::OnStateSucceeded, EStateTreeTransitionType::GotoState);
+						Transition.State = Tree1RootStateSub2A.GetLinkToState();
+					}
+				}
+				//Tree1RootStateSub2A
+				{
+					Tree1RootStateSub2A.LinkedSubtree = Sub2.GetLinkToState();
+					Tree1RootStateSub2A.AddTransition(EStateTreeTransitionTrigger::OnStateSucceeded, EStateTreeTransitionType::NextState);
+				}
+				{
+					UStateTreeState& State = Root.AddChildState("Tree1RootStateSub2B", EStateTreeStateType::Linked);
+					State.LinkedSubtree = Sub2.GetLinkToState();
+					State.AddTransition(EStateTreeTransitionTrigger::OnStateSucceeded, EStateTreeTransitionType::NextState);
+				}
+				{
+					UStateTreeState& State = Root.AddChildState("Tree1RootStateLinkTree2A", EStateTreeStateType::LinkedAsset);
+					State.LinkedAsset = &StateTree2;
+					State.AddTransition(EStateTreeTransitionTrigger::OnStateSucceeded, EStateTreeTransitionType::NextState);
+				}
+				{
+					UStateTreeState& State = Root.AddChildState("Tree1RootStateLinkTree3", EStateTreeStateType::LinkedAsset);
+					State.LinkedAsset = &StateTree3;
+					State.AddTransition(EStateTreeTransitionTrigger::OnStateSucceeded, EStateTreeTransitionType::NextState);
+				}
+				{
+					UStateTreeState& State = Root.AddChildState("Tree1RootStateLinkTree2B", EStateTreeStateType::LinkedAsset);
+					State.LinkedAsset = &StateTree2;
+					State.AddTransition(EStateTreeTransitionTrigger::OnStateSucceeded, EStateTreeTransitionType::NextState);
+				}
+				{
+					UStateTreeState& State = Root.AddChildState("Tree1RootState1", EStateTreeStateType::State);
+					FStateTreeTransition& Transition = State.AddTransition(EStateTreeTransitionTrigger::OnTick, EStateTreeTransitionType::GotoState);
+					Transition.State = Root.GetLinkToState();
+					Transition.bDelayTransition = true;
+					Transition.DelayDuration = 0.999f;
+				}
+			}
+			//Tree1Sub1
+			{
+				{
+					UStateTreeState& State = Sub1.AddChildState("Tree1Sub1StateLinkTree2A", EStateTreeStateType::LinkedAsset);
+					State.LinkedAsset = &StateTree2;
+					State.AddTransition(EStateTreeTransitionTrigger::OnStateSucceeded, EStateTreeTransitionType::NextState);
+				}
+				{
+					UStateTreeState& State = Sub1.AddChildState("Tree1Sub1StateLinkTree2B", EStateTreeStateType::LinkedAsset);
+					State.LinkedAsset = &StateTree2;
+					State.AddTransition(EStateTreeTransitionTrigger::OnStateSucceeded, EStateTreeTransitionType::NextState);
+				}
+				{
+					UStateTreeState& State = Sub1.AddChildState("Tree1Sub1StateSub2A", EStateTreeStateType::Linked);
+					State.LinkedSubtree = Sub2.GetLinkToState();
+					State.AddTransition(EStateTreeTransitionTrigger::OnStateSucceeded, EStateTreeTransitionType::NextState);
+				}
+				{
+					UStateTreeState& State = Sub1.AddChildState("Tree1Sub1StateSub2B", EStateTreeStateType::Linked);
+					State.LinkedSubtree = Sub2.GetLinkToState();
+					State.AddTransition(EStateTreeTransitionTrigger::OnStateSucceeded, EStateTreeTransitionType::NextState);
+				}
+				{
+					UStateTreeState& State = Sub1.AddChildState("TreeASub1State1", EStateTreeStateType::State);
+					FStateTreeTransition& Transition = State.AddTransition(EStateTreeTransitionTrigger::OnTick, EStateTreeTransitionType::Succeeded);
+					Transition.bDelayTransition = true;
+					Transition.DelayDuration = 0.999f;
+				}
+			}
+			//Tree1Sub2
+			{
+				{
+					UStateTreeState& State = Sub2.AddChildState("Tree1Sub2State1", EStateTreeStateType::State);
+					State.AddTransition(EStateTreeTransitionTrigger::OnStateSucceeded, EStateTreeTransitionType::NextState);
+				}
+				{
+					UStateTreeState& State = Sub2.AddChildState("Tree1Sub2State2", EStateTreeStateType::State);
+					FStateTreeTransition& Transition = State.AddTransition(EStateTreeTransitionTrigger::OnTick, EStateTreeTransitionType::Succeeded);
+					Transition.bDelayTransition = true;
+					Transition.DelayDuration = 0.999f;
+				}
+			}
+		}
+		//Tree 2
+		{
+			UStateTreeEditorData& EditorData = *Cast<UStateTreeEditorData>(StateTree2.EditorData);
+			UStateTreeState& Root = EditorData.AddSubTree("Tree2StateRoot");
+			UStateTreeState& Sub1 = EditorData.AddSubTree("Tree2Sub1");
+			Sub1.Type = EStateTreeStateType::Subtree;
+			//Root
+			{
+				{
+					UStateTreeState& State = Root.AddChildState("Tree2RootStateSub1A", EStateTreeStateType::Linked);
+					State.LinkedSubtree = Sub1.GetLinkToState();
+					State.AddTransition(EStateTreeTransitionTrigger::OnStateSucceeded, EStateTreeTransitionType::NextState);
+				}
+				{
+					UStateTreeState& State = Root.AddChildState("Tree2RootStateSub1B", EStateTreeStateType::Linked);
+					State.LinkedSubtree = Sub1.GetLinkToState();
+					State.AddTransition(EStateTreeTransitionTrigger::OnStateSucceeded, EStateTreeTransitionType::NextState);
+				}
+				{
+					UStateTreeState& State = Root.AddChildState("Tree2RootState1", EStateTreeStateType::State);
+					FStateTreeTransition& Transition = State.AddTransition(EStateTreeTransitionTrigger::OnTick, EStateTreeTransitionType::Succeeded);
+					Transition.bDelayTransition = true;
+					Transition.DelayDuration = 0.999f;
+				}
+			}
+			//Tree2Sub1
+			{
+				{
+					UStateTreeState& State = Sub1.AddChildState("Tree2Sub1State1", EStateTreeStateType::State);
+					State.AddTransition(EStateTreeTransitionTrigger::OnStateSucceeded, EStateTreeTransitionType::NextState);
+				}
+				{
+					UStateTreeState& State = Sub1.AddChildState("Tree2Sub1State2", EStateTreeStateType::State);
+					FStateTreeTransition& Transition = State.AddTransition(EStateTreeTransitionTrigger::OnTick, EStateTreeTransitionType::Succeeded);
+					Transition.bDelayTransition = true;
+					Transition.DelayDuration = 0.999f;
+				}
+			}
+		}
+		//Tree 3
+		{
+			UStateTreeEditorData& EditorData = *Cast<UStateTreeEditorData>(StateTree3.EditorData);
+			UStateTreeState& Root = EditorData.AddSubTree("Tree3StateRoot");
+			UStateTreeState& Sub1 = EditorData.AddSubTree("Tree3Sub1");
+			Sub1.Type = EStateTreeStateType::Subtree;
+			//Root
+			{
+				UStateTreeState& State = Root.AddChildState("Tree2RootStateSub1A", EStateTreeStateType::Linked);
+				State.LinkedSubtree = Sub1.GetLinkToState();
+				State.AddTransition(EStateTreeTransitionTrigger::OnStateSucceeded, EStateTreeTransitionType::Succeeded);
+			}
+			//Tree3Sub1
+			{
+				UStateTreeState& State = Sub1.AddChildState("Tree3Sub1State1", EStateTreeStateType::State);
+				FStateTreeTransition& Transition = State.AddTransition(EStateTreeTransitionTrigger::OnTick, EStateTreeTransitionType::Succeeded);
+				Transition.bDelayTransition = true;
+				Transition.DelayDuration = 0.999f;
+
+				TStateTreeEditorNode<FStateTreeRandomCondition>& Cond = State.AddEnterCondition<FStateTreeRandomCondition>();
+				Cond.GetNode().EvaluationMode = EStateTreeConditionEvaluationMode::ForcedFalse;
+			}
+		}
+
+		// Compile tree
+		{
+			FStateTreeCompilerLog Log;
+			FStateTreeCompiler Compiler(Log);
+			const bool bResult = Compiler.Compile(StateTree3);
+			AITEST_TRUE("StateTree3 should get compiled", bResult);
+		}
+		{
+			FStateTreeCompilerLog Log;
+			FStateTreeCompiler Compiler(Log);
+			const bool bResult = Compiler.Compile(StateTree2);
+			AITEST_TRUE("StateTree2 should get compiled", bResult);
+		}
+		{
+			FStateTreeCompilerLog Log;
+			FStateTreeCompiler Compiler(Log);
+			const bool bResult = Compiler.Compile(StateTree1);
+			AITEST_TRUE("StateTree1 should get compiled", bResult);
+		}
+
+		// Create context
+		FStateTreeInstanceData InstanceData;
+		FTestStateTreeExecutionContext Exec(StateTree1, StateTree1, InstanceData);
+		{
+			const bool bInitSucceeded = Exec.IsValid();
+			AITEST_TRUE("StateTree should init", bInitSucceeded);
+		}
+		{
+			UE::StateTree::FActiveStatePath ExecPath = InstanceData.GetExecutionState()->GetActiveStatePath();
+			AITEST_TRUE("ExecPath should be empty.", ExecPath.Num() == 0);
+		}
+
+		// Test variable and helper functions
+		int32 ActiveCounter = 0;
+		TArray<UE::StateTree::FActiveFrameId> PreviousFrameIds;
+		auto PreTestFrameId = [&PreviousFrameIds, &InstanceData]()
+			{
+				PreviousFrameIds.Reset();
+				for (int32 Index = 0; Index < InstanceData.GetExecutionState()->ActiveFrames.Num(); ++Index)
+				{
+					PreviousFrameIds.Add(InstanceData.GetExecutionState()->ActiveFrames[Index].FrameId);
+				}
+			};
+		auto TestFrameId = [&PreviousFrameIds, &InstanceData](int32 CorrectAmount)
+			{
+				int32 Index = 0;
+				for (; Index < PreviousFrameIds.Num() && Index < CorrectAmount; ++Index)
+				{
+					if (PreviousFrameIds[Index] != InstanceData.GetExecutionState()->ActiveFrames[Index].FrameId)
+					{
+						return false;
+					}
+				}
+				for (; Index < InstanceData.GetExecutionState()->ActiveFrames.Num() && Index < PreviousFrameIds.Num(); ++Index)
+				{
+					if (PreviousFrameIds[Index] == InstanceData.GetExecutionState()->ActiveFrames[Index].FrameId)
+					{
+						return false;
+					}
+				}
+				return true;
+			};
+
+		// Start tests
+		{
+			const EStateTreeRunStatus Status = Exec.Start();
+			AITEST_EQUAL("Start should complete with Running", Status, EStateTreeRunStatus::Running);
+			AITEST_TRUE("Should be in the correct state", Exec.ExpectInActiveStates("Tree1Root", "Tree1RootState1", "Tree1State1StateSub1A", "Tree1Sub1", "Tree1Sub1StateLinkTree2A", "Tree2StateRoot", "Tree2RootStateSub1A", "Tree2Sub1", "Tree2Sub1State1"));
+
+			const UE::StateTree::FActiveStatePath ExecPath = InstanceData.GetExecutionState()->GetActiveStatePath();
+			AITEST_TRUE("Has the correct number of path elements", ExecPath.Num() == 9);
+			AITEST_TRUE("Has the correct number of active states", InstanceData.GetExecutionState()->ActiveFrames.Num() == 4);
+			{
+				AITEST_TRUE("Frame for Tree1Root is active", InstanceData.GetExecutionState()->ActiveFrames[0].FrameId == UE::StateTree::FActiveFrameId(++ActiveCounter));
+				AITEST_TRUE("State Tree1Root is active", ExecPath.Contains(UE::StateTree::FActiveStateId(++ActiveCounter)));
+				AITEST_TRUE("State Tree1Root is active", ExecPath.Contains(UE::StateTree::FActiveStatePath::FElement(FStateTreeStateHandle(0), UE::StateTree::FActiveStateId(ActiveCounter))));
+				AITEST_TRUE("State Tree1RootState1 is active", ExecPath.Contains(UE::StateTree::FActiveStateId(++ActiveCounter)));
+				AITEST_TRUE("State Tree1RootState1 is active", ExecPath.Contains(UE::StateTree::FActiveStatePath::FElement(FStateTreeStateHandle(1), UE::StateTree::FActiveStateId(ActiveCounter))));
+				AITEST_TRUE("State Tree1State1StateSub1A is active", ExecPath.Contains(UE::StateTree::FActiveStateId(++ActiveCounter)));
+				AITEST_TRUE("State Tree1State1StateSub1A is active", ExecPath.Contains(UE::StateTree::FActiveStatePath::FElement(FStateTreeStateHandle(2), UE::StateTree::FActiveStateId(4))));
+			}
+			{
+				AITEST_TRUE("Frame for Tree1Sub1 is active", InstanceData.GetExecutionState()->ActiveFrames[1].FrameId == UE::StateTree::FActiveFrameId(++ActiveCounter));
+				AITEST_TRUE("State Tree1Sub1 is active", ExecPath.Contains(UE::StateTree::FActiveStateId(++ActiveCounter)));
+				AITEST_TRUE("State Tree1Sub1StateLinkTree2A is active", ExecPath.Contains(UE::StateTree::FActiveStateId(++ActiveCounter)));
+			}
+			{
+				AITEST_TRUE("Frame for Tree2StateRoot is active", InstanceData.GetExecutionState()->ActiveFrames[2].FrameId == UE::StateTree::FActiveFrameId(++ActiveCounter));
+				AITEST_TRUE("State Tree2StateRoot is active", ExecPath.Contains(UE::StateTree::FActiveStateId(++ActiveCounter)));
+				AITEST_TRUE("State Tree2RootStateSub1A is active", ExecPath.Contains(UE::StateTree::FActiveStateId(++ActiveCounter)));
+			}
+			{
+				AITEST_TRUE("Frame for Tree2Sub1 is active", InstanceData.GetExecutionState()->ActiveFrames[3].FrameId == UE::StateTree::FActiveFrameId(++ActiveCounter));
+				AITEST_TRUE("State Tree2Sub1 is active", ExecPath.Contains(UE::StateTree::FActiveStateId(++ActiveCounter)));
+				AITEST_TRUE("State Tree2Sub1State1 is active", ExecPath.Contains(UE::StateTree::FActiveStateId(++ActiveCounter)));
+			}
+			{
+				AITEST_FALSE("No accidental increment", ExecPath.Contains(UE::StateTree::FActiveStateId(ActiveCounter+1)));
+			}
+			Exec.LogClear();
+		}
+		{
+			PreTestFrameId();
+			EStateTreeRunStatus Status = Exec.Tick(1.0f);
+			AITEST_EQUAL("Start should complete with Running", Status, EStateTreeRunStatus::Running);
+			AITEST_TRUE("Should be in the correct state", Exec.ExpectInActiveStates("Tree1Root", "Tree1RootState1", "Tree1State1StateSub1A", "Tree1Sub1", "Tree1Sub1StateLinkTree2A", "Tree2StateRoot", "Tree2RootStateSub1A", "Tree2Sub1", "Tree2Sub1State2"));
+
+			UE::StateTree::FActiveStatePath ExecPath = InstanceData.GetExecutionState()->GetActiveStatePath();
+			AITEST_TRUE("Has the correct number of path elements", ExecPath.Num() == 9);
+			AITEST_TRUE("Has the correct number of active states", InstanceData.GetExecutionState()->ActiveFrames.Num() == 4);
+			AITEST_TRUE("State Tree2Sub1State2 is active", ExecPath.Contains(UE::StateTree::FActiveStateId(++ActiveCounter)));
+			AITEST_FALSE("No accidental increment", ExecPath.Contains(UE::StateTree::FActiveStateId(ActiveCounter + 1)));
+			AITEST_TRUE("All frames are the same", TestFrameId(4));
+			Exec.LogClear();
+		}
+		{
+			PreTestFrameId();
+			EStateTreeRunStatus Status = Exec.Tick(1.0f);
+			AITEST_EQUAL("Start should complete with Running", Status, EStateTreeRunStatus::Running);
+			AITEST_TRUE("Should be in the correct state", Exec.ExpectInActiveStates("Tree1Root", "Tree1RootState1", "Tree1State1StateSub1A", "Tree1Sub1", "Tree1Sub1StateLinkTree2A", "Tree2StateRoot", "Tree2RootStateSub1B", "Tree2Sub1", "Tree2Sub1State1"));
+
+			const UE::StateTree::FActiveStatePath ExecPath = InstanceData.GetExecutionState()->GetActiveStatePath();
+			AITEST_TRUE("Has the correct number of path elements", ExecPath.Num() == 9);
+			AITEST_TRUE("Has the correct number of active states", InstanceData.GetExecutionState()->ActiveFrames.Num() == 4);
+			AITEST_TRUE("The last frame changed", TestFrameId(3));
+			Exec.LogClear();
+		}
+		{
+			PreTestFrameId();
+			EStateTreeRunStatus Status = Exec.Tick(1.0f);
+			AITEST_EQUAL("Start should complete with Running", Status, EStateTreeRunStatus::Running);
+			AITEST_TRUE("Should be in the correct state", Exec.ExpectInActiveStates("Tree1Root", "Tree1RootState1", "Tree1State1StateSub1A", "Tree1Sub1", "Tree1Sub1StateLinkTree2A", "Tree2StateRoot", "Tree2RootStateSub1B", "Tree2Sub1", "Tree2Sub1State2"));
+
+			UE::StateTree::FActiveStatePath ExecPath = InstanceData.GetExecutionState()->GetActiveStatePath();
+			AITEST_TRUE("Has the correct number of path elements", ExecPath.Num() == 9);
+			AITEST_TRUE("Has the correct number of active states", InstanceData.GetExecutionState()->ActiveFrames.Num() == 4);
+			AITEST_TRUE("All frames are the same", TestFrameId(4));
+			Exec.LogClear();
+		}
+		{
+			PreTestFrameId();
+			EStateTreeRunStatus Status = Exec.Tick(1.0f);
+			AITEST_EQUAL("Start should complete with Running", Status, EStateTreeRunStatus::Running);
+			AITEST_TRUE("Should be in the correct state", Exec.ExpectInActiveStates("Tree1Root", "Tree1RootState1", "Tree1State1StateSub1A", "Tree1Sub1", "Tree1Sub1StateLinkTree2A", "Tree2StateRoot", "Tree2RootState1"));
+
+			UE::StateTree::FActiveStatePath ExecPath = InstanceData.GetExecutionState()->GetActiveStatePath();
+			AITEST_TRUE("Has the correct number of path elements", ExecPath.Num() == 7);
+			AITEST_TRUE("Has the correct number of active states", InstanceData.GetExecutionState()->ActiveFrames.Num() == 3);
+			AITEST_TRUE("All frames are the same", TestFrameId(3));
+			Exec.LogClear();
+		}
+		{
+			PreTestFrameId();
+			EStateTreeRunStatus Status = Exec.Tick(1.0f);
+			AITEST_EQUAL("Start should complete with Running", Status, EStateTreeRunStatus::Running);
+			AITEST_TRUE("Should be in the correct state", Exec.ExpectInActiveStates("Tree1Root", "Tree1RootState1", "Tree1State1StateSub1A", "Tree1Sub1", "Tree1Sub1StateLinkTree2B", "Tree2StateRoot", "Tree2RootStateSub1A", "Tree2Sub1", "Tree2Sub1State1"));
+
+			UE::StateTree::FActiveStatePath ExecPath = InstanceData.GetExecutionState()->GetActiveStatePath();
+			AITEST_TRUE("Has the correct number of path elements", ExecPath.Num() == 9);
+			AITEST_TRUE("Has the correct number of active states", InstanceData.GetExecutionState()->ActiveFrames.Num() == 4);
+			AITEST_TRUE("All frames are the same", TestFrameId(2));
+			Exec.LogClear();
+		}
+		{
+			PreTestFrameId();
+			EStateTreeRunStatus Status = Exec.Tick(1.0f);
+			AITEST_EQUAL("Start should complete with Running", Status, EStateTreeRunStatus::Running);
+			AITEST_TRUE("Should be in the correct state", Exec.ExpectInActiveStates("Tree1Root", "Tree1RootState1", "Tree1State1StateSub1A", "Tree1Sub1", "Tree1Sub1StateLinkTree2B", "Tree2StateRoot", "Tree2RootStateSub1A", "Tree2Sub1", "Tree2Sub1State2"));
+			AITEST_TRUE("All frames are the same", TestFrameId(4));
+			Exec.LogClear();
+
+			PreTestFrameId();
+			Status = Exec.Tick(1.0f);
+			AITEST_EQUAL("Start should complete with Running", Status, EStateTreeRunStatus::Running);
+			AITEST_TRUE("Should be in the correct state", Exec.ExpectInActiveStates("Tree1Root", "Tree1RootState1", "Tree1State1StateSub1A", "Tree1Sub1", "Tree1Sub1StateLinkTree2B", "Tree2StateRoot", "Tree2RootStateSub1B", "Tree2Sub1", "Tree2Sub1State1")); 
+			AITEST_TRUE("All frames are the same", TestFrameId(3));
+			Exec.LogClear();
+
+			PreTestFrameId();
+			Status = Exec.Tick(1.0f);
+			AITEST_EQUAL("Start should complete with Running", Status, EStateTreeRunStatus::Running);
+			AITEST_TRUE("Should be in the correct state", Exec.ExpectInActiveStates("Tree1Root", "Tree1RootState1", "Tree1State1StateSub1A", "Tree1Sub1", "Tree1Sub1StateLinkTree2B", "Tree2StateRoot", "Tree2RootStateSub1B", "Tree2Sub1", "Tree2Sub1State2"));
+			AITEST_TRUE("Has the correct number of active states", InstanceData.GetExecutionState()->ActiveFrames.Num() == 4);
+			AITEST_TRUE("All frames are the same", TestFrameId(4));
+			Exec.LogClear();
+
+			PreTestFrameId();
+			Status = Exec.Tick(1.0f);
+			AITEST_EQUAL("Start should complete with Running", Status, EStateTreeRunStatus::Running);
+			AITEST_TRUE("Should be in the correct state", Exec.ExpectInActiveStates("Tree1Root", "Tree1RootState1", "Tree1State1StateSub1A", "Tree1Sub1", "Tree1Sub1StateLinkTree2B", "Tree2StateRoot", "Tree2RootState1"));
+			AITEST_TRUE("Has the correct number of active states", InstanceData.GetExecutionState()->ActiveFrames.Num() == 3);
+			AITEST_TRUE("All frames are the same", TestFrameId(3));
+			Exec.LogClear();
+		}
+		{
+			PreTestFrameId();
+			EStateTreeRunStatus Status = Exec.Tick(1.0f);
+			AITEST_EQUAL("Start should complete with Running", Status, EStateTreeRunStatus::Running);
+			AITEST_TRUE("Should be in the correct state", Exec.ExpectInActiveStates("Tree1Root", "Tree1RootState1", "Tree1State1StateSub1A", "Tree1Sub1", "Tree1Sub1StateSub2A", "Tree1Sub2", "Tree1Sub2State1"));
+			AITEST_TRUE("Has the correct number of active states", InstanceData.GetExecutionState()->ActiveFrames.Num() == 3);
+			AITEST_TRUE("All frames are the same", TestFrameId(2));
+			Exec.LogClear();
+
+			PreTestFrameId();
+			Status = Exec.Tick(1.0f);
+			AITEST_EQUAL("Start should complete with Running", Status, EStateTreeRunStatus::Running);
+			AITEST_TRUE("Should be in the correct state", Exec.ExpectInActiveStates("Tree1Root", "Tree1RootState1", "Tree1State1StateSub1A", "Tree1Sub1", "Tree1Sub1StateSub2A", "Tree1Sub2", "Tree1Sub2State2"));
+			AITEST_TRUE("Has the correct number of active states", InstanceData.GetExecutionState()->ActiveFrames.Num() == 3);
+			AITEST_TRUE("All frames are the same", TestFrameId(3));
+			Exec.LogClear();
+		}
+
+		Exec.Stop();
+
+		return true;
+	}
+};
+IMPLEMENT_AI_INSTANT_TEST(FStateTreeTest_StatePath_LinkStates, "System.StateTree.StatePath.LinkStates");
+
+} // namespace UE::StateTree::Tests
+
+UE_ENABLE_OPTIMIZATION_SHIP
+
+#undef LOCTEXT_NAMESPACE
