StateTree: Fix scheduled tick when a state has an override but should tick the next frame because of events or delayed transitions.

diff --git a/Engine/Plugins/Runtime/StateTree/Source/StateTreeModule/Private/StateTreeExecutionContext.cpp b/Engine/Plugins/Runtime/StateTree/Source/StateTreeModule/Private/StateTreeExecutionContext.cpp
--- a/Engine/Plugins/Runtime/StateTree/Source/StateTreeModule/Private/StateTreeExecutionContext.cpp
+++ b/Engine/Plugins/Runtime/StateTree/Source/StateTreeModule/Private/StateTreeExecutionContext.cpp
@@ -195,9 +195,10 @@ FStateTreeScheduledTick FStateTreeMinimalExecutionContext::GetNextScheduledTick(
 	const bool bHasEvents = EventQueue.HasEvents();
 	const bool bHasBroadcastedDelegates = Storage.HasBroadcastedDelegates();
 
+	// We wish to return in order: EveryFrames, then NextFrame, then CustomTickRate, then Sleep.
 	// Do we have a state that requires a tick or is waiting for an event.
+	TOptional<float> StateCustomTickRate;
 	{
-		TOptional<float> StateCustomTickRate;
 		bool bHasTaskWithEveryFramesTick = false;
 		for (int32 FrameIndex = 0; FrameIndex < Exec.ActiveFrames.Num(); ++FrameIndex)
 		{
@@ -229,20 +230,23 @@ FStateTreeScheduledTick FStateTreeMinimalExecutionContext::GetNextScheduledTick(
 							)
 						{
 							// todo: ShouldTickTransitions has onevent or ontick. both can be already triggered and we are waiting for the delay
-							// todo: find out if we need EveryFrames or NextFrame
 							bHasTaskWithEveryFramesTick = true;
 						}
 					}
 				}
 			}
 		}
-		//@TODO: This doesn't work. A transition request might be pending it should tick before the requested CustomTickRate.
-		if (StateCustomTickRate.IsSet())
+
+		if (!StateCustomTickRate.IsSet() && bHasTaskWithEveryFramesTick)
 		{
-			return FStateTreeScheduledTick::MakeCustomTickRate(StateCustomTickRate.GetValue());
+			return FStateTreeScheduledTick::MakeEveryFrames();
 		}
-		if (bHasTaskWithEveryFramesTick)
+
+		// If one state has a custom tick rate, then it overrides the tick rate for all states.
+		//Only return the CustomTickRate if it's > than NextFrame, the custom tick rate will be processed at the end.
+		if (StateCustomTickRate.IsSet() && StateCustomTickRate.GetValue() <= 0.0f)
 		{
+			// A state might override the custom tick rate with > 0, then another state overrides it again with 0 to tick back every frame.
 			return FStateTreeScheduledTick::MakeEveryFrames();
 		}
 	}
@@ -268,12 +272,16 @@ FStateTreeScheduledTick FStateTreeMinimalExecutionContext::GetNextScheduledTick(
 	// Min of all delayed transitions.
 	if (Exec.DelayedTransitions.Num() > 0)
 	{
-		float NextDeltaTime = std::numeric_limits<float>::max();
 		for (const FStateTreeTransitionDelayedState& Transition : Exec.DelayedTransitions)
 		{
-			NextDeltaTime = FMath::Min(NextDeltaTime, Transition.TimeLeft);
+			StateCustomTickRate = StateCustomTickRate.IsSet() ? FMath::Min(StateCustomTickRate.GetValue(), Transition.TimeLeft) : Transition.TimeLeft;
 		}
-		return FStateTreeScheduledTick::MakeCustomTickRate(NextDeltaTime);
+	}
+
+	// Custom tick rate for tasks and transitions.
+	if (StateCustomTickRate.IsSet())
+	{
+		return FStateTreeScheduledTick::MakeCustomTickRate(StateCustomTickRate.GetValue());
 	}
 
 	return FStateTreeScheduledTick::MakeSleep();
@@ -2656,7 +2664,7 @@ void FStateTreeExecutionContext::ExitState(const FStateTreeTransitionResult& Tra
 		if (Transition.NextActiveFrames.IsEmpty())
 		{
 			RemoveAllDelegateListeners();
-}
+		}
 		else if (ensure(FirstModifiedFrame != INDEX_NONE && FirstModifiedState != INDEX_NONE))
 		{
 			const FStateTreeExecutionFrame* ChangedFrame = &Exec.ActiveFrames[FirstModifiedFrame];
