StateTree: Modify comments to UStateTree::CustomTickRate

#rnx

diff --git a/Engine/Plugins/Runtime/StateTree/Source/StateTreeEditorModule/Public/StateTreeState.h b/Engine/Plugins/Runtime/StateTree/Source/StateTreeEditorModule/Public/StateTreeState.h
--- a/Engine/Plugins/Runtime/StateTree/Source/StateTreeEditorModule/Public/StateTreeState.h
+++ b/Engine/Plugins/Runtime/StateTree/Source/StateTreeEditorModule/Public/StateTreeState.h
@@ -370,6 +370,7 @@ public:
 	 * If set the state cannot sleep.
 	 * If set all the other states (children or parents) will also tick at that rate.
 	 * If more than one active states has a custom tick rate then the smallest custom tick rate wins.
+	 * If not set, the state will tick every frame unless the state tree is allowed to sleep.
 	 */
 	UPROPERTY(EditDefaultsOnly, Category = "State", meta = (EditCondition = "bHasCustomTickRate", ClampMin = 0.0f))
 	float CustomTickRate = 0.0f;
