StateTree: Delay the DetailView new selection to let the focus event go through and commit the data.

diff --git a/Engine/Plugins/Runtime/StateTree/Source/StateTreeEditorModule/Private/StateTreeEditorMode.cpp b/Engine/Plugins/Runtime/StateTree/Source/StateTreeEditorModule/Private/StateTreeEditorMode.cpp
--- a/Engine/Plugins/Runtime/StateTree/Source/StateTreeEditorModule/Private/StateTreeEditorMode.cpp
+++ b/Engine/Plugins/Runtime/StateTree/Source/StateTreeEditorModule/Private/StateTreeEditorMode.cpp
@@ -111,10 +111,13 @@ void UStateTreeEditorMode::OnSchemaChanged(const UStateTree& InStateTree)
 }
 
 void UStateTreeEditorMode::ForceRefreshDetailsView() const
-{ 					
+{
 	if (TSharedPtr<IDetailsView> DetailsView = GetDetailsView())
 	{
-		DetailsView->ForceRefresh();
+		if (!GEditor->GetTimerManager()->IsTimerActive(SetObjectTimerHandle))
+		{
+			DetailsView->ForceRefresh();
+		}
 	}
 }
 
@@ -442,25 +445,40 @@ void UStateTreeEditorMode::HandleModelAssetChanged()
 	UpdateAsset();
 }
 
-void UStateTreeEditorMode::HandleModelSelectionChanged(const TArray<TWeakObjectPtr<UStateTreeState>>& SelectedStates) const
+void UStateTreeEditorMode::HandleModelSelectionChanged(const TArray<TWeakObjectPtr<UStateTreeState>>& InSelectedStates) const
 {
+UE_LOG(LogTemp, Warning, TEXT("UStateTreeEditorMode::HandleModelSelectionChanged"));
 	if (TSharedPtr<IDetailsView> DetailsView = GetDetailsView())
 	{
-		TArray<UObject*> Selected;
-		for (const TWeakObjectPtr<UStateTreeState>& WeakState : SelectedStates)
-		{
-			if (UStateTreeState* State = WeakState.Get())
-			{
-				Selected.Add(State);
-			}
-		}
-		DetailsView->SetObjects(Selected);
+		//// Delay the selection by one frame to handle focus lost.
+		//constexpr bool bLoop = false;
+		//constexpr float DelayedTime = UE_SMALL_NUMBER;
+		//GEditor->GetTimerManager()->SetTimer(
+		//	SetObjectTimerHandle,
+		//	FTimerDelegate::CreateLambda([WeakDetailsView = DetailsView.ToWeakPtr(), InSelectedStates]()
+		//		{
+		//			if (TSharedPtr<IDetailsView> DetailsView = WeakDetailsView.Pin())
+		//			{
+						TArray<UObject*> Selected;
+						Selected.Reserve(InSelectedStates.Num());
+						for (const TWeakObjectPtr<UStateTreeState>& WeakState : InSelectedStates)
+						{
+							if (UStateTreeState* State = WeakState.Get())
+							{
+								Selected.Add(State);
+							}
+						}
+						DetailsView->SetObjects(Selected);
+				//	}
+				//}),
+				//DelayedTime,
+				//bLoop);
 	}
 }
 
 void UStateTreeEditorMode::HandleModelBringNodeToFocus(const UStateTreeState* State, const FGuid NodeID) const
 {
-	auto BingToFocus = [](const FPropertyPath& HighlightPath, TSharedPtr<IDetailsView>& DetailsView)
+	auto BingToFocus = [Self = this](const FPropertyPath& HighlightPath, TSharedPtr<IDetailsView>& DetailsView)
 		{
 			if (HighlightPath.IsValid())
 			{
@@ -469,16 +487,23 @@ void UStateTreeEditorMode::HandleModelBringNodeToFocus(const UStateTreeState* St
 				DetailsView->HighlightProperty(HighlightPath);
 
 				constexpr bool bLoop = false;
-				FTimerHandle HighlightTimerHandle;
 				GEditor->GetTimerManager()->SetTimer(
-					HighlightTimerHandle,
-					FTimerDelegate::CreateLambda([SelectionDetailsView = DetailsView]()
+					Self->HighlightTimerHandle,
+					FTimerDelegate::CreateLambda([WeakSelectionDetailsView = DetailsView.ToWeakPtr()]()
 						{
-							SelectionDetailsView->HighlightProperty({});
+							if (TSharedPtr<IDetailsView> SelectionDetailsView = WeakSelectionDetailsView.Pin())
+							{
+								SelectionDetailsView->HighlightProperty({});
+							}
 						}),
 					1.0f,
 					bLoop);
 			}
+			else if (Self->HighlightTimerHandle.IsValid())
+			{
+				// NB. SetTimer also clear the timer.
+				GEditor->GetTimerManager()->ClearTimer(Self->HighlightTimerHandle);
+			}
 		};
 
 	if (State)
diff --git a/Engine/Plugins/Runtime/StateTree/Source/StateTreeEditorModule/Public/StateTreeEditorMode.h b/Engine/Plugins/Runtime/StateTree/Source/StateTreeEditorModule/Public/StateTreeEditorMode.h
--- a/Engine/Plugins/Runtime/StateTree/Source/StateTreeEditorModule/Public/StateTreeEditorMode.h
+++ b/Engine/Plugins/Runtime/StateTree/Source/StateTreeEditorModule/Public/StateTreeEditorMode.h
@@ -72,6 +72,9 @@ protected:
 	bool bLastCompileSucceeded = true;
 	bool bForceAssetDetailViewToRefresh = false;
 
+	mutable FTimerHandle SetObjectTimerHandle;
+	mutable FTimerHandle HighlightTimerHandle;
+
 	TWeakObjectPtr<UStateTree> CachedStateTree = nullptr;
 
 	TSharedPtr<FStateTreeBindingExtension> DetailsViewExtensionHandler;
